moved remaining viewport settings from EditorUserSettings into LevelEditorViewportSettings class

diff --git a/Engine/Config/BaseEditorUserSettings.ini b/Engine/Config/BaseEditorUserSettings.ini
--- a/Engine/Config/BaseEditorUserSettings.ini
+++ b/Engine/Config/BaseEditorUserSettings.ini
@@ -59,18 +59,6 @@ bShowFrameRateAndMemory=False
 ; True if we should automatically reimport animset assets when a change to source content is detected
 bAutoReimportAnimSets=False
 
-; If true, Clicking a BSP selects the brush and ctrl+shift+click selects the surface. If false, vice versa
-bClickBSPSelectsBrush=True
-
-; If true, BSP will auto-update
-bBSPAutoUpdate=True
-
-; If true, Navigation will auto-update
-bNavigationAutoUpdate=True
-
-; If true, Replaces respects the scale of the original actor. If false, Replace sets the scale to 1.0
-bReplaceRespectsScale=True
-
 ; Enables audio feedback for certain operations in Unreal Editor, such as entering and exiting Play mode
 bEnableEditorSounds=True
 
@@ -91,7 +79,6 @@ bAutoApplyLightingEnable=true
 
 ; Misc
 HoverHighlightIntensity=0.000000
-bAllowTranslateRotateZWidget=False
 bDisplayUIExtensionPoints=False
 bUseCurvesForDistributions=False
 bAutoloadCheckedOutPackages=False
@@ -254,9 +241,17 @@ MouseSensitivty=.2f
 ; Whether to use mouse position as direct widget position
 bUseAbsoluteTranslation=True
 bLevelStreamingVolumePrevis=False
-
 bUseUE3OrbitControls=False
-bLeftMouseDragMovesCamera=False
+
+bAllowTranslateRotateZWidget=False
+; If true, Clicking a BSP selects the brush and ctrl+shift+click selects the surface. If false, vice versa
+bClickBSPSelectsBrush=True
+; If true, BSP will auto-update
+bBSPAutoUpdate=True
+; If true, Navigation will auto-update
+bNavigationAutoUpdate=True
+; If true, Replaces respects the scale of the original actor. If false, Replace sets the scale to 1.0
+bReplaceRespectsScale=True
 
 bUsePowerOf2SnapSize=False
 
diff --git a/Engine/Source/Editor/LevelEditor/Private/LevelEditorActions.cpp b/Engine/Source/Editor/LevelEditor/Private/LevelEditorActions.cpp
--- a/Engine/Source/Editor/LevelEditor/Private/LevelEditorActions.cpp
+++ b/Engine/Source/Editor/LevelEditor/Private/LevelEditorActions.cpp
@@ -2158,7 +2158,7 @@ bool FLevelEditorActionCallbacks::CanSetWidgetMode( FWidget::EWidgetMode WidgetM
 
 bool FLevelEditorActionCallbacks::IsTranslateRotateModeVisible()
 {
-	return GEditor->GetEditorUserSettings().bAllowTranslateRotateZWidget;
+	return GetDefault<ULevelEditorViewportSettings>()->bAllowTranslateRotateZWidget;
 }
 
 void FLevelEditorActionCallbacks::SetCoordinateSystem( ECoordSystem CoordinateSystem )
diff --git a/Engine/Source/Editor/UnrealEd/Classes/Editor/EditorUserSettings.h b/Engine/Source/Editor/UnrealEd/Classes/Editor/EditorUserSettings.h
--- a/Engine/Source/Editor/UnrealEd/Classes/Editor/EditorUserSettings.h
+++ b/Engine/Source/Editor/UnrealEd/Classes/Editor/EditorUserSettings.h
@@ -34,26 +34,6 @@ class UEditorUserSettings : public UObject
 	// =====================================================================
 	// The following options are exposed in the Preferences Editor 
 
-	/** Allow translate/rotate widget */
-	UPROPERTY(EditAnywhere, config, Category=LevelEditing, meta=( DisplayName = "Enable Combined Translate/Rotate Widget" ))
-	uint32 bAllowTranslateRotateZWidget:1;
-
-	/** If true, Clicking a BSP selects the brush and ctrl+shift+click selects the surface. If false, vice versa */
-	UPROPERTY(EditAnywhere, config, Category=LevelEditing, meta=( DisplayName = "Clicking BSP Enables Brush" ))
-	uint32 bClickBSPSelectsBrush:1;
-
-	/** If true, BSP will auto-update */
-	UPROPERTY(EditAnywhere, config, Category=LevelEditing, meta=( DisplayName = "Update BSP Automatically" ))
-	uint32 bBSPAutoUpdate:1;
-
-	/** If true, Navigation will auto-update */
-	UPROPERTY(EditAnywhere, config, Category=LevelEditing, meta=( DisplayName = "Update Navigation Automatically" ))
-	uint32 bNavigationAutoUpdate:1;
-
-	/** If enabled, replacing actors will respect the scale of the original actor.  Otherwise, the replaced actors will have a scale of 1.0 */
-	UPROPERTY(EditAnywhere, config, Category=LevelEditing, meta=( DisplayName = "Preserve Actor Scale on Replace" ))
-	uint32 bReplaceRespectsScale:1;
-
 	/** If enabled, any newly opened UI menus, menu bars, and toolbars will show the developer hooks that would accept extensions */
 	UPROPERTY(EditAnywhere, config, Category=DeveloperTools)
 	uint32 bDisplayUIExtensionPoints:1;
diff --git a/Engine/Source/Editor/UnrealEd/Classes/Settings/LevelEditorViewportSettings.h b/Engine/Source/Editor/UnrealEd/Classes/Settings/LevelEditorViewportSettings.h
--- a/Engine/Source/Editor/UnrealEd/Classes/Settings/LevelEditorViewportSettings.h
+++ b/Engine/Source/Editor/UnrealEd/Classes/Settings/LevelEditorViewportSettings.h
@@ -199,7 +199,31 @@ public:
 	/** When checked, orbit the camera by using the L or U keys when unchecked, Alt and Left Mouse Drag will orbit around the look at point */
 	UPROPERTY(EditAnywhere, config, Category=Controls, meta=(DisplayName="Use UE3 Orbit Controls"), AdvancedDisplay)
 	bool bUseUE3OrbitControls;
+
+public:
+
+	/** Allow translate/rotate widget */
+	UPROPERTY(EditAnywhere, config, Category=Editing, meta=( DisplayName = "Enable Combined Translate/Rotate Widget" ))
+	uint32 bAllowTranslateRotateZWidget:1;
+
+	/** If true, Clicking a BSP selects the brush and ctrl+shift+click selects the surface. If false, vice versa */
+	UPROPERTY(EditAnywhere, config, Category=Editing, meta=( DisplayName = "Clicking BSP Enables Brush" ))
+	uint32 bClickBSPSelectsBrush:1;
+
+	/** If true, BSP will auto-update */
+	UPROPERTY(EditAnywhere, config, Category=Editing, meta=( DisplayName = "Update BSP Automatically" ))
+	uint32 bBSPAutoUpdate:1;
+
+	/** If true, Navigation will auto-update */
+	UPROPERTY(EditAnywhere, config, Category=Editing, meta=( DisplayName = "Update Navigation Automatically" ))
+	uint32 bNavigationAutoUpdate:1;
+
+	/** If enabled, replacing actors will respect the scale of the original actor.  Otherwise, the replaced actors will have a scale of 1.0 */
+	UPROPERTY(EditAnywhere, config, Category=Editing, meta=( DisplayName = "Preserve Actor Scale on Replace" ))
+	uint32 bReplaceRespectsScale:1;
+
 public:
+
 	/** If enabled will use power of 2 grid settings (e.g, 1,2,4,8,16,...,1024) instead of decimal grid sizes */
 	UPROPERTY(EditAnywhere, config, Category=GridSnapping, meta=(DisplayName = "User Power of Two Snap Size"))
 	bool bUsePowerOf2SnapSize;
@@ -232,8 +256,6 @@ public:
 	UPROPERTY(EditAnywhere, config, AdvancedDisplay, Category=GridSnapping)
 	TArray<float> ScalingGridSizes;
 
-public:
-
 	/** If enabled, actor positions will snap to the grid. */
 	UPROPERTY(EditAnywhere, config, Category=GridSnapping, meta=(DisplayName = "Enable Grid Snapping"))
 	uint32 GridEnabled:1;
@@ -346,10 +368,11 @@ public:
 	float BillboardScale;
 
 public:
+
 	/**
 	 * @return The instance settings for the given viewport; null if no settings were found for this viewport
 	 */
-	const FLevelEditorViewportInstanceSettings* GetViewportInstanceSettings(const FString& InConfigName) const
+	const FLevelEditorViewportInstanceSettings* GetViewportInstanceSettings( const FString& InConfigName ) const
 	{
 		for(auto It = PerInstanceSettings.CreateConstIterator(); It; ++It)
 		{
@@ -366,7 +389,7 @@ public:
 	/**
 	 * Set the instance settings for the given viewport
 	 */
-	void SetViewportInstanceSettings(const FString& InConfigName, const FLevelEditorViewportInstanceSettings& InConfigSettings)
+	void SetViewportInstanceSettings( const FString& InConfigName, const FLevelEditorViewportInstanceSettings& InConfigSettings )
 	{
 		check(!InConfigName.IsEmpty());
 
diff --git a/Engine/Source/Editor/UnrealEd/Private/Editor.cpp b/Engine/Source/Editor/UnrealEd/Private/Editor.cpp
--- a/Engine/Source/Editor/UnrealEd/Private/Editor.cpp
+++ b/Engine/Source/Editor/UnrealEd/Private/Editor.cpp
@@ -390,7 +390,7 @@ void UEditorEngine::InitEditor(IEngineLoop* InEngineLoop)
 	
 	// Set navigation system property indicating whether navigation is supposed to rebuild automatically 
 	FWorldContext &EditorContext = GEditor->GetEditorWorldContext();
-	UNavigationSystem::SetNavigationAutoUpdateEnabled(GEditor->GetEditorUserSettings().bNavigationAutoUpdate, EditorContext.World()->GetNavigationSystem() );
+	UNavigationSystem::SetNavigationAutoUpdateEnabled(GetDefault<ULevelEditorViewportSettings>()->bNavigationAutoUpdate, EditorContext.World()->GetNavigationSystem() );
 
 	// Allocate temporary model.
 	TempModel = new UModel( FPostConstructInitializeProperties(),NULL, 1 );
@@ -4966,7 +4966,7 @@ void UEditorEngine::ReplaceSelectedActors(UActorFactory* Factory, const FAssetDa
 			USceneComponent* const NewActorRootComponent = NewActor->GetRootComponent();
 			if(NewActorRootComponent)
 			{
-				if(!GEditorModeTools().GetReplaceRespectsScale() || OldActor->GetRootComponent() == NULL )
+				if(!GetDefault<ULevelEditorViewportSettings>()->bReplaceRespectsScale || OldActor->GetRootComponent() == NULL )
 				{
 					NewActorRootComponent->SetRelativeScale3D(FVector(1.0f, 1.0f, 1.0f));
 				}
diff --git a/Engine/Source/Editor/UnrealEd/Private/Editor/EditorUserSettings.cpp b/Engine/Source/Editor/UnrealEd/Private/Editor/EditorUserSettings.cpp
--- a/Engine/Source/Editor/UnrealEd/Private/Editor/EditorUserSettings.cpp
+++ b/Engine/Source/Editor/UnrealEd/Private/Editor/EditorUserSettings.cpp
@@ -89,22 +89,6 @@ void UEditorUserSettings::PostEditChangeProperty( FPropertyChangedEvent& Propert
 		//GDistributionType == 0 for curves
 		GDistributionType = (bUseCurvesForDistributions) ? 0 : 1;
 	}
-	else if ( Name == FName(TEXT("bAllowTranslateRotateZWidget")) )
-	{
-		if ( bAllowTranslateRotateZWidget )
-		{
-			GEditorModeTools().SetWidgetMode( FWidget::WM_TranslateRotateZ );
-		}
-		else if ( GEditorModeTools().GetWidgetMode() == FWidget::WM_TranslateRotateZ )
-		{
-			GEditorModeTools().SetWidgetMode( FWidget::WM_Translate );
-		}
-	}
-	else if ( Name == FName(TEXT("bNavigationAutoUpdate")) )
-	{
-		FWorldContext &EditorContext = GEditor->GetEditorWorldContext();
-		UNavigationSystem::SetNavigationAutoUpdateEnabled(bNavigationAutoUpdate, EditorContext.World()->GetNavigationSystem() );
-	}
 
 	GEditor->SaveEditorUserSettings();
 
diff --git a/Engine/Source/Editor/UnrealEd/Private/EditorModes.cpp b/Engine/Source/Editor/UnrealEd/Private/EditorModes.cpp
--- a/Engine/Source/Editor/UnrealEd/Private/EditorModes.cpp
+++ b/Engine/Source/Editor/UnrealEd/Private/EditorModes.cpp
@@ -1907,7 +1907,7 @@ void FEditorModeTools::CycleWidgetMode (void)
 		do
 		{
 			Wk++;
-			if ((Wk == FWidget::WM_TranslateRotateZ) && (!GEditor->GetEditorUserSettings().bAllowTranslateRotateZWidget))
+			if ((Wk == FWidget::WM_TranslateRotateZ) && (!GetDefault<ULevelEditorViewportSettings>()->bAllowTranslateRotateZWidget))
 			{
 				Wk++;
 			}
diff --git a/Engine/Source/Editor/UnrealEd/Private/EditorServer.cpp b/Engine/Source/Editor/UnrealEd/Private/EditorServer.cpp
--- a/Engine/Source/Editor/UnrealEd/Private/EditorServer.cpp
+++ b/Engine/Source/Editor/UnrealEd/Private/EditorServer.cpp
@@ -1390,7 +1390,7 @@ void UEditorEngine::RebuildMap(UWorld* InWorld, EMapRebuildType RebuildType)
 void UEditorEngine::RebuildLevel(ULevel& Level)
 {
 	// Early out if BSP auto-updating is disabled
-	if (!GEditorModeTools().GetBSPAutoUpdate())
+	if (!GetDefault<ULevelEditorViewportSettings>()->bBSPAutoUpdate)
 	{
 		return;
 	}
@@ -1478,7 +1478,7 @@ void UEditorEngine::RebuildModelFromBrushes(UModel* Model, bool bSelectedBrushes
 void UEditorEngine::RebuildAlteredBSP()
 {
 	// Early out if BSP auto-updating is disabled
-	if (!GEditorModeTools().GetBSPAutoUpdate())
+	if (!GetDefault<ULevelEditorViewportSettings>()->bBSPAutoUpdate)
 	{
 		return;
 	}
diff --git a/Engine/Source/Editor/UnrealEd/Private/LevelViewportClickHandlers.cpp b/Engine/Source/Editor/UnrealEd/Private/LevelViewportClickHandlers.cpp
--- a/Engine/Source/Editor/UnrealEd/Private/LevelViewportClickHandlers.cpp
+++ b/Engine/Source/Editor/UnrealEd/Private/LevelViewportClickHandlers.cpp
@@ -676,7 +676,7 @@ namespace ClickHandlers
 		else if( Click.GetKey() == EKeys::LeftMouseButton && Click.IsShiftDown() && Click.IsControlDown() )
 		{
 
-			if( !GEditorModeTools().GetClickBSPSelectsBrush() )
+			if( !GetDefault<ULevelEditorViewportSettings>()->bClickBSPSelectsBrush )
 			{
 				// Add to the actor selection set the brush actor that belongs to this BSP surface.
 				// Check Surf.Actor, as it can be NULL after deleting brushes and before rebuilding BSP.
@@ -862,7 +862,7 @@ namespace ClickHandlers
 		else
 		{	
 			bool bDeselectAlreadyHandled = false;
-			if( GEditorModeTools().GetClickBSPSelectsBrush() )
+			if( GetDefault<ULevelEditorViewportSettings>()->bClickBSPSelectsBrush )
 			{
 				// Add to the actor selection set the brush actor that belongs to this BSP surface.
 				// Check Surf.Actor, as it can be NULL after deleting brushes and before rebuilding BSP.
diff --git a/Engine/Source/Editor/UnrealEd/Private/NavigationBuildingNotification.cpp b/Engine/Source/Editor/UnrealEd/Private/NavigationBuildingNotification.cpp
--- a/Engine/Source/Editor/UnrealEd/Private/NavigationBuildingNotification.cpp
+++ b/Engine/Source/Editor/UnrealEd/Private/NavigationBuildingNotification.cpp
@@ -148,7 +148,7 @@ void FNavigationBuildingNotificationImpl::Tick(float DeltaTime)
 		FWorldContext &EditorContext = EEngine->GetEditorWorldContext();
 		
 		const bool bBuildInProgress = EditorContext.World() != NULL && EditorContext.World()->GetNavigationSystem() != NULL 
-			&& EditorContext.World()->GetNavigationSystem()->IsNavigationBuildInProgress( GEditor->AccessEditorUserSettings().bNavigationAutoUpdate ? true : false) == true;
+			&& EditorContext.World()->GetNavigationSystem()->IsNavigationBuildInProgress( GetDefault<ULevelEditorViewportSettings>()->bNavigationAutoUpdate ? true : false) == true;
 
 		if (!bPreviouslyDetectedBuild && bBuildInProgress)
 		{
diff --git a/Engine/Source/Editor/UnrealEd/Private/SEditorViewport.cpp b/Engine/Source/Editor/UnrealEd/Private/SEditorViewport.cpp
--- a/Engine/Source/Editor/UnrealEd/Private/SEditorViewport.cpp
+++ b/Engine/Source/Editor/UnrealEd/Private/SEditorViewport.cpp
@@ -382,7 +382,7 @@ bool SEditorViewport::IsWidgetModeActive( FWidget::EWidgetMode Mode ) const
 
 bool SEditorViewport::IsTranslateRotateModeVisible() const
 {
-	return GEditor->GetEditorUserSettings().bAllowTranslateRotateZWidget;
+	return GetDefault<ULevelEditorViewportSettings>()->bAllowTranslateRotateZWidget;
 }
 
 bool SEditorViewport::IsCoordSystemActive( ECoordSystem CoordSystem ) const
@@ -400,7 +400,7 @@ void SEditorViewport::OnCycleWidgetMode()
 	{
 		++WidgetModeAsInt;
 
-		if ((WidgetModeAsInt == FWidget::WM_TranslateRotateZ) && (!GEditor->GetEditorUserSettings().bAllowTranslateRotateZWidget))
+		if ((WidgetModeAsInt == FWidget::WM_TranslateRotateZ) && (!GetDefault<ULevelEditorViewportSettings>()->bAllowTranslateRotateZWidget))
 		{
 			++WidgetModeAsInt;
 		}
diff --git a/Engine/Source/Editor/UnrealEd/Private/Settings/SettingsClasses.cpp b/Engine/Source/Editor/UnrealEd/Private/Settings/SettingsClasses.cpp
--- a/Engine/Source/Editor/UnrealEd/Private/Settings/SettingsClasses.cpp
+++ b/Engine/Source/Editor/UnrealEd/Private/Settings/SettingsClasses.cpp
@@ -160,7 +160,18 @@ void ULevelEditorViewportSettings::PostEditChangeProperty( struct FPropertyChang
 
 	const FName Name = (PropertyChangedEvent.Property != nullptr) ? PropertyChangedEvent.Property->GetFName() : NAME_None;
 
-	if (Name == FName(TEXT("bHighlightWithBrackets")))
+	if (Name == FName(TEXT("bAllowTranslateRotateZWidget")))
+	{
+		if (bAllowTranslateRotateZWidget)
+		{
+			GEditorModeTools().SetWidgetMode(FWidget::WM_TranslateRotateZ);
+		}
+		else if (GEditorModeTools().GetWidgetMode() == FWidget::WM_TranslateRotateZ)
+		{
+			GEditorModeTools().SetWidgetMode(FWidget::WM_Translate);
+		}
+	}
+	else if (Name == FName(TEXT("bHighlightWithBrackets")))
 	{
 		GEngine->SetSelectedMaterialColor(bHighlightWithBrackets ? FLinearColor::Black : GetDefault<UEditorStyleSettings>()->SelectionColor);
 	}
@@ -168,6 +179,11 @@ void ULevelEditorViewportSettings::PostEditChangeProperty( struct FPropertyChang
 	{
 		GEngine->HoverHighlightIntensity = HoverHighlightIntensity;
 	}
+	else if (Name == FName(TEXT("bNavigationAutoUpdate")))
+	{
+		FWorldContext &EditorContext = GEditor->GetEditorWorldContext();
+		UNavigationSystem::SetNavigationAutoUpdateEnabled(bNavigationAutoUpdate, EditorContext.World()->GetNavigationSystem());
+	}
 	else if (Name == FName(TEXT("SelectionHighlightIntensity")))
 	{
 		GEngine->SelectionHighlightIntensity = SelectionHighlightIntensity;
diff --git a/Engine/Source/Editor/UnrealEd/Public/EditorModes.h b/Engine/Source/Editor/UnrealEd/Public/EditorModes.h
--- a/Engine/Source/Editor/UnrealEd/Public/EditorModes.h
+++ b/Engine/Source/Editor/UnrealEd/Public/EditorModes.h
@@ -538,10 +538,6 @@ public:
 
 	void CycleWidgetMode (void);
 
-	/** Allow translate/rotate widget */
-	bool GetAllowTranslateRotateZWidget()                 { return GEditor->GetEditorUserSettings().bAllowTranslateRotateZWidget; }
-	void SetAllowTranslateRotateZWidget (const bool bInOnOff) { GEditor->AccessEditorUserSettings().bAllowTranslateRotateZWidget = bInOnOff; }
-
 	/**Save Widget Settings to Ini file*/
 	void SaveWidgetSettings();
 	/**Load Widget Settings from Ini file*/
@@ -649,16 +645,6 @@ public:
 	 */
 	UNREALED_API FWidget::EWidgetMode GetWidgetMode() const;
 
-	/**
-	 * Sets editor wide usage of replace respecting scale
-	 */
-	void SetReplaceRespectsScale(bool bInReplaceRespectsScale) { GEditor->AccessEditorUserSettings().bReplaceRespectsScale = bInReplaceRespectsScale; }
-	/**
-	 * Gets the current state of editor wide usage of replace respecting scale
-	 * @ return - If true, replace respects scale
-	 */
-	bool GetReplaceRespectsScale(void) const { return GEditor->GetEditorUserSettings().bReplaceRespectsScale; }
-
 	/**
 	 * Gets the current state of script editor usage of show friendly names
 	 * @ return - If true, replace variable names with sanatized ones
@@ -669,38 +655,12 @@ public:
 	 * Sets Interp Editor inversion of panning on drag
 	 */
 	void SetInterpPanInvert(bool bInInterpPanInverted) { bInterpPanInverted = bInInterpPanInverted; }
+
 	/**
 	 * Gets the current state of the interp editor panning
 	 * @ return - If true, the interp editor left-right panning is inverted
 	 */
 	bool GetInterpPanInvert(void) const { return bInterpPanInverted; }
-	/**
-	 * Sets clicking BSP surface selecting brush
-	 */
-	void SetClickBSPSelectsBrush(bool bInClickBSPSelectsBrush) { GEditor->AccessEditorUserSettings().bClickBSPSelectsBrush = bInClickBSPSelectsBrush; }
-	/**
-	 * Gets the current state of clicking BSP surface selecting brush
-	 * @ return - If true, the clicking BSP surface selects brush
-	 */
-	bool GetClickBSPSelectsBrush(void) const { return GEditor->GetEditorUserSettings().bClickBSPSelectsBrush; }
-	/**
-	 * Sets BSP auto updating
-	 */
-	void SetBSPAutoUpdate(bool bInBSPAutoUpdate) { GEditor->AccessEditorUserSettings().bBSPAutoUpdate = bInBSPAutoUpdate; }
-	/**
-	 * Gets the current state of BSP auto-updating
-	 * @ return - If true, BSP auto-updating is enabled
-	 */
-	bool GetBSPAutoUpdate(void) const { return GEditor->GetEditorUserSettings().bBSPAutoUpdate; }
-	/**
-	 * Sets Navigation auto updating
-	 */
-	void SetNavigationAutoUpdate(bool bInNavigationAutoUpdate) { GEditor->AccessEditorUserSettings().bNavigationAutoUpdate = bInNavigationAutoUpdate; }
-	/**
-	 * Gets the current state of Navigation auto-updating
-	 * @ return - If true, Navigation auto-updating is enabled
-	 */
-	bool GetNavigationAutoUpdate(void) const { return GEditor->GetEditorUserSettings().bNavigationAutoUpdate; }
 
 	/**
 	 * Sets a bookmark in the levelinfo file, allocating it if necessary.
