#UE4:
* improved levelgrid fading, now 10 based

diff --git a/Engine/Source/Editor/UnrealEd/Private/EditorComponents.cpp b/Engine/Source/Editor/UnrealEd/Private/EditorComponents.cpp
--- a/Engine/Source/Editor/UnrealEd/Private/EditorComponents.cpp
+++ b/Engine/Source/Editor/UnrealEd/Private/EditorComponents.cpp
@@ -448,10 +448,8 @@ void FEditorCommonDrawHelper::DrawOldGrid(const FSceneView* View,FPrimitiveDrawI
 			FVector EndY( 0.0f, -HALF_WORLD_MAX1, -HALF_WORLD_MAX );
 			FVector StartX( +HALF_WORLD_MAX1, 0.0f, -HALF_WORLD_MAX );
 			FVector EndX( -HALF_WORLD_MAX1, 0.0f, -HALF_WORLD_MAX );
-			DrawGridSection( GEditor->GetGridSize(), &StartY, &EndY, &StartY.X, &EndY.X, 0, View, PDI, EAxisLines::Minor );
-			DrawGridSection( GEditor->GetGridSize(), &StartX, &EndX, &StartX.Y, &EndX.Y, 1, View, PDI, EAxisLines::Minor );
-			DrawGridSection( GEditor->GetGridSize(), &StartY, &EndY, &StartY.X, &EndY.X, 0, View, PDI, EAxisLines::Major );
-			DrawGridSection( GEditor->GetGridSize(), &StartX, &EndX, &StartX.Y, &EndX.Y, 1, View, PDI, EAxisLines::Major );
+			DrawGridSection( GEditor->GetGridSize(), &StartY, &EndY, &StartY.X, &EndY.X, 0, View, PDI);
+			DrawGridSection( GEditor->GetGridSize(), &StartX, &EndX, &StartX.Y, &EndX.Y, 1, View, PDI);
 			DrawOriginAxisLine( &StartY, &EndY, &StartY.X, &EndY.X, View, PDI, AxisColors[1] );
 			DrawOriginAxisLine( &StartX, &EndX, &StartX.Y, &EndX.Y, View, PDI, AxisColors[0] );
 		}
@@ -461,10 +459,8 @@ void FEditorCommonDrawHelper::DrawOldGrid(const FSceneView* View,FPrimitiveDrawI
 			FVector EndZ( 0.0f, -HALF_WORLD_MAX, -HALF_WORLD_MAX1 );
 			FVector StartX( +HALF_WORLD_MAX1, -HALF_WORLD_MAX, 0.0f );
 			FVector EndX( -HALF_WORLD_MAX1, -HALF_WORLD_MAX, 0.0f );
-			DrawGridSection( GEditor->GetGridSize(), &StartZ, &EndZ, &StartZ.X, &EndZ.X, 0, View, PDI, EAxisLines::Minor );
-			DrawGridSection( GEditor->GetGridSize(), &StartX, &EndX, &StartX.Z, &EndX.Z, 2, View, PDI, EAxisLines::Minor );
-			DrawGridSection( GEditor->GetGridSize(), &StartZ, &EndZ, &StartZ.X, &EndZ.X, 0, View, PDI, EAxisLines::Major );
-			DrawGridSection( GEditor->GetGridSize(), &StartX, &EndX, &StartX.Z, &EndX.Z, 2, View, PDI, EAxisLines::Major );
+			DrawGridSection( GEditor->GetGridSize(), &StartZ, &EndZ, &StartZ.X, &EndZ.X, 0, View, PDI);
+			DrawGridSection( GEditor->GetGridSize(), &StartX, &EndX, &StartX.Z, &EndX.Z, 2, View, PDI);
 			DrawOriginAxisLine( &StartZ, &EndZ, &StartZ.X, &EndZ.X, View, PDI, AxisColors[2] );
 			DrawOriginAxisLine( &StartX, &EndX, &StartX.Z, &EndX.Z, View, PDI, AxisColors[0] );
 		}
@@ -474,10 +470,8 @@ void FEditorCommonDrawHelper::DrawOldGrid(const FSceneView* View,FPrimitiveDrawI
 			FVector EndZ( +HALF_WORLD_MAX, 0.0f, -HALF_WORLD_MAX1 );
 			FVector StartY( +HALF_WORLD_MAX, +HALF_WORLD_MAX1, 0.0f );
 			FVector EndY( +HALF_WORLD_MAX, -HALF_WORLD_MAX1, 0.0f );
-			DrawGridSection( GEditor->GetGridSize(), &StartZ, &EndZ, &StartZ.Y, &EndZ.Y, 1, View, PDI, EAxisLines::Minor );
-			DrawGridSection( GEditor->GetGridSize(), &StartY, &EndY, &StartY.Z, &EndY.Z, 2, View, PDI, EAxisLines::Minor );
-			DrawGridSection( GEditor->GetGridSize(), &StartZ, &EndZ, &StartZ.Y, &EndZ.Y, 1, View, PDI, EAxisLines::Major );
-			DrawGridSection( GEditor->GetGridSize(), &StartY, &EndY, &StartY.Z, &EndY.Z, 2, View, PDI, EAxisLines::Major );
+			DrawGridSection( GEditor->GetGridSize(), &StartZ, &EndZ, &StartZ.Y, &EndZ.Y, 1, View, PDI);
+			DrawGridSection( GEditor->GetGridSize(), &StartY, &EndY, &StartY.Z, &EndY.Z, 2, View, PDI);
 			DrawOriginAxisLine( &StartZ, &EndZ, &StartZ.Y, &EndZ.Y, View, PDI, AxisColors[2] );
 			DrawOriginAxisLine( &StartY, &EndY, &StartY.Z, &EndY.Z, View, PDI, AxisColors[1] );
 		}
@@ -499,7 +493,7 @@ void FEditorCommonDrawHelper::DrawOldGrid(const FSceneView* View,FPrimitiveDrawI
 }
 
 
-void FEditorCommonDrawHelper::DrawGridSection(float ViewportGridY,FVector* A,FVector* B,float* AX,float* BX,int32 Axis,const FSceneView* View,FPrimitiveDrawInterface* PDI, const EAxisLines::Type LineType )
+void FEditorCommonDrawHelper::DrawGridSection(float ViewportGridY,FVector* A,FVector* B,float* AX,float* BX,int32 Axis,const FSceneView* View,FPrimitiveDrawInterface* PDI )
 {
 	if( ViewportGridY == 0 )
 	{
@@ -507,6 +501,9 @@ void FEditorCommonDrawHelper::DrawGridSection(float ViewportGridY,FVector* A,FVe
 		return;
 	}
 
+	// todo
+	static int32 Exponent = 10;	// or 2?
+
 	const FMatrix InvViewProjMatrix = View->ViewMatrices.ProjMatrix.Inverse() * View->ViewMatrices.ViewMatrix.Inverse();
 	int32 FirstLine = FMath::Trunc(InvViewProjMatrix.TransformPosition(FVector(-1,-1,0.5f)).Component(Axis) / ViewportGridY);
 	int32 LastLine = FMath::Trunc(InvViewProjMatrix.TransformPosition(FVector(+1,+1,0.5f)).Component(Axis) / ViewportGridY);
@@ -517,42 +514,54 @@ void FEditorCommonDrawHelper::DrawGridSection(float ViewportGridY,FVector* A,FVe
 
 	const float SizeX = View->ViewRect.Width();
 	const float Zoom = (1.0f / View->ViewMatrices.ProjMatrix.M[0][0]) * 2.0f / SizeX;
-	const int32 Dist = FMath::Trunc(SizeX * Zoom / ViewportGridY);
+	const float Dist = SizeX * Zoom / ViewportGridY;
+
+	// when the grid fades
+	static float Tweak = 4.0f;
 
-	// Figure out alpha interpolator for fading in the grid lines.
-	float Alpha;
-	int32 IncBits=0;
+	float IncValue = FMath::LogX(Exponent, Dist / (float)(SizeX / Tweak));
+	
+	int32 IncScale = 1;
+
+	for(float x = 0; x < IncValue; ++x)
 	{
-		if( Dist+Dist >= SizeX/4 )
-		{
-			while( (Dist>>IncBits) >= SizeX/4 )
-			{
-				IncBits++;
-			}
-			Alpha = 2 - 2*(float)Dist / (float)((1<<IncBits) * SizeX/4);
-		}
-		else
-		{
-			Alpha = 1.0;
-		}
+		IncScale *= Exponent;
+	}
+
+	// 0 excluded for hard transitions .. 0.5f for very soft transitions
+	const float TransitionRegion = 0.5f;
+
+	const float InvTransitionRegion = 1.0f / TransitionRegion;
+	float Fract = IncValue - floorf(IncValue);
+	float AlphaA = FMath::Clamp(Fract * InvTransitionRegion, 0.0f, 1.0f);
+	float AlphaB = FMath::Clamp(InvTransitionRegion - Fract * InvTransitionRegion, 0.0f, 1.0f);
+
+	if(IncValue < -1.0f)
+	{
+		// no fade in magnification case
+		AlphaA = 1.0f;
+		AlphaB = 1.0f;
 	}
 
 	const int32 MajorLineInterval = GEditor->GetGridInterval();
 
-	const FLinearColor Background = FColor(View->BackgroundColor).ReinterpretAsLinear();
+	const FLinearColor Background = View->BackgroundColor;
+
+	FLinearColor MajorColor = FMath::Lerp(Background, FLinearColor::White, 0.05f);
+	FLinearColor MinorColor = FMath::Lerp(Background, FLinearColor::White, 0.02f);
 
 	// Draw major and minor grid lines
-	const int32 FirstLineClamped = FMath::Max<int32>(FirstLine - 1,-HALF_WORLD_MAX/ViewportGridY) >> IncBits;
-	const int32 LastLineClamped = FMath::Min<int32>(LastLine + 1, +HALF_WORLD_MAX/ViewportGridY) >> IncBits;
-	for( int32 LineIndex = FirstLineClamped; LineIndex < LastLineClamped; ++LineIndex )
+	const int32 FirstLineClamped = FMath::Max<int32>(FirstLine - 1,-HALF_WORLD_MAX/ViewportGridY) / IncScale;
+	const int32 LastLineClamped = FMath::Min<int32>(LastLine + 1, +HALF_WORLD_MAX/ViewportGridY) / IncScale;
+	for( int32 LineIndex = FirstLineClamped; LineIndex <= LastLineClamped; ++LineIndex )
 	{
-		*AX = FPlatformMath::Trunc(LineIndex * ViewportGridY) << IncBits;
-		*BX = FPlatformMath::Trunc(LineIndex * ViewportGridY) << IncBits;
+		*AX = FPlatformMath::Trunc(LineIndex * ViewportGridY) * IncScale;
+		*BX = FPlatformMath::Trunc(LineIndex * ViewportGridY) * IncScale;
 
 		// Only minor lines fade out with ortho zoom distance.  Origin lines and major lines are drawn
 		// at 100% opacity, but with a brighter value
 		const bool bIsMajorLine = (LineIndex % MajorLineInterval) == 0;
-		if( bIsMajorLine || LineType != EAxisLines::Major )
+
 		{
 			// Don't bother drawing the world origin line.  We'll do that later.
 			const bool bIsOriginLine = ( LineIndex == 0 );
@@ -561,14 +570,14 @@ void FEditorCommonDrawHelper::DrawGridSection(float ViewportGridY,FVector* A,FVe
 				FLinearColor Color;
 				if( bIsMajorLine )
 				{
-					Color = Background * 1.5f;
+					Color = FMath::Lerp( Background, MajorColor, AlphaA );
 				}
 				else
 				{
-					Color = FMath::Lerp( Background, Background * 1.25f, Alpha );
+					Color = FMath::Lerp( Background, MinorColor, AlphaB );
 				}
 
-				PDI->DrawLine(*A,*B,FLinearColor(Color.Quantize()),SDPG_World);
+				PDI->DrawLine(*A,*B,Color,SDPG_World);
 			}
 		}
 	}
diff --git a/Engine/Source/Editor/UnrealEd/Public/EditorComponents.h b/Engine/Source/Editor/UnrealEd/Public/EditorComponents.h
--- a/Engine/Source/Editor/UnrealEd/Public/EditorComponents.h
+++ b/Engine/Source/Editor/UnrealEd/Public/EditorComponents.h
@@ -72,7 +72,7 @@ private:
 	/**
 	 * Draws a section(vertical lines or horizontal lines) of a viewport's grid.
 	 */
-	static void DrawGridSection(float ViewportGridY,FVector* A,FVector* B,float* AX,float* BX,int32 Axis,const FSceneView* View,FPrimitiveDrawInterface* PDI, const EAxisLines::Type LineType);
+	static void DrawGridSection(float ViewportGridY,FVector* A,FVector* B,float* AX,float* BX,int32 Axis,const FSceneView* View,FPrimitiveDrawInterface* PDI);
 
 	/**
 	 * Draws an origin axis line for the orthographic grid
