moved miscellaneous level editor settings into their own section

diff --git a/Engine/Config/BaseEditorUserSettings.ini b/Engine/Config/BaseEditorUserSettings.ini
--- a/Engine/Config/BaseEditorUserSettings.ini
+++ b/Engine/Config/BaseEditorUserSettings.ini
@@ -32,15 +32,6 @@
 ; NOTE: The reason we default this to true, even though it can be annoying, is that it can be a frustrating experience for new users to not be able to select a translucent object of interest, until they at least learn about this feature.
 bAllowSelectTranslucent=True
 
-; If true audio will be enabled in the editor. Does not affect PIE 
-bEnableRealTimeAudio=False
-
-; If true audio will continue to play when the editor does not have focus
-bAllowBackgroundAudio=False
-
-; Volume level for the editor
-EditorVolumeLevel=1.0
-
 ; True if we should move actors to their appropriate grid volume levels immediately after most operations
 bUpdateActorsInGridLevelsImmediately=False
 
@@ -59,9 +50,6 @@
 ; True if we should automatically reimport animset assets when a change to source content is detected
 bAutoReimportAnimSets=False
 
-; Enables audio feedback for certain operations in Unreal Editor, such as entering and exiting Play mode
-bEnableEditorSounds=True
-
 ; Scene Outliner Settings
 bShowOnlySelectedActors=False
 
@@ -74,9 +62,6 @@
 ; Blueprint Editor SCS Editor settings
 bSCSEditorShowGrid=true
 
-; Auto apply lighting after it has been built
-bAutoApplyLightingEnable=true
-
 ; Misc
 HoverHighlightIntensity=0.000000
 bDisplayUIExtensionPoints=False
@@ -163,6 +148,26 @@
 TextDiffToolPath=(FilePath="p4merge.exe")
 
 
+[/Script/UnrealEd.LevelEditorMiscSettings]
+; Auto apply lighting after it has been built
+bAutoApplyLightingEnable=true
+; If true, BSP will auto-update
+bBSPAutoUpdate=True
+; If true, Navigation will auto-update
+bNavigationAutoUpdate=True
+; If true, Replaces respects the scale of the original actor. If false, Replace sets the scale to 1.0
+bReplaceRespectsScale=True
+
+; If true audio will continue to play when the editor does not have focus
+bAllowBackgroundAudio=False
+; If true audio will be enabled in the editor. Does not affect PIE 
+bEnableRealTimeAudio=False
+; Volume level for the editor
+EditorVolumeLevel=1.0
+; Enables audio feedback for certain operations in Unreal Editor, such as entering and exiting Play mode
+bEnableEditorSounds=True
+
+
 [/Script/UnrealEd.LevelEditorPlaySettings]
 PlayFromHerePlayerStartClassName=/Script/Engine.PlayerStartPIE
 
@@ -231,8 +236,9 @@
 bPanMovesCanvas=True
 ; If true, zooms centering on the mouse position.  If false, the zoom is around the center of the viewport
 bCenterZoomAroundCursor=True
-; If true all orthographic viewports are linked to the same position and move together
-bUseLinkedOrthographicViewports=True
+bAllowTranslateRotateZWidget=False
+; If true, Clicking a BSP selects the brush and ctrl+shift+click selects the surface. If false, vice versa
+bClickBSPSelectsBrush=True
 ;Mouse speed when dragging in viewport
 CameraSpeed=4
 ;Mouse speed when using middle mouse scroll in viewport
@@ -243,16 +249,6 @@
 bLevelStreamingVolumePrevis=False
 bUseUE3OrbitControls=False
 
-bAllowTranslateRotateZWidget=False
-; If true, Clicking a BSP selects the brush and ctrl+shift+click selects the surface. If false, vice versa
-bClickBSPSelectsBrush=True
-; If true, BSP will auto-update
-bBSPAutoUpdate=True
-; If true, Navigation will auto-update
-bNavigationAutoUpdate=True
-; If true, Replaces respects the scale of the original actor. If false, Replace sets the scale to 1.0
-bReplaceRespectsScale=True
-
 bUsePowerOf2SnapSize=False
 
 .DecimalGridSizes=1
@@ -332,6 +328,8 @@
 bEnableViewportHoverFeedback=False
 ; If enabled, selected objects will be highlighted with brackets in all modes rather than a special highlight color.
 bHighlightWithBrackets=False
+; If true all orthographic viewports are linked to the same position and move together
+bUseLinkedOrthographicViewports=True
 ; If true, objects must be entirely encompassed by the selection box in ortho. viewports to be selected
 bStrictBoxSelection=False
 ; If enabled, selected objects will have an outline around them
diff --git a/Engine/Source/Editor/EditorSettingsViewer/Private/EditorSettingsViewerModule.cpp b/Engine/Source/Editor/EditorSettingsViewer/Private/EditorSettingsViewerModule.cpp
--- a/Engine/Source/Editor/EditorSettingsViewer/Private/EditorSettingsViewerModule.cpp
+++ b/Engine/Source/Editor/EditorSettingsViewer/Private/EditorSettingsViewerModule.cpp
@@ -171,6 +171,13 @@ protected:
 			LOCTEXT("LevelEditorViewportSettingsDescription", "Configure the look and feel of the Level Editor view ports."),
 			TWeakObjectPtr<UObject>(GetMutableDefault<ULevelEditorViewportSettings>())
 		);
+
+		// miscellaneous settings
+		SettingsModule.RegisterSettings("Editor", "LevelEditor", "Misc",
+			LOCTEXT("LevelEditorMiscSettingsName", "Miscellaneous"),
+			LOCTEXT("LevelEditorMiscSettingsDescription", "Configure miscellaneous settings for the Level Editor."),
+			TWeakObjectPtr<UObject>(GetMutableDefault<ULevelEditorMiscSettings>())
+		);
 	}
 
 	/**
diff --git a/Engine/Source/Editor/UnrealEd/Classes/Editor/EditorUserSettings.h b/Engine/Source/Editor/UnrealEd/Classes/Editor/EditorUserSettings.h
--- a/Engine/Source/Editor/UnrealEd/Classes/Editor/EditorUserSettings.h
+++ b/Engine/Source/Editor/UnrealEd/Classes/Editor/EditorUserSettings.h
@@ -248,17 +248,6 @@ class UEditorUserSettings : public UObject
 	UPROPERTY(config)
 	uint32 bAllowSelectTranslucent:1;
 
-	UPROPERTY(EditAnywhere, config, Category=Sound)
-	uint32 bAllowBackgroundAudio:1;
-
-	/** If true audio will be enabled in the editor. Does not affect PIE **/
-	UPROPERTY(config)
-	uint32 bEnableRealTimeAudio:1;
-
-	/** Global volume setting for the editor */
-	UPROPERTY(config)
-	float EditorVolumeLevel;
-
 	/** True if the actor count is displayed in the slate level browser */
 	UPROPERTY(config)
 	uint32 bDisplayActorCountInLevelBrowser:1;
@@ -279,17 +268,9 @@ class UEditorUserSettings : public UObject
 	UPROPERTY(config)
 	uint32 bDisplayEditorOffsetInLevelBrowser:1;
 
-	/** Enables audio feedback for certain operations in Unreal Editor, such as entering and exiting Play mode */
-	UPROPERTY(EditAnywhere, config, Category=Sound)
-	uint32 bEnableEditorSounds:1;
-
 	UPROPERTY()
 	class UBlueprintPaletteFavorites* BlueprintFavorites;
 
-	/** Whether to automatically save after a time interval */
-	UPROPERTY( EditAnywhere, config, Category = Build, meta = ( DisplayName = "Enable auto apply lighting" ) )
-	uint32 bAutoApplyLightingEnable : 1;
-
 public:
 
 	/**
@@ -307,6 +288,7 @@ public:
 	int32 MaterialQualityLevel;
 
 public:
+
 	/** Delegate for when a user setting has changed */
 	DECLARE_EVENT_OneParam(UEditorUserSettings, FUserSettingChangedEvent, FName /*PropertyName*/);
 	FUserSettingChangedEvent& OnUserSettingChanged() { return UserSettingChangedEvent; }
diff --git a/Engine/Source/Editor/UnrealEd/Classes/Settings/LevelEditorMiscSettings.h b/Engine/Source/Editor/UnrealEd/Classes/Settings/LevelEditorMiscSettings.h
new file mode 100644
--- /dev/null
+++ b/Engine/Source/Editor/UnrealEd/Classes/Settings/LevelEditorMiscSettings.h
@@ -0,0 +1,63 @@
+// Copyright 1998-2014 Epic Games, Inc. All Rights Reserved.
+
+/*=============================================================================
+	LevelEditorMiscSettings.h: Declares the ULevelEditorMiscSettings class.
+=============================================================================*/
+
+#pragma once
+
+#include "LevelEditorMiscSettings.generated.h"
+
+
+/**
+ * Implements the Level Editor's miscellaneous settings.
+ */
+UCLASS(config=EditorUserSettings)
+class UNREALED_API ULevelEditorMiscSettings
+	: public UObject
+{
+	GENERATED_UCLASS_BODY()
+
+public:
+
+	/** Whether to automatically save after a time interval */
+	UPROPERTY(EditAnywhere, config, Category=Editing, meta=(DisplayName="Apply Lighting Automatically"))
+	uint32 bAutoApplyLightingEnable:1;
+
+	/** If true, BSP will auto-update */
+	UPROPERTY(EditAnywhere, config, Category=Editing, meta=( DisplayName = "Update BSP Automatically" ))
+	uint32 bBSPAutoUpdate:1;
+
+	/** If true, Navigation will auto-update */
+	UPROPERTY(EditAnywhere, config, Category=Editing, meta=( DisplayName = "Update Navigation Automatically" ))
+	uint32 bNavigationAutoUpdate:1;
+
+	/** If enabled, replacing actors will respect the scale of the original actor.  Otherwise, the replaced actors will have a scale of 1.0 */
+	UPROPERTY(EditAnywhere, config, Category=Editing, meta=( DisplayName = "Preserve Actor Scale on Replace" ))
+	uint32 bReplaceRespectsScale:1;
+
+public:
+
+	UPROPERTY(EditAnywhere, config, Category=Sound)
+	uint32 bAllowBackgroundAudio:1;
+
+	/** If true audio will be enabled in the editor. Does not affect PIE **/
+	UPROPERTY(config)
+	uint32 bEnableRealTimeAudio:1;
+
+	/** Global volume setting for the editor */
+	UPROPERTY(config)
+	float EditorVolumeLevel;
+
+	/** Enables audio feedback for certain operations in Unreal Editor, such as entering and exiting Play mode */
+	UPROPERTY(EditAnywhere, config, Category=Sound)
+	uint32 bEnableEditorSounds:1;
+
+protected:
+
+	// Begin UObject overrides
+
+	virtual void PostEditChangeProperty( struct FPropertyChangedEvent& PropertyChangedEvent ) OVERRIDE;
+
+	// End UObject overrides
+};
diff --git a/Engine/Source/Editor/UnrealEd/Classes/Settings/LevelEditorViewportSettings.h b/Engine/Source/Editor/UnrealEd/Classes/Settings/LevelEditorViewportSettings.h
--- a/Engine/Source/Editor/UnrealEd/Classes/Settings/LevelEditorViewportSettings.h
+++ b/Engine/Source/Editor/UnrealEd/Classes/Settings/LevelEditorViewportSettings.h
@@ -154,11 +154,13 @@ public:
 	UPROPERTY(EditAnywhere, config, Category=Controls, meta=(DisplayName = "Orthographic Zoom to Cursor Position"))
 	uint32 bCenterZoomAroundCursor:1;
 
-	/**
-	 * If checked all orthographic view ports are linked to the same position and move together.
-	 */
-	UPROPERTY(EditAnywhere, config, Category=Controls, meta=(DisplayName = "Link Orthographic Viewport Movement"))
-	uint32 bUseLinkedOrthographicViewports:1;
+	/** Allow translate/rotate widget */
+	UPROPERTY(EditAnywhere, config, Category=Controls, meta=( DisplayName = "Enable Combined Translate/Rotate Widget" ))
+	uint32 bAllowTranslateRotateZWidget:1;
+
+	/** If true, Clicking a BSP selects the brush and ctrl+shift+click selects the surface. If false, vice versa */
+	UPROPERTY(EditAnywhere, config, Category=Controls, meta=( DisplayName = "Clicking BSP Enables Brush" ))
+	uint32 bClickBSPSelectsBrush:1;
 
 	/**
 	 * How fast the perspective camera moves when flying through the world.
@@ -200,28 +202,6 @@ public:
 	UPROPERTY(EditAnywhere, config, Category=Controls, meta=(DisplayName="Use UE3 Orbit Controls"), AdvancedDisplay)
 	bool bUseUE3OrbitControls;
 
-public:
-
-	/** Allow translate/rotate widget */
-	UPROPERTY(EditAnywhere, config, Category=Editing, meta=( DisplayName = "Enable Combined Translate/Rotate Widget" ))
-	uint32 bAllowTranslateRotateZWidget:1;
-
-	/** If true, Clicking a BSP selects the brush and ctrl+shift+click selects the surface. If false, vice versa */
-	UPROPERTY(EditAnywhere, config, Category=Editing, meta=( DisplayName = "Clicking BSP Enables Brush" ))
-	uint32 bClickBSPSelectsBrush:1;
-
-	/** If true, BSP will auto-update */
-	UPROPERTY(EditAnywhere, config, Category=Editing, meta=( DisplayName = "Update BSP Automatically" ))
-	uint32 bBSPAutoUpdate:1;
-
-	/** If true, Navigation will auto-update */
-	UPROPERTY(EditAnywhere, config, Category=Editing, meta=( DisplayName = "Update Navigation Automatically" ))
-	uint32 bNavigationAutoUpdate:1;
-
-	/** If enabled, replacing actors will respect the scale of the original actor.  Otherwise, the replaced actors will have a scale of 1.0 */
-	UPROPERTY(EditAnywhere, config, Category=Editing, meta=( DisplayName = "Preserve Actor Scale on Replace" ))
-	uint32 bReplaceRespectsScale:1;
-
 public:
 
 	/** If enabled will use power of 2 grid settings (e.g, 1,2,4,8,16,...,1024) instead of decimal grid sizes */
@@ -324,6 +304,12 @@ public:
 	UPROPERTY(EditAnywhere, config, Category=LookAndFeel, meta=(DisplayName = "Highlight Selected Objects with Brackets"))
 	uint32 bHighlightWithBrackets:1;
 
+	/**
+	 * If checked all orthographic view ports are linked to the same position and move together.
+	 */
+	UPROPERTY(EditAnywhere, config, Category=LookAndFeel, meta=(DisplayName = "Link Orthographic Viewport Movement"))
+	uint32 bUseLinkedOrthographicViewports:1;
+
 	/** True if viewport box selection requires objects to be fully encompassed by the selection box to be selected */
 	UPROPERTY(config)
 	uint32 bStrictBoxSelection:1;
diff --git a/Engine/Source/Editor/UnrealEd/Private/Editor.cpp b/Engine/Source/Editor/UnrealEd/Private/Editor.cpp
--- a/Engine/Source/Editor/UnrealEd/Private/Editor.cpp
+++ b/Engine/Source/Editor/UnrealEd/Private/Editor.cpp
@@ -390,7 +390,7 @@ void UEditorEngine::InitEditor(IEngineLoop* InEngineLoop)
 	
 	// Set navigation system property indicating whether navigation is supposed to rebuild automatically 
 	FWorldContext &EditorContext = GEditor->GetEditorWorldContext();
-	UNavigationSystem::SetNavigationAutoUpdateEnabled(GetDefault<ULevelEditorViewportSettings>()->bNavigationAutoUpdate, EditorContext.World()->GetNavigationSystem() );
+	UNavigationSystem::SetNavigationAutoUpdateEnabled(GetDefault<ULevelEditorMiscSettings>()->bNavigationAutoUpdate, EditorContext.World()->GetNavigationSystem() );
 
 	// Allocate temporary model.
 	TempModel = new UModel( FPostConstructInitializeProperties(),NULL, 1 );
@@ -838,7 +838,7 @@ void UEditorEngine::Tick( float DeltaSeconds, bool bIdleMode )
 	bool IsRealtime = false;
 
 	// True if a viewport has realtime audio	// If any realtime audio is enabled in the editor
-	bool bAudioIsRealtime = GEditor->AccessEditorUserSettings().bEnableRealTimeAudio;
+	bool bAudioIsRealtime = GetDefault<ULevelEditorMiscSettings>()->bEnableRealTimeAudio;
 
 	// By default we tick the editor world.  
 	// When in PIE if we are in immersive we do not tick the editor world unless there is a visible editor viewport.
@@ -943,12 +943,12 @@ void UEditorEngine::Tick( float DeltaSeconds, bool bIdleMode )
 	// Find out if the editor has focus. Audio should only play if the editor has focus.
 	const bool bHasFocus = FPlatformProcess::IsThisApplicationForeground();
 
-	if (bHasFocus || GEditor->AccessEditorUserSettings().bAllowBackgroundAudio)
+	if (bHasFocus || GetDefault<ULevelEditorMiscSettings>()->bAllowBackgroundAudio)
 	{
 		if (!PlayWorld)
 		{
 			// Adjust the global volume multiplier if the window has focus and there is no pie world or no viewport overriding audio.
-			GVolumeMultiplier = GEditor->AccessEditorUserSettings().EditorVolumeLevel;
+			GVolumeMultiplier = GetDefault<ULevelEditorMiscSettings>()->EditorVolumeLevel;
 		}
 		else
 		{
@@ -1516,28 +1516,28 @@ bool UEditorEngine::IsRealTimeAudioMuted() const
 	{
 		return true;
 	}
-	return EditorUserSettings->bEnableRealTimeAudio ? false : true;
+	return GetDefault<ULevelEditorMiscSettings>()->bEnableRealTimeAudio ? false : true;
 }
 
 void UEditorEngine::MuteRealTimeAudio(bool bMute)
 {
-	AccessEditorUserSettings().bEnableRealTimeAudio = bMute ? false : true;
-	AccessEditorUserSettings().PostEditChange();
+	ULevelEditorMiscSettings* LevelEditorMiscSettings = GetMutableDefault<ULevelEditorMiscSettings>();
+
+	LevelEditorMiscSettings->bEnableRealTimeAudio = bMute ? false : true;
+	LevelEditorMiscSettings->PostEditChange();
 }
 
 float UEditorEngine::GetRealTimeAudioVolume() const
 {
-	if (EditorUserSettings == NULL)
-	{
-		return 0.f;
-	}
-	return EditorUserSettings->EditorVolumeLevel;
+	return GetDefault<ULevelEditorMiscSettings>()->EditorVolumeLevel;
 }
 
 void UEditorEngine::SetRealTimeAudioVolume(float VolumeLevel)
 {
-	AccessEditorUserSettings().EditorVolumeLevel = VolumeLevel;
-	AccessEditorUserSettings().PostEditChange();
+	ULevelEditorMiscSettings* LevelEditorMiscSettings = GetMutableDefault<ULevelEditorMiscSettings>();
+
+	LevelEditorMiscSettings->EditorVolumeLevel = VolumeLevel;
+	LevelEditorMiscSettings->PostEditChange();
 }
 
 bool UEditorEngine::UpdateSingleViewportClient(FEditorViewportClient* InViewportClient, const bool bInAllowNonRealtimeViewportToDraw, bool bLinkedOrthoMovement )
@@ -1729,7 +1729,7 @@ void UEditorEngine::PlayPreviewSound( USoundBase* Sound,  USoundNode* SoundNode
 void UEditorEngine::PlayEditorSound( const FString& SoundAssetName )
 {
 	// Only play sounds if the user has that feature enabled
-	if( GetEditorUserSettings().bEnableEditorSounds )
+	if( GetDefault<ULevelEditorMiscSettings>()->bEnableEditorSounds )
 	{
 		USoundBase* Sound = Cast<USoundBase>( StaticFindObject( USoundBase::StaticClass(), NULL, *SoundAssetName ) );
 		if( Sound == NULL )
@@ -4966,7 +4966,7 @@ void UEditorEngine::ReplaceSelectedActors(UActorFactory* Factory, const FAssetDa
 			USceneComponent* const NewActorRootComponent = NewActor->GetRootComponent();
 			if(NewActorRootComponent)
 			{
-				if(!GetDefault<ULevelEditorViewportSettings>()->bReplaceRespectsScale || OldActor->GetRootComponent() == NULL )
+				if(!GetDefault<ULevelEditorMiscSettings>()->bReplaceRespectsScale || OldActor->GetRootComponent() == NULL )
 				{
 					NewActorRootComponent->SetRelativeScale3D(FVector(1.0f, 1.0f, 1.0f));
 				}
diff --git a/Engine/Source/Editor/UnrealEd/Private/Editor/EditorUserSettings.cpp b/Engine/Source/Editor/UnrealEd/Private/Editor/EditorUserSettings.cpp
--- a/Engine/Source/Editor/UnrealEd/Private/Editor/EditorUserSettings.cpp
+++ b/Engine/Source/Editor/UnrealEd/Private/Editor/EditorUserSettings.cpp
@@ -64,8 +64,6 @@ UEditorUserSettings::UEditorUserSettings(const class FPostConstructInitializePro
 	MaterialQualityLevel = 1;
 
 	BlueprintFavorites = ConstructObject<UBlueprintPaletteFavorites>(UBlueprintPaletteFavorites::StaticClass(), this);
-
-	bAutoApplyLightingEnable = true;
 }
 
 void UEditorUserSettings::PostInitProperties()
diff --git a/Engine/Source/Editor/UnrealEd/Private/EditorServer.cpp b/Engine/Source/Editor/UnrealEd/Private/EditorServer.cpp
--- a/Engine/Source/Editor/UnrealEd/Private/EditorServer.cpp
+++ b/Engine/Source/Editor/UnrealEd/Private/EditorServer.cpp
@@ -1390,7 +1390,7 @@ void UEditorEngine::RebuildMap(UWorld* InWorld, EMapRebuildType RebuildType)
 void UEditorEngine::RebuildLevel(ULevel& Level)
 {
 	// Early out if BSP auto-updating is disabled
-	if (!GetDefault<ULevelEditorViewportSettings>()->bBSPAutoUpdate)
+	if (!GetDefault<ULevelEditorMiscSettings>()->bBSPAutoUpdate)
 	{
 		return;
 	}
@@ -1478,7 +1478,7 @@ void UEditorEngine::RebuildModelFromBrushes(UModel* Model, bool bSelectedBrushes
 void UEditorEngine::RebuildAlteredBSP()
 {
 	// Early out if BSP auto-updating is disabled
-	if (!GetDefault<ULevelEditorViewportSettings>()->bBSPAutoUpdate)
+	if (!GetDefault<ULevelEditorMiscSettings>()->bBSPAutoUpdate)
 	{
 		return;
 	}
diff --git a/Engine/Source/Editor/UnrealEd/Private/NavigationBuildingNotification.cpp b/Engine/Source/Editor/UnrealEd/Private/NavigationBuildingNotification.cpp
--- a/Engine/Source/Editor/UnrealEd/Private/NavigationBuildingNotification.cpp
+++ b/Engine/Source/Editor/UnrealEd/Private/NavigationBuildingNotification.cpp
@@ -148,7 +148,7 @@ void FNavigationBuildingNotificationImpl::Tick(float DeltaTime)
 		FWorldContext &EditorContext = EEngine->GetEditorWorldContext();
 		
 		const bool bBuildInProgress = EditorContext.World() != NULL && EditorContext.World()->GetNavigationSystem() != NULL 
-			&& EditorContext.World()->GetNavigationSystem()->IsNavigationBuildInProgress( GetDefault<ULevelEditorViewportSettings>()->bNavigationAutoUpdate ? true : false) == true;
+			&& EditorContext.World()->GetNavigationSystem()->IsNavigationBuildInProgress( GetDefault<ULevelEditorMiscSettings>()->bNavigationAutoUpdate ? true : false) == true;
 
 		if (!bPreviouslyDetectedBuild && bBuildInProgress)
 		{
diff --git a/Engine/Source/Editor/UnrealEd/Private/Settings/SettingsClasses.cpp b/Engine/Source/Editor/UnrealEd/Private/Settings/SettingsClasses.cpp
--- a/Engine/Source/Editor/UnrealEd/Private/Settings/SettingsClasses.cpp
+++ b/Engine/Source/Editor/UnrealEd/Private/Settings/SettingsClasses.cpp
@@ -117,7 +117,7 @@ void UEditorLoadingSavingSettings::PostEditChangeProperty( struct FPropertyChang
 }
 
 
-/* ULevelEditorPlaySettings interface
+/* UEditorMiscSettings interface
  *****************************************************************************/
 
 UEditorMiscSettings::UEditorMiscSettings( const class FPostConstructInitializeProperties& PCIP )
@@ -125,6 +125,35 @@ UEditorMiscSettings::UEditorMiscSettings( const class FPostConstructInitializePr
 { }
 
 
+/* ULevelEditorMiscSettings interface
+ *****************************************************************************/
+
+ULevelEditorMiscSettings::ULevelEditorMiscSettings( const class FPostConstructInitializeProperties& PCIP )
+	: Super(PCIP)
+{
+	bAutoApplyLightingEnable = true;
+}
+
+
+void ULevelEditorMiscSettings::PostEditChangeProperty( struct FPropertyChangedEvent& PropertyChangedEvent )
+{
+	Super::PostEditChangeProperty(PropertyChangedEvent);
+
+	const FName Name = (PropertyChangedEvent.Property != nullptr) ? PropertyChangedEvent.Property->GetFName() : NAME_None;
+
+	if (Name == FName(TEXT("bNavigationAutoUpdate")))
+	{
+		FWorldContext &EditorContext = GEditor->GetEditorWorldContext();
+		UNavigationSystem::SetNavigationAutoUpdateEnabled(bNavigationAutoUpdate, EditorContext.World()->GetNavigationSystem());
+	}
+
+	if (!FUnrealEdMisc::Get().IsDeletePreferences())
+	{
+		SaveConfig();
+	}
+}
+
+
 /* ULevelEditorPlaySettings interface
  *****************************************************************************/
 
@@ -179,11 +208,6 @@ void ULevelEditorViewportSettings::PostEditChangeProperty( struct FPropertyChang
 	{
 		GEngine->HoverHighlightIntensity = HoverHighlightIntensity;
 	}
-	else if (Name == FName(TEXT("bNavigationAutoUpdate")))
-	{
-		FWorldContext &EditorContext = GEditor->GetEditorWorldContext();
-		UNavigationSystem::SetNavigationAutoUpdateEnabled(bNavigationAutoUpdate, EditorContext.World()->GetNavigationSystem());
-	}
 	else if (Name == FName(TEXT("SelectionHighlightIntensity")))
 	{
 		GEngine->SelectionHighlightIntensity = SelectionHighlightIntensity;
@@ -228,6 +252,7 @@ void ULevelEditorViewportSettings::PostEditChangeProperty( struct FPropertyChang
 		UBillboardComponent::SetEditorScale(BillboardScale);
 		UArrowComponent::SetEditorScale(BillboardScale);
 	}
+
 	if (!FUnrealEdMisc::Get().IsDeletePreferences())
 	{
 		SaveConfig();
diff --git a/Engine/Source/Editor/UnrealEd/Private/StaticLightingSystem/StaticLightingSystem.cpp b/Engine/Source/Editor/UnrealEd/Private/StaticLightingSystem/StaticLightingSystem.cpp
--- a/Engine/Source/Editor/UnrealEd/Private/StaticLightingSystem/StaticLightingSystem.cpp
+++ b/Engine/Source/Editor/UnrealEd/Private/StaticLightingSystem/StaticLightingSystem.cpp
@@ -2050,7 +2050,7 @@ void FStaticLightingSystem::UpdateAutomaticImportanceVolumeBounds( const FBox& M
 
 bool FStaticLightingSystem::CanAutoApplyLighting() const
 {
-	const bool bAutoApplyEnabled = GetDefault<UEditorUserSettings>()->bAutoApplyLightingEnable;
+	const bool bAutoApplyEnabled = GetDefault<ULevelEditorMiscSettings>()->bAutoApplyLightingEnable;
 	const bool bSlowTask = GIsSlowTask;
 	const bool bInterpEditMode = GEditorModeTools().IsModeActive( FBuiltinEditorModes::EM_InterpEdit );
 	const bool bPlayWorldValid = GUnrealEd->PlayWorld != nullptr;
