* Improving the delete dialog to better handle deleting redirectors.
* We no longer display non-assets in the list of items being deleted, to hide complexity from the underlying implementation.
* Fixing the black arow/text on the replace references dropdown on the delete dialog.

diff --git a/Engine/Source/Editor/UnrealEd/Private/AssetDeleteModel.cpp b/Engine/Source/Editor/UnrealEd/Private/AssetDeleteModel.cpp
--- a/Engine/Source/Editor/UnrealEd/Private/AssetDeleteModel.cpp
+++ b/Engine/Source/Editor/UnrealEd/Private/AssetDeleteModel.cpp
@@ -36,22 +36,6 @@ void FAssetDeleteModel::AddObjectToDelete( UObject* InObject )
 {
 	PendingDeletes.Add(MakeShareable(new FPendingDelete(InObject)));
 
-	// Blueprints actually contain 3 assets, the UBlueprint, GeneratedClass and SkeletonGeneratedClass
-	// we need to add them all to the pending deleted objects so that the memory reference detection system
-	// can know those references don't count towards the in memory references.
-	if ( InObject->IsA(UBlueprint::StaticClass()) )
-	{
-		UBlueprint* BlueprintObj = Cast<UBlueprint>(InObject);
-
-		TSharedPtr<FPendingDelete> GeneratedClass = MakeShareable(new FPendingDelete(BlueprintObj->GeneratedClass));
-		GeneratedClass->IsInternal(true);
-		TSharedPtr<FPendingDelete> SkeletonGeneratedClass = MakeShareable(new FPendingDelete(BlueprintObj->SkeletonGeneratedClass));
-		SkeletonGeneratedClass->IsInternal(true);
-
-		PendingDeletes.Add(GeneratedClass);
-		PendingDeletes.Add(SkeletonGeneratedClass);
-	}
-
 	SetState(StartScanning);
 }
 
@@ -364,10 +348,55 @@ int32 FAssetDeleteModel::GetDeletedObjectCount() const
 	return ObjectsDeleted;
 }
 
+// FPendingDelete
+//-----------------------------------------------------------------
+
+FPendingDelete::FPendingDelete(UObject* InObject)
+	: Object(InObject)
+	, bReferencesChecked(false)
+	, bIsReferencedInMemory(false)
+	, bIsReferencedInMemoryByUndo(false)
+	, RemainingDiskReferences(0)
+	, RemainingMemoryReferences(0)
+	, bIsInternal(false)
+{
+	// Blueprints actually contain 3 assets, the UBlueprint, GeneratedClass and SkeletonGeneratedClass
+	// we need to add them all to the pending deleted objects so that the memory reference detection system
+	// can know those references don't count towards the in memory references.
+	if ( InObject->IsA(UBlueprint::StaticClass()) )
+	{
+		UBlueprint* BlueprintObj = Cast<UBlueprint>(InObject);
+
+		if ( BlueprintObj->GeneratedClass != NULL )
+		{
+			InternalObjects.Add(BlueprintObj->GeneratedClass);
+		}
+
+		if ( BlueprintObj->SkeletonGeneratedClass != NULL )
+		{
+			InternalObjects.Add(BlueprintObj->SkeletonGeneratedClass);
+		}
+	}
+
+	FAssetData InAssetData(InObject);
+	// Filter out any non assets
+	if ( !InAssetData.IsUAsset() )
+	{
+		bIsInternal = true;
+	}
+}
+
 bool FPendingDelete::IsObjectContained(const UObject* InObject) const
 {
 	const UObject* InObjectParent = InObject;
 
+	// If the objects are in the same package then it should be safe to delete them since the package
+	// will be marked for garbage collection.
+	if ( Object->GetOutermost() == InObject->GetOutermost() )
+	{
+		return true;
+	}
+
 	// We need to check if the object or any of it's parents are children of the object being deleted, and so
 	// can safely be ignored.
 	while ( InObjectParent != NULL )
@@ -377,6 +406,15 @@ bool FPendingDelete::IsObjectContained(const UObject* InObject) const
 			return true;
 		}
 
+		// Also check if it's a child of any of the internal objects.
+		for ( int32 InternalIndex = 0; InternalIndex < InternalObjects.Num(); InternalIndex++ )
+		{
+			if ( InternalObjects[InternalIndex] == InObjectParent )
+			{
+				return true;
+			}
+		}
+
 		InObjectParent = InObjectParent->GetOuter();
 	}
 
diff --git a/Engine/Source/Editor/UnrealEd/Private/Dialogs/SDeleteAssetsDialog.cpp b/Engine/Source/Editor/UnrealEd/Private/Dialogs/SDeleteAssetsDialog.cpp
--- a/Engine/Source/Editor/UnrealEd/Private/Dialogs/SDeleteAssetsDialog.cpp
+++ b/Engine/Source/Editor/UnrealEd/Private/Dialogs/SDeleteAssetsDialog.cpp
@@ -430,7 +430,8 @@ TSharedRef<SWidget> SDeleteAssetsDialog::BuildReplaceReferencesWidget()
 		SAssignNew( ConsolidationPickerComboButton, SComboButton )
 		.HAlign( EHorizontalAlignment::HAlign_Fill )
 		.VAlign( EVerticalAlignment::VAlign_Center )
-		.ComboButtonStyle(FEditorStyle::Get(), "ToolbarComboButton")
+		.ComboButtonStyle( FEditorStyle::Get(), "ToolbarComboButton" )
+		.ForegroundColor( FLinearColor::White )
 		.ContentPadding( 3 )
 		.MenuPlacement( EMenuPlacement::MenuPlacement_BelowAnchor )
 		.OnGetMenuContent( this, &SDeleteAssetsDialog::MakeConsolidationAssetPicker )
@@ -818,6 +819,12 @@ void SDeleteAssetsDialog::ExecuteDeleteReferencers()
 
 bool SDeleteAssetsDialog::OnShouldFilterAsset( const FAssetData& InAssetData ) const
 {
+	// Filter out any redirectors that are not to the main UAsset
+	if ( InAssetData.IsRedirector() && !InAssetData.IsUAsset() )
+	{
+		return true;
+	}
+
 	// If it's in the set of references then don't filter it.
 	if ( DeleteModel->GetAssetReferences().Contains( InAssetData.PackageName ) )
 	{
diff --git a/Engine/Source/Editor/UnrealEd/Public/AssetDeleteModel.h b/Engine/Source/Editor/UnrealEd/Public/AssetDeleteModel.h
--- a/Engine/Source/Editor/UnrealEd/Public/AssetDeleteModel.h
+++ b/Engine/Source/Editor/UnrealEd/Public/AssetDeleteModel.h
@@ -8,16 +8,7 @@
 struct FPendingDelete
 {
 public:
-	FPendingDelete( UObject* InObject )
-		: Object( InObject )
-		, bReferencesChecked( false )
-		, bIsReferencedInMemory( false )
-		, bIsReferencedInMemoryByUndo( false )
-		, RemainingDiskReferences( 0 )
-		, RemainingMemoryReferences( 0 )
-		, bIsInternal(false)
-	{
-	}
+	FPendingDelete( UObject* InObject );
 
 	/** Checks for references on disk and in memory for this object filling out all information. */
 	void CheckForReferences();
@@ -58,6 +49,9 @@ private:
 	/** The object to delete */
 	UObject* Object;
 
+	/** Internal objects being deleted that we need to make sure aren't counted as memory references. */
+	TArray<UObject*> InternalObjects;
+
 	/** A flag indicating that references have been checked, so don't check again. */
 	bool bReferencesChecked;
 
diff --git a/Engine/Source/Runtime/AssetRegistry/Public/AssetData.h b/Engine/Source/Runtime/AssetRegistry/Public/AssetData.h
--- a/Engine/Source/Runtime/AssetRegistry/Public/AssetData.h
+++ b/Engine/Source/Runtime/AssetRegistry/Public/AssetData.h
@@ -114,6 +114,17 @@ public:
 		return FString::Printf(TEXT("%s'%s'"), *AssetClass.ToString(), *ObjectPath.ToString());
 	}
 
+	/** Returns true if the this asset is a redirector. */
+	bool IsRedirector() const
+	{
+		if ( AssetClass == UObjectRedirector::StaticClass()->GetFName() )
+		{
+			return true;
+		}
+
+		return false;
+	}
+
 	/** Returns the class UClass if it is loaded. It is not possible to load the class if it is unloaded since we only have the short name. */
 	UClass* GetClass() const
 	{
