UE4 - Delete Asset Dialog
* The code that navigates to actors in the world that are referencing assets about to be deleted, now attempts to also check if the referencer is a component and get the actor from it.

diff --git a/Engine/Source/Editor/UnrealEd/Private/AssetDeleteModel.cpp b/Engine/Source/Editor/UnrealEd/Private/AssetDeleteModel.cpp
--- a/Engine/Source/Editor/UnrealEd/Private/AssetDeleteModel.cpp
+++ b/Engine/Source/Editor/UnrealEd/Private/AssetDeleteModel.cpp
@@ -323,16 +323,26 @@ bool FAssetDeleteModel::GoToNextReferenceInLevel() const
 	{
 		for ( const FReferencerInformation& Reference : PendingDelete->MemoryReferences.ExternalReferences )
 		{
-			const AActor* ReferencedActor = Cast<const AActor>(Reference.Referencer);
+			const AActor* ReferencingActor = Cast<const AActor>(Reference.Referencer);
 
-			if ( ReferencedActor )
+			// If the referencer isn't an actor, it might be a component.
+			if ( !ReferencingActor )
 			{
-				if ( ReferencedActor->GetWorld() == RepresentingWorld )
+				const UActorComponent* ReferencingComponent = Cast<const UActorComponent>(Reference.Referencer);
+				if ( ReferencingComponent )
 				{
-					GEditor->SelectActor(const_cast<AActor*>( ReferencedActor ), true, true, true);
+					ReferencingActor = ReferencingComponent->GetOwner();
+				}
+			}
+
+			if ( ReferencingActor )
+			{
+				if ( ReferencingActor->GetWorld() == RepresentingWorld )
+				{
+					GEditor->SelectActor(const_cast<AActor*>( ReferencingActor ), true, true, true);
 
 					// Point the camera at it
-					GUnrealEd->Exec(ReferencedActor->GetWorld(), TEXT("CAMERA ALIGN ACTIVEVIEWPORTONLY"));
+					GUnrealEd->Exec(ReferencingActor->GetWorld(), TEXT("CAMERA ALIGN ACTIVEVIEWPORTONLY"));
 
 					return true;
 				}
