#UE4 - Fix FBX crash with file made from blender

diff --git a/Engine/Source/Editor/UnrealEd/Private/Fbx/FbxMainImport.cpp b/Engine/Source/Editor/UnrealEd/Private/Fbx/FbxMainImport.cpp
--- a/Engine/Source/Editor/UnrealEd/Private/Fbx/FbxMainImport.cpp
+++ b/Engine/Source/Editor/UnrealEd/Private/Fbx/FbxMainImport.cpp
@@ -1097,31 +1097,43 @@ void FFbxImporter::RecursiveFindFbxSkelMesh(FbxNode* Node, TArray< TArray<FbxNod
 		const int32 fbxDeformerCount = SkelMeshNode->GetMesh()->GetDeformerCount();
 		FbxSkin* Deformer = static_cast<FbxSkin*>( SkelMeshNode->GetMesh()->GetDeformer(0, FbxDeformer::eSkin) );
 		
-		if (Deformer != NULL)
+		if (Deformer != NULL )
 		{
-			FbxNode* Link = Deformer->GetCluster(0)->GetLink(); //Get the bone influences by this first cluster
-			Link = GetRootSkeleton(Link); // Get the skeleton root itself
-
-			int32 i;
-			for (i = 0; i < SkeletonArray.Num(); i++)
+			if( Deformer->GetClusterCount() > 0 )
 			{
-				if ( Link == SkeletonArray[i] )
+				FbxNode* Link = Deformer->GetCluster(0)->GetLink(); //Get the bone influences by this first cluster
+				Link = GetRootSkeleton(Link); // Get the skeleton root itself
+
+				int32 i;
+				for (i = 0; i < SkeletonArray.Num(); i++)
 				{
-					// append to existed outSkelMeshArray element
-					TArray<FbxNode*>* TempArray = outSkelMeshArray[i];
+					if (Link == SkeletonArray[i])
+					{
+						// append to existed outSkelMeshArray element
+						TArray<FbxNode*>* TempArray = outSkelMeshArray[i];
+						TempArray->Add(NodeToAdd);
+						break;
+					}
+				}
+
+				// if there is no outSkelMeshArray element that is bind to this skeleton
+				// create new element for outSkelMeshArray
+				if (i == SkeletonArray.Num())
+				{
+					TArray<FbxNode*>* TempArray = new TArray<FbxNode*>();
 					TempArray->Add(NodeToAdd);
-					break;
+					outSkelMeshArray.Add(TempArray);
+					SkeletonArray.Add(Link);
 				}
 			}
-
-			// if there is no outSkelMeshArray element that is bind to this skeleton
-			// create new element for outSkelMeshArray
-			if ( i == SkeletonArray.Num() )
+			else
 			{
-				TArray<FbxNode*>* TempArray = new TArray<FbxNode*>();
-				TempArray->Add(NodeToAdd);
-				outSkelMeshArray.Add(TempArray);
-				SkeletonArray.Add(Link);
+				AddTokenizedErrorMessage(
+					FTokenizedMessage::Create(
+						EMessageSeverity::Warning, 
+						FText::Format( LOCTEXT("FBX_NoWeightsOnDeformer", "Ignoring mesh {0} because it but no weights."), FText::FromString( ANSI_TO_TCHAR(SkelMeshNode->GetName()) ) )
+					)
+				);
 			}
 		}
 	}
