Added support for cook on the fly to work with multiple platforms (slowly removing -targetplatform commandline option).
Phase 1
Added platform to the cooked file requests so that the cooker just cooks files for the requested platform.

diff --git a/Engine/Source/Developer/TargetPlatform/Private/TargetPlatformManagerModule.cpp b/Engine/Source/Developer/TargetPlatform/Private/TargetPlatformManagerModule.cpp
--- a/Engine/Source/Developer/TargetPlatform/Private/TargetPlatformManagerModule.cpp
+++ b/Engine/Source/Developer/TargetPlatform/Private/TargetPlatformManagerModule.cpp
@@ -93,6 +93,28 @@ public:
 		return NULL;
 	}
 
+	virtual const TArray<ITargetPlatform*>& GetCookingTargetPlatforms() OVERRIDE
+	{
+		static bool bInitialized = false;
+		static TArray<ITargetPlatform*> Results;
+
+		if ( !bInitialized || bForceCacheUpdate )
+		{
+			Results = GetActiveTargetPlatforms();
+
+			FString PlatformStr;
+			if (FParse::Value(FCommandLine::Get(), TEXT("TARGETPLATFORM="), PlatformStr))
+			{
+				if (PlatformStr == TEXT("None"))
+				{
+					Results = Platforms;
+				}
+			}
+		}
+
+		return Results;
+	}
+
 	virtual const TArray<ITargetPlatform*>& GetActiveTargetPlatforms() OVERRIDE
 	{
 		static bool bInitialized = false;
diff --git a/Engine/Source/Developer/TargetPlatform/Public/Interfaces/ITargetPlatformManagerModule.h b/Engine/Source/Developer/TargetPlatform/Public/Interfaces/ITargetPlatformManagerModule.h
--- a/Engine/Source/Developer/TargetPlatform/Public/Interfaces/ITargetPlatformManagerModule.h
+++ b/Engine/Source/Developer/TargetPlatform/Public/Interfaces/ITargetPlatformManagerModule.h
@@ -75,6 +75,14 @@ public:
 	 */
 	virtual ITargetPlatform* FindTargetPlatform(FString Name) = 0;
 
+
+	/**
+	 * Return the list of platforms which we need to support when cooking (only set when actually cooking)
+	 *
+	 * @return Collection of platforms.
+	 */
+	virtual const TArray<ITargetPlatform*>& GetCookingTargetPlatforms() = 0;
+
 	/**
 	 * Return the list of the ITargetPlatforms that we want to build data for.
 	 *
diff --git a/Engine/Source/Editor/UnrealEd/Classes/Commandlets/CookCommandlet.h b/Engine/Source/Editor/UnrealEd/Classes/Commandlets/CookCommandlet.h
--- a/Engine/Source/Editor/UnrealEd/Classes/Commandlets/CookCommandlet.h
+++ b/Engine/Source/Editor/UnrealEd/Classes/Commandlets/CookCommandlet.h
@@ -76,8 +76,20 @@ class UCookCommandlet
 	 *	@return	bool			true if packages was cooked
 	 */
 	bool SaveCookedPackage( UPackage* Package, uint32 SaveFlags, bool& bOutWasUpToDate );
+	/**
+	 *	Cook (save) the given package
+	 *
+	 *	@param	Package				The package to cook/save
+	 *	@param	SaveFlags			The flags to pass to the SavePackage function
+	 *	@param	bOutWasUpToDate		Upon return, if true then the cooked package was cached (up to date)
+	 *	@param  TargetPlatformNames Only cook for target platforms which are included in this array (if empty cook for all target platforms specified on commandline options)
+	 *									TargetPlatformNames is in and out value returns the platforms which the SaveCookedPackage function saved for
+	 *
+	 *	@return	bool			true if packages was cooked
+	 */
+	bool SaveCookedPackage( UPackage* Package, uint32 SaveFlags, bool& bOutWasUpToDate, TArray<FString> &TargetPlatformNames );
 
-	bool ShouldCook(const FString& InFilename);
+	bool ShouldCook(const FString& InFilename, const FString &InPlatformname = TEXT(""));
 
 public:
 
@@ -90,7 +102,7 @@ public:
 private:
 
 	// Callback for handling a network file request.
-	void HandleNetworkFileServerFileRequest( const FString& Filename, TArray<FString>& UnsolicitedFiles );
+	void HandleNetworkFileServerFileRequest( const FString& Filename, const FString &Platformname, TArray<FString>& UnsolicitedFiles );
 
 	// Callback for recompiling shaders
 	void HandleNetworkFileServerRecompileShaders(const struct FShaderRecompileData& RecompileData);
diff --git a/Engine/Source/Editor/UnrealEd/Private/Commandlets/CookCommandlet.cpp b/Engine/Source/Editor/UnrealEd/Private/Commandlets/CookCommandlet.cpp
--- a/Engine/Source/Editor/UnrealEd/Private/Commandlets/CookCommandlet.cpp
+++ b/Engine/Source/Editor/UnrealEd/Private/Commandlets/CookCommandlet.cpp
@@ -95,22 +95,65 @@ UCookCommandlet::UCookCommandlet( const class FPostConstructInitializeProperties
 /* UCookCommandlet interface
  *****************************************************************************/
 
+struct FFilePlatformRequest
+{
+public:
+	FFilePlatformRequest() { }
+	FFilePlatformRequest( const FString &InFilename, const FString &InPlatformname ) : Filename( FPaths::ConvertRelativePathToFull( InFilename) ), Platformname( InPlatformname ) 
+	{
+		// FPaths::ConvertRelativePathToFull(InFilename);
+	}
+
+	FString Filename;
+	FString Platformname;
+
+	bool IsValid()  const
+	{
+		return Filename.Len() > 0;
+	}
+
+	void Clear()
+	{
+
+		Filename = TEXT("");
+		Platformname = TEXT("");
+	}
+
+	bool operator ==( const FFilePlatformRequest &InFileRequest ) const
+	{
+		if ( InFileRequest.Filename == Filename )
+		{
+			if ( InFileRequest.Platformname == Platformname )
+				return true;
+		}
+		return false;
+	}
+
+	
+};
+uint32 GetTypeHash(const FFilePlatformRequest &Key)
+{
+	return GetTypeHash( Key.Filename ) ^ GetTypeHash( Key.Platformname );
+}
+
+FThreadSafeArray<FFilePlatformRequest> ThreadSafeFilenameArray;
+FThreadSafeArray<FFilePlatformRequest> UnsolicitedCookedPackages;
+
+
 
-FThreadSafeArray<FString> ThreadSafeFilenameArray;
-FThreadSafeArray<FString> UnsolicitedCookedPackages;
 
 struct FThreadSafeFilenameSet
 {
 	FCriticalSection	SynchronizationObject;
-	TSet<FString> FilesProcessed;
+	TSet<FFilePlatformRequest> FilesProcessed;
 
-	void Add(const FString& Filename)
+	void Add(const FFilePlatformRequest &Filename)
 	{
 		FScopeLock ScopeLock(&SynchronizationObject);
-		check(Filename.Len());
+		check(Filename.IsValid());
 		FilesProcessed.Add(Filename);
 	}
-	bool Exists(const FString& Filename)
+	bool Exists(const FFilePlatformRequest &Filename)
 	{
 		FScopeLock ScopeLock(&SynchronizationObject);
 		return FilesProcessed.Contains(Filename);
@@ -182,28 +225,36 @@ bool UCookCommandlet::CookOnTheFly( bool BindAnyPort, int32 Timeout, bool bForce
 
 	while (!GIsRequestingExit)
 	{
-		TSet<FString> CookedPackages;
-		FString ToBuild;
+		FFilePlatformRequest ToBuild;
 
 		while (!GIsRequestingExit)
 		{
-			while (ToBuild.Len() && !GIsRequestingExit)
+			// This is all the target platforms which we needed to process requests for this iteration
+			// we use this in the unsolicited packages processing below
+			TArray<FString> AllTargetPlatformNames;
+
+
+			while (ToBuild.IsValid() && !GIsRequestingExit)
 			{
 				bool bWasUpToDate = false;
+				TArray<FString> TargetPlatformNames;
+				if( ToBuild.Platformname.Len() )
+				{
+					TargetPlatformNames.Add( ToBuild.Platformname );
+				}
 
-				FString Name = FPackageName::FilenameToLongPackageName(*ToBuild);
-				if (!CookedPackages.Contains(Name))
+				FString Name = FPackageName::FilenameToLongPackageName(*ToBuild.Filename);
+				if (!ThreadSafeFilenameSet.Exists(ToBuild))
 				{
-					UPackage* Package = LoadPackage( NULL, *ToBuild, LOAD_None );
+					UPackage* Package = LoadPackage( NULL, *ToBuild.Filename, LOAD_None );
 
 					if( Package == NULL )
 					{
-						UE_LOG(LogCookCommandlet, Error, TEXT("Error loading %s!"), *ToBuild );
+						UE_LOG(LogCookCommandlet, Error, TEXT("Error loading %s!"), *ToBuild.Filename );
 					}
 					else
 					{
-						CookedPackages.Add(Name);
-						if( SaveCookedPackage(Package, SAVE_KeepGUID | SAVE_Async, bWasUpToDate) )
+						if( SaveCookedPackage(Package, SAVE_KeepGUID | SAVE_Async, bWasUpToDate, TargetPlatformNames) )
 						{
 							// Update flags used to determine garbage collection.
 							if (Package->ContainsMap())
@@ -225,9 +276,17 @@ bool UCookCommandlet::CookOnTheFly( bool BindAnyPort, int32 Timeout, bool bForce
 					}
 				}
 
-				ThreadSafeFilenameSet.Add(ToBuild);
+				for ( int Index = 0; Index < TargetPlatformNames.Num(); ++Index )
+				{
+					AllTargetPlatformNames.AddUnique(TargetPlatformNames[Index]);
+
+					FFilePlatformRequest FileRequest( ToBuild.Filename, TargetPlatformNames[Index]);
+					
+					ThreadSafeFilenameSet.Add( FileRequest );
+				}
+
 
-				ToBuild = TEXT("");
+				ToBuild.Clear();
 
 				ThreadSafeFilenameArray.Dequeue(&ToBuild);
 			}
@@ -241,25 +300,42 @@ bool UCookCommandlet::CookOnTheFly( bool BindAnyPort, int32 Timeout, bool bForce
 				if (Package)
 				{
 					FString Name = Package->GetPathName();
+					FString PackageFilename(GetPackageFilename(Package));
+					FPaths::MakeStandardFilename(PackageFilename);
 
-					if (!CookedPackages.Contains(Name))
-					{
-						CookedPackages.Add(Name);
+					const FString FullPath = FPaths::ConvertRelativePathToFull(PackageFilename);
 
+					bool AlreadyCooked = true;
+					for ( int Index = 0; Index < AllTargetPlatformNames.Num(); ++Index )
+					{
+						if ( !ThreadSafeFilenameSet.Exists(FFilePlatformRequest(FullPath, AllTargetPlatformNames[Index])) )
+						{
+							AlreadyCooked = false;
+							break;
+						}
+					}
+					if ( !AlreadyCooked )
+					{
 						bool bUnsolicitedWasUpdateToDate = false;
 
-						if (SaveCookedPackage(Package, SAVE_KeepGUID | SAVE_Async, bUnsolicitedWasUpdateToDate))
+						if (SaveCookedPackage(Package, SAVE_KeepGUID | SAVE_Async, bUnsolicitedWasUpdateToDate, AllTargetPlatformNames))
 						{
-							FString PackageFilename(GetPackageFilename(Package));
-							FPaths::MakeStandardFilename(PackageFilename);
+							
 
-							if ((FPaths::FileExists(PackageFilename) == true) && 
+							if ((FPaths::FileExists(PackageFilename) == true) &&
 								(PackageFilename.Len() > 0) && 
 								(bUnsolicitedWasUpdateToDate == false))
 							{
 								UE_LOG(LogCookCommandlet, Display, TEXT("UnsolicitedCookedPackages: %s"), *PackageFilename);
 
-								UnsolicitedCookedPackages.Enqueue(PackageFilename);
+								for ( int Index = 0; Index < AllTargetPlatformNames.Num(); ++Index )
+								{
+									const FFilePlatformRequest Request( FullPath, AllTargetPlatformNames[Index] );
+									UnsolicitedCookedPackages.Enqueue( Request );
+									ThreadSafeFilenameSet.Add( Request );
+
+									UE_LOG( LogCookCommandlet, Display, TEXT("Unsolicited saved package %s.%s"), *Request.Filename, *Request.Platformname );
+								}
 							}
 
 							// Update flags used to determine garbage collection.
@@ -275,24 +351,24 @@ bool UCookCommandlet::CookOnTheFly( bool BindAnyPort, int32 Timeout, bool bForce
 						}
 					}
 
-					ToBuild = TEXT("");
+					ToBuild.Clear();
 
 					ThreadSafeFilenameArray.Dequeue(&ToBuild);
 
-					if (ToBuild.Len())
+					if (ToBuild.IsValid())
 					{
 						break;
 					}
 				}
 			}
 
-			while (!ToBuild.Len() && !GIsRequestingExit)
+			while (!ToBuild.IsValid() && !GIsRequestingExit)
 			{
-				ToBuild = TEXT("");
+				ToBuild.Clear();
 
 				ThreadSafeFilenameArray.Dequeue(&ToBuild);
 
-				if (!ToBuild.Len())
+				if (!ToBuild.IsValid())
 				{
 					if (NonMapPackageCountSinceLastGC > 0)
 					{
@@ -406,10 +482,11 @@ bool UCookCommandlet::GetPackageTimestamp( const FString& InFilename, FDateTime&
 	return false;
 }
 
-bool UCookCommandlet::ShouldCook(const FString& InFileName)
+bool UCookCommandlet::ShouldCook(const FString& InFileName, const FString &InPlatformName)
 {
 	bool bDoCook = false;
 
+
 	FString PkgFile;
 	FString PkgFilename;
 	FDateTime DependentTimeStamp = FDateTime::MinValue();
@@ -428,7 +505,19 @@ bool UCookCommandlet::ShouldCook(const FString& InFileName)
 	PkgFilename = SandboxFile->ConvertToAbsolutePathForExternalAppForWrite(*PkgFilename);
 
 	ITargetPlatformManagerModule& TPM = GetTargetPlatformManagerRef();
-	static const TArray<ITargetPlatform*>& Platforms =  TPM.GetActiveTargetPlatforms();
+
+	static const TArray<ITargetPlatform*> &ActiveTargetPlatforms = TPM.GetActiveTargetPlatforms();
+
+	TArray<ITargetPlatform*> Platforms;
+
+	if ( InPlatformName.Len() > 0 )
+	{
+		Platforms.Add( TPM.FindTargetPlatform( InPlatformName ) );
+	}
+	else 
+	{
+		Platforms = ActiveTargetPlatforms;
+	}
 
 	for (int32 Index = 0; Index < Platforms.Num() && !bDoCook; Index++)
 	{
@@ -451,7 +540,14 @@ bool UCookCommandlet::ShouldCook(const FString& InFileName)
 	return bDoCook;
 }
 
-bool UCookCommandlet::SaveCookedPackage( UPackage* Package, uint32 SaveFlags, bool& bOutWasUpToDate )
+
+bool UCookCommandlet::SaveCookedPackage( UPackage* Package, uint32 SaveFlags, bool& bOutWasUpToDate ) 
+{
+	TArray<FString> TargetPlatformNames; 
+	return SaveCookedPackage( Package, SaveFlags, bOutWasUpToDate, TargetPlatformNames );
+}
+
+bool UCookCommandlet::SaveCookedPackage( UPackage* Package, uint32 SaveFlags, bool& bOutWasUpToDate, TArray<FString> &TargetPlatformNames )
 {
 	bool bSavedCorrectly = true;
 
@@ -491,7 +587,41 @@ bool UCookCommandlet::SaveCookedPackage( UPackage* Package, uint32 SaveFlags, bo
 		}
 
 		ITargetPlatformManagerModule& TPM = GetTargetPlatformManagerRef();
-		static const TArray<ITargetPlatform*>& Platforms =  TPM.GetActiveTargetPlatforms();
+
+		static TArray<ITargetPlatform*> ActiveStartupPlatforms = TPM.GetActiveTargetPlatforms();
+
+		TArray<ITargetPlatform*> Platforms;
+
+		if ( TargetPlatformNames.Num() )
+		{
+			for ( int Index = 0; Index < TargetPlatformNames.Num(); ++Index )
+			{
+				const FString &TargetPlatformName = TargetPlatformNames[Index];
+
+
+				const TArray<ITargetPlatform*>& TargetPlatforms = TPM.GetTargetPlatforms();	
+
+
+				for ( int Index = 0; Index < TargetPlatforms.Num(); ++Index )
+				{
+					ITargetPlatform *TargetPlatform = TargetPlatforms[ Index ];
+					if ( TargetPlatform->PlatformName() == TargetPlatformName )
+					{
+						Platforms.Add( TargetPlatform );
+					}
+				}
+			}
+		}
+		else
+		{
+			Platforms = ActiveStartupPlatforms;
+
+			for ( int Index = 0; Index < Platforms.Num(); ++Index )
+			{
+				TargetPlatformNames.Add( Platforms[Index]->PlatformName() );
+			}
+		}
+		
 
 		for (int32 Index = 0; Index < Platforms.Num(); Index++)
 		{
@@ -973,7 +1103,7 @@ void UCookCommandlet::CollectFilesToCook(TArray<FString>& FilesInPath)
 
 	// make sure we cook the default maps
 	ITargetPlatformManagerModule& TPM = GetTargetPlatformManagerRef();
-	static const TArray<ITargetPlatform*>& Platforms =  TPM.GetActiveTargetPlatforms();
+	static const TArray<ITargetPlatform*>& Platforms =  TPM.GetTargetPlatforms();
 	for (int32 Index = 0; Index < Platforms.Num(); Index++)
 	{
 		// load the platform specific ini to get its DefaultMap
@@ -1258,26 +1388,37 @@ bool UCookCommandlet::Cook(const TArray<ITargetPlatform*>& Platforms, TArray<FSt
 /* UCookCommandlet callbacks
  *****************************************************************************/
 
-void UCookCommandlet::HandleNetworkFileServerFileRequest( const FString& Filename, TArray<FString>& UnsolicitedFiles )
+void UCookCommandlet::HandleNetworkFileServerFileRequest( const FString& Filename, const FString &Platformname, TArray<FString>& UnsolicitedFiles )
 {
 	bool bIsCookable = FPackageName::IsPackageExtension(*FPaths::GetExtension(Filename, true));
 
 	if (!bIsCookable)
 	{
-		UnsolicitedCookedPackages.DequeueAll(UnsolicitedFiles);
+		TArray<FFilePlatformRequest> FileRequests;
+		UnsolicitedCookedPackages.DequeueAll(FileRequests);
+		for ( int Index=  0; Index < FileRequests.Num(); ++Index)
+		{
+			UnsolicitedFiles.Add( FileRequests[Index].Filename);
+		}
 		UPackage::WaitForAsyncFileWrites();
 		return;
 	}
 
-	ThreadSafeFilenameArray.Enqueue(Filename);
+	FFilePlatformRequest FileRequest( Filename, Platformname );
+	ThreadSafeFilenameArray.Enqueue(FileRequest);
 
 	do
 	{
 		FPlatformProcess::Sleep(0.0f);
 	}
-	while (!ThreadSafeFilenameSet.Exists(Filename));
+	while (!ThreadSafeFilenameSet.Exists(FileRequest));
 
-	UnsolicitedCookedPackages.DequeueAll(UnsolicitedFiles);
+	TArray<FFilePlatformRequest> FileRequests;
+	UnsolicitedCookedPackages.DequeueAll(FileRequests);
+	for ( int Index=  0; Index < FileRequests.Num(); ++Index)
+	{
+		UnsolicitedFiles.Add( FileRequests[Index].Filename);
+	}
 	
 	UPackage::WaitForAsyncFileWrites();
 }
diff --git a/Engine/Source/Runtime/CoreUObject/Private/UObject/SavePackage.cpp b/Engine/Source/Runtime/CoreUObject/Private/UObject/SavePackage.cpp
--- a/Engine/Source/Runtime/CoreUObject/Private/UObject/SavePackage.cpp
+++ b/Engine/Source/Runtime/CoreUObject/Private/UObject/SavePackage.cpp
@@ -2463,6 +2463,27 @@ bool UPackage::SavePackage( UPackage* InOuter, UObject* Base, EObjectFlags TopLe
 			PackageExportTagger.TagPackageExports( ExportTaggerArchive, true );
 			ExportTaggerArchive.SetFilterEditorOnly(FilterEditorOnly);
 
+
+
+
+			// Kick off any Precaching required for the target platform to save these objects
+			// only need to do this if we are cooking a different platform then the one which is currently running
+			if ( !!TargetPlatform )
+			{
+				TArray<UObject*> TagExpObjects;
+				GetObjectsWithAnyMarks(TagExpObjects, OBJECTMARK_TagExp);
+				for ( int Index =0; Index < TagExpObjects.Num(); ++Index)
+				{
+					if ( TagExpObjects[Index]->HasAnyMarks( OBJECTMARK_TagExp ) )
+					{
+						TagExpObjects[Index]->BeginCacheForCookedPlatformData( TargetPlatform );
+
+						// UE_LOG(LogSavePackage, Log, TEXT("begincache %s for %s"), *TagExpObjects[Index]->GetFullName(), *TargetPlatform->PlatformName());
+					}
+				}
+			}
+
+
 		
 			{
 				// set GIsSavingPackage here as it is now illegal to create any new object references; they potentially wouldn't be saved correctly
@@ -2584,10 +2605,10 @@ bool UPackage::SavePackage( UPackage* InOuter, UObject* Base, EObjectFlags TopLe
 					}
 				}
 
-
+				
 				if ( EndSavingIfCancelled( Linker, TempFilename ) ) { return false; }
 				GWarn->UpdateProgress( ++CurSaveStep, TotalSaveSteps );
-		
+
 				TArray<UObject*> PrivateObjects;
 				TArray<UObject*> ObjectsInOtherMaps;
 				TArray<UObject*> LevelObjects;
@@ -3110,7 +3131,7 @@ bool UPackage::SavePackage( UPackage* InOuter, UObject* Base, EObjectFlags TopLe
 					Linker->StartScriptSHAGeneration();
 				}
 
-				
+
 				// Save exports.
 				int32 LastExportSaveStep = 0;
 				for( int32 i=0; i<Linker->ExportMap.Num(); i++ )
@@ -3172,7 +3193,7 @@ bool UPackage::SavePackage( UPackage* InOuter, UObject* Base, EObjectFlags TopLe
 
 						// Save the object data.
 						Export.SerialOffset = Linker->Tell();
-	//						UE_LOG(LogSavePackage, Log, TEXT("export %s"), *Export.Object->GetFullName());
+						// UE_LOG(LogSavePackage, Log, TEXT("export %s for %s"), *Export.Object->GetFullName(), *Linker->CookingTarget()->PlatformName());
 						if ( Export.Object->HasAnyFlags(RF_ClassDefaultObject) )
 						{
 							Export.Object->GetClass()->SerializeDefaultObject(Export.Object, *Linker);
diff --git a/Engine/Source/Runtime/CoreUObject/Public/UObject/UObject.h b/Engine/Source/Runtime/CoreUObject/Public/UObject/UObject.h
--- a/Engine/Source/Runtime/CoreUObject/Public/UObject/UObject.h
+++ b/Engine/Source/Runtime/CoreUObject/Public/UObject/UObject.h
@@ -534,6 +534,18 @@ public:
 	/** Make sure this object has been shut down. */
 	void ConditionalShutdownAfterError();
 
+
+	/**
+	 * Starts caching of platform specific data for the target platform
+	 * Called when cooking before serialization so that object can prepare platform specific data
+	 * Not called during normal loading of objects
+	 * 
+	 * @param	TargetPlatform	target platform to cache platform specific data for
+	 * @return	true if everything went well, false will abort save...
+	 */
+	virtual void BeginCacheForCookedPlatformData( const ITargetPlatform *TargetPlatform ) {  }
+
+
 	/**
 	 * Determine if this object has SomeObject in its archetype chain.
 	 */
diff --git a/Engine/Source/Runtime/Engine/Classes/Materials/Material.h b/Engine/Source/Runtime/Engine/Classes/Materials/Material.h
--- a/Engine/Source/Runtime/Engine/Classes/Materials/Material.h
+++ b/Engine/Source/Runtime/Engine/Classes/Materials/Material.h
@@ -766,6 +766,7 @@ public:
 	ENGINE_API virtual void Serialize(FArchive& Ar) OVERRIDE;
 	ENGINE_API virtual void PostDuplicate(bool bDuplicateForPIE) OVERRIDE;
 	ENGINE_API virtual void PostLoad() OVERRIDE;
+	ENGINE_API virtual void BeginCacheForCookedPlatformData( const ITargetPlatform *TargetPlatform ) OVERRIDE;
 #if WITH_EDITOR
 	ENGINE_API virtual void PreEditChange(UProperty* PropertyAboutToChange) OVERRIDE;
 	ENGINE_API virtual void PostEditChangeProperty(FPropertyChangedEvent& PropertyChangedEvent) OVERRIDE;
diff --git a/Engine/Source/Runtime/Engine/Classes/Materials/MaterialInstance.h b/Engine/Source/Runtime/Engine/Classes/Materials/MaterialInstance.h
--- a/Engine/Source/Runtime/Engine/Classes/Materials/MaterialInstance.h
+++ b/Engine/Source/Runtime/Engine/Classes/Materials/MaterialInstance.h
@@ -261,6 +261,7 @@ public:
 	// Begin UObject interface.
 	virtual SIZE_T GetResourceSize(EResourceSizeMode::Type Mode) OVERRIDE;
 	virtual void PostInitProperties() OVERRIDE;	
+	virtual void BeginCacheForCookedPlatformData(  const ITargetPlatform *TargetPlatform ) OVERRIDE;
 	virtual void Serialize(FArchive& Ar) OVERRIDE;
 	virtual void PostLoad() OVERRIDE;
 	virtual void BeginDestroy() OVERRIDE;
diff --git a/Engine/Source/Runtime/Engine/Private/Materials/Material.cpp b/Engine/Source/Runtime/Engine/Private/Materials/Material.cpp
--- a/Engine/Source/Runtime/Engine/Private/Materials/Material.cpp
+++ b/Engine/Source/Runtime/Engine/Private/Materials/Material.cpp
@@ -10,6 +10,7 @@
 #include "MaterialInstance.h"
 #include "UObjectAnnotation.h"
 #include "MaterialCompiler.h"
+#include "TargetPlatform.h"
 
 #if WITH_EDITOR
 #include "UnrealEd.h"
@@ -1933,6 +1934,24 @@ void UMaterial::PostLoad()
 	}
 }
 
+void UMaterial::BeginCacheForCookedPlatformData( const ITargetPlatform *TargetPlatform )
+{
+	TArray<FName> DesiredShaderFormats;
+	TargetPlatform->GetShaderFormats( DesiredShaderFormats );
+
+	if (DesiredShaderFormats.Num())
+	{
+		// Cache for all the shader formats that the cooking target requires
+		for (int32 FormatIndex = 0; FormatIndex < DesiredShaderFormats.Num(); FormatIndex++)
+		{
+			const EShaderPlatform TargetPlatform = ShaderFormatToLegacyShaderPlatform(DesiredShaderFormats[FormatIndex]);
+
+			// Begin caching shaders for the target platform and store the material resource being compiled into CachedMaterialResourcesForCooking
+			CacheResourceShadersForCooking(TargetPlatform, CachedMaterialResourcesForCooking);
+		}
+	}
+}
+
 #if WITH_EDITOR
 bool UMaterial::CanEditChange(const UProperty* InProperty) const
 {
diff --git a/Engine/Source/Runtime/Engine/Private/Materials/MaterialInstance.cpp b/Engine/Source/Runtime/Engine/Private/Materials/MaterialInstance.cpp
--- a/Engine/Source/Runtime/Engine/Private/Materials/MaterialInstance.cpp
+++ b/Engine/Source/Runtime/Engine/Private/Materials/MaterialInstance.cpp
@@ -4,6 +4,7 @@
 #include "EngineMaterialClasses.h"
 #include "MaterialInstance.h"
 #include "MaterialShader.h"
+#include "TargetPlatform.h"
 
 /**
  * Cache uniform expressions for the given material.
@@ -1487,6 +1488,19 @@ void TrimToOverriddenOnly(TArray<ParameterType>& Parameters)
 	}
 }
 
+void UMaterialInstance::BeginCacheForCookedPlatformData( const ITargetPlatform *TargetPlatform )
+{
+	TArray<FName> DesiredShaderFormats;
+	TargetPlatform->GetShaderFormats( DesiredShaderFormats );
+	// Cache shaders for each shader format, storing the results in CachedMaterialResourcesForCooking so they will be available during saving
+	for (int32 FormatIndex = 0; FormatIndex < DesiredShaderFormats.Num(); FormatIndex++)
+	{
+		const EShaderPlatform TargetPlatform = ShaderFormatToLegacyShaderPlatform(DesiredShaderFormats[FormatIndex]);
+
+		CacheResourceShadersForCooking(TargetPlatform, CachedMaterialResourcesForCooking);
+	}
+}
+
 void UMaterialInstance::Serialize(FArchive& Ar)
 {
 	Super::Serialize(Ar);
diff --git a/Engine/Source/Runtime/NetworkFileSystem/Private/NetworkFileServerConnection.cpp b/Engine/Source/Runtime/NetworkFileSystem/Private/NetworkFileServerConnection.cpp
--- a/Engine/Source/Runtime/NetworkFileSystem/Private/NetworkFileServerConnection.cpp
+++ b/Engine/Source/Runtime/NetworkFileSystem/Private/NetworkFileServerConnection.cpp
@@ -29,7 +29,13 @@ FNetworkFileServerClientConnection::FNetworkFileServerClientConnection( FSocket*
 		RecompileShadersDelegate = InRecompileShadersDelegate;
 	}
 
-	Thread = FRunnableThread::Create(this, TEXT("FNetworkFileServerClientConnection"), false, false, 8 * 1024, TPri_AboveNormal);
+#if UE_BUILD_DEBUG
+	// this thread needs more space in debug builds as it tries to log messages and such
+	const static uint32 NetworkFileServerThreadSize = 2 * 1024 * 1024; 
+#else
+	const static uint32 NetworkFileServerThreadSize = 8 * 1024; 
+#endif
+	Thread = FRunnableThread::Create(this, TEXT("FNetworkFileServerClientConnection"), false, false, NetworkFileServerThreadSize, TPri_AboveNormal);
 }
 
 
@@ -346,7 +352,7 @@ void FNetworkFileServerClientConnection::ProcessOpenFile( FArchive& In, FArchive
 	}
 
 	TArray<FString> NewUnsolictedFiles;
-	FileRequestDelegate.ExecuteIfBound(Filename, NewUnsolictedFiles);
+	FileRequestDelegate.ExecuteIfBound(Filename, ConnectedPlatformName, NewUnsolictedFiles);
 
 	FDateTime ServerTimeStamp = Sandbox->GetTimeStamp(*Filename);
 	int64 ServerFileSize = 0;
@@ -493,7 +499,7 @@ void FNetworkFileServerClientConnection::ProcessGetFileInfo( FArchive& In, FArch
 	if (Info.FileExists)
 	{
 		TArray<FString> NewUnsolictedFiles;
-		FileRequestDelegate.ExecuteIfBound(Filename, NewUnsolictedFiles);
+		FileRequestDelegate.ExecuteIfBound(Filename, ConnectedPlatformName, NewUnsolictedFiles);
 	}
 
 	// get the rest of the info
@@ -691,7 +697,7 @@ void FNetworkFileServerClientConnection::ProcessGetFileList( FArchive& In, FArch
 	ConnectedPlatformName = TEXT("");
 	// figure out the best matching target platform for the set of valid ones
 	ITargetPlatformManagerModule& TPM = GetTargetPlatformManagerRef();
-	static const TArray<ITargetPlatform*>& ActiveTargetPlatforms =  TPM.GetActiveTargetPlatforms();
+	static const TArray<ITargetPlatform*>& ActiveTargetPlatforms =  TPM.GetTargetPlatforms();
 	for (int32 TPIndex = 0; TPIndex < TargetPlatformNames.Num(); TPIndex++)
 	{
 		UE_LOG(LogFileServer, Display, TEXT("    Possible Target Platform from client: %s"), *TargetPlatformNames[TPIndex]);
@@ -949,7 +955,7 @@ void FNetworkFileServerClientConnection::ProcessSyncFile( FArchive& In, FArchive
 
 	TArray<FString> NewUnsolictedFiles;
 
-	FileRequestDelegate.ExecuteIfBound(Filename, NewUnsolictedFiles);
+	FileRequestDelegate.ExecuteIfBound(Filename, ConnectedPlatformName, NewUnsolictedFiles);
 
 	for (int32 Index = 0; Index < NewUnsolictedFiles.Num(); Index++)
 	{
diff --git a/Engine/Source/Runtime/NetworkFileSystem/Public/Interfaces/INetworkFileSystemModule.h b/Engine/Source/Runtime/NetworkFileSystem/Public/Interfaces/INetworkFileSystemModule.h
--- a/Engine/Source/Runtime/NetworkFileSystem/Public/Interfaces/INetworkFileSystemModule.h
+++ b/Engine/Source/Runtime/NetworkFileSystem/Public/Interfaces/INetworkFileSystemModule.h
@@ -13,7 +13,7 @@
  * The first parameter is the name of the requested file.
  * The second parameter will hold the list of unsolicited files to send back.
  */
-DECLARE_DELEGATE_TwoParams(FFileRequestDelegate, const FString&, TArray<FString>&);
+DECLARE_DELEGATE_ThreeParams(FFileRequestDelegate, const FString&, const FString&, TArray<FString>&);
 
 struct FShaderRecompileData
 {
