Disabled hit proxy update when rotating camera with trackpad gestures

diff --git a/Engine/Source/Editor/UnrealEd/Private/EditorViewportClient.cpp b/Engine/Source/Editor/UnrealEd/Private/EditorViewportClient.cpp
--- a/Engine/Source/Editor/UnrealEd/Private/EditorViewportClient.cpp
+++ b/Engine/Source/Editor/UnrealEd/Private/EditorViewportClient.cpp
@@ -2868,7 +2868,7 @@ void FEditorViewportClient::UpdateGestureDelta()
 	{
 		MoveViewportCamera( CurrentGestureDragDelta, CurrentGestureRotDelta, false );
 		
-		Invalidate();
+		Invalidate( true, false );
 		
 		CurrentGestureDragDelta = FVector::ZeroVector;
 		CurrentGestureRotDelta = FRotator::ZeroRotator;
