TTP #331140: Preventing the BP rename automation test from renaming blueprints that are externally referenced.

diff --git a/Engine/Source/Editor/UnrealEd/Private/Tests/BlueprintAutomationTests.cpp b/Engine/Source/Editor/UnrealEd/Private/Tests/BlueprintAutomationTests.cpp
--- a/Engine/Source/Editor/UnrealEd/Private/Tests/BlueprintAutomationTests.cpp
+++ b/Engine/Source/Editor/UnrealEd/Private/Tests/BlueprintAutomationTests.cpp
@@ -1095,6 +1095,13 @@ bool FBlueprintRenameAndCloneTest::RunTest(const FString& BlueprintAssetPath)
 		AddError(FString::Printf(TEXT("You have unsaved changes made to '%s', please save them before running this test."), *BlueprintAssetPath));
 	}
 
+	bool bIsAlreadyLoaded = false;
+	if (FBlueprintAutomationTestUtilities::IsBlueprintLoaded(BlueprintAssetPath))
+	{
+		bIsAlreadyLoaded = true;
+		AddWarning(FString::Printf(TEXT("'%s' is already loaded, and possibly referenced by external objects (unable to perform rename tests... please run again in an empty map)."), *BlueprintAssetPath));
+	}
+
 	UBlueprint* const OriginalBlueprint = Cast<UBlueprint>(StaticLoadObject(UBlueprint::StaticClass(), NULL, *BlueprintAssetPath));
 	if (OriginalBlueprint == NULL)
 	{
@@ -1115,6 +1122,7 @@ bool FBlueprintRenameAndCloneTest::RunTest(const FString& BlueprintAssetPath)
 		}
 
 		// rename
+		if (!bIsAlreadyLoaded)
 		{
 			// store the original package so we can manually invalidate it after the move
 			UPackage* const OriginalPackage = Cast<UPackage>(OriginalBlueprint->GetOutermost());
@@ -1136,9 +1144,10 @@ bool FBlueprintRenameAndCloneTest::RunTest(const FString& BlueprintAssetPath)
 			// the blueprint has been moved out of this package, invalidate it so 
 			// we don't save it in this state and so we can reload the blueprint later
 			FBlueprintAutomationTestUtilities::InvalidatePackage(OriginalPackage);
+
+			FBlueprintAutomationTestUtilities::UnloadBlueprint(OriginalBlueprint);
 		}
 
-		FBlueprintAutomationTestUtilities::UnloadBlueprint(OriginalBlueprint);
 		// make sure the unloaded blueprints are properly flushed (for future tests)
 		CollectGarbage(RF_Native);
 	}
