#UE4:
* fixed 2D LevelGrid major lines rendering for power of two snap sizes
  (3D not done yet)

diff --git a/Engine/Source/Editor/UnrealEd/Classes/Editor/EditorEngine.h b/Engine/Source/Editor/UnrealEd/Classes/Editor/EditorEngine.h
--- a/Engine/Source/Editor/UnrealEd/Classes/Editor/EditorEngine.h
+++ b/Engine/Source/Editor/UnrealEd/Classes/Editor/EditorEngine.h
@@ -1898,6 +1898,11 @@ public:
 	 */
 	float GetGridSize();
 
+	/** 
+	 * @return - if the grid size is part of the 1,2,48,16,.. list or not
+	 */
+	bool IsGridSizePowerOfTwo() const;
+
 	/** 
 	 * Sets the selected positional snap grid setting
 	 *
diff --git a/Engine/Source/Editor/UnrealEd/Private/EditorComponents.cpp b/Engine/Source/Editor/UnrealEd/Private/EditorComponents.cpp
--- a/Engine/Source/Editor/UnrealEd/Private/EditorComponents.cpp
+++ b/Engine/Source/Editor/UnrealEd/Private/EditorComponents.cpp
@@ -502,7 +502,7 @@ void FEditorCommonDrawHelper::DrawGridSection(float ViewportGridY,FVector* A,FVe
 	}
 
 	// todo
-	static int32 Exponent = 10;	// or 2?
+	static int32 Exponent = GEditor->IsGridSizePowerOfTwo() ? 8 : 10;
 
 	const FMatrix InvViewProjMatrix = View->ViewMatrices.ProjMatrix.Inverse() * View->ViewMatrices.ViewMatrix.Inverse();
 	int32 FirstLine = FMath::Trunc(InvViewProjMatrix.TransformPosition(FVector(-1,-1,0.5f)).Component(Axis) / ViewportGridY);
@@ -536,7 +536,7 @@ void FEditorCommonDrawHelper::DrawGridSection(float ViewportGridY,FVector* A,FVe
 	float AlphaA = FMath::Clamp(Fract * InvTransitionRegion, 0.0f, 1.0f);
 	float AlphaB = FMath::Clamp(InvTransitionRegion - Fract * InvTransitionRegion, 0.0f, 1.0f);
 
-	if(IncValue < -1.0f)
+	if(IncValue < -0.5f)
 	{
 		// no fade in magnification case
 		AlphaA = 1.0f;
diff --git a/Engine/Source/Editor/UnrealEd/Private/EditorConstraints.cpp b/Engine/Source/Editor/UnrealEd/Private/EditorConstraints.cpp
--- a/Engine/Source/Editor/UnrealEd/Private/EditorConstraints.cpp
+++ b/Engine/Source/Editor/UnrealEd/Private/EditorConstraints.cpp
@@ -1,7 +1,7 @@
 // Copyright 1998-2014 Epic Games, Inc. All Rights Reserved.
 
 /*=============================================================================
-	EditorConstraints.cpp: Editor movement contraints.
+	EditorConstraints.cpp: Editor movement constraints.
 =============================================================================*/
 
 #include "UnrealEd.h"
@@ -13,6 +13,13 @@ float UEditorEngine::GetGridSize()
 	return PosGridSizes[GetDefault<ULevelEditorViewportSettings>()->CurrentPosGridSize ];
 }
 
+bool UEditorEngine::IsGridSizePowerOfTwo() const
+{
+	const ULevelEditorViewportSettings* ViewportSettings = GetDefault<ULevelEditorViewportSettings>();
+
+	return ViewportSettings->bUsePowerOf2SnapSize;
+}
+
 void UEditorEngine::SetGridSize( int32 InIndex )
 {
 	FinishAllSnaps();
