#UE4 - Add option to ignore floor snapping when dragging from the content browser

diff --git a/Engine/Config/BaseEditorUserSettings.ini b/Engine/Config/BaseEditorUserSettings.ini
--- a/Engine/Config/BaseEditorUserSettings.ini
+++ b/Engine/Config/BaseEditorUserSettings.ini
@@ -307,6 +307,7 @@ bUsePowerOf2SnapSize=False
 GridEnabled=True
 RotGridEnabled=True
 SnapScaleEnabled=True
+bSnapNewObjectsToFloor=True
 ; If enabled, use the old-style multiplicative/percentage scaling method instead of the new additive/fraction method
 bUsePercentageBasedScaling=False
 ; If true actor snap will be enabled in the editor
diff --git a/Engine/Source/Editor/UnrealEd/Classes/Settings/LevelEditorViewportSettings.h b/Engine/Source/Editor/UnrealEd/Classes/Settings/LevelEditorViewportSettings.h
--- a/Engine/Source/Editor/UnrealEd/Classes/Settings/LevelEditorViewportSettings.h
+++ b/Engine/Source/Editor/UnrealEd/Classes/Settings/LevelEditorViewportSettings.h
@@ -248,6 +248,9 @@ public:
 	UPROPERTY(EditAnywhere, config, Category=GridSnapping, meta=(DisplayName = "Enable Scale Snapping"))
 	uint32 SnapScaleEnabled:1;
 
+	/** If enabled the when dragging new objects out of the content browser, it will snap the objects Z coordinate to the floor below it (if any) instead of the Z grid snapping location */
+	UPROPERTY(EditAnywhere, config, Category=GridSnapping)
+	bool bSnapNewObjectsToFloor;
 private:
 
 	/** If enabled, use the old-style multiplicative/percentage scaling method instead of the new additive/fraction method */
diff --git a/Engine/Source/Editor/UnrealEd/Private/LevelEditorViewport.cpp b/Engine/Source/Editor/UnrealEd/Private/LevelEditorViewport.cpp
--- a/Engine/Source/Editor/UnrealEd/Private/LevelEditorViewport.cpp
+++ b/Engine/Source/Editor/UnrealEd/Private/LevelEditorViewport.cpp
@@ -258,12 +258,16 @@ static TArray<AActor*> AttemptDropObjAsActors( ULevel* InLevel, UObject* ObjToUs
 	{
 		if (bUsedHitProxy)
 		{
+			const bool bSnapNewObjectsToFloor = GetDefault<ULevelEditorViewportSettings>()->bSnapNewObjectsToFloor;
 			for (auto ActorIt = PlacedActors.CreateConstIterator(); ActorIt; ++ActorIt)
 			{
 				AActor* Actor = *ActorIt;
 
 				FVector Loc = Actor->GetActorLocation();
-				Loc.Z = PreSnapLocation.Z;
+				if( bSnapNewObjectsToFloor )
+				{
+					Loc.Z = PreSnapLocation.Z;
+				}
 
 				Actor->TeleportTo(Loc, Actor->GetActorRotation(), false, true);
 			}
@@ -1327,13 +1331,15 @@ bool FLevelEditorViewportClient::UpdateDropPreviewActors(int32 MouseX, int32 Mou
 		}
 		else
 		{
+			const bool bSnapNewObjectsToFloor = GetDefault<ULevelEditorViewportSettings>()->bSnapNewObjectsToFloor;
+
 			// Move the actor to the target location while preserving the relative offset to the other actors
 			FVector Collision = FirstDraggingActor->GetPlacementExtent();
 			FVector PreSnapLocation = MouseLocation + MousePlane * (FVector::BoxPushOut(MousePlane, Collision) + 0.1f);
 			FVector NewLocation = PreSnapLocation;
 			FSnappingUtils::SnapPointToGrid( NewLocation, FVector(0, 0, 0) );
 
-			if( HitProxy != NULL && Cursor.GetViewportType() == LVT_Perspective )
+			if( bSnapNewObjectsToFloor && HitProxy != NULL && Cursor.GetViewportType() == LVT_Perspective )
 			{
 				NewLocation.Z = PreSnapLocation.Z;
 			}
