Improvements to the Curve Editor
* Numberic boxes now used for Time/Value.
* Useless index number on item no longer displayed in text field.
* Numberic boxes no longer cover the numbers on the side.
* The numberic boxes are no longer transparent.
* There's an option in the right click menu to "Add key", for users who don't know the mystical Shift+Click.

Improvements to the Timeline Editor
* The Nominal Height is now 300.  100 Was way too small to be functional.

diff --git a/Engine/Source/Editor/Kismet/Private/STimelineEditor.cpp b/Engine/Source/Editor/Kismet/Private/STimelineEditor.cpp
--- a/Engine/Source/Editor/Kismet/Private/STimelineEditor.cpp
+++ b/Engine/Source/Editor/Kismet/Private/STimelineEditor.cpp
@@ -769,7 +769,7 @@ void STimelineEdTrack::CopyCurveData( const FRichCurve* SrcCurve, FRichCurve* De
 
 void STimelineEdTrack::ResetExternalCurveInfo( )
 {
-	ExternalCurveName = FString( TEXT("None") );
+	ExternalCurveName = FString( TEXT( "None" ) );
 	ExternalCurvePath = FString( TEXT( "None" ) );
 }
 
@@ -784,7 +784,7 @@ void STimelineEditor::Construct(const FArguments& InArgs, TSharedPtr<FBlueprintE
 	Kismet2Ptr = InKismet2;
 	TimelineObj = NULL;
 
-	NominalTimelineDesiredHeight = 100.0f;
+	NominalTimelineDesiredHeight = 300.0f;
 	TimelineDesiredSize = FVector2D(128.0f, NominalTimelineDesiredHeight);
 
 	// Leave these uninitialized at first.  We'll zoom to fit the tracks which will set the correct values
diff --git a/Engine/Source/Editor/UnrealEd/Private/SCurveEditor.cpp b/Engine/Source/Editor/UnrealEd/Private/SCurveEditor.cpp
--- a/Engine/Source/Editor/UnrealEd/Private/SCurveEditor.cpp
+++ b/Engine/Source/Editor/UnrealEd/Private/SCurveEditor.cpp
@@ -39,6 +39,8 @@ void SCurveEditor::Construct(const FArguments& InArgs)
 	bZoomToFit = InArgs._ZoomToFit;
 	DesiredSize = InArgs._DesiredSize;
 
+	bIsUsingSlider = false;
+
 	// if editor size is set, use it, otherwise, use default value
 	if (DesiredSize.Get().IsZero())
 	{
@@ -91,12 +93,14 @@ void SCurveEditor::Construct(const FArguments& InArgs)
 		[
 			SNew(SHorizontalBox)
 			.Visibility( this, &SCurveEditor::GetCurveAreaVisibility )
-			+SHorizontalBox::Slot()
+			
+			+ SHorizontalBox::Slot()
 			.AutoWidth()
 			[
 				CurveSelector
 			]
-			+SHorizontalBox::Slot()
+
+			+ SHorizontalBox::Slot()
 			.AutoWidth()
 			[
 				SNew(SBorder)
@@ -105,9 +109,11 @@ void SCurveEditor::Construct(const FArguments& InArgs)
 				.BorderImage( FEditorStyle::GetBrush("NoBorder") )
 				.DesiredSizeScale(FVector2D(256.0f,32.0f))
 				.Visibility(this, &SCurveEditor::GetControlVisibility)
+				.Padding(FMargin(18, 12))
 				[
 					SNew(SHorizontalBox)
-					+SHorizontalBox::Slot()
+					
+					+ SHorizontalBox::Slot()
 					.AutoWidth()
 					[
 						SNew(SButton)
@@ -120,7 +126,8 @@ void SCurveEditor::Construct(const FArguments& InArgs)
 							.ColorAndOpacity( FSlateColor::UseForeground() ) 
 						]
 					]
-					+SHorizontalBox::Slot()
+					
+					+ SHorizontalBox::Slot()
 					.AutoWidth()
 					[
 						SNew(SButton)
@@ -133,58 +140,71 @@ void SCurveEditor::Construct(const FArguments& InArgs)
 							.ColorAndOpacity( FSlateColor::UseForeground() ) 
 						]
 					]
-					+SHorizontalBox::Slot()
-					.AutoWidth()
-					[
-						SNew(SBorder)
-						.Visibility(this, &SCurveEditor::GetEditVisibility)
-						.VAlign(VAlign_Center)
-						[
-							SNew(SEditableText)
-							.MinDesiredWidth(64.0f)
-							.IsReadOnly(true)
-							.IsEnabled(true)
-							.Font(FEditorStyle::GetFontStyle("CurveEd.InfoFont"))
-							.ColorAndOpacity( FLinearColor::White )
-							.Text(this, &SCurveEditor::GetKeyInfoText)
-						]
-					]
-					+SHorizontalBox::Slot()
+					
+					+ SHorizontalBox::Slot()
 					.AutoWidth()
 					[
 						SNew(SBorder)
+						.BorderImage(FEditorStyle::GetBrush("NoBorder"))
 						.Visibility(this, &SCurveEditor::GetEditVisibility)
 						.VAlign(VAlign_Center)
 						[
-							SNew(SEditableText)
-							.MinDesiredWidth(64.0f)
+							SNew(SNumericEntryBox<float>)
 							.IsEnabled(this, &SCurveEditor::GetInputEditEnabled)
 							.Font(FEditorStyle::GetFontStyle("CurveEd.InfoFont"))
-							.ColorAndOpacity( FLinearColor::White )
-							.SelectAllTextWhenFocused(true)
-							.Text(this, &SCurveEditor::GetInputInfoText)
-							.OnTextCommitted(this, &SCurveEditor::NewTimeEntered)
+							.Value(this, &SCurveEditor::OnGetTime)
+							.UndeterminedString(LOCTEXT("MultipleValues", "Multiple Values").ToString())
+							.OnValueCommitted(this, &SCurveEditor::OnTimeComitted)
+							.OnValueChanged(this, &SCurveEditor::OnTimeChanged)
+							.OnBeginSliderMovement(this, &SCurveEditor::OnBeginSliderMovement, LOCTEXT("SetValue", "Set New Time"))
+							.OnEndSliderMovement(this, &SCurveEditor::OnEndSliderMovement)
+							.LabelVAlign(VAlign_Center)
+							.AllowSpin(true)
+							.MinValue(TOptional<float>())
+							.MaxValue(TOptional<float>())
+							.MaxSliderValue(TOptional<float>())
+							.MinSliderValue(TOptional<float>())
+							.Label()
+							[
+								SNew(STextBlock)
+								.Text(LOCTEXT("Time", "Time"))
+							]
 						]
 					]
-					+SHorizontalBox::Slot()
+					
+					+ SHorizontalBox::Slot()
 					.AutoWidth()
 					[
 						SNew(SBorder)
+						.BorderImage( FEditorStyle::GetBrush("NoBorder") )
 						.Visibility(this, &SCurveEditor::GetEditVisibility)
 						.VAlign(VAlign_Center)
 						[
-							SNew(SEditableText)
-							.MinDesiredWidth(64.0f)
+							SNew(SNumericEntryBox<float>)
 							.Font(FEditorStyle::GetFontStyle("CurveEd.InfoFont"))
-							.ColorAndOpacity( FLinearColor::White )
-							.SelectAllTextWhenFocused(true)
-							.Text(this, &SCurveEditor::GetOutputInfoText)
-							.OnTextCommitted(this, &SCurveEditor::NewValueEntered)
+							.Value(this, &SCurveEditor::OnGetValue)
+							.UndeterminedString(LOCTEXT("MultipleValues", "Multiple Values").ToString())
+							.OnValueCommitted(this, &SCurveEditor::OnValueComitted)
+							.OnValueChanged(this, &SCurveEditor::OnValueChanged)
+							.OnBeginSliderMovement(this, &SCurveEditor::OnBeginSliderMovement, LOCTEXT("SetValue", "Set New Value"))
+							.OnEndSliderMovement(this, &SCurveEditor::OnEndSliderMovement)
+							.LabelVAlign(VAlign_Center)
+							.AllowSpin(true)
+							.MinValue(TOptional<float>())
+							.MaxValue(TOptional<float>())
+							.MaxSliderValue(TOptional<float>())
+							.MinSliderValue(TOptional<float>())
+							.Label()
+							[
+								SNew(STextBlock)
+								.Text(LOCTEXT("Value", "Value"))
+							]
 						]
 					]
 				]
 			]
 		]
+
 		+ SVerticalBox::Slot()
 		.VAlign(VAlign_Bottom)
 		.FillHeight(.75f)
@@ -913,38 +933,7 @@ FReply SCurveEditor::OnMouseButtonDown( const FGeometry& InMyGeometry, const FPo
 				// shift click on background = add new key
 				if(bShiftDown)
 				{
-					if(SelectedCurves.Num() == 0)
-					{
-						if(Curves.Num() == 1)
-						{
-							SelectCurve(Curves[0].CurveToEdit, false);
-						}
-						else
-						{
-							PushWarningMenu(InMouseEvent.GetScreenSpacePosition(), LOCTEXT("NoCurveSelected", "Please select a curve"));
-						}						
-					}
-					
-					for(auto It = SelectedCurves.CreateIterator();It;++It)
-					{
-						FRichCurve* SelectedCurve = *It;
-						if(IsValidCurve(SelectedCurve))
-						{
-							const FScopedTransaction Transaction( LOCTEXT("CurveEditor_AddKey", "Add Key(s)") );
-							CurveOwner->ModifyOwner();
-
-							FTrackScaleInfo ScaleInfo(ViewMinInput.Get(),  ViewMaxInput.Get(), ViewMinOutput, ViewMaxOutput, InMyGeometry.Size);
-
-							FVector2D LocalClickPos = InMyGeometry.AbsoluteToLocal(InMouseEvent.GetScreenSpacePosition());
-
-							float NewInput = ScaleInfo.LocalXToInput( LocalClickPos.X );
-							float NewOutput = ScaleInfo.LocalYToOutput( LocalClickPos.Y );
-							FKeyHandle NewKeyHandle = SelectedCurve->AddKey(NewInput, NewOutput);
-
-							EmptySelection();
-							AddToSelection(FSelectedCurveKey(SelectedCurve,NewKeyHandle));
-						}
-					}
+					AddNewKey(InMyGeometry, InMouseEvent.GetScreenSpacePosition());
 				}
 				else
 				{
@@ -967,7 +956,6 @@ FReply SCurveEditor::OnMouseButtonDown( const FGeometry& InMyGeometry, const FPo
 							HitTestCurves(InMyGeometry, InMouseEvent);
 						}
 					}
-		
 				}
 			}
 			else if(bRightMouseButton)
@@ -989,6 +977,42 @@ FReply SCurveEditor::OnMouseButtonDown( const FGeometry& InMyGeometry, const FPo
 	return FReply::Unhandled();
 }
 
+void SCurveEditor::AddNewKey(FGeometry InMyGeometry, FVector2D ScreenPosition)
+{
+	if ( SelectedCurves.Num() == 0 )
+	{
+		if ( Curves.Num() == 1 )
+		{
+			SelectCurve(Curves[0].CurveToEdit, false);
+		}
+		else
+		{
+			PushWarningMenu(ScreenPosition, LOCTEXT("NoCurveSelected", "Please select a curve"));
+		}
+	}
+
+	for ( auto It = SelectedCurves.CreateIterator(); It; ++It )
+	{
+		FRichCurve* SelectedCurve = *It;
+		if ( IsValidCurve(SelectedCurve) )
+		{
+			const FScopedTransaction Transaction(LOCTEXT("CurveEditor_AddKey", "Add Key(s)"));
+			CurveOwner->ModifyOwner();
+
+			FTrackScaleInfo ScaleInfo(ViewMinInput.Get(), ViewMaxInput.Get(), ViewMinOutput, ViewMaxOutput, InMyGeometry.Size);
+
+			FVector2D LocalClickPos = InMyGeometry.AbsoluteToLocal(ScreenPosition);
+
+			float NewInput = ScaleInfo.LocalXToInput(LocalClickPos.X);
+			float NewOutput = ScaleInfo.LocalYToOutput(LocalClickPos.Y);
+			FKeyHandle NewKeyHandle = SelectedCurve->AddKey(NewInput, NewOutput);
+
+			EmptySelection();
+			AddToSelection(FSelectedCurveKey(SelectedCurve, NewKeyHandle));
+		}
+	}
+}
+
 void SCurveEditor::ResetDragValues()
 {
 	TotalDragDist = 0.f;
@@ -1032,9 +1056,9 @@ FReply SCurveEditor::OnMouseButtonUp( const FGeometry& InMyGeometry, const FPoin
 			}
 			else if(
 				!HitTestCubicTangents(InMyGeometry, ScreenPos).IsValid() && 
-				!HitTestKeys(InMyGeometry,ScreenPos).IsValid())
+				!HitTestKeys(InMyGeometry, ScreenPos).IsValid())
 			{
-				CreateContextMenu();
+				CreateContextMenu(InMyGeometry, ScreenPos);
 			}
 		}
 
@@ -1158,20 +1182,87 @@ FReply SCurveEditor::OnKeyDown( const FGeometry& MyGeometry, const FKeyboardEven
 	}
 }
 
-void SCurveEditor::NewValueEntered( const FText& NewText, ETextCommit::Type CommitInfo )
+TOptional<float> SCurveEditor::OnGetTime() const
 {
-	// Don't digest the number if we just clicked away from the popup
-	if ((CommitInfo == ETextCommit::OnEnter) || (CommitInfo == ETextCommit::OnUserMovedFocus))
+	if ( SelectedKeys.Num() == 1 )
 	{
-		float NewValue = FCString::Atof( *NewText.ToString() );
+		return GetKeyTime(SelectedKeys[0]);
+	}
+
+	// Value couldn't be accessed.  Return an unset value
+	return TOptional<float>();
+}
 
+void SCurveEditor::OnTimeComitted(float NewTime, ETextCommit::Type CommitType)
+{
+	// Don't digest the number if we just clicked away from the pop-up
+	if ( !bIsUsingSlider && ((CommitType == ETextCommit::OnEnter) || ( CommitType == ETextCommit::OnUserMovedFocus )) )
+	{
+		if ( SelectedKeys.Num() >= 1 )
+		{
+			auto Key = SelectedKeys[0];
+			if ( IsValidCurve(Key.Curve) )
+			{
+				const FScopedTransaction Transaction(LOCTEXT("CurveEditor_NewTime", "New Time Entered"));
+				CurveOwner->ModifyOwner();
+				Key.Curve->SetKeyTime(Key.KeyHandle, NewTime);
+			}
+		}
+
+		FSlateApplication::Get().DismissAllMenus();
+	}
+}
+
+void SCurveEditor::OnTimeChanged(float NewTime)
+{
+	if ( bIsUsingSlider )
+	{
+		if ( SelectedKeys.Num() >= 1 )
+		{
+			auto Key = SelectedKeys[0];
+			if ( IsValidCurve(Key.Curve) )
+			{
+				CurveOwner->ModifyOwner();
+				Key.Curve->SetKeyTime(Key.KeyHandle, NewTime);
+			}
+		}
+	}
+}
+
+TOptional<float> SCurveEditor::OnGetValue() const
+{
+	TOptional<float> Value;
+
+	// Return the value string if all selected keys have the same output string, otherwise empty
+	if ( SelectedKeys.Num() > 0 )
+	{
+		Value = GetKeyValue(SelectedKeys[0]);
+		for ( int32 i=1; i < SelectedKeys.Num(); i++ )
+		{
+			TOptional<float> NewValue = GetKeyValue(SelectedKeys[i]);
+			bool bAreEqual = ( ( !Value.IsSet() && !NewValue.IsSet() ) || ( Value.IsSet() && NewValue.IsSet() && Value.GetValue() == NewValue.GetValue() ) );
+			if ( !bAreEqual )
+			{
+				return TOptional<float>();
+			}
+		}
+	}
+
+	return Value;
+}
+
+void SCurveEditor::OnValueComitted(float NewValue, ETextCommit::Type CommitType)
+{
+	// Don't digest the number if we just clicked away from the popup
+	if ( !bIsUsingSlider && ((CommitType == ETextCommit::OnEnter) || ( CommitType == ETextCommit::OnUserMovedFocus )) )
+	{
 		// Iterate over selected set
-		for(int32 i=0; i<SelectedKeys.Num(); i++)
+		for ( int32 i=0; i < SelectedKeys.Num(); i++ )
 		{
 			auto Key = SelectedKeys[i];
-			if(IsValidCurve(Key.Curve))
+			if ( IsValidCurve(Key.Curve) )
 			{
-				const FScopedTransaction Transaction( LOCTEXT("CurveEditor_NewValue", "New Value Entered") );
+				const FScopedTransaction Transaction(LOCTEXT("CurveEditor_NewValue", "New Value Entered"));
 				CurveOwner->ModifyOwner();
 
 				// Fill in each element of this key
@@ -1183,28 +1274,41 @@ void SCurveEditor::NewValueEntered( const FText& NewText, ETextCommit::Type Comm
 	}
 }
 
-void SCurveEditor::NewTimeEntered( const FText& NewText, ETextCommit::Type CommitInfo )
+void SCurveEditor::OnValueChanged(float NewValue)
 {
-	// Don't digest the number if we just clicked away from the pop-up
-	if ((CommitInfo == ETextCommit::OnEnter) || (CommitInfo == ETextCommit::OnUserMovedFocus))
+	if ( bIsUsingSlider )
 	{
-		float NewTime = FCString::Atof( *NewText.ToString() );
+		const FScopedTransaction Transaction(LOCTEXT("CurveEditor_NewValue", "New Value Entered"));
 
-		if(SelectedKeys.Num() >= 1)
+		// Iterate over selected set
+		for ( int32 i=0; i < SelectedKeys.Num(); i++ )
 		{
-			auto Key = SelectedKeys[0];
-			if(IsValidCurve(Key.Curve))
+			auto Key = SelectedKeys[i];
+			if ( IsValidCurve(Key.Curve) )
 			{
-				const FScopedTransaction Transaction( LOCTEXT("CurveEditor_NewTime", "New Time Entered") );
 				CurveOwner->ModifyOwner();
-				Key.Curve->SetKeyTime(Key.KeyHandle, NewTime);
+
+				// Fill in each element of this key
+				Key.Curve->SetKeyValue(Key.KeyHandle, NewValue);
 			}
 		}
-		
-		FSlateApplication::Get().DismissAllMenus();
 	}
 }
 
+void SCurveEditor::OnBeginSliderMovement(FText TransactionName)
+{
+	bIsUsingSlider = true;
+
+	GEditor->BeginTransaction(TransactionName);
+}
+
+void SCurveEditor::OnEndSliderMovement(float NewValue)
+{
+	bIsUsingSlider = false;
+
+	GEditor->EndTransaction();
+}
+
 
 SCurveEditor::FSelectedCurveKey SCurveEditor::HitTestKeys(const FGeometry& InMyGeometry, const FVector2D& HitScreenPosition)
 {	
@@ -1290,28 +1394,24 @@ void SCurveEditor::MoveSelectedKeysByDelta(FVector2D InputDelta)
 	}
 }
 
-FText SCurveEditor::GetKeyValueText(FSelectedCurveKey Key) const
+TOptional<float> SCurveEditor::GetKeyValue(FSelectedCurveKey Key) const
 {
-	FText ValueText;
 	if(IsValidCurve(Key.Curve))
 	{
-		float Value = Key.Curve->GetKeyValue(Key.KeyHandle);
-		ValueText = FText::AsNumber( Value );
+		return Key.Curve->GetKeyValue(Key.KeyHandle);
 	}
 
-	return ValueText;
+	return TOptional<float>();
 }
 
-FText SCurveEditor::GetKeyTimeText(FSelectedCurveKey Key) const
+TOptional<float> SCurveEditor::GetKeyTime(FSelectedCurveKey Key) const
 {
-	FText TimeText;
 	if ( IsValidCurve(Key.Curve) )
 	{
-		float Time = Key.Curve->GetKeyTime(Key.KeyHandle);
-		TimeText = FText::AsNumber( Time );
+		return Key.Curve->GetKeyTime(Key.KeyHandle);
 	}
 
-	return TimeText;
+	return TOptional<float>();
 }
 
 void SCurveEditor::EmptySelection()
@@ -1448,52 +1548,23 @@ FReply SCurveEditor::ZoomToFitVertical()
 	return FReply::Handled();
 }
 
-FText SCurveEditor::GetInputInfoText() const
-{
-	if ( SelectedKeys.Num() == 1 )
-	{
-		return GetKeyTimeText( SelectedKeys[0] );
-	}
-
-	return LOCTEXT("NoInputInfo", "---");
-}
-
-FText SCurveEditor::GetOutputInfoText() const
+void SCurveEditor::CreateContextMenu(const FGeometry& InMyGeometry, const FVector2D& ScreenPosition)
 {
-	// Return the value string if all selected keys have the same output string, otherwise empty
-	FText OutputText;
-	if ( SelectedKeys.Num() > 0 )
-	{
-		OutputText = GetKeyValueText(SelectedKeys[0]);
-		for(int32 i=1; i<SelectedKeys.Num(); i++)
-		{
-			FText KeyOutput = GetKeyValueText( SelectedKeys[i] );
-			if ( !KeyOutput.EqualTo( OutputText ) )
-			{
-				OutputText = FText::GetEmpty();
-			}
-		}
-	}
-
-	return OutputText;
-}
+	const bool CloseAfterSelection = true;
+	FMenuBuilder MenuBuilder( CloseAfterSelection, NULL );
 
-FText SCurveEditor::GetKeyInfoText() const
-{
-	if ( SelectedKeys.Num() == 1 && SelectedKeys[0].Curve->IsKeyHandleValid(SelectedKeys[0].KeyHandle) )
-	{
-		return FText::Format(LOCTEXT("KeyInfo", "Index {0}"), FText::AsNumber( SelectedKeys[0].Curve->GetIndexSafe(SelectedKeys[0].KeyHandle) ) );
-	}
-	else
+	MenuBuilder.BeginSection("EditCurveEditorActions", LOCTEXT("Actions", "Actions"));
 	{
-		return LOCTEXT("NoKeyInfo", "---");
+		FUIAction Action = FUIAction(FExecuteAction::CreateSP(this, &SCurveEditor::AddNewKey, InMyGeometry, ScreenPosition));
+		MenuBuilder.AddMenuEntry
+		(
+			LOCTEXT("AddKey", "Add Key"),
+			LOCTEXT("AddKey_ToolTip", "Create a new point in the curve.  You can also use (Shift + Left Mouse Button)"),
+			FSlateIcon(),
+			Action
+		);
 	}
-}
-
-void SCurveEditor::CreateContextMenu()
-{
-	const bool CloseAfterSelection = false;
-	FMenuBuilder MenuBuilder( CloseAfterSelection, NULL );
+	MenuBuilder.EndSection();
 
 	MenuBuilder.BeginSection("CurveEditorActions", LOCTEXT("CurveAction", "Curve Actions") );
 	{
diff --git a/Engine/Source/Editor/UnrealEd/Public/SCurveEditor.h b/Engine/Source/Editor/UnrealEd/Public/SCurveEditor.h
--- a/Engine/Source/Editor/UnrealEd/Public/SCurveEditor.h
+++ b/Engine/Source/Editor/UnrealEd/Public/SCurveEditor.h
@@ -194,7 +194,11 @@ private:
 		/** Indicates if it is the arrival tangent, or the leave tangent */
 		bool			  bIsArrival;
 	};
+
 private:
+	/** Adds a new key to the curve. */
+	void AddNewKey(FGeometry InMyGeometry, FVector2D ScreenPosition);
+
 	/** Get the color associated with a given curve */
 	FLinearColor	GetCurveColor(FRichCurve* Curve) const;
 
@@ -209,8 +213,6 @@ private:
 
 	/** Called when new value for a key is entered */
 	void NewValueEntered( const FText& NewText, ETextCommit::Type CommitInfo );
-	/** Called when a new time for a key is entered */
-	void NewTimeEntered( const FText& NewText, ETextCommit::Type CommitInfo );
 
 	void NewHorizontalGridScaleEntered(const FString& NewText, bool bCommitFromEnter );
 	void NewVerticalGridScaleEntered(const FString& NewText, bool bCommitFromEnter );
@@ -239,9 +241,9 @@ private:
 	bool AreKeysSelected() const;
 
 	/** Get the value of the desired key as text */
-	FText GetKeyValueText(FSelectedCurveKey Key) const;
-	/** Get the time of the desired key as text */
-	FText GetKeyTimeText(FSelectedCurveKey Key) const;
+	TOptional<float> GetKeyValue(FSelectedCurveKey Key) const;
+	/** Get the time of the desired key */
+	TOptional<float> GetKeyTime(FSelectedCurveKey Key) const;
 
 	/** Move the selected keys by the supplied change in input (X) */
 	void MoveSelectedKeysByDelta(FVector2D InputDelta);
@@ -255,9 +257,17 @@ private:
 	FReply ZoomToFitHorizontal();
 	FReply ZoomToFitVertical();
 
-	FText GetInputInfoText() const;
-	FText GetOutputInfoText() const;
-	FText GetKeyInfoText() const;
+
+	TOptional<float> OnGetTime() const;
+	void OnTimeComitted(float NewValue, ETextCommit::Type CommitType);
+	void OnTimeChanged(float NewValue);
+
+	TOptional<float> OnGetValue() const;
+	void OnValueComitted(float NewValue, ETextCommit::Type CommitType);
+	void OnValueChanged(float NewValue);
+
+	void OnBeginSliderMovement(FText TransactionName);
+	void OnEndSliderMovement(float NewValue);
 
 	EVisibility GetCurveAreaVisibility() const;
 	EVisibility GetControlVisibility() const;
@@ -267,7 +277,7 @@ private:
 	bool GetInputEditEnabled() const;
 
 	/** Function to create context menu on mouse right click*/
-	void CreateContextMenu();
+	void CreateContextMenu(const FGeometry& InMyGeometry, const FVector2D& ScreenPosition);
 
 	/** Callback function called when item is select in the context menu */
 	void OnCreateExternalCurveClicked();
@@ -374,7 +384,9 @@ private:
 	bool AreCurvesVisible() const { return bAlwaysDisplayColorCurves || bAreCurvesVisible; }
 	bool IsGradientEditorVisible() const { return bIsGradientEditorVisible; }
 	bool IsLinearColorCurve() const;
+
 protected:
+
 	/** Set Default output values when range is too small **/
 	UNREALED_API virtual void SetDefaultOutput(const float MinZoomRange);
 	/** Get Time Step for vertical line drawing **/
@@ -473,6 +485,9 @@ protected:
 
 	/** True if you want the curve editor to fit to zoom **/
 	bool				bZoomToFit;
+
+	/** True if the sliders are being used to adjust point values **/
+	bool bIsUsingSlider;
 };
 
 #endif // __SCurveEditor_h__
