#UE4 - Invert middle mouse pan option no longer inverts left+right mouse pan

diff --git a/Engine/Source/Editor/UnrealEd/Private/EditorViewportClient.cpp b/Engine/Source/Editor/UnrealEd/Private/EditorViewportClient.cpp
--- a/Engine/Source/Editor/UnrealEd/Private/EditorViewportClient.cpp
+++ b/Engine/Source/Editor/UnrealEd/Private/EditorViewportClient.cpp
@@ -557,9 +557,6 @@ FSceneView* FEditorViewportClient::CalcSceneView(FSceneViewFamily* ViewFamily)
 				FPlane(0,	1,	 0,	0),
 				FPlane(0,	0,	 0,	1));
 
-			//@TODO: Note: This code needs to be kept in sync with the code in ULocalPlayer::GetProjectionData
-			check(bUseControllingActorViewInfo);
-
 			const float EffectiveAspectRatio = bConstrainAspectRatio ? AspectRatio : (ViewportSizeXY.X / (float)ViewportSizeXY.Y);
 			const float YScale = 1.0f / EffectiveAspectRatio;
 			OrthoWidth = ControllingActorViewInfo.OrthoWidth / 2.0f;
@@ -2932,7 +2929,7 @@ void FEditorViewportClient::ConvertMovementToDragRot(const FVector& InDelta,
 			else if( MiddleMouseButtonDown || bIsUsingTrackpad || ( ( LeftMouseButtonDown || bIsUsingTrackpad ) && RightMouseButtonDown ) )
 			{
 				// Pan left/right/up/down
-				bool bInvert = !bIsUsingTrackpad && GetDefault<ULevelEditorViewportSettings>()->bInvertMiddleMousePan;
+				bool bInvert = !bIsUsingTrackpad && MiddleMouseButtonDown && GetDefault<ULevelEditorViewportSettings>()->bInvertMiddleMousePan;
 
 
 				float Direction = bInvert ? 1 : -1;
