#summary Updated metadata gathering commandlet to be configuration driven. Fixed a bug with local edits to metadata overriding localization. Updated property editors using comboboxes for enums to use localized data. Added Region & Language Settings for toggling use of localized field names.
#ttp 331251 - L10N: Editor: Tooltips and Property Names need to be localized

diff --git a/Engine/Config/BaseEditor.ini b/Engine/Config/BaseEditor.ini
--- a/Engine/Config/BaseEditor.ini
+++ b/Engine/Config/BaseEditor.ini
@@ -1,5 +1,8 @@
 [Internationalization]
 +LocalizationPaths=../../../Engine/Content/Localization/Editor
++ToolTipLocalizationPaths=../../../Engine/Content/Localization/Editor/ToolTips
++PropertyNameLocalizationPaths=../../../Engine/Content/Localization/Editor/PropertyNames
+ShouldLoadLocalizedPropertyNames=True
 
 [UnrealEd.GeneralSettings]
 EnableStylusPressure=True
diff --git a/Engine/Config/Localization/Editor.ini b/Engine/Config/Localization/Editor.ini
--- a/Engine/Config/Localization/Editor.ini
+++ b/Engine/Config/Localization/Editor.ini
@@ -23,11 +23,14 @@ SourceFileSearchFilters=*.mm
 ;Gather text from metadata
 [GatherTextStep1]
 CommandletClass=GatherTextFromMetaData
-IncludePaths=Engine/Source/Editor
-IncludePaths=Engine/Source/Runtime
-IncludePaths=Engine/Source/Developer
+IncludePaths=*/Engine/Source/Editor/*
+IncludePaths=*/Engine/Source/Runtime/*
+IncludePaths=*/Engine/Source/Developer/*
 ExcludePaths=*/Engine/Config/Localization/*
 ExcludePaths=*/Engine/Source/Developer/NoRedist/CommunityPortalServices/*
+InputKeys=Category
+OutputNamespaces=UObjectCategories
+OutputKeys="{MetaDataValue}"
 
 ;Gather text from assets
 [GatherTextStep2]
diff --git a/Engine/Config/Localization/PortableObjectExport.ini b/Engine/Config/Localization/PortableObjectExport.ini
--- a/Engine/Config/Localization/PortableObjectExport.ini
+++ b/Engine/Config/Localization/PortableObjectExport.ini
@@ -16,3 +16,17 @@ SourcePath=../../../Engine/Content/Localization/Editor
 DestinationPath=../../../Engine/Content/Localization/Editor
 ManifestName=Editor.manifest
 ArchiveName=Editor.archive
+
+[GatherTextStep2]
+CommandletClass=InternationalizationExport
+SourcePath=../../../Engine/Content/Localization/Editor
+DestinationPath=../../../Engine/Content/Localization/Editor
+ManifestName=ToolTips.manifest
+ArchiveName=ToolTips.archive
+
+[GatherTextStep3]
+CommandletClass=InternationalizationExport
+SourcePath=../../../Engine/Content/Localization/Editor
+DestinationPath=../../../Engine/Content/Localization/Editor
+ManifestName=PropertyNames.manifest
+ArchiveName=PropertyNames.archive
\ No newline at end of file
diff --git a/Engine/Config/Localization/PropertyNames.ini b/Engine/Config/Localization/PropertyNames.ini
new file mode 100644
--- /dev/null
+++ b/Engine/Config/Localization/PropertyNames.ini
@@ -0,0 +1,43 @@
+[CommonSettings]
+SourcePath=../../../Engine/Content/Localization/Editor/PropertyNames
+DestinationPath=../../../Engine/Content/Localization/Editor/PropertyNames
+ManifestName=PropertyNames.manifest
+ArchiveName=PropertyNames.archive
+SourceCulture=en
+CulturesToGenerate=en
+CulturesToGenerate=ja
+CulturesToGenerate=ko
+CulturesToGenerate=zh
+
+;Gather text from metadata
+[GatherTextStep0]
+CommandletClass=GatherTextFromMetaData
+IncludePaths=*/Engine/Source/Editor/*
+IncludePaths=*/Engine/Source/Runtime/*
+IncludePaths=*/Engine/Source/Developer/*
+ExcludePaths=*/Engine/Config/Localization/*
+ExcludePaths=*/Engine/Source/Developer/NoRedist/CommunityPortalServices/*
+InputKeys=DisplayName
+OutputNamespaces=UObjectDisplayNames
+OutputKeys="{FieldPath}"
+
+;Write Manifest
+[GatherTextStep1]
+CommandletClass=GenerateGatherManifest
+
+;Write Archives
+[GatherTextStep2]
+CommandletClass=GenerateGatherArchive
+bPurgeOldEmptyEntries=true
+
+;Write Localized Text Resource
+[GatherTextStep3]
+CommandletClass=GenerateTextLocalizationResource
+ResourceName=PropertyNames.locres
+
+;Write Text Localization Report
+[GatherTextStep4]
+CommandletClass=GenerateTextLocalizationReport
+DestinationPath=\\epicgames.net\root\UE3\Localization\WordCounts
+bConflictReport=true
+ConflictReportName=PropertyNameConflicts.txt
\ No newline at end of file
diff --git a/Engine/Config/Localization/ResourceFileGen.ini b/Engine/Config/Localization/ResourceFileGen.ini
--- a/Engine/Config/Localization/ResourceFileGen.ini
+++ b/Engine/Config/Localization/ResourceFileGen.ini
@@ -20,4 +20,28 @@ CulturesToGenerate=en
 CulturesToGenerate=ja
 CulturesToGenerate=ko
 CulturesToGenerate=zh
-ResourceName=Editor.locres
\ No newline at end of file
+ResourceName=Editor.locres
+
+;Write Localized Text Resource for Editor Tool Tips
+[GatherTextStep1]
+CommandletClass=GenerateTextLocalizationResource
+SourcePath=../../../Engine/Content/Localization/Editor/ToolTips
+DestinationPath=../../../Engine/Content/Localization/Editor/ToolTips
+ManifestName=ToolTips.manifest
+CulturesToGenerate=en
+CulturesToGenerate=ja
+CulturesToGenerate=ko
+CulturesToGenerate=zh
+ResourceName=ToolTips.locres
+
+;Write Localized Text Resource for Editor Property Names
+[GatherTextStep1]
+CommandletClass=GenerateTextLocalizationResource
+SourcePath=../../../Engine/Content/Localization/Editor/PropertyNames
+DestinationPath=../../../Engine/Content/Localization/Editor/PropertyNames
+ManifestName=PropertyNames.manifest
+CulturesToGenerate=en
+CulturesToGenerate=ja
+CulturesToGenerate=ko
+CulturesToGenerate=zh
+ResourceName=PropertyNames.locres
\ No newline at end of file
diff --git a/Engine/Config/Localization/ToolTips.ini b/Engine/Config/Localization/ToolTips.ini
new file mode 100644
--- /dev/null
+++ b/Engine/Config/Localization/ToolTips.ini
@@ -0,0 +1,43 @@
+[CommonSettings]
+SourcePath=../../../Engine/Content/Localization/Editor/ToolTips
+DestinationPath=../../../Engine/Content/Localization/Editor/ToolTips
+ManifestName=ToolTips.manifest
+ArchiveName=ToolTips.archive
+SourceCulture=en
+CulturesToGenerate=en
+CulturesToGenerate=ja
+CulturesToGenerate=ko
+CulturesToGenerate=zh
+
+;Gather text from metadata
+[GatherTextStep0]
+CommandletClass=GatherTextFromMetaData
+IncludePaths=*/Engine/Source/Editor/*
+IncludePaths=*/Engine/Source/Runtime/*
+IncludePaths=*/Engine/Source/Developer/*
+ExcludePaths=*/Engine/Config/Localization/*
+ExcludePaths=*/Engine/Source/Developer/NoRedist/CommunityPortalServices/*
+InputKeys=ToolTip
+OutputNamespaces=UObjectToolTips
+OutputKeys="{FieldPath}"
+
+;Write Manifest
+[GatherTextStep1]
+CommandletClass=GenerateGatherManifest
+
+;Write Archives
+[GatherTextStep2]
+CommandletClass=GenerateGatherArchive
+bPurgeOldEmptyEntries=true
+
+;Write Localized Text Resource
+[GatherTextStep3]
+CommandletClass=GenerateTextLocalizationResource
+ResourceName=ToolTips.locres
+
+;Write Text Localization Report
+[GatherTextStep4]
+CommandletClass=GenerateTextLocalizationReport
+DestinationPath=\\epicgames.net\root\UE3\Localization\WordCounts
+bConflictReport=true
+ConflictReportName=TooltipConflicts.txt
\ No newline at end of file
diff --git a/Engine/Config/Localization/WordCount.ini b/Engine/Config/Localization/WordCount.ini
--- a/Engine/Config/Localization/WordCount.ini
+++ b/Engine/Config/Localization/WordCount.ini
@@ -21,4 +21,26 @@ ManifestName=Editor.manifest
 CulturesToGenerate=ja
 CulturesToGenerate=ko
 CulturesToGenerate=zh
-WordCountReportName=Editor.csv
\ No newline at end of file
+WordCountReportName=Editor.csv
+
+;Write Text Localization Report for Editor Tool Tips
+[GatherTextStep2]
+CommandletClass=GenerateTextLocalizationReport
+bWordCountReport=true
+SourcePath=../../../Engine/Content/Localization/Editor
+ManifestName=ToolTips.manifest
+CulturesToGenerate=ja
+CulturesToGenerate=ko
+CulturesToGenerate=zh
+WordCountReportName=ToolTips.csv
+
+;Write Text Localization Report for Editor Property Names
+[GatherTextStep3]
+CommandletClass=GenerateTextLocalizationReport
+bWordCountReport=true
+SourcePath=../../../Engine/Content/Localization/Editor
+ManifestName=PropertyNames.manifest
+CulturesToGenerate=ja
+CulturesToGenerate=ko
+CulturesToGenerate=zh
+WordCountReportName=PropertyNames.csv
\ No newline at end of file
diff --git a/Engine/Source/Developer/InternationalizationSettings/Classes/InternationalizationSettingsModel.h b/Engine/Source/Developer/InternationalizationSettings/Classes/InternationalizationSettingsModel.h
--- a/Engine/Source/Developer/InternationalizationSettings/Classes/InternationalizationSettingsModel.h
+++ b/Engine/Source/Developer/InternationalizationSettings/Classes/InternationalizationSettingsModel.h
@@ -22,6 +22,8 @@ public:
 	void ResetToDefault();
 	FString GetCultureName() const;
 	void SetCultureName(const FString& CultureName);
+	bool ShouldLoadLocalizedPropertyNames() const;
+	void ShouldLoadLocalizedPropertyNames(const bool Value);
 
 public:
 	/**
diff --git a/Engine/Source/Developer/InternationalizationSettings/Private/InternationalizationSettingsModel.cpp b/Engine/Source/Developer/InternationalizationSettings/Private/InternationalizationSettingsModel.cpp
--- a/Engine/Source/Developer/InternationalizationSettings/Private/InternationalizationSettingsModel.cpp
+++ b/Engine/Source/Developer/InternationalizationSettings/Private/InternationalizationSettingsModel.cpp
@@ -19,6 +19,10 @@ void UInternationalizationSettingsModel::SaveDefaults()
 	FString SavedCultureName;
 	GConfig->GetString( TEXT("Internationalization"), TEXT("Culture"), SavedCultureName, GEditorGameAgnosticIni );
 	GConfig->SetString( TEXT("Internationalization"), TEXT("Culture"), *SavedCultureName, GEngineIni );
+
+	bool bShouldLoadLocalizedPropertyNames = true;
+	GConfig->GetBool( TEXT("Internationalization"), TEXT("ShouldLoadLocalizedPropertyNames"), bShouldLoadLocalizedPropertyNames, GEditorGameAgnosticIni );
+	GConfig->SetBool( TEXT("Internationalization"), TEXT("ShouldLoadLocalizedPropertyNames"), bShouldLoadLocalizedPropertyNames, GEngineIni );
 }
 
 void UInternationalizationSettingsModel::ResetToDefault()
@@ -26,6 +30,11 @@ void UInternationalizationSettingsModel::ResetToDefault()
 	FString SavedCultureName;
 	GConfig->GetString( TEXT("Internationalization"), TEXT("Culture"), SavedCultureName, GEngineIni );
 	GConfig->SetString( TEXT("Internationalization"), TEXT("Culture"), *SavedCultureName, GEditorGameAgnosticIni );
+
+	bool bShouldLoadLocalizedPropertyNames = true;
+	GConfig->GetBool( TEXT("Internationalization"), TEXT("ShouldLoadLocalizedPropertyNames"), bShouldLoadLocalizedPropertyNames, GEngineIni );
+	GConfig->SetBool( TEXT("Internationalization"), TEXT("ShouldLoadLocalizedPropertyNames"), bShouldLoadLocalizedPropertyNames, GEditorGameAgnosticIni );
+
 	SettingChangedEvent.Broadcast();
 }
 
@@ -43,4 +52,21 @@ void UInternationalizationSettingsModel::SetCultureName(const FString& CultureNa
 {
 	GConfig->SetString( TEXT("Internationalization"), TEXT("Culture"), *CultureName, GEditorGameAgnosticIni );
 	SettingChangedEvent.Broadcast();
+}
+
+
+bool UInternationalizationSettingsModel::ShouldLoadLocalizedPropertyNames() const
+{
+	bool bShouldLoadLocalizedPropertyNames = true;
+	if( !GConfig->GetBool( TEXT("Internationalization"), TEXT("ShouldLoadLocalizedPropertyNames"), bShouldLoadLocalizedPropertyNames, GEditorGameAgnosticIni ) )
+	{
+		GConfig->GetBool( TEXT("Internationalization"), TEXT("ShouldLoadLocalizedPropertyNames"), bShouldLoadLocalizedPropertyNames, GEngineIni );
+	}
+	return bShouldLoadLocalizedPropertyNames;
+}
+
+void UInternationalizationSettingsModel::ShouldLoadLocalizedPropertyNames(const bool Value)
+{
+	GConfig->SetBool( TEXT("Internationalization"), TEXT("ShouldLoadLocalizedPropertyNames"), Value, GEditorGameAgnosticIni );
+	SettingChangedEvent.Broadcast();
 }
\ No newline at end of file
diff --git a/Engine/Source/Developer/InternationalizationSettings/Private/InternationalizationSettingsModelDetails.cpp b/Engine/Source/Developer/InternationalizationSettings/Private/InternationalizationSettingsModelDetails.cpp
--- a/Engine/Source/Developer/InternationalizationSettings/Private/InternationalizationSettingsModelDetails.cpp
+++ b/Engine/Source/Developer/InternationalizationSettings/Private/InternationalizationSettingsModelDetails.cpp
@@ -185,6 +185,31 @@ void FInternationalizationSettingsModelDetails::CustomizeDetails( IDetailLayoutB
 		]
 	];
 
+	const FText FieldNamesToolTipText = LOCTEXT("EditorFieldNamesTooltip", "Toggle showing localized field names (requires restart to take effect)");
+
+	CategoryBuilder.AddCustomRow("FieldNames")
+	.NameContent()
+	[
+		SNew(SHorizontalBox)
+		+SHorizontalBox::Slot()
+		.Padding( FMargin( 0, 1, 0, 1 ) )
+		.FillWidth(1.0f)
+		[
+			SNew(STextBlock)
+			.Text(LOCTEXT("EditorFieldNamesLabel", "Use Localized Field Names"))
+			.Font(DetailBuilder.GetDetailFont())
+			.ToolTipText(FieldNamesToolTipText)
+		]
+	]
+	.ValueContent()
+	.MaxDesiredWidth(300.0f)
+	[
+		SAssignNew(FieldNamesCheckBox, SCheckBox)
+		.IsChecked(Model->ShouldLoadLocalizedPropertyNames() ? ESlateCheckBoxState::Checked : ESlateCheckBoxState::Unchecked)
+		.ToolTipText(FieldNamesToolTipText)
+		.OnCheckStateChanged(this, &FInternationalizationSettingsModelDetails::ShoudLoadLocalizedFieldNamesCheckChanged)
+	];
+
 	CategoryBuilder.AddCustomRow("RestartWarning")
 	.Visibility( TAttribute<EVisibility>(this, &FInternationalizationSettingsModelDetails::GetInternationalizationRestartRowVisibility) )
 	.WholeRowContent()
@@ -418,12 +443,19 @@ EVisibility FInternationalizationSettingsModelDetails::GetInternationalizationRe
 	return RequiresRestart ? EVisibility::Visible : EVisibility::Collapsed;
 }
 
+void FInternationalizationSettingsModelDetails::ShoudLoadLocalizedFieldNamesCheckChanged(ESlateCheckBoxState::Type CheckState)
+{
+	HandleShutdownPostPackagesSaved();
+}
+
+
 void FInternationalizationSettingsModelDetails::HandleShutdownPostPackagesSaved()
 {
 	if ( SelectedCulture.IsValid() )
 	{
 		check(Model.IsValid());
 		Model->SetCultureName(SelectedCulture->GetName());
+		Model->ShouldLoadLocalizedPropertyNames(FieldNamesCheckBox->IsChecked());
 		if(SelectedCulture != FInternationalization::GetCurrentCulture())
 		{
 			RequiresRestart = true;
diff --git a/Engine/Source/Developer/InternationalizationSettings/Public/InternationalizationSettingsModelDetails.h b/Engine/Source/Developer/InternationalizationSettings/Public/InternationalizationSettingsModelDetails.h
--- a/Engine/Source/Developer/InternationalizationSettings/Public/InternationalizationSettingsModelDetails.h
+++ b/Engine/Source/Developer/InternationalizationSettings/Public/InternationalizationSettingsModelDetails.h
@@ -62,6 +62,9 @@ private:
 	/** Delegate called to determine whether to collapse or display the restart row. */
 	EVisibility GetInternationalizationRestartRowVisibility() const;
 
+	/** Delegate called when the the checked state of whether or not field names should be localized has changed. */
+	void ShoudLoadLocalizedFieldNamesCheckChanged(ESlateCheckBoxState::Type CheckState);
+
 	/** Write to config now the Editor is shutting down (all packages are saved) */
 	void HandleShutdownPostPackagesSaved();
 
@@ -89,5 +92,8 @@ private:
 	/** The dropdown for the regions */
 	TSharedPtr<SComboBox< TSharedPtr<FCulture> >> RegionComboBox;
 
+	/** The check boxc for using localized field names */
+	TSharedPtr<SCheckBox> FieldNamesCheckBox;
+
 	bool RequiresRestart;
 };
\ No newline at end of file
diff --git a/Engine/Source/Developer/MovieSceneTools/Private/PropertyTrackEditor.cpp b/Engine/Source/Developer/MovieSceneTools/Private/PropertyTrackEditor.cpp
--- a/Engine/Source/Developer/MovieSceneTools/Private/PropertyTrackEditor.cpp
+++ b/Engine/Source/Developer/MovieSceneTools/Private/PropertyTrackEditor.cpp
@@ -31,7 +31,7 @@ public:
 	FPropertySection( UMovieSceneSection& InSectionObject, FName SectionName )
 		: SectionObject( InSectionObject )
 	{
-		DisplayName = EngineUtils::SanitizeDisplayName( SectionName.ToString(), false );
+		DisplayName = FName::NameToDisplayString( SectionName.ToString(), false );
 	}
 
 	/** ISequencerSection interface */
diff --git a/Engine/Source/Editor/BehaviorTreeEditor/Private/EdGraphSchema_BehaviorTree.cpp b/Engine/Source/Editor/BehaviorTreeEditor/Private/EdGraphSchema_BehaviorTree.cpp
--- a/Engine/Source/Editor/BehaviorTreeEditor/Private/EdGraphSchema_BehaviorTree.cpp
+++ b/Engine/Source/Editor/BehaviorTreeEditor/Private/EdGraphSchema_BehaviorTree.cpp
@@ -275,7 +275,7 @@ void UEdGraphSchema_BehaviorTree::GetGraphNodeContextActions(FGraphContextMenuBu
 
 		for (int32 i = 0; i < NodeClasses.Num(); i++)
 		{
-			const FString NodeTypeName = EngineUtils::SanitizeDisplayName(NodeClasses[i].ToString(), false);
+			const FString NodeTypeName = FName::NameToDisplayString(NodeClasses[i].ToString(), false);
 
 			UBehaviorTreeGraphNode* OpNode = NewObject<UBehaviorTreeGraphNode_Decorator>(Graph);
 			OpNode->ClassData = NodeClasses[i];
@@ -290,7 +290,7 @@ void UEdGraphSchema_BehaviorTree::GetGraphNodeContextActions(FGraphContextMenuBu
 		FClassBrowseHelper::GatherClasses(UBTService::StaticClass(), NodeClasses);
 		for (int32 i = 0; i < NodeClasses.Num(); i++)
 		{
-			const FString NodeTypeName = EngineUtils::SanitizeDisplayName(NodeClasses[i].ToString(), false);
+			const FString NodeTypeName = FName::NameToDisplayString(NodeClasses[i].ToString(), false);
 
 			UBehaviorTreeGraphNode* OpNode = NewObject<UBehaviorTreeGraphNode_Service>(Graph);
 			OpNode->ClassData = NodeClasses[i];
@@ -323,7 +323,7 @@ void UEdGraphSchema_BehaviorTree::GetGraphContextActions(FGraphContextMenuBuilde
 
 		for (int32 i = 0; i < NodeClasses.Num(); i++)
 		{
-			const FString NodeTypeName = EngineUtils::SanitizeDisplayName(NodeClasses[i].ToString(), false);
+			const FString NodeTypeName = FName::NameToDisplayString(NodeClasses[i].ToString(), false);
 
 			TSharedPtr<FBehaviorTreeSchemaAction_NewNode> AddOpAction = AddNewNodeAction(ContextMenuBuilder, TEXT("Composites"), NodeTypeName, "");
 
@@ -343,7 +343,7 @@ void UEdGraphSchema_BehaviorTree::GetGraphContextActions(FGraphContextMenuBuilde
 
 		for (int32 i = 0; i < NodeClasses.Num(); i++)
 		{
-			const FString NodeTypeName = EngineUtils::SanitizeDisplayName(NodeClasses[i].ToString(), false);
+			const FString NodeTypeName = FName::NameToDisplayString(NodeClasses[i].ToString(), false);
 
 			TSharedPtr<FBehaviorTreeSchemaAction_NewNode> AddOpAction = AddNewNodeAction(ContextMenuBuilder, TEXT("Tasks"), NodeTypeName, "");
 
diff --git a/Engine/Source/Editor/BlueprintGraph/Private/EdGraphSchema_K2.cpp b/Engine/Source/Editor/BlueprintGraph/Private/EdGraphSchema_K2.cpp
--- a/Engine/Source/Editor/BlueprintGraph/Private/EdGraphSchema_K2.cpp
+++ b/Engine/Source/Editor/BlueprintGraph/Private/EdGraphSchema_K2.cpp
@@ -1913,7 +1913,7 @@ FString UEdGraphSchema_K2::GetPinDisplayName(const UEdGraphPin* Pin) const
 
 		if( GEditor && GetDefault<UEditorStyleSettings>()->bShowFriendlyNames )
 		{
-			DisplayName = EngineUtils::SanitizeDisplayName(DisplayName, Pin->PinType.PinCategory == PC_Boolean);
+			DisplayName = FName::NameToDisplayString(DisplayName, Pin->PinType.PinCategory == PC_Boolean);
 		}
 	}
 	return DisplayName;
diff --git a/Engine/Source/Editor/BlueprintGraph/Private/K2ActionMenuBuilder.cpp b/Engine/Source/Editor/BlueprintGraph/Private/K2ActionMenuBuilder.cpp
--- a/Engine/Source/Editor/BlueprintGraph/Private/K2ActionMenuBuilder.cpp
+++ b/Engine/Source/Editor/BlueprintGraph/Private/K2ActionMenuBuilder.cpp
@@ -1477,7 +1477,7 @@ void FK2ActionMenuBuilder::GetAllActionsForClass(FBlueprintPaletteListBuilder& P
 					{
 						PropertyCategory = K2ActionCategories::DefaultEventDispatcherCategory;
 					}
-					PropertyCategory = EngineUtils::SanitizeDisplayName(PropertyCategory, /*bIsBool =*/false);
+					PropertyCategory = FName::NameToDisplayString(PropertyCategory, /*bIsBool =*/false);
 
 					TSharedPtr<FEdGraphSchemaAction_K2Delegate> NewFuncAction = MakeShareable(new FEdGraphSchemaAction_K2Delegate(PropertyCategory, PropertyDesc, PropertyTooltip, 0));
 					NewFuncAction->SetDelegateInfo(Property->GetFName(), InClass);
diff --git a/Engine/Source/Editor/BlueprintGraph/Private/K2Node_CallFunction.cpp b/Engine/Source/Editor/BlueprintGraph/Private/K2Node_CallFunction.cpp
--- a/Engine/Source/Editor/BlueprintGraph/Private/K2Node_CallFunction.cpp
+++ b/Engine/Source/Editor/BlueprintGraph/Private/K2Node_CallFunction.cpp
@@ -83,7 +83,7 @@ FString UK2Node_CallFunction::GetNodeTitle(ENodeTitleType::Type TitleType) const
 		FunctionName = FunctionReference.GetMemberName().ToString();
 		if ((GEditor != NULL) && (GetDefault<UEditorStyleSettings>()->bShowFriendlyNames))
 		{
-			FunctionName = EngineUtils::SanitizeDisplayName(FunctionName, false);
+			FunctionName = FName::NameToDisplayString(FunctionName, false);
 		}
 	}
 
@@ -693,7 +693,7 @@ FString UK2Node_CallFunction::GetUserFacingFunctionName(const UFunction* Functio
 
 	if( GEditor && GetDefault<UEditorStyleSettings>()->bShowFriendlyNames )
 	{
-		FunctionName = EngineUtils::SanitizeDisplayName(FunctionName, false);
+		FunctionName = FName::NameToDisplayString(FunctionName, false);
 	}
 	return FunctionName;
 }
@@ -727,7 +727,7 @@ FString UK2Node_CallFunction::GetDefaultCategoryForFunction(const UFunction* Fun
 		FString FuncCategory = Function->GetMetaData(FBlueprintMetadata::MD_FunctionCategory);
 		if( GEditor && GetDefault<UEditorStyleSettings>()->bShowFriendlyNames )
 		{
-			FuncCategory = EngineUtils::SanitizeDisplayName( FuncCategory, false );
+			FuncCategory = FName::NameToDisplayString( FuncCategory, false );
 		}
 		NodeCategory += FuncCategory;
 	}
diff --git a/Engine/Source/Editor/BlueprintGraph/Private/K2Node_CallParentFunction.cpp b/Engine/Source/Editor/BlueprintGraph/Private/K2Node_CallParentFunction.cpp
--- a/Engine/Source/Editor/BlueprintGraph/Private/K2Node_CallParentFunction.cpp
+++ b/Engine/Source/Editor/BlueprintGraph/Private/K2Node_CallParentFunction.cpp
@@ -18,7 +18,7 @@ FString UK2Node_CallParentFunction::GetNodeTitle(ENodeTitleType::Type TitleType)
 	
 	if( GEditor && GetDefault<UEditorStyleSettings>()->bShowFriendlyNames )
 	{
-		Result = EngineUtils::SanitizeDisplayName(Result, false);
+		Result = FName::NameToDisplayString(Result, false);
 	}
 	return Result;
 }
diff --git a/Engine/Source/Editor/BlueprintGraph/Private/K2Node_FunctionEntry.cpp b/Engine/Source/Editor/BlueprintGraph/Private/K2Node_FunctionEntry.cpp
--- a/Engine/Source/Editor/BlueprintGraph/Private/K2Node_FunctionEntry.cpp
+++ b/Engine/Source/Editor/BlueprintGraph/Private/K2Node_FunctionEntry.cpp
@@ -107,7 +107,7 @@ FString UK2Node_FunctionEntry::GetNodeTitle(ENodeTitleType::Type TitleType) cons
 
 	if( GEditor && GetDefault<UEditorStyleSettings>()->bShowFriendlyNames )
 	{
-		DisplayInfo.DisplayName = FText::FromString( EngineUtils::SanitizeDisplayName(DisplayInfo.DisplayName.ToString(), false) );
+		DisplayInfo.DisplayName = FText::FromString( FName::NameToDisplayString(DisplayInfo.DisplayName.ToString(), false) );
 	}
 	return DisplayInfo.DisplayName.ToString();
 }
diff --git a/Engine/Source/Editor/Blutility/Private/BlutilityDetailsPanel.cpp b/Engine/Source/Editor/Blutility/Private/BlutilityDetailsPanel.cpp
--- a/Engine/Source/Editor/Blutility/Private/BlutilityDetailsPanel.cpp
+++ b/Engine/Source/Editor/Blutility/Private/BlutilityDetailsPanel.cpp
@@ -98,7 +98,7 @@ void FEditorUtilityInstanceDetails::CustomizeDetails(IDetailLayoutBuilder& Detai
 			{
 				++NumButtons;
 
-				const FString ButtonCaption = EngineUtils::SanitizeDisplayName(*Function->GetName(), false);
+				const FString ButtonCaption = FName::NameToDisplayString(*Function->GetName(), false);
 
 				//@TODO: Expose the code in UK2Node_CallFunction::GetUserFacingFunctionName / etc...
 				FString Tooltip = Function->GetToolTipText().ToString();
diff --git a/Engine/Source/Editor/ContentBrowser/Private/AssetViewWidgets.cpp b/Engine/Source/Editor/ContentBrowser/Private/AssetViewWidgets.cpp
--- a/Engine/Source/Editor/ContentBrowser/Private/AssetViewWidgets.cpp
+++ b/Engine/Source/Editor/ContentBrowser/Private/AssetViewWidgets.cpp
@@ -433,7 +433,7 @@ TSharedRef<SToolTip> SAssetViewItem::CreateToolTipWidget() const
 
 						// We have no type information by this point, so no idea if it's a bool :(
 						const bool bIsBool = false;
-						AddToToolTipInfoBox( InfoBox, FText::FromString(EngineUtils::SanitizeDisplayName(TagIt.Key().ToString(), bIsBool)), FText::FromString(ValueString), bImportant );
+						AddToToolTipInfoBox( InfoBox, FText::FromString(FName::NameToDisplayString(TagIt.Key().ToString(), bIsBool)), FText::FromString(ValueString), bImportant );
 					}
 				}
 			}
diff --git a/Engine/Source/Editor/DetailCustomizations/Private/BrushDetails.cpp b/Engine/Source/Editor/DetailCustomizations/Private/BrushDetails.cpp
--- a/Engine/Source/Editor/DetailCustomizations/Private/BrushDetails.cpp
+++ b/Engine/Source/Editor/DetailCustomizations/Private/BrushDetails.cpp
@@ -159,7 +159,7 @@ void FBrushDetails::CustomizeDetails( IDetailLayoutBuilder& DetailLayout )
 				}
 				else
 				{
-					return FText::FromString(EngineUtils::SanitizeDisplayName(BrushBuilder->GetClass()->GetName(), false));
+					return FText::FromString(FName::NameToDisplayString(BrushBuilder->GetClass()->GetName(), false));
 				}
 			}
 
diff --git a/Engine/Source/Editor/GameProjectGeneration/Private/SNewClassDialog.cpp b/Engine/Source/Editor/GameProjectGeneration/Private/SNewClassDialog.cpp
--- a/Engine/Source/Editor/GameProjectGeneration/Private/SNewClassDialog.cpp
+++ b/Engine/Source/Editor/GameProjectGeneration/Private/SNewClassDialog.cpp
@@ -391,7 +391,7 @@ TSharedRef<ITableRow> SNewClassDialog::MakeParentClassListViewWidget(TSharedPtr<
 		return SNew( STableRow<TSharedPtr<FParentClassItem>>, OwnerTable );
 	}
 
-	const FString ClassName = EngineUtils::SanitizeDisplayName( Class->GetName(), false );
+	const FString ClassName = FName::NameToDisplayString( Class->GetName(), false );
 	FString ClassDescription = Class->GetToolTipText().ToString();
 	int32 NewLineIndex = 0;
 	if (ClassDescription.FindChar('.', NewLineIndex))
diff --git a/Engine/Source/Editor/GraphEditor/Private/SGraphActionMenu.cpp b/Engine/Source/Editor/GraphEditor/Private/SGraphActionMenu.cpp
--- a/Engine/Source/Editor/GraphEditor/Private/SGraphActionMenu.cpp
+++ b/Engine/Source/Editor/GraphEditor/Private/SGraphActionMenu.cpp
@@ -503,7 +503,7 @@ void SGraphActionMenu::GenerateFilteredItems(bool bPreserveExpansion)
 	TArray<FString> SanitizedFilterTerms;
 	for (int32 iFilters = 0; iFilters < FilterTerms.Num() ; iFilters++)
 	{
-		FString EachString = EngineUtils::SanitizeDisplayName( FilterTerms[iFilters], false );
+		FString EachString = FName::NameToDisplayString( FilterTerms[iFilters], false );
 		EachString = EachString.Replace( TEXT( " " ), TEXT( "" ) );
 		SanitizedFilterTerms.Add( EachString );
 	}
diff --git a/Engine/Source/Editor/LevelEditor/Private/LevelEditorCreateActorMenu.cpp b/Engine/Source/Editor/LevelEditor/Private/LevelEditorCreateActorMenu.cpp
--- a/Engine/Source/Editor/LevelEditor/Private/LevelEditorCreateActorMenu.cpp
+++ b/Engine/Source/Editor/LevelEditor/Private/LevelEditorCreateActorMenu.cpp
@@ -68,7 +68,7 @@ class SAssetMenuEntry : public SCompoundWidget
 		FText AssetDisplayName = FText::FromName( Asset.AssetName );
 		if ( IsClass )
 		{
-			AssetDisplayName = FText::FromString( EngineUtils::SanitizeDisplayName( Asset.AssetName.ToString(), false ) );
+			AssetDisplayName = FText::FromString( FName::NameToDisplayString( Asset.AssetName.ToString(), false ) );
 		}
 
 		FText ActorTypeDisplayName;
@@ -80,7 +80,7 @@ class SAssetMenuEntry : public SCompoundWidget
 			if ( IsClass && Cast<UClass>( MenuItem.AssetData.GetAsset() )->IsChildOf( AActor::StaticClass() ) )
 			{
 				DefaultActor = Cast<AActor>( Cast<UClass>( MenuItem.AssetData.GetAsset() )->ClassDefaultObject );
-				ActorTypeDisplayName = FText::FromString( EngineUtils::SanitizeDisplayName( DefaultActor->GetClass()->GetName(), false ) );
+				ActorTypeDisplayName = FText::FromString( FName::NameToDisplayString( DefaultActor->GetClass()->GetName(), false ) );
 			}
 
 			const FSlateBrush* IconBrush = NULL;
diff --git a/Engine/Source/Editor/PlacementMode/Private/SPlacementModeTools.cpp b/Engine/Source/Editor/PlacementMode/Private/SPlacementModeTools.cpp
--- a/Engine/Source/Editor/PlacementMode/Private/SPlacementModeTools.cpp
+++ b/Engine/Source/Editor/PlacementMode/Private/SPlacementModeTools.cpp
@@ -98,7 +98,7 @@ void SPlacementAssetEntry::Construct(const FArguments& InArgs, UActorFactory* In
 	AssetDisplayName = FText::FromName( AssetData.AssetName );
 	if ( IsClass )
 	{
-		AssetDisplayName = FText::FromString( EngineUtils::SanitizeDisplayName( AssetData.AssetName.ToString(), false ) );
+		AssetDisplayName = FText::FromString( FName::NameToDisplayString( AssetData.AssetName.ToString(), false ) );
 	}
 
 	FText ActorTypeDisplayName;
@@ -106,7 +106,7 @@ void SPlacementAssetEntry::Construct(const FArguments& InArgs, UActorFactory* In
 	if ( IsClass && Cast<UClass>( AssetData.GetAsset() )->IsChildOf( AActor::StaticClass() ) )
 	{
 		DefaultActor = Cast<AActor>( Cast<UClass>( AssetData.GetAsset() )->ClassDefaultObject );
-		ActorTypeDisplayName = FText::FromString( EngineUtils::SanitizeDisplayName( DefaultActor->GetClass()->GetName(), false ) );
+		ActorTypeDisplayName = FText::FromString( FName::NameToDisplayString( DefaultActor->GetClass()->GetName(), false ) );
 	}
 
 	if ( FactoryToUse != nullptr )
diff --git a/Engine/Source/Editor/PropertyEditor/Private/CategoryPropertyNode.cpp b/Engine/Source/Editor/PropertyEditor/Private/CategoryPropertyNode.cpp
--- a/Engine/Source/Editor/PropertyEditor/Private/CategoryPropertyNode.cpp
+++ b/Engine/Source/Editor/PropertyEditor/Private/CategoryPropertyNode.cpp
@@ -29,7 +29,7 @@ FString FCategoryPropertyNode::GetDisplayName() const
 	//if in "readable display name mode" return that
 	if ( FPropertySettings::Get().ShowFriendlyPropertyNames() )
 	{
-		FString DisplaySubcategoryName = EngineUtils::SanitizeDisplayName( SubcategoryName, false );
+		FString DisplaySubcategoryName = FName::NameToDisplayString( SubcategoryName, false );
 
 		//if there is a localized version, return that
 		FText LocalizedName;
diff --git a/Engine/Source/Editor/PropertyEditor/Private/DetailCategoryBuilderImpl.cpp b/Engine/Source/Editor/PropertyEditor/Private/DetailCategoryBuilderImpl.cpp
--- a/Engine/Source/Editor/PropertyEditor/Private/DetailCategoryBuilderImpl.cpp
+++ b/Engine/Source/Editor/PropertyEditor/Private/DetailCategoryBuilderImpl.cpp
@@ -452,7 +452,7 @@ void FDetailCategoryImpl::SetDisplayName( FName InCategoryName, const FString& L
 		FName BaseClassName = BaseClass ? BaseClass->GetFName() : FName("Generic");
 
 		FString CategoryStr = InCategoryName != NAME_None ? InCategoryName.ToString() : BaseClassName.ToString();
-		FString SourceCategoryStr = EngineUtils::SanitizeDisplayName( CategoryStr, false );
+		FString SourceCategoryStr = FName::NameToDisplayString( CategoryStr, false );
 
 		bool FoundText = false;
 		FText DisplayNameText;
diff --git a/Engine/Source/Editor/PropertyEditor/Private/ItemPropertyNode.cpp b/Engine/Source/Editor/PropertyEditor/Private/ItemPropertyNode.cpp
--- a/Engine/Source/Editor/PropertyEditor/Private/ItemPropertyNode.cpp
+++ b/Engine/Source/Editor/PropertyEditor/Private/ItemPropertyNode.cpp
@@ -337,7 +337,7 @@ FString FItemPropertyNode::GetDisplayName() const
 					}
 					if( GetDefault<UEditorStyleSettings>()->bShowFriendlyNames )
 					{
-						PropertyDisplayName = EngineUtils::SanitizeDisplayName( PropertyDisplayName, bIsBoolProperty );
+						PropertyDisplayName = FName::NameToDisplayString( PropertyDisplayName, bIsBoolProperty );
 					}
 
 					FinalDisplayName = PropertyDisplayName;
diff --git a/Engine/Source/Editor/PropertyEditor/Private/Presentation/PropertyTable/PropertyTableColumn.cpp b/Engine/Source/Editor/PropertyEditor/Private/Presentation/PropertyTable/PropertyTableColumn.cpp
--- a/Engine/Source/Editor/PropertyEditor/Private/Presentation/PropertyTable/PropertyTableColumn.cpp
+++ b/Engine/Source/Editor/PropertyEditor/Private/Presentation/PropertyTable/PropertyTableColumn.cpp
@@ -460,7 +460,7 @@ void FPropertyTableColumn::GenerateColumnDisplayName()
 						}
 					}
 
-					PropertyName = EngineUtils::SanitizeDisplayName( PropertyName, bIsBoolProperty );
+					PropertyName = FName::NameToDisplayString( PropertyName, bIsBoolProperty );
 				}
 
 				NewName += PropertyName;
diff --git a/Engine/Source/Editor/PropertyEditor/Private/PropertyHandleImpl.cpp b/Engine/Source/Editor/PropertyEditor/Private/PropertyHandleImpl.cpp
--- a/Engine/Source/Editor/PropertyEditor/Private/PropertyHandleImpl.cpp
+++ b/Engine/Source/Editor/PropertyEditor/Private/PropertyHandleImpl.cpp
@@ -66,18 +66,8 @@ FPropertyAccess::Result FPropertyValueImpl::GetPropertyValueString( FString& Out
 				UByteProperty* ByteProperty = Cast<UByteProperty>(Property);
 				if ( ByteProperty != NULL && ByteProperty->Enum != NULL )
 				{
-					// see if we have alternate text to use for displaying the value
-					UMetaData* PackageMetaData = ByteProperty->Enum->GetOutermost()->GetMetaData();
-					if ( PackageMetaData )
-					{
-						FName AltDisplayName = FName(*(OutString+TEXT(".DisplayName")));
-						FString ValueText = PackageMetaData->GetValue(ByteProperty->Enum, AltDisplayName);
-						if ( ValueText.Len() > 0 )
-						{
-							// render the alternate text for this enum value
-							OutString = ValueText;
-						}
-					}
+					uint8 EnumValue = ByteProperty->GetPropertyValue(ValueAddress);
+					OutString = ByteProperty->Enum->GetEnumString(EnumValue);
 				}
 			}
 			else
@@ -123,18 +113,8 @@ FPropertyAccess::Result FPropertyValueImpl::GetPropertyValueText( FText& OutText
 				UByteProperty* ByteProperty = Cast<UByteProperty>(Property);
 				if ( ByteProperty != NULL && ByteProperty->Enum != NULL )
 				{
-					// see if we have alternate text to use for displaying the value
-					UMetaData* PackageMetaData = ByteProperty->Enum->GetOutermost()->GetMetaData();
-					if ( PackageMetaData )
-					{
-						FName AltDisplayName = FName(*(OutString+TEXT(".DisplayName")));
-						FString ValueText = PackageMetaData->GetValue(ByteProperty->Enum, AltDisplayName);
-						if ( ValueText.Len() > 0 )
-						{
-							// render the alternate text for this enum value
-							OutString = ValueText;
-						}
-					}
+					uint8 EnumValueIndex = ByteProperty->GetPropertyValue(ValueAddress);
+					OutString = ByteProperty->Enum->GetEnumString(EnumValueIndex);
 				}
 
 				OutText = FText::FromString(OutString);
diff --git a/Engine/Source/Editor/UnrealEd/Classes/Commandlets/GatherTextFromMetadataCommandlet.h b/Engine/Source/Editor/UnrealEd/Classes/Commandlets/GatherTextFromMetadataCommandlet.h
--- a/Engine/Source/Editor/UnrealEd/Classes/Commandlets/GatherTextFromMetadataCommandlet.h
+++ b/Engine/Source/Editor/UnrealEd/Classes/Commandlets/GatherTextFromMetadataCommandlet.h
@@ -18,6 +18,14 @@ public:
 	// End UCommandlet Interface
 
 private:
-	void GatherTextFromUObjects(const TArray<FString>& IncludePaths, const TArray<FString>& ExcludePaths);
-	void GatherTextFromUObject(UField* const Field);
+	struct FGatherParameters
+	{
+		TArray<FString> InputKeys;
+		TArray<FString> OutputNamespaces;
+		TArray<FText> OutputKeys;
+	};
+
+private:
+	void GatherTextFromUObjects(const TArray<FString>& IncludePaths, const TArray<FString>& ExcludePaths, const FGatherParameters& Arguments);
+	void GatherTextFromUObject(UField* const Field, const FGatherParameters& Arguments);
 };
\ No newline at end of file
diff --git a/Engine/Source/Editor/UnrealEd/Private/AssetThumbnail.cpp b/Engine/Source/Editor/UnrealEd/Private/AssetThumbnail.cpp
--- a/Engine/Source/Editor/UnrealEd/Private/AssetThumbnail.cpp
+++ b/Engine/Source/Editor/UnrealEd/Private/AssetThumbnail.cpp
@@ -483,7 +483,7 @@ private:
 
 			if ( ClassDisplayName.IsEmpty() )
 			{
-				ClassDisplayName = FText::FromString( EngineUtils::SanitizeDisplayName(*Class->GetName(), false) );
+				ClassDisplayName = FText::FromString( FName::NameToDisplayString(*Class->GetName(), false) );
 			}
 		}
 
diff --git a/Engine/Source/Editor/UnrealEd/Private/Commandlets/GatherTextFromMetadataCommandlet.cpp b/Engine/Source/Editor/UnrealEd/Private/Commandlets/GatherTextFromMetadataCommandlet.cpp
--- a/Engine/Source/Editor/UnrealEd/Private/Commandlets/GatherTextFromMetadataCommandlet.cpp
+++ b/Engine/Source/Editor/UnrealEd/Private/Commandlets/GatherTextFromMetadataCommandlet.cpp
@@ -65,7 +65,7 @@ int32 UGatherTextFromMetaDataCommandlet::Main( const FString& Params )
 
 	//Required module names
 	TArray<FString> RequiredModuleNames;
-	GetConfigArray(*SectionName, TEXT("RequiredModuleNames"), ExcludePaths, GatherTextConfigPath);
+	GetConfigArray(*SectionName, TEXT("RequiredModuleNames"), RequiredModuleNames, GatherTextConfigPath);
 
 	// Pre-load all required modules so that UFields from those modules can be guaranteed a chance to have their metadata gathered.
 	for(const FString& RequiredModuleName : RequiredModuleNames)
@@ -73,8 +73,18 @@ int32 UGatherTextFromMetaDataCommandlet::Main( const FString& Params )
 		FModuleManager::Get().LoadModule(*RequiredModuleName);
 	}
 
+	FGatherParameters Arguments;
+	GetConfigArray(*SectionName, TEXT("InputKeys"), Arguments.InputKeys, GatherTextConfigPath);
+	GetConfigArray(*SectionName, TEXT("OutputNamespaces"), Arguments.OutputNamespaces, GatherTextConfigPath);
+	TArray<FString> OutputKeys;
+	GetConfigArray(*SectionName, TEXT("OutputKeys"), OutputKeys, GatherTextConfigPath);
+	for(const auto& OutputKey : OutputKeys)
+	{
+		Arguments.OutputKeys.Add(FText::FromString(OutputKey));
+	}
+
 	// Execute gather.
-	GatherTextFromUObjects(IncludePaths, ExcludePaths);
+	GatherTextFromUObjects(IncludePaths, ExcludePaths, Arguments);
 
 	// Add any manifest dependencies if they were provided
 	TArray<FString> ManifestDependenciesList;
@@ -89,7 +99,7 @@ int32 UGatherTextFromMetaDataCommandlet::Main( const FString& Params )
 	return 0;
 }
 
-void UGatherTextFromMetaDataCommandlet::GatherTextFromUObjects(const TArray<FString>& IncludePaths, const TArray<FString>& ExcludePaths)
+void UGatherTextFromMetaDataCommandlet::GatherTextFromUObjects(const TArray<FString>& IncludePaths, const TArray<FString>& ExcludePaths, const FGatherParameters& Arguments)
 {
 	for(TObjectIterator<UField> It; It; ++It)
 	{
@@ -101,7 +111,7 @@ void UGatherTextFromMetaDataCommandlet::GatherTextFromUObjects(const TArray<FStr
 		{
 			for(int32 i = 0; i < IncludePaths.Num(); ++i)
 			{
-				if(SourceFilePath.MatchesWildcard(FPaths::RootDir() + IncludePaths[i] + TEXT("*")))
+				if(SourceFilePath.MatchesWildcard(IncludePaths[i]))
 				{
 					return true;
 				}
@@ -130,33 +140,38 @@ void UGatherTextFromMetaDataCommandlet::GatherTextFromUObjects(const TArray<FStr
 			continue;
 		}
 
-		GatherTextFromUObject(*It);
+		GatherTextFromUObject(*It, Arguments);
 	}
 }
 
-void UGatherTextFromMetaDataCommandlet::GatherTextFromUObject(UField* const Field)
+void UGatherTextFromMetaDataCommandlet::GatherTextFromUObject(UField* const Field, const FGatherParameters& Arguments)
 {
-	const int32 MetaDataCount = 3;
-	const FString MetaDataKeys[MetaDataCount] =
-	{	TEXT("ToolTip"),			TEXT("DisplayName"),			TEXT("Category")			};
-	const FString AssociatedNamespaces[MetaDataCount] =
-	{	TEXT("UObjectToolTips"),	TEXT("UObjectDisplayNames"),	TEXT("UObjectCategories")	};
-
 	// Gather for object.
 	{
-		for(int32 i = 0; i < MetaDataCount; ++i)
+		if( !Field->HasMetaData( TEXT("DisplayName") ) )
+		{
+			Field->SetMetaData( TEXT("DisplayName"), *FName::NameToDisplayString( Field->GetName(), Field->IsA( UBoolProperty::StaticClass() ) ) );
+		}
+
+		for(int32 i = 0; i < Arguments.InputKeys.Num(); ++i)
 		{
-			const FString& MetaDataValue = Field->GetMetaData(*MetaDataKeys[i]);
-			if(!(MetaDataValue.IsEmpty()))
+			FFormatNamedArguments PatternArguments;
+			PatternArguments.Add( TEXT("FieldPath"), FText::FromString( Field->GetFullGroupName(true) + TEXT(".") + Field->GetName() ) );
+
+			if( Field->HasMetaData( *Arguments.InputKeys[i] ) )
 			{
-				const FString Namespace = AssociatedNamespaces[i];
-				FLocItem LocItem(MetaDataValue);
-				FContext Context;
-				Context.Key =	MetaDataKeys[i] == TEXT("Category")
-							?	MetaDataValue
-							:	Field->GetFullGroupName(true) + TEXT(".") + Field->GetName();
-				Context.SourceLocation = TEXT("Run-time MetaData");
-				ManifestInfo->AddEntry(TEXT("EntryDescription"), Namespace, LocItem, Context);
+				const FString& MetaDataValue = Field->GetMetaData(*Arguments.InputKeys[i]);
+				if( !MetaDataValue.IsEmpty() )
+				{
+					PatternArguments.Add( TEXT("MetaDataValue"), FText::FromString(MetaDataValue) );
+
+					const FString Namespace = Arguments.OutputNamespaces[i];
+					FLocItem LocItem(MetaDataValue);
+					FContext Context;
+					Context.Key = FText::Format(Arguments.OutputKeys[i], PatternArguments).ToString();
+					Context.SourceLocation = FString::Printf(TEXT("From metadata for key %s of member %s in %s"), *Arguments.InputKeys[i], *Field->GetName(), *Field->GetFullGroupName(true));
+					ManifestInfo->AddEntry(TEXT("EntryDescription"), Namespace, LocItem, Context);
+				}
 			}
 		}
 	}
@@ -169,19 +184,30 @@ void UGatherTextFromMetaDataCommandlet::GatherTextFromUObject(UField* const Fiel
 			const int32 ValueCount = Enum->NumEnums();
 			for(int32 i = 0; i < ValueCount; ++i)
 			{
-				for(int32 j = 0; j < MetaDataCount; ++j)
+				if( !Enum->HasMetaData(TEXT("DisplayName"), i) )
 				{
-					const FString& MetaDataValue = Enum->GetMetaData(*MetaDataKeys[j], i);
-					if(!(MetaDataValue.IsEmpty()))
+					Enum->SetMetaData(TEXT("DisplayName"), *FName::NameToDisplayString(Enum->GetEnumName(i), false), i);
+				}
+
+				for(int32 j = 0; j < Arguments.InputKeys.Num(); ++j)
+				{
+					FFormatNamedArguments PatternArguments;
+					PatternArguments.Add( TEXT("FieldPath"), FText::FromString( Enum->GetFullGroupName(true) + TEXT(".") + Enum->GetName() + TEXT(".") + Enum->GetEnumName(i) ) );
+
+					if( Enum->HasMetaData(*Arguments.InputKeys[j], i) )
 					{
-						const FString Namespace = AssociatedNamespaces[j];
-						FLocItem LocItem(MetaDataValue);
-						FContext Context;
-						Context.Key =	MetaDataKeys[j] == TEXT("Category")
-									?	MetaDataValue
-									:	Enum->GetFullGroupName(true) + TEXT(".") + Enum->GetName() + TEXT(".") + Enum->GetEnumName(i);
-						Context.SourceLocation = TEXT("Run-time MetaData");
-						ManifestInfo->AddEntry(TEXT("EntryDescription"), Namespace, LocItem, Context);
+						const FString& MetaDataValue = Enum->GetMetaData(*Arguments.InputKeys[j], i);
+						if( !MetaDataValue.IsEmpty() )
+						{
+							PatternArguments.Add( TEXT("MetaDataValue"), FText::FromString(MetaDataValue) );
+
+							const FString Namespace = Arguments.OutputNamespaces[j];
+							FLocItem LocItem(MetaDataValue);
+							FContext Context;
+							Context.Key = FText::Format(Arguments.OutputKeys[j], PatternArguments).ToString();
+							Context.SourceLocation = FString::Printf(TEXT("From metadata for key %s of enum value %s of enum %s in %s"), *Arguments.InputKeys[j], *Enum->GetEnumName(i), *Enum->GetName(), *Enum->GetFullGroupName(true));
+							ManifestInfo->AddEntry(TEXT("EntryDescription"), Namespace, LocItem, Context);
+						}
 					}
 				}
 			}
diff --git a/Engine/Source/Editor/UnrealEd/Private/Editor.cpp b/Engine/Source/Editor/UnrealEd/Private/Editor.cpp
--- a/Engine/Source/Editor/UnrealEd/Private/Editor.cpp
+++ b/Engine/Source/Editor/UnrealEd/Private/Editor.cpp
@@ -4520,7 +4520,7 @@ FString UEditorEngine::GetFriendlyName( const UProperty* Property, UClass* Owner
 		if ( DefaultFriendlyName.Len() == 0 )
 		{
 			const bool bIsBool = Cast<const UBoolProperty>(Property) != NULL;
-			return EngineUtils::SanitizeDisplayName( Property->GetName(), bIsBool );
+			return FName::NameToDisplayString( Property->GetName(), bIsBool );
 		}
 		return DefaultFriendlyName;
 	}
diff --git a/Engine/Source/Editor/UnrealEd/Private/Factories/Factory.cpp b/Engine/Source/Editor/UnrealEd/Private/Factories/Factory.cpp
--- a/Engine/Source/Editor/UnrealEd/Private/Factories/Factory.cpp
+++ b/Engine/Source/Editor/UnrealEd/Private/Factories/Factory.cpp
@@ -91,7 +91,7 @@ FText UFactory::GetDisplayName() const
 		}
 
 		// Factories whose classes do not have asset type actions should just display the sanitized class name
-		return FText::FromString( EngineUtils::SanitizeDisplayName(*LocalSupportedClass->GetName(), false) );
+		return FText::FromString( FName::NameToDisplayString(*LocalSupportedClass->GetName(), false) );
 	}
 
 	// Factories that have no supported class have no display name.
diff --git a/Engine/Source/Editor/UnrealEd/Private/Kismet2/BlueprintEditorUtils.cpp b/Engine/Source/Editor/UnrealEd/Private/Kismet2/BlueprintEditorUtils.cpp
--- a/Engine/Source/Editor/UnrealEd/Private/Kismet2/BlueprintEditorUtils.cpp
+++ b/Engine/Source/Editor/UnrealEd/Private/Kismet2/BlueprintEditorUtils.cpp
@@ -2726,7 +2726,7 @@ bool FBlueprintEditorUtils::AddMemberVariable(UBlueprint* Blueprint, const FName
 
 	NewVar.VarName = NewVarName;
 	NewVar.VarGuid = FGuid::NewGuid();
-	NewVar.FriendlyName = EngineUtils::SanitizeDisplayName( NewVarName.ToString(), (NewVarType.PinCategory == K2Schema->PC_Boolean) ? true : false );
+	NewVar.FriendlyName = FName::NameToDisplayString( NewVarName.ToString(), (NewVarType.PinCategory == K2Schema->PC_Boolean) ? true : false );
 	NewVar.VarType = NewVarType;
 	// default new vars to 'kismet read/write' and 'only editable on owning CDO' 
 	NewVar.PropertyFlags |= (CPF_Edit | CPF_BlueprintVisible | CPF_DisableEditOnInstance);
@@ -2867,7 +2867,7 @@ void FBlueprintEditorUtils::RenameMemberVariable(UBlueprint* Blueprint, const FN
 			const UEdGraphSchema_K2* K2Schema = GetDefault<UEdGraphSchema_K2>();
 			FBPVariableDescription& Variable = Blueprint->NewVariables[VarIndex];
 			Variable.VarName = NewName;
-			Variable.FriendlyName = EngineUtils::SanitizeDisplayName( NewName.ToString(), (Variable.VarType.PinCategory == K2Schema->PC_Boolean) ? true : false );
+			Variable.FriendlyName = FName::NameToDisplayString( NewName.ToString(), (Variable.VarType.PinCategory == K2Schema->PC_Boolean) ? true : false );
 
 			// Update any existing references to the old name
 			FBlueprintEditorUtils::ReplaceVariableReferences(Blueprint, OldName, NewName);
diff --git a/Engine/Source/Editor/UnrealEd/Private/SComponentClassCombo.cpp b/Engine/Source/Editor/UnrealEd/Private/SComponentClassCombo.cpp
--- a/Engine/Source/Editor/UnrealEd/Private/SComponentClassCombo.cpp
+++ b/Engine/Source/Editor/UnrealEd/Private/SComponentClassCombo.cpp
@@ -332,7 +332,7 @@ FString SComponentClassCombo::GetSanitizedComponentName( UClass* ComponentClass
 {
 	FString OriginalName = ComponentClass->GetName();
 	FString ComponentNameWithoutComponent = OriginalName.Replace( TEXT("Component"), TEXT(""), ESearchCase::IgnoreCase );
-	return EngineUtils::SanitizeDisplayName( ComponentNameWithoutComponent, false );
+	return FName::NameToDisplayString( ComponentNameWithoutComponent, false );
 }
 
 #undef LOCTEXT_NAMESPACE
diff --git a/Engine/Source/Editor/UnrealEd/Private/UnrealEdEngine.cpp b/Engine/Source/Editor/UnrealEd/Private/UnrealEdEngine.cpp
--- a/Engine/Source/Editor/UnrealEd/Private/UnrealEdEngine.cpp
+++ b/Engine/Source/Editor/UnrealEd/Private/UnrealEdEngine.cpp
@@ -518,7 +518,7 @@ FString FClassPickerDefaults::GetName() const
 	{
 		UClass* ItemClass = LoadClass<UObject>(NULL, *ClassName, NULL, LOAD_None, NULL);
 		check( ItemClass );
-		return EngineUtils::SanitizeDisplayName(ItemClass->GetName(), false);
+		return FName::NameToDisplayString(ItemClass->GetName(), false);
 	}
 	
 	const FText* PreExistingName = LocNames.Find( LocTextNameID );
diff --git a/Engine/Source/Editor/UnrealEd/Public/GameModeInfoCustomizer.h b/Engine/Source/Editor/UnrealEd/Public/GameModeInfoCustomizer.h
--- a/Engine/Source/Editor/UnrealEd/Public/GameModeInfoCustomizer.h
+++ b/Engine/Source/Editor/UnrealEd/Public/GameModeInfoCustomizer.h
@@ -29,7 +29,7 @@ public:
 		{
 			DisplayName = Prop->GetName();
 		}
-		DisplayName = EngineUtils::SanitizeDisplayName(DisplayName, false);
+		DisplayName = FName::NameToDisplayString(DisplayName, false);
 
 		return
 			SNew(STextBlock)
diff --git a/Engine/Source/Runtime/Core/Private/Internationalization/Text.cpp b/Engine/Source/Runtime/Core/Private/Internationalization/Text.cpp
--- a/Engine/Source/Runtime/Core/Private/Internationalization/Text.cpp
+++ b/Engine/Source/Runtime/Core/Private/Internationalization/Text.cpp
@@ -324,13 +324,13 @@ DEF_ASPERCENT(double)
 DEF_ASPERCENT(float)
 #undef DEF_ASPERCENT
 
-bool FText::FindText( const FString& Namespace, const FString& Key, FText& OutText )
+bool FText::FindText( const FString& Namespace, const FString& Key, FText& OutText, const FString* const SourceString )
 {
 	TSharedPtr< FString > FoundString = FTextLocalizationManager::Get().FindString( Namespace, Key );
 
 	if ( FoundString.IsValid() )
 	{
-		OutText = FText( FString(), Namespace, Key );
+		OutText = FText( SourceString ? *SourceString : FString(), Namespace, Key );
 	}
 
 	return FoundString.IsValid();
diff --git a/Engine/Source/Runtime/Core/Private/Internationalization/TextLocalizationManager.cpp b/Engine/Source/Runtime/Core/Private/Internationalization/TextLocalizationManager.cpp
--- a/Engine/Source/Runtime/Core/Private/Internationalization/TextLocalizationManager.cpp
+++ b/Engine/Source/Runtime/Core/Private/Internationalization/TextLocalizationManager.cpp
@@ -196,6 +196,17 @@ void FTextLocalizationManager::LoadResources(const bool ShouldLoadEditor, const
 	if(ShouldLoadEditor)
 	{
 		LocalizationPaths += FPaths::GetEditorLocalizationPaths();
+		LocalizationPaths += FPaths::GetToolTipLocalizationPaths();
+
+		bool bShouldLoadLocalizedPropertyNames = true;
+		if( !GConfig->GetBool( TEXT("Internationalization"), TEXT("ShouldLoadLocalizedPropertyNames"), bShouldLoadLocalizedPropertyNames, GEditorGameAgnosticIni ) )
+		{
+			GConfig->GetBool( TEXT("Internationalization"), TEXT("ShouldLoadLocalizedPropertyNames"), bShouldLoadLocalizedPropertyNames, GEngineIni );
+		}
+		if(bShouldLoadLocalizedPropertyNames)
+		{
+			LocalizationPaths += FPaths::GetPropertyNameLocalizationPaths();
+		}
 	}
 	LocalizationPaths += FPaths::GetEngineLocalizationPaths();
 
@@ -324,7 +335,7 @@ void FTextLocalizationManager::OnCultureChanged()
 	LoadResources(ShouldLoadEditor, ShouldLoadGame);
 }
 
-TSharedPtr<FString> FTextLocalizationManager::FindString( const FString& Namespace, const FString& Key )
+TSharedPtr<FString> FTextLocalizationManager::FindString( const FString& Namespace, const FString& Key, const FString* const SourceString )
 {
 	FScopeLock ScopeLock( &SynchronizationObject );
 
@@ -334,7 +345,7 @@ TSharedPtr<FString> FTextLocalizationManager::FindString( const FString& Namespa
 	// Find key table's entry.
 	const FStringEntry* LiveEntry = LiveKeyTable ? LiveKeyTable->Find( Key ) : NULL;
 
-	if ( LiveEntry != NULL )
+	if( LiveEntry != NULL && ( !SourceString || LiveEntry->SourceStringHash == FCrc::StrCrc32(**SourceString) ) )
 	{
 		return LiveEntry->String;
 	}
diff --git a/Engine/Source/Runtime/Core/Private/Misc/Paths.cpp b/Engine/Source/Runtime/Core/Private/Misc/Paths.cpp
--- a/Engine/Source/Runtime/Core/Private/Misc/Paths.cpp
+++ b/Engine/Source/Runtime/Core/Private/Misc/Paths.cpp
@@ -247,6 +247,56 @@ const TArray<FString>& FPaths::GetEditorLocalizationPaths()
 	return Results;
 }
 
+const TArray<FString>& FPaths::GetPropertyNameLocalizationPaths()
+{
+	static TArray<FString> Results;
+	static bool HasInitialized = false;
+
+	if(!HasInitialized)
+	{
+		if(GConfig && GConfig->IsReadyForUse())
+		{
+			GConfig->GetArray( TEXT("Internationalization"), TEXT("PropertyNameLocalizationPaths"), Results, GEditorIni );
+			if(!Results.Num())
+			{
+				UE_LOG(LogInit, Warning, TEXT("We appear to have no property name localization data? Editor can't run."));
+			}
+			HasInitialized = true;
+		}
+		else
+		{
+			Results.AddUnique(TEXT("../../../Engine/Content/Localization/Editor/PropertyNames")); // Hardcoded convention.
+		}
+	}
+
+	return Results;
+}
+
+const TArray<FString>& FPaths::GetToolTipLocalizationPaths()
+{
+	static TArray<FString> Results;
+	static bool HasInitialized = false;
+
+	if(!HasInitialized)
+	{
+		if(GConfig && GConfig->IsReadyForUse())
+		{
+			GConfig->GetArray( TEXT("Internationalization"), TEXT("ToolTipLocalizationPaths"), Results, GEditorIni );
+			if(!Results.Num())
+			{
+				UE_LOG(LogInit, Warning, TEXT("We appear to have no tool tip localization data? Editor can't run."));
+			}
+			HasInitialized = true;
+		}
+		else
+		{
+			Results.AddUnique(TEXT("../../../Engine/Content/Localization/Editor/ToolTips")); // Hardcoded convention.
+		}
+	}
+
+	return Results;
+}
+
 const TArray<FString>& FPaths::GetGameLocalizationPaths()
 {
 	static TArray<FString> Results;
diff --git a/Engine/Source/Runtime/Core/Private/UObject/UnrealNames.cpp b/Engine/Source/Runtime/Core/Private/UObject/UnrealNames.cpp
--- a/Engine/Source/Runtime/Core/Private/UObject/UnrealNames.cpp
+++ b/Engine/Source/Runtime/Core/Private/UObject/UnrealNames.cpp
@@ -234,6 +234,124 @@ FCriticalSection* FName::GetCriticalSection()
 	return CriticalSection;
 }
 
+FString FName::NameToDisplayString( const FString& InDisplayName, const bool bIsBool )
+{
+	// Copy the characters out so that we can modify the string in place
+	const TArray< TCHAR > Chars = InDisplayName.GetCharArray();
+
+	// This is used to indicate that we are in a run of uppercase letter and/or digits.  The code attempts to keep
+	// these characters together as breaking them up often looks silly (i.e. "Draw Scale 3 D" as opposed to "Draw Scale 3D"
+	bool bInARun = false;
+	bool bWasSpace = false;
+	bool bWasOpenParen = false;
+	FString OutDisplayName;
+	for( int32 CharIndex = 0 ; CharIndex < Chars.Num() ; ++CharIndex )
+	{
+		TCHAR ch = Chars[CharIndex];
+
+		bool bLowerCase = FChar::IsLower( ch );
+		bool bUpperCase = FChar::IsUpper( ch );
+		bool bIsDigit = FChar::IsDigit( ch );
+		bool bIsUnderscore = FChar::IsUnderscore( ch );
+
+		// Skip the first character if the property is a bool (they should all start with a lowercase 'b', which we don't want to keep)
+		if( CharIndex == 0 && bIsBool && ch == 'b' )
+		{
+			// Check if next character is uppercase as it may be a user created string that doesn't follow the rules of Unreal variables
+			if (Chars.Num() > 1 && FChar::IsUpper(Chars[1]))
+			{
+				continue;
+			}
+		}
+
+		// If the current character is upper case or a digit, and the previous character wasn't, then we need to insert a space if there wasn't one previously
+		if( (bUpperCase || bIsDigit) && !bInARun && !bWasOpenParen)
+		{
+			if( !bWasSpace && OutDisplayName.Len() > 0 )
+			{
+				OutDisplayName += TEXT( ' ' );
+				bWasSpace = true;
+			}
+			bInARun = true;
+		}
+
+		// A lower case character will break a run of upper case letters and/or digits
+		if( bLowerCase )
+		{
+			bInARun = false;
+		}
+
+		// An underscore denotes a space, so replace it and continue the run
+		if( bIsUnderscore )
+		{
+			ch = TEXT( ' ' );
+			bInARun = true;
+		}
+
+		// If this is the first character in the string, then it will always be upper-case
+		if( OutDisplayName.Len() == 0 )
+		{
+			ch = FChar::ToUpper( ch );
+		}
+		else if( bWasSpace || bWasOpenParen)	// If this is first character after a space, then make sure it is case-correct
+		{
+			// Some words are always forced lowercase
+			const TCHAR* Articles[] =
+			{
+				TEXT( "In" ),
+				TEXT( "As" ),
+				TEXT( "To" ),
+				TEXT( "Or" ),
+				TEXT( "At" ),
+				TEXT( "On" ),
+				TEXT( "If" ),
+				TEXT( "Be" ),
+				TEXT( "By" ),
+				TEXT( "The" ),
+				TEXT( "For" ),
+				TEXT( "And" ),
+				TEXT( "With" ),
+				TEXT( "When" ),
+				TEXT( "From" ),
+			};
+
+			// Search for a word that needs case repaired
+			bool bIsArticle = false;
+			for( int32 CurArticleIndex = 0; CurArticleIndex < ARRAY_COUNT( Articles ); ++CurArticleIndex )
+			{
+				// Make sure the character following the string we're testing is not lowercase (we don't want to match "in" with "instance")
+				const int32 ArticleLength = FCString::Strlen( Articles[ CurArticleIndex ] );
+				if( ( Chars.Num() - CharIndex ) > ArticleLength && !FChar::IsLower( Chars[ CharIndex + ArticleLength ] ) && Chars[ CharIndex + ArticleLength ] != '\0' )
+				{
+					// Does this match the current article?
+					if( FCString::Strncmp( &Chars[ CharIndex ], Articles[ CurArticleIndex ], ArticleLength ) == 0 )
+					{
+						bIsArticle = true;
+						break;
+					}
+				}
+			}
+
+			// Start of a keyword, force to lowercase
+			if( bIsArticle )
+			{
+				ch = FChar::ToLower( ch );				
+			}
+			else	// First character after a space that's not a reserved keyword, make sure it's uppercase
+			{
+				ch = FChar::ToUpper( ch );
+			}
+		}
+
+		bWasSpace = ( ch == TEXT( ' ' ) ? true : false );
+		bWasOpenParen = ( ch == TEXT( '(' ) ? true : false );
+
+		OutDisplayName += ch;
+	}
+
+	return OutDisplayName;
+}
+
 
 // Static variables.
 FNameEntry*						FName::NameHash[ FNameDefs::NameHashBucketCount ];
diff --git a/Engine/Source/Runtime/Core/Public/Internationalization/Internationalization.h b/Engine/Source/Runtime/Core/Public/Internationalization/Internationalization.h
--- a/Engine/Source/Runtime/Core/Public/Internationalization/Internationalization.h
+++ b/Engine/Source/Runtime/Core/Public/Internationalization/Internationalization.h
@@ -35,7 +35,7 @@ public:
 
 	static CORE_API void GetTimeZonesIDs(TArray<FString>& TimeZonesIDs);
 
-	//Set the current culture by Name
+	//Set the current culture by name
 	static CORE_API void SetCurrentCulture(const FString& Name);
 
 	//@return the current culture
@@ -82,7 +82,6 @@ private:
 		FInternationalizationLifetimeObject();
 	} static InternationalizationLifetimeObject;
 
-
 	static CORE_API TArray< TSharedRef< FCulture > > AllCultures;
 
 	static CORE_API int CurrentCultureIndex;
diff --git a/Engine/Source/Runtime/Core/Public/Internationalization/Text.h b/Engine/Source/Runtime/Core/Public/Internationalization/Text.h
--- a/Engine/Source/Runtime/Core/Public/Internationalization/Text.h
+++ b/Engine/Source/Runtime/Core/Public/Internationalization/Text.h
@@ -162,7 +162,7 @@ public:
 	 * Attempts to find an existing FText using the representation found in the loc tables for the specified namespace and key
 	 * @return true if OutText was properly set; otherwise false and OutText will be untouched
 	 */
-	static bool FindText( const FString& Namespace, const FString& Key, FText& OutText );
+	static bool FindText( const FString& Namespace, const FString& Key, FText& OutText, const FString* const SourceString = nullptr );
 
 	/**
 	 * Generate an FText representing the pass name
diff --git a/Engine/Source/Runtime/Core/Public/Internationalization/TextLocalizationManager.h b/Engine/Source/Runtime/Core/Public/Internationalization/TextLocalizationManager.h
--- a/Engine/Source/Runtime/Core/Public/Internationalization/TextLocalizationManager.h
+++ b/Engine/Source/Runtime/Core/Public/Internationalization/TextLocalizationManager.h
@@ -63,7 +63,7 @@ private:
 public:
 	static CORE_API FTextLocalizationManager& Get();
 
-	TSharedPtr<FString> FindString( const FString& Namespace, const FString& Key );
+	TSharedPtr<FString> FindString( const FString& Namespace, const FString& Key, const FString* const SourceString = nullptr );
 	TSharedRef<FString> GetString(const FString& Namespace, const FString& Key, const FString* const SourceString);
 
 	void CORE_API RegenerateResources(const FString& ConfigFilePath);
diff --git a/Engine/Source/Runtime/Core/Public/Misc/Paths.h b/Engine/Source/Runtime/Core/Public/Misc/Paths.h
--- a/Engine/Source/Runtime/Core/Public/Misc/Paths.h
+++ b/Engine/Source/Runtime/Core/Public/Misc/Paths.h
@@ -220,6 +220,16 @@ public:
 	 */
 	static const TArray<FString>& GetEditorLocalizationPaths();
 
+	/** 
+	 * Returns a list of property name localization paths
+	 */
+	static const TArray<FString>& GetPropertyNameLocalizationPaths();
+
+		/** 
+	 * Returns a list of tool tip localization paths
+	 */
+	static const TArray<FString>& GetToolTipLocalizationPaths();
+
 	/** 
 	 * Returns a list of game-specific localization paths
 	 */
diff --git a/Engine/Source/Runtime/Core/Public/UObject/NameTypes.h b/Engine/Source/Runtime/Core/Public/UObject/NameTypes.h
--- a/Engine/Source/Runtime/Core/Public/UObject/NameTypes.h
+++ b/Engine/Source/Runtime/Core/Public/UObject/NameTypes.h
@@ -701,6 +701,18 @@ public:
 	static FNameEntry*** GetNameTableForDebuggerVisualizers_MT();
 	/** Run autotest on FNames. */
 	static void AutoTest();
+	
+	/**
+	 * Takes a string and breaks it down into a human readable string.
+	 * For example - "bCreateSomeStuff" becomes "Create Some Stuff?" and "DrawScale3D" becomes "Draw Scale 3D".
+	 * 
+	 * @param	InDisplayName	[In, Out] The name to sanitize
+	 * @param	bIsBool				True if the name is a bool
+	 *
+	 * @return	the sanitized version of the display name
+	 */
+	static FString NameToDisplayString( const FString& InDisplayName, const bool bIsBool );
+
 private:
 
 	/** Index into the Names array (used to find String portion of the string/number pair) */
diff --git a/Engine/Source/Runtime/CoreUObject/Private/UObject/Class.cpp b/Engine/Source/Runtime/CoreUObject/Private/UObject/Class.cpp
--- a/Engine/Source/Runtime/CoreUObject/Private/UObject/Class.cpp
+++ b/Engine/Source/Runtime/CoreUObject/Private/UObject/Class.cpp
@@ -188,11 +188,21 @@ void UField::AddCppProperty( UProperty* Property )
 FText UField::GetDisplayNameText() const
 {
 	FText LocalizedDisplayName;
-	const FString NativeDisplayName = GetMetaData( TEXT("DisplayName") );
 
 	static const FString Namespace = TEXT("UObjectDisplayNames");
 	const FString Key = GetFullGroupName(true) + TEXT(".") + GetName();
-	if ( !(FText::FindText( Namespace, Key, /*OUT*/LocalizedDisplayName )) || *FTextInspector::GetSourceString(LocalizedDisplayName) != NativeDisplayName )
+
+	FString NativeDisplayName;
+	if( HasMetaData( TEXT("DisplayName") ) )
+	{
+		NativeDisplayName = GetMetaData( TEXT("DisplayName") );
+	}
+	else
+	{
+		NativeDisplayName = FName::NameToDisplayString(GetName(), IsA<UBoolProperty>());
+	}
+
+	if ( !( FText::FindText( Namespace, Key, /*OUT*/LocalizedDisplayName, &NativeDisplayName ) ) )
 	{
 		LocalizedDisplayName = FText::FromString(NativeDisplayName );
 	}
diff --git a/Engine/Source/Runtime/CoreUObject/Private/UObject/Enum.cpp b/Engine/Source/Runtime/CoreUObject/Private/UObject/Enum.cpp
--- a/Engine/Source/Runtime/CoreUObject/Private/UObject/Enum.cpp
+++ b/Engine/Source/Runtime/CoreUObject/Private/UObject/Enum.cpp
@@ -274,14 +274,24 @@ bool UEnum::GenerateMaxEnum()
 FText UEnum::GetDisplayNameText(int32 NameIndex) const
 {
 	FText LocalizedDisplayName;
-	const FString NativeDisplayName = GetMetaData( TEXT("DisplayName"), NameIndex );
 
-	static const FString Namespace = TEXT("UEnumDisplayNames");
+	static const FString Namespace = TEXT("UObjectDisplayNames");
 	const FString Key =	NameIndex == INDEX_NONE
 		?			GetFullGroupName(true) + TEXT(".") + GetName()
 		:			GetFullGroupName(true) + TEXT(".") + GetName() + TEXT(".") + GetEnumName(NameIndex);
 
-	if ( !(FText::FindText( Namespace, Key, /*OUT*/LocalizedDisplayName )) || *FTextInspector::GetSourceString(LocalizedDisplayName) != NativeDisplayName)
+
+	FString NativeDisplayName;
+	if( HasMetaData( TEXT("DisplayName"), NameIndex ) )
+	{
+		NativeDisplayName = GetMetaData( TEXT("DisplayName"), NameIndex );
+	}
+	else
+	{
+		NativeDisplayName = FName::NameToDisplayString(GetEnumName(NameIndex), false);
+	}
+
+	if ( !( FText::FindText( Namespace, Key, /*OUT*/LocalizedDisplayName, &NativeDisplayName) ) )
 	{
 		LocalizedDisplayName = FText::FromString(NativeDisplayName);
 	}
diff --git a/Engine/Source/Runtime/Engine/Private/AI/BehaviorTree/BlueprintNodeHelpers.cpp b/Engine/Source/Runtime/Engine/Private/AI/BehaviorTree/BlueprintNodeHelpers.cpp
--- a/Engine/Source/Runtime/Engine/Private/AI/BehaviorTree/BlueprintNodeHelpers.cpp
+++ b/Engine/Source/Runtime/Engine/Private/AI/BehaviorTree/BlueprintNodeHelpers.cpp
@@ -100,7 +100,7 @@ namespace BlueprintNodeHelpers
 		}
 
 		const bool bIsBool = Prop->IsA(UBoolProperty::StaticClass());
-		return FString::Printf(TEXT("%s: %s"), *EngineUtils::SanitizeDisplayName(Prop->GetName(), bIsBool), *ExportedStringValue);
+		return FString::Printf(TEXT("%s: %s"), *FName::NameToDisplayString(Prop->GetName(), bIsBool), *ExportedStringValue);
 	}
 
 	void CollectBlackboardSelectors(const UObject* Ob, const UClass* StopAtClass, TArray<FName>& KeyNames)
diff --git a/Engine/Source/Runtime/Engine/Private/EngineUtils.cpp b/Engine/Source/Runtime/Engine/Private/EngineUtils.cpp
--- a/Engine/Source/Runtime/Engine/Private/EngineUtils.cpp
+++ b/Engine/Source/Runtime/Engine/Private/EngineUtils.cpp
@@ -312,125 +312,6 @@ bool EngineUtils::FindOrLoadAssetsByPath(const FString& Path, TArray<UObject*>&
 }
 
 
-FString EngineUtils::SanitizeDisplayName( const FString& InDisplayName, const bool bIsBool )
-{
-	// Copy the characters out so that we can modify the string in place
-	const TArray< TCHAR > Chars = InDisplayName.GetCharArray();
-
-	// This is used to indicate that we are in a run of uppercase letter and/or digits.  The code attempts to keep
-	// these characters together as breaking them up often looks silly (i.e. "Draw Scale 3 D" as opposed to "Draw Scale 3D"
-	bool bInARun = false;
-	bool bWasSpace = false;
-	bool bWasOpenParen = false;
-	FString OutDisplayName;
-	for( int32 CharIndex = 0 ; CharIndex < Chars.Num() ; ++CharIndex )
-	{
-		TCHAR ch = Chars[CharIndex];
-
-		bool bLowerCase = FChar::IsLower( ch );
-		bool bUpperCase = FChar::IsUpper( ch );
-		bool bIsDigit = FChar::IsDigit( ch );
-		bool bIsUnderscore = FChar::IsUnderscore( ch );
-
-		// Skip the first character if the property is a bool (they should all start with a lowercase 'b', which we don't want to keep)
-		if( CharIndex == 0 && bIsBool && ch == 'b' )
-		{
-			// Check if next character is uppercase as it may be a user created string that doesn't follow the rules of Unreal variables
-			if (Chars.Num() > 1 && FChar::IsUpper(Chars[1]))
-			{
-				continue;
-			}
-		}
-
-		// If the current character is upper case or a digit, and the previous character wasn't, then we need to insert a space if there wasn't one previously
-		if( (bUpperCase || bIsDigit) && !bInARun && !bWasOpenParen)
-		{
-			if( !bWasSpace && OutDisplayName.Len() > 0 )
-			{
-				OutDisplayName += TEXT( ' ' );
-				bWasSpace = true;
-			}
-			bInARun = true;
-		}
-
-		// A lower case character will break a run of upper case letters and/or digits
-		if( bLowerCase )
-		{
-			bInARun = false;
-		}
-
-		// An underscore denotes a space, so replace it and continue the run
-		if( bIsUnderscore )
-		{
-			ch = TEXT( ' ' );
-			bInARun = true;
-		}
-
-		// If this is the first character in the string, then it will always be upper-case
-		if( OutDisplayName.Len() == 0 )
-		{
-			ch = FChar::ToUpper( ch );
-		}
-		else if( bWasSpace || bWasOpenParen)	// If this is first character after a space, then make sure it is case-correct
-		{
-			// Some words are always forced lowercase
-			const TCHAR* Articles[] =
-			{
-				TEXT( "In" ),
-				TEXT( "As" ),
-				TEXT( "To" ),
-				TEXT( "Or" ),
-				TEXT( "At" ),
-				TEXT( "On" ),
-				TEXT( "If" ),
-				TEXT( "Be" ),
-				TEXT( "By" ),
-				TEXT( "The" ),
-				TEXT( "For" ),
-				TEXT( "And" ),
-				TEXT( "With" ),
-				TEXT( "When" ),
-				TEXT( "From" ),
-			};
-
-			// Search for a word that needs case repaired
-			bool bIsArticle = false;
-			for( int32 CurArticleIndex = 0; CurArticleIndex < ARRAY_COUNT( Articles ); ++CurArticleIndex )
-			{
-				// Make sure the character following the string we're testing is not lowercase (we don't want to match "in" with "instance")
-				const int32 ArticleLength = FCString::Strlen( Articles[ CurArticleIndex ] );
-				if( ( Chars.Num() - CharIndex ) > ArticleLength && !FChar::IsLower( Chars[ CharIndex + ArticleLength ] ) && Chars[ CharIndex + ArticleLength ] != '\0' )
-				{
-					// Does this match the current article?
-					if( FCString::Strncmp( &Chars[ CharIndex ], Articles[ CurArticleIndex ], ArticleLength ) == 0 )
-					{
-						bIsArticle = true;
-						break;
-					}
-				}
-			}
-
-			// Start of a keyword, force to lowercase
-			if( bIsArticle )
-			{
-				ch = FChar::ToLower( ch );				
-			}
-			else	// First character after a space that's not a reserved keyword, make sure it's uppercase
-			{
-				ch = FChar::ToUpper( ch );
-			}
-		}
-
-		bWasSpace = ( ch == TEXT( ' ' ) ? true : false );
-		bWasOpenParen = ( ch == TEXT( '(' ) ? true : false );
-
-		OutDisplayName += ch;
-	}
-
-	return OutDisplayName;
-}
-
-
 /** This will set the StreamingLevels TMap with the current Streaming Level Status and also set which level the player is in **/
 void GetLevelStreamingStatus( UWorld* World, TMap<FName,int32>& StreamingLevels, FString& LevelPlayerIsInName )
 {
diff --git a/Engine/Source/Runtime/Engine/Public/EngineUtils.h b/Engine/Source/Runtime/Engine/Public/EngineUtils.h
--- a/Engine/Source/Runtime/Engine/Public/EngineUtils.h
+++ b/Engine/Source/Runtime/Engine/Public/EngineUtils.h
@@ -669,16 +669,6 @@ namespace EngineUtils
 {
 	/** Loads all the assets found in the specified path and subpaths */
 	ENGINE_API bool FindOrLoadAssetsByPath(const FString& Path, TArray<UObject*>& OutAssets);
-	/**
-	 * Takes a string and breaks it down into a human readable string.
-	 * For example - "bCreateSomeStuff" becomes "Create Some Stuff?" and "DrawScale3D" becomes "Draw Scale 3D".
-	 * 
-	 * @param	InDisplayName	[In, Out] The name to sanitize
-	 * @param	bIsBool				True if the name is a bool
-	 *
-	 * @return	the sanitized version of the display name
-	 */
-	ENGINE_API FString SanitizeDisplayName( const FString& InDisplayName, const bool bIsBool );
 }
 
 /** Helper class for serializing flags describing which data have been stripped (if any). */
