Fixed issue with minimum size too large to create physics asset issue.

diff --git a/Engine/Source/Editor/UnrealEd/Private/Fbx/FbxSkeletalMeshImport.cpp b/Engine/Source/Editor/UnrealEd/Private/Fbx/FbxSkeletalMeshImport.cpp
--- a/Engine/Source/Editor/UnrealEd/Private/Fbx/FbxSkeletalMeshImport.cpp
+++ b/Engine/Source/Editor/UnrealEd/Private/Fbx/FbxSkeletalMeshImport.cpp
@@ -1268,7 +1268,7 @@ USkeletalMesh* UnFbx::FFbxImporter::ImportSkeletalMesh(UObject* InParent, TArray
 					bool bSuccess = FPhysicsAssetUtils::CreateFromSkeletalMesh(NewPhysicsAsset, SkeletalMesh, NewBodyData, CreationErrorMessage);
 					if (!bSuccess)
 					{
-						AddTokenizedErrorMessage(FTokenizedMessage::Create(EMessageSeverity::Error, CreationErrorMessage));
+						AddTokenizedErrorMessage(FTokenizedMessage::Create(EMessageSeverity::Warning, CreationErrorMessage));
 						// delete the asset since we could not have create physics asset
 						TArray<UObject*> ObjectsToDelete;
 						ObjectsToDelete.Add(NewPhysicsAsset);
diff --git a/Engine/Source/Editor/UnrealEd/Private/PhysicsAssetUtils.cpp b/Engine/Source/Editor/UnrealEd/Private/PhysicsAssetUtils.cpp
--- a/Engine/Source/Editor/UnrealEd/Private/PhysicsAssetUtils.cpp
+++ b/Engine/Source/Editor/UnrealEd/Private/PhysicsAssetUtils.cpp
@@ -96,17 +96,14 @@ static float GetMaximalMinSizeBelow(int32 BoneIndex, USkeletalMesh* SkelMesh, co
 	return MaximalMinBoxSize;
 }
 
-
-bool CreateFromSkeletalMesh(UPhysicsAsset* PhysicsAsset, USkeletalMesh* SkelMesh, FPhysAssetCreateParams& Params, FText & OutErrorMessage)
+bool CreateFromSkeletalMeshInternal(UPhysicsAsset* PhysicsAsset, USkeletalMesh* SkelMesh, FPhysAssetCreateParams& Params)
 {
-	PhysicsAsset->PreviewSkeletalMesh = SkelMesh;
-
 	IMeshUtilities& MeshUtilities = FModuleManager::Get().LoadModuleChecked<IMeshUtilities>("MeshUtilities");
 
 	// For each bone, get the vertices most firmly attached to it.
 	TArray<FBoneVertInfo> Infos;
 	MeshUtilities.CalcBoneVertInfos(SkelMesh, Infos, (Params.VertWeight == EVW_DominantWeight));
-	check( Infos.Num() == SkelMesh->RefSkeleton.GetNum() );
+	check(Infos.Num() == SkelMesh->RefSkeleton.GetNum());
 
 	bool bHitRoot = false;
 
@@ -116,7 +113,7 @@ bool CreateFromSkeletalMesh(UPhysicsAsset* PhysicsAsset, USkeletalMesh* SkelMesh
 		FName BoneName = SkelMesh->RefSkeleton.GetBoneName(i);
 
 		int32 ParentIndex = INDEX_NONE;
-		FName ParentName = NAME_None; 
+		FName ParentName = NAME_None;
 		int32 ParentBodyIndex = INDEX_NONE;
 
 		// If we have already found the 'physics root', we expect a parent.
@@ -142,29 +139,29 @@ bool CreateFromSkeletalMesh(UPhysicsAsset* PhysicsAsset, USkeletalMesh* SkelMesh
 			bMakeBone = true;
 		}
 		// If we have passed the physics 'root', and this bone has no physical parent, ignore it.
-		else if( !(bHitRoot && ParentBodyIndex == INDEX_NONE) )
+		else if(!(bHitRoot && ParentBodyIndex == INDEX_NONE))
 		{
 			// If bone is big enough - create physics.
-			if( CalcBoneInfoLength( Infos[i] ) > Params.MinBoneSize )
+			if(CalcBoneInfoLength(Infos[i]) > Params.MinBoneSize)
 			{
 				bMakeBone = true;
 			}
 
 			// If its too small, and we have set the option, see if it has any large children.
-			if( !bMakeBone && Params.bWalkPastSmall )
+			if(!bMakeBone && Params.bWalkPastSmall)
 			{
-				if( GetMaximalMinSizeBelow(i, SkelMesh, Infos) > Params.MinBoneSize )
+				if(GetMaximalMinSizeBelow(i, SkelMesh, Infos) > Params.MinBoneSize)
 				{
 					bMakeBone = true;
 				}
 			}
 		}
 
-		if( bMakeBone )
+		if(bMakeBone)
 		{
 			// Go ahead and make this bone physical.
-			int32 NewBodyIndex = CreateNewBody( PhysicsAsset, BoneName );
-			UBodySetup* bs = PhysicsAsset->BodySetup[ NewBodyIndex ];
+			int32 NewBodyIndex = CreateNewBody(PhysicsAsset, BoneName);
+			UBodySetup* bs = PhysicsAsset->BodySetup[NewBodyIndex];
 			check(bs->BoneName == BoneName);
 
 			// Fill in collision info for this bone.
@@ -173,8 +170,8 @@ bool CreateFromSkeletalMesh(UPhysicsAsset* PhysicsAsset, USkeletalMesh* SkelMesh
 			// If not root - create joint to parent body.
 			if(bHitRoot && Params.bCreateJoints)
 			{
-				int32 NewConstraintIndex = CreateNewConstraint( PhysicsAsset, BoneName );
-				UPhysicsConstraintTemplate* CS = PhysicsAsset->ConstraintSetup[ NewConstraintIndex ];
+				int32 NewConstraintIndex = CreateNewConstraint(PhysicsAsset, BoneName);
+				UPhysicsConstraintTemplate* CS = PhysicsAsset->ConstraintSetup[NewConstraintIndex];
 
 				// Transform of child from parent is just child ref-pose entry.
 				FMatrix RelTM = SkelMesh->GetRefPoseMatrix(i);
@@ -187,13 +184,13 @@ bool CreateFromSkeletalMesh(UPhysicsAsset* PhysicsAsset, USkeletalMesh* SkelMesh
 				// Place joint at origin of child
 				CS->DefaultInstance.ConstraintBone1 = BoneName;
 				CS->DefaultInstance.Pos1 = FVector::ZeroVector;
-				CS->DefaultInstance.PriAxis1 = FVector(1,0,0);
-				CS->DefaultInstance.SecAxis1 = FVector(0,1,0);
+				CS->DefaultInstance.PriAxis1 = FVector(1, 0, 0);
+				CS->DefaultInstance.SecAxis1 = FVector(0, 1, 0);
 
 				CS->DefaultInstance.ConstraintBone2 = ParentName;
 				CS->DefaultInstance.Pos2 = RelTM.GetOrigin();
-				CS->DefaultInstance.PriAxis2 = RelTM.GetScaledAxis( EAxis::X );
-				CS->DefaultInstance.SecAxis2 = RelTM.GetScaledAxis( EAxis::Y );
+				CS->DefaultInstance.PriAxis2 = RelTM.GetScaledAxis(EAxis::X);
+				CS->DefaultInstance.SecAxis2 = RelTM.GetScaledAxis(EAxis::Y);
 
 				// Disable collision between constrained bodies by default.
 				PhysicsAsset->DisableCollision(NewBodyIndex, ParentBodyIndex);
@@ -203,7 +200,14 @@ bool CreateFromSkeletalMesh(UPhysicsAsset* PhysicsAsset, USkeletalMesh* SkelMesh
 		}
 	}
 
-	const bool bSuccess = PhysicsAsset->BodySetup.Num() > 0;
+	return PhysicsAsset->BodySetup.Num() > 0;
+}
+
+bool CreateFromSkeletalMesh(UPhysicsAsset* PhysicsAsset, USkeletalMesh* SkelMesh, FPhysAssetCreateParams& Params, FText & OutErrorMessage)
+{
+	PhysicsAsset->PreviewSkeletalMesh = SkelMesh;
+
+	bool bSuccess = CreateFromSkeletalMeshInternal(PhysicsAsset, SkelMesh, Params);
 	if (bSuccess)
 	{
 		// sets physics asset here now, so whoever creates 
@@ -213,13 +217,20 @@ bool CreateFromSkeletalMesh(UPhysicsAsset* PhysicsAsset, USkeletalMesh* SkelMesh
 	}
 	else
 	{
-		if( Params.bWalkPastSmall )
+		// try lower minimum bone size 
+		Params.MinBoneSize = 1.f;
+
+		bSuccess = CreateFromSkeletalMeshInternal(PhysicsAsset, SkelMesh, Params);
+		if(bSuccess)
 		{
-			OutErrorMessage = FText::Format(NSLOCTEXT("CreatePhysicsAsset", "CreatePhysicsAssetLinkFailedWalkPastSmallBones", "Minimum Bone Size too large, even with WalkPastSmallBones set, for all bones to create Physics Asset '{0}' from Skeletal Mesh '{1}'."), FText::FromString(PhysicsAsset->GetName()), FText::FromString(SkelMesh->GetName()));
+			// sets physics asset here now, so whoever creates 
+			// new physics asset from skeletalmesh will set properly here
+			SkelMesh->PhysicsAsset = PhysicsAsset;
+			SkelMesh->MarkPackageDirty();
 		}
 		else
 		{
-			OutErrorMessage = FText::Format(NSLOCTEXT("CreatePhysicsAsset", "CreatePhysicsAssetLinkFailed", "Minimum Bone Size too large for all bones to create Physics Asset '{0}' from Skeletal Mesh '{1}'."), FText::FromString(PhysicsAsset->GetName()), FText::FromString(SkelMesh->GetName()));
+			OutErrorMessage = FText::Format(NSLOCTEXT("CreatePhysicsAsset", "CreatePhysicsAssetLinkFailed", "The bone size is too small to create Physics Asset '{0}' from Skeletal Mesh '{1}'. You will have to create physics asset manually."), FText::FromString(PhysicsAsset->GetName()), FText::FromString(SkelMesh->GetName()));
 		}
 	}
 
