#summary 	Fixed crash in FixupRedirectsCommandlet
#ttp 331519 	CRASH: Occurs when Fixup Redirects attempts to check into Perforce.

diff --git a/Engine/Source/Editor/UnrealEd/Private/Commandlets/FixupRedirectsCommandlet.cpp b/Engine/Source/Editor/UnrealEd/Private/Commandlets/FixupRedirectsCommandlet.cpp
--- a/Engine/Source/Editor/UnrealEd/Private/Commandlets/FixupRedirectsCommandlet.cpp
+++ b/Engine/Source/Editor/UnrealEd/Private/Commandlets/FixupRedirectsCommandlet.cpp
@@ -756,7 +756,7 @@ int32 UFixupRedirectsCommandlet::Main( const FString& Params )
 			FString PackageName(FPackageName::FilenameToLongPackageName(Filename));
 			FSourceControlStatePtr SourceControlState = SourceControlProvider.GetState(SourceControlHelpers::PackageFilename(Filename), EStateCacheUsage::ForceUpdate);
 
-			if( SourceControlState.IsValid() && SourceControlState->IsCheckedOut() || SourceControlState->IsAdded() || SourceControlState->IsDeleted() )
+			if( SourceControlState.IsValid() && (SourceControlState->IsCheckedOut() || SourceControlState->IsAdded() || SourceControlState->IsDeleted()) )
 			{
 				// Only submit engine packages if we're requested to
 				if( bUpdateEnginePackages || !Filename.StartsWith( FPaths::EngineDir() ) )
