TTP 310741
* Better placement of the different color boxes for the curve editor.
* Clicking on nothing no longer deselects your currently selected color.

diff --git a/Engine/Source/Editor/UnrealEd/Private/SCurveEditor.cpp b/Engine/Source/Editor/UnrealEd/Private/SCurveEditor.cpp
--- a/Engine/Source/Editor/UnrealEd/Private/SCurveEditor.cpp
+++ b/Engine/Source/Editor/UnrealEd/Private/SCurveEditor.cpp
@@ -75,7 +75,7 @@ void SCurveEditor::Construct(const FArguments& InArgs)
 
 	SAssignNew(WarningMessageText, SErrorText );
 
-	auto CurveSelector	=	
+	TSharedRef<SBorder> CurveSelector =
 		SNew(SBorder)
 		.BorderImage( FEditorStyle::GetBrush("NoBorder") )
 		.Visibility(this, &SCurveEditor::GetControlVisibility)
@@ -88,14 +88,22 @@ void SCurveEditor::Construct(const FArguments& InArgs)
 	this->ChildSlot
 	[
 		SNew( SVerticalBox )
+
 		+ SVerticalBox::Slot()
 		.FillHeight(1.0f)
 		[
 			SNew(SHorizontalBox)
 			.Visibility( this, &SCurveEditor::GetCurveAreaVisibility )
+
+			+ SHorizontalBox::Slot()
+			.Padding(FMargin(20, 12, 0, 0))
+			[
+				SNew(SSpacer)
+			]
 			
 			+ SHorizontalBox::Slot()
 			.AutoWidth()
+			.Padding(FMargin(2, 12, 0, 0))
 			[
 				CurveSelector
 			]
@@ -109,7 +117,7 @@ void SCurveEditor::Construct(const FArguments& InArgs)
 				.BorderImage( FEditorStyle::GetBrush("NoBorder") )
 				.DesiredSizeScale(FVector2D(256.0f,32.0f))
 				.Visibility(this, &SCurveEditor::GetControlVisibility)
-				.Padding(FMargin(18, 12))
+				.Padding(FMargin(2, 12, 0, 0))
 				[
 					SNew(SHorizontalBox)
 					
@@ -239,13 +247,14 @@ SCurveEditor::~SCurveEditor()
 
 TSharedRef<SWidget> SCurveEditor::CreateCurveSelectionWidget() const
 {
-	auto Box =  SNew(SVerticalBox);
+	TSharedRef<SVerticalBox> Box = SNew(SVerticalBox);
+
 	//Curve selection buttons only useful if you have more then 1 curve
 	if(Curves.Num() > 1)
 	{
-		for(auto It = Curves.CreateConstIterator();It;++It)
+		for ( const FRichCurveEditInfo& CurveInfo : Curves )
 		{
-			FRichCurve* Curve = It->CurveToEdit;
+			FRichCurve* Curve = CurveInfo.CurveToEdit;
 
 			Box->AddSlot()
 			.AutoHeight()
@@ -1733,11 +1742,6 @@ bool SCurveEditor::HitTestCurves(  const FGeometry& InMyGeometry, const FPointer
 		}
 	}
 
-	if(!bHit)
-	{
-		//Clicked on emptyness, so clear curve selection
-		SelectedCurves.Empty();
-	}
 	return bHit;
 }
 
