Curve editor improvements
* When doing hit tests on curves we prioritize the one that is already selected.
* Fixed a uninitialized drag bug that would cause the context menu to always show up after first drag.

diff --git a/Engine/Source/Editor/UnrealEd/Private/SCurveEditor.cpp b/Engine/Source/Editor/UnrealEd/Private/SCurveEditor.cpp
--- a/Engine/Source/Editor/UnrealEd/Private/SCurveEditor.cpp
+++ b/Engine/Source/Editor/UnrealEd/Private/SCurveEditor.cpp
@@ -41,6 +41,8 @@ void SCurveEditor::Construct(const FArguments& InArgs)
 
 	bIsUsingSlider = false;
 
+	ResetDragValues();
+
 	// if editor size is set, use it, otherwise, use default value
 	if (DesiredSize.Get().IsZero())
 	{
@@ -61,9 +63,9 @@ void SCurveEditor::Construct(const FArguments& InArgs)
 	bMovingTangent = false;
 
 	//Simple r/g/b for now
-	CurveColors.Add(FLinearColor(1.0f,0.0f,0.0f));
-	CurveColors.Add(FLinearColor(0.0f,1.0f,0.0f));
-	CurveColors.Add(FLinearColor(0.0f,0.0f,1.0f));
+	CurveColors.Add(FLinearColor(1.0f, 0.0f, 0.0f));
+	CurveColors.Add(FLinearColor(0.0f, 1.0f, 0.0f));
+	CurveColors.Add(FLinearColor(0.0f, 0.0f, 1.0f));
 
 	TransactionIndex = -1;
 
@@ -75,8 +77,7 @@ void SCurveEditor::Construct(const FArguments& InArgs)
 
 	SAssignNew(WarningMessageText, SErrorText );
 
-	TSharedRef<SBorder> CurveSelector =
-		SNew(SBorder)
+	TSharedRef<SBorder> CurveSelector = SNew(SBorder)
 		.BorderImage( FEditorStyle::GetBrush("NoBorder") )
 		.Visibility(this, &SCurveEditor::GetControlVisibility)
 		[
@@ -268,7 +269,7 @@ TSharedRef<SWidget> SCurveEditor::CreateCurveSelectionWidget() const
 					.IsChecked(this, &SCurveEditor::GetCheckState, Curve)
 					.OnCheckStateChanged(this, &SCurveEditor::OnUserSelectedCurve, Curve)
 					.ToolTipText(this, &SCurveEditor::GetCurveCheckBoxToolTip, Curve)
-				]		
+				]
 			];
 		}
 	}
@@ -442,21 +443,21 @@ int32 SCurveEditor::OnPaint(const FGeometry& AllottedGeometry, const FSlateRect&
 			const int32 UnSelectedCurveLayer = LayerId+1;
 
 			//Paint the curves, any selected curves will be on top
-			for(auto It = Curves.CreateConstIterator();It;++It)
+			for ( const FRichCurveEditInfo& CurveInfo : Curves )
 			{
-				FRichCurve* Curve = It->CurveToEdit;
+				FRichCurve* Curve = CurveInfo.CurveToEdit;
 				int32 LayerToUse = (SelectedCurves.Find(Curve) != INDEX_NONE) ? SelectedCurveLayer : UnSelectedCurveLayer;
 				{
 					PaintCurve(Curve, CurveAreaGeometry, ScaleInfo, OutDrawElements, LayerToUse, MyClippingRect, DrawEffects, InWidgetStyle);
 				}
 			}
-			LayerId = LayerId+2;
+			LayerId = LayerId + 2;
 		}
 
 		//Paint the keys on top of the curve
-		for(auto It = Curves.CreateConstIterator();It;++It)
+		for ( const FRichCurveEditInfo& CurveInfo : Curves )
 		{
-			FRichCurve* Curve = It->CurveToEdit;
+			FRichCurve* Curve = CurveInfo.CurveToEdit;
 			LayerId = PaintKeys(Curve, ScaleInfo, OutDrawElements, LayerId, CurveAreaGeometry, MyClippingRect, DrawEffects, InWidgetStyle);
 		}
 	}
@@ -469,12 +470,12 @@ int32 SCurveEditor::OnPaint(const FGeometry& AllottedGeometry, const FSlateRect&
 	return LayerId;
 }
 
-void SCurveEditor::PaintCurve(  FRichCurve* Curve, const FGeometry &AllottedGeometry, FTrackScaleInfo &ScaleInfo, FSlateWindowElementList &OutDrawElements, 
+void SCurveEditor::PaintCurve( FRichCurve* Curve, const FGeometry &AllottedGeometry, FTrackScaleInfo &ScaleInfo, FSlateWindowElementList &OutDrawElements, 
 	int32 LayerId, const FSlateRect& MyClippingRect, ESlateDrawEffect::Type DrawEffects, const FWidgetStyle &InWidgetStyle )const
 {
-	if(Curve != NULL)
+	if (Curve != NULL)
 	{
-		if(bDrawCurve)
+		if (bDrawCurve)
 		{
 			FLinearColor Color = InWidgetStyle.GetColorAndOpacityTint() * GetCurveColor(Curve);
 
@@ -487,7 +488,7 @@ void SCurveEditor::PaintCurve(  FRichCurve* Curve, const FGeometry &AllottedGeom
 			TArray<FVector2D> LinePoints;
 			int32 CurveDrawInterval = 1;
 
-			if(Curve->GetNumKeys() < 2) 
+			if (Curve->GetNumKeys() < 2) 
 			{
 				//Not enough point, just draw flat line
 				float Value = Curve->Eval(0.0f);
@@ -495,7 +496,7 @@ void SCurveEditor::PaintCurve(  FRichCurve* Curve, const FGeometry &AllottedGeom
 				LinePoints.Add(FVector2D(0.0f, Y));
 				LinePoints.Add(FVector2D(AllottedGeometry.Size.X, Y));
 
-				FSlateDrawElement::MakeLines(OutDrawElements,LayerId,AllottedGeometry.ToPaintGeometry(),LinePoints,MyClippingRect,DrawEffects,Color);
+				FSlateDrawElement::MakeLines(OutDrawElements, LayerId, AllottedGeometry.ToPaintGeometry(), LinePoints, MyClippingRect, DrawEffects, Color);
 				LinePoints.Empty();
 			}
 			else
@@ -513,13 +514,13 @@ void SCurveEditor::PaintCurve(  FRichCurve* Curve, const FGeometry &AllottedGeom
 					//Arrival line
 					LinePoints.Add(FVector2D(0.0f, ArriveY));
 					LinePoints.Add(FVector2D(ArriveX, ArriveY));
-					FSlateDrawElement::MakeLines(OutDrawElements,LayerId,AllottedGeometry.ToPaintGeometry(),LinePoints,MyClippingRect,DrawEffects,Color);
+					FSlateDrawElement::MakeLines(OutDrawElements, LayerId, AllottedGeometry.ToPaintGeometry(), LinePoints, MyClippingRect, DrawEffects, Color);
 					LinePoints.Empty();
 
 					//Leave line
 					LinePoints.Add(FVector2D(AllottedGeometry.Size.X, LeaveY));
 					LinePoints.Add(FVector2D(LeaveX, LeaveY));
-					FSlateDrawElement::MakeLines(OutDrawElements,LayerId,AllottedGeometry.ToPaintGeometry(),LinePoints,MyClippingRect,DrawEffects,Color);
+					FSlateDrawElement::MakeLines(OutDrawElements, LayerId, AllottedGeometry.ToPaintGeometry(), LinePoints, MyClippingRect, DrawEffects, Color);
 					LinePoints.Empty();
 				}
 
@@ -529,7 +530,7 @@ void SCurveEditor::PaintCurve(  FRichCurve* Curve, const FGeometry &AllottedGeom
 				for(int32 i = 0;i<Keys.Num()-1;++i)
 				{
 					CreateLinesForSegment(Curve, Keys[i], Keys[i+1],LinePoints, ScaleInfo);
-					FSlateDrawElement::MakeLines(OutDrawElements,LayerId,AllottedGeometry.ToPaintGeometry(),LinePoints,MyClippingRect,DrawEffects,Color);
+					FSlateDrawElement::MakeLines(OutDrawElements, LayerId, AllottedGeometry.ToPaintGeometry(), LinePoints, MyClippingRect, DrawEffects, Color);
 					LinePoints.Empty();
 				}
 			}
@@ -1104,7 +1105,7 @@ FReply SCurveEditor::OnMouseMove( const FGeometry& InMyGeometry, const FPointerE
 	if( Curve != NULL)
 	{
 		// When mouse moves, if we are moving a key, update its 'input' position
-		if(bPanning  || bMovingKeys || bMovingTangent)
+		if (bPanning || bMovingKeys || bMovingTangent)
 		{
 			FTrackScaleInfo ScaleInfo(ViewMinInput.Get(),  ViewMaxInput.Get(), ViewMinOutput, ViewMaxOutput, InMyGeometry.Size);
 			FVector2D ScreenDelta = InMouseEvent.GetCursorDelta();
@@ -1673,12 +1674,13 @@ FLinearColor SCurveEditor::GetCurveColor( FRichCurve* TestCurve ) const
 	check(TestCurve);
 	for(int32 I = 0;I != Curves.Num() && I < CurveColors.Num();I++)
 	{
-		FRichCurve* Curve =Curves[I].CurveToEdit;
+		FRichCurve* Curve = Curves[I].CurveToEdit;
 		if(Curve == TestCurve)
 		{
 			return CurveColors[I];
 		}
 	}
+
 	return FLinearColor::Black;
 }
 
@@ -1721,7 +1723,7 @@ bool SCurveEditor::HitTestCurves(  const FGeometry& InMyGeometry, const FPointer
 
 		const FVector2D HitPosition = InMyGeometry.AbsoluteToLocal( InMouseEvent.GetScreenSpacePosition()  );
 
-
+		TArray<FRichCurve*> CurvesHit;
 
 		for(auto It = Curves.CreateConstIterator();It;++It)
 		{
@@ -1731,15 +1733,55 @@ bool SCurveEditor::HitTestCurves(  const FGeometry& InMyGeometry, const FPointer
 				float Time		 = ScaleInfo.LocalXToInput(HitPosition.X);
 				float KeyScreenY = ScaleInfo.OutputToLocalY(Curve->Eval(Time));
 
-
 				if( HitPosition.Y > (KeyScreenY - (0.5f * CONST_CurveSize.Y)) &&
 					HitPosition.Y < (KeyScreenY + (0.5f * CONST_CurveSize.Y)))
 				{
-					SelectCurve(Curve, bControlDown); 
+					CurvesHit.Add(Curve);
 					bHit = true;
 				}
 			}
 		}
+
+		// Select the curves that were hit
+		if ( bHit )
+		{
+			if ( bControlDown )
+			{
+				// If control is down and we have to do multi-selection then don't worry about order, select everything that was hit.
+				for ( FRichCurve* Curve : CurvesHit )
+				{
+					SelectCurve(Curve, bControlDown);
+				}
+			}
+			else
+			{
+				FRichCurve* AlreadySelectedCurve = NULL;
+
+				// Without control down we want to bias the decision towards the curve that is already selected.
+				for ( FRichCurve* Curve : CurvesHit )
+				{
+					if ( SelectedCurves.Contains(Curve) )
+					{
+						AlreadySelectedCurve = Curve;
+						break;
+					}
+				}
+
+				if ( AlreadySelectedCurve != NULL )
+				{
+					// Select the already selected curve
+					SelectCurve(AlreadySelectedCurve, bControlDown);
+				}
+				else
+				{
+					// Select which ever curve was hit last
+					for ( FRichCurve* Curve : CurvesHit )
+					{
+						SelectCurve(Curve, bControlDown);
+					}
+				}
+			}
+		}
 	}
 
 	return bHit;
