#ttp 330870 - UE4: FBX import type selector is confusing
#branch UE4
#proj Editor.UnrealEd
#summary Added extra text indicating that a particular import type has been auto-detected, and a warning which is displayed if the user overrides this detected import type.
#add Added two extra VerticalBox slots in SImportTypeButton, displaying the auto-detected import type, and a warning which is displayed if the user overrides this type.
#change Changed the import type string and tooltip to be fetched from the EFBXImportType metadata instead of specified in SImportTypeButton::Construct.
#change Changed tooltip metadata in EFBXImportType to match the tooltip text in SImportTypeButton::Construct.
#change Changed all FString references to FText in FbxOptionWindow.cpp.

diff --git a/Engine/Source/Editor/UnrealEd/Classes/Factories/FbxImportUI.h b/Engine/Source/Editor/UnrealEd/Classes/Factories/FbxImportUI.h
--- a/Engine/Source/Editor/UnrealEd/Classes/Factories/FbxImportUI.h
+++ b/Engine/Source/Editor/UnrealEd/Classes/Factories/FbxImportUI.h
@@ -13,11 +13,11 @@
 UENUM()
 enum EFBXImportType
 {
-	/** Static mesh */
+	/** Select Static Mesh if you'd like to import static mesh. */
 	FBXIT_StaticMesh UMETA(DisplayName="Static Mesh"),
-	/** Skeletal mesh */
+	/** Select Skeletal Mesh if you'd like to import skeletal mesh. */
 	FBXIT_SkeletalMesh UMETA(DisplayName="Skeletal Mesh"),
-	/** Animation */
+	/** Select Animation if you'd like to import only animation. */
 	FBXIT_Animation UMETA(DisplayName="Animation"),
 	FBXIT_MAX,
 };
diff --git a/Engine/Source/Editor/UnrealEd/Private/Fbx/FbxOptionWindow.cpp b/Engine/Source/Editor/UnrealEd/Private/Fbx/FbxOptionWindow.cpp
--- a/Engine/Source/Editor/UnrealEd/Private/Fbx/FbxOptionWindow.cpp
+++ b/Engine/Source/Editor/UnrealEd/Private/Fbx/FbxOptionWindow.cpp
@@ -21,34 +21,61 @@ public:
 
 	void Construct( const FArguments& InArgs )
 	{
+		DetectedImportType = InArgs._DefaultImportType;
 		CurrentChoice = InArgs._DefaultImportType;
 		OnSelectionChanged = InArgs._OnSelectionChanged;
 
+		const UEnum* ImportTypeEnum = FindObject<UEnum>(ANY_PACKAGE, TEXT("EFBXImportType"), true);
+		check(ImportTypeEnum != nullptr);
+
+		FText AutoDetected = FText::Format(LOCTEXT("FbxOptionWindow_AutoDetectedPrefix", "Auto-detected: {0}"), ImportTypeEnum->GetEnumText(CurrentChoice));
+
 		this->ChildSlot
 		[
-			SNew(SVerticalBox)
-
-			+SVerticalBox::Slot().AutoHeight() .Padding(0, 2)
+			SNew(SBox)
+			// TODO: Ideally I would like to specify that the VerticalBox slot containing the radio buttons defines the width of the VerticalBox.
+			// While that's not an option, the width is fixed, so that the warning STextBlock slot doesn't expand the window to its length.
+			.WidthOverride(400.0f)
 			[
-				SNew(STextBlock)
-				.TextStyle(FEditorStyle::Get(), "FBXLargeFont")
-				.Text(LOCTEXT("FbxOptionWindow_SelectType", "Select Asset Type").ToString())
-			]
+				SNew(SVerticalBox)
 
-			+SVerticalBox::Slot().AutoHeight()
-			[
-				SNew(SUniformGridPanel)
-				+ SUniformGridPanel::Slot(0, 0)
+				+SVerticalBox::Slot().AutoHeight().Padding(0, 2)
 				[
-					CreateRadioButton(LOCTEXT("ImportTypeButton_StaticMesh", "Static Mesh").ToString(), FBXIT_StaticMesh, LOCTEXT("FBXImportOption_StaticMesh", "Select StaticMesh if you'd like to import staticmesh."), FEditorStyle::GetBrush("FBXIcon.StaticMesh"), 0)
+					SNew(STextBlock)
+					.TextStyle(FEditorStyle::Get(), "FBXLargeFont")
+					.Text(LOCTEXT("FbxOptionWindow_SelectType", "Select Asset Type"))
 				]
-				+ SUniformGridPanel::Slot(1, 0)
+
+				+SVerticalBox::Slot().AutoHeight().Padding(0, 5)
 				[
-					CreateRadioButton(LOCTEXT("ImportTypeButton_SkeletalMesh", "Skeletal Mesh").ToString(), FBXIT_SkeletalMesh, LOCTEXT("FBXImportOption_SkeletalMesh", "Select SkeletalMesh if you'd like to import skeletalmesh."), FEditorStyle::GetBrush("FBXIcon.SkeletalMesh"), 1)
+					SNew(STextBlock)
+					.Text(AutoDetected)
 				]
-				+ SUniformGridPanel::Slot(2, 0)
+
+				+SVerticalBox::Slot().AutoHeight()
+				[
+					SNew(SUniformGridPanel)
+					+SUniformGridPanel::Slot(0, 0)
+					[
+						CreateRadioButton(*ImportTypeEnum, FBXIT_StaticMesh, FEditorStyle::GetBrush("FBXIcon.StaticMesh"), 0)
+					]
+					+SUniformGridPanel::Slot(1, 0)
+					[
+						CreateRadioButton(*ImportTypeEnum, FBXIT_SkeletalMesh, FEditorStyle::GetBrush("FBXIcon.SkeletalMesh"), 1)
+					]
+					+SUniformGridPanel::Slot(2, 0)
+					[
+						CreateRadioButton(*ImportTypeEnum, FBXIT_Animation, FEditorStyle::GetBrush("FBXIcon.Animation"), 2)
+					]
+				]
+
+				+SVerticalBox::Slot().AutoHeight().Padding(0, 5)
 				[
-					CreateRadioButton(LOCTEXT("ImportTypeButton_Animation", "Animation").ToString(), FBXIT_Animation, LOCTEXT("FBXImportOption_Animation", "Select Animation if you'd like to import only animation."), FEditorStyle::GetBrush("FBXIcon.Animation"), 2)
+					SNew(STextBlock)
+					.Text(LOCTEXT("FbxOptionWindow_TypeMismatch", "Warning: Asset does not appear to correspond to selected type, and may not import correctly."))
+					.ColorAndOpacity(FColor::Yellow)
+					.AutoWrapText(true)
+					.Visibility(this, &SImportTypeButton::IsWarningVisible)
 				]
 			]
 		];
@@ -56,14 +83,14 @@ public:
 
 private:
 
-	TSharedRef<SWidget> CreateRadioButton( const FString& RadioText, EFBXImportType RadioButtonChoice, const FText & ToolTipText, const FSlateBrush* Brush, int32 Position)
+	TSharedRef<SWidget> CreateRadioButton( const UEnum& Enum, EFBXImportType RadioButtonChoice, const FSlateBrush* Brush, int32 Position)
 	{
 		return
 			SNew(SCheckBox)
 			.Style(FEditorStyle::Get(), (Position == 0) ? "Property.ToggleButton.Start" : ((Position==1)? "Property.ToggleButton.Middle" : "Property.ToggleButton.End") )
 			.IsChecked( this, &SImportTypeButton::IsRadioChecked, RadioButtonChoice )
 			.OnCheckStateChanged( this, &SImportTypeButton::OnRadioChanged, RadioButtonChoice )
-			.ToolTipText(ToolTipText)
+			.ToolTipText(Enum.GetToolTipText(RadioButtonChoice))
 			[
 				SNew(SHorizontalBox)
 
@@ -83,7 +110,7 @@ private:
 				.Padding(6, 2)
 				[
 					SNew(STextBlock)
-					.Text(RadioText)
+					.Text(Enum.GetEnumText(RadioButtonChoice))
 					.ColorAndOpacity(FLinearColor(0.72f, 0.72f, 0.72f, 1.f))
 				]
 			];
@@ -96,6 +123,13 @@ private:
 			: ESlateCheckBoxState::Unchecked;
 	}
 
+	EVisibility IsWarningVisible() const
+	{
+		return (CurrentChoice == DetectedImportType)
+			? EVisibility::Collapsed
+			: EVisibility::Visible;
+	}
+
 	void OnRadioChanged( ESlateCheckBoxState::Type NewRadioState, EFBXImportType RadioThatChanged )
 	{
 		if (NewRadioState == ESlateCheckBoxState::Checked)
@@ -109,6 +143,7 @@ private:
 		}
 	}
 
+	EFBXImportType DetectedImportType;
 	EFBXImportType CurrentChoice;
 	FOnImportTypeChanged OnSelectionChanged;
 };
@@ -140,11 +175,11 @@ public:
 			SNew(SVerticalBox)
 			+SVerticalBox::Slot().AutoHeight() .Padding(1)
 			[
-				CreateRadioButton( LOCTEXT("AnimImportLengthOption_ExportTime", "Exported Time").ToString(), FBXALIT_ExportedTime )
+				CreateRadioButton( LOCTEXT("AnimImportLengthOption_ExportTime", "Exported Time"), FBXALIT_ExportedTime )
 			]
 			+SVerticalBox::Slot().AutoHeight() .Padding(1)
 			[
-				CreateRadioButton( LOCTEXT("AnimImportLengthOption_AnimTime", "Animated Time").ToString(), FBXALIT_AnimatedKey )
+				CreateRadioButton( LOCTEXT("AnimImportLengthOption_AnimTime", "Animated Time"), FBXALIT_AnimatedKey )
 			]
 			+SVerticalBox::Slot().AutoHeight() .Padding(1)
 			[
@@ -152,7 +187,7 @@ public:
 
 				+SHorizontalBox::Slot().AutoWidth() .Padding(2)
 				[
-					CreateRadioButton( LOCTEXT("AnimImportLengthOption_SetRange", "Set Range").ToString(), FBXALIT_SetRange )
+					CreateRadioButton( LOCTEXT("AnimImportLengthOption_SetRange", "Set Range"), FBXALIT_SetRange )
 				]
 
 				+SHorizontalBox::Slot().AutoWidth() .Padding(2) 
@@ -179,7 +214,7 @@ public:
 				+SHorizontalBox::Slot().AutoWidth() .Padding(2)
 				[
 					SNew(STextBlock)
-					.Text(LOCTEXT("AnimImportLengthOption_InFrameNumber", " In Frame Number").ToString())
+					.Text(LOCTEXT("AnimImportLengthOption_InFrameNumber", " In Frame Number"))
 				]
 			]
 		];
@@ -187,7 +222,7 @@ public:
 
 private:
 
-	TSharedRef<SWidget> CreateRadioButton( const FString& RadioText, EFBXAnimationLengthImportType RadioButtonChoice )
+	TSharedRef<SWidget> CreateRadioButton( const FText& RadioText, EFBXAnimationLengthImportType RadioButtonChoice )
 	{
 		return
 			SNew(SCheckBox)
@@ -313,7 +348,7 @@ void SFbxOptionWindow::Construct(const FArguments& InArgs)
 				[
 					SNew(SButton)
 					.HAlign(HAlign_Center)
-					.Text(LOCTEXT("FbxOptionWindow_ImportAll", "Import All").ToString())
+					.Text(LOCTEXT("FbxOptionWindow_ImportAll", "Import All"))
 					.ToolTipText( LOCTEXT("FbxOptionWindow_ImportAll_ToolTip", "Import all files with these same settings") )
 					.IsEnabled(this, &SFbxOptionWindow::CanImport)
 					.OnClicked( this, &SFbxOptionWindow::OnImportAll )
@@ -322,7 +357,7 @@ void SFbxOptionWindow::Construct(const FArguments& InArgs)
 				[
 					SAssignNew( ImportButton, SButton)
 					.HAlign(HAlign_Center)
-					.Text(LOCTEXT("FbxOptionWindow_Import", "Import").ToString())
+					.Text(LOCTEXT("FbxOptionWindow_Import", "Import"))
 					.IsEnabled(this, &SFbxOptionWindow::CanImport)
 					.OnClicked( this, &SFbxOptionWindow::OnImport )
 				]
@@ -330,8 +365,8 @@ void SFbxOptionWindow::Construct(const FArguments& InArgs)
 				[
 					SNew(SButton) 
 					.HAlign(HAlign_Center)
-					.Text(LOCTEXT("FbxOptionWindow_Cancel", "Cancel").ToString())
-					.ToolTipText(LOCTEXT("FbxOptionWindow_Cancel_ToolTip", "Cancels importing this FBX file").ToString())
+					.Text(LOCTEXT("FbxOptionWindow_Cancel", "Cancel"))
+					.ToolTipText(LOCTEXT("FbxOptionWindow_Cancel_ToolTip", "Cancels importing this FBX file"))
 					.OnClicked( this, &SFbxOptionWindow::OnCancel )
 				]
 			]
@@ -490,7 +525,7 @@ TSharedRef<SWidget> SFbxOptionWindow::ConstructMiscOption()
 	[
 		SNew(STextBlock)
 		.TextStyle( FEditorStyle::Get(), "FBXSmallFont" ) 
-		.Text( LOCTEXT("FbxOptionWindow_Misc", "Misc").ToString() )
+		.Text( LOCTEXT("FbxOptionWindow_Misc", "Misc") )
 	];
 
 	NewBox->AddSlot().AutoHeight() .Padding(2)
@@ -500,7 +535,7 @@ TSharedRef<SWidget> SFbxOptionWindow::ConstructMiscOption()
 		.OnCheckStateChanged(this, &SFbxOptionWindow::SetGeneral_OverrideFullName)
 		[
 			SNew(STextBlock)
-			.Text(LOCTEXT("FbxOptionWindow_NameOverride", "Override FullName").ToString())
+			.Text(LOCTEXT("FbxOptionWindow_NameOverride", "Override FullName"))
 		]
 	];
 
@@ -531,7 +566,7 @@ TSharedRef<SWidget> SFbxOptionWindow::ConstructNormalImportOptions()
 		.AutoWidth()
 		[
 			SNew( STextBlock )
-			.Text(LOCTEXT("FBXOptionWindow_NormalImport", "Normals").ToString())
+			.Text(LOCTEXT("FBXOptionWindow_NormalImport", "Normals"))
 		]
 		+ SHorizontalBox::Slot()
 		.VAlign(VAlign_Center)
@@ -540,7 +575,7 @@ TSharedRef<SWidget> SFbxOptionWindow::ConstructNormalImportOptions()
 		[
 			SNew( SComboBox< TSharedPtr<EFBXNormalImportMethod> > )
 			.ContentPadding(1.f)
-			.ToolTipText(LOCTEXT("FBXOptionWindow_NormalInputMethod", "Options for importing normals and tangents").ToString())
+			.ToolTipText(LOCTEXT("FBXOptionWindow_NormalInputMethod", "Options for importing normals and tangents"))
 			.OptionsSource( &NormalImportOptions )
 			.InitiallySelectedItem( NormalImportOptions[ GetCurrentNormalImportMethod() ] )
 			.OnSelectionChanged( this, &SFbxOptionWindow::OnNormalImportMethodChanged )
@@ -566,7 +601,7 @@ TSharedRef<SWidget> SFbxOptionWindow::ConstructStaticMeshBasic()
 			.VAlign(VAlign_Center)
 		[
 			SNew(STextBlock)
-			.Text(LOCTEXT("FbxOptionWindow_StaticMeshLODGroup", "LOD Group").ToString())
+			.Text(LOCTEXT("FbxOptionWindow_StaticMeshLODGroup", "LOD Group"))
 		]
 
 		+SHorizontalBox::Slot()
@@ -596,7 +631,7 @@ TSharedRef<SWidget> SFbxOptionWindow::ConstructStaticMeshAdvanced()
 	[
 		SNew(STextBlock)
 		.TextStyle( FEditorStyle::Get(), "FBXSmallFont" ) 
-		.Text( LOCTEXT("FbxOptionWindow_StaticMesh", "Mesh").ToString())
+		.Text( LOCTEXT("FbxOptionWindow_StaticMesh", "Mesh"))
 	];
 
 
@@ -608,7 +643,7 @@ TSharedRef<SWidget> SFbxOptionWindow::ConstructStaticMeshAdvanced()
 		.Content()
 		[
 			SNew(STextBlock)
-			.Text(LOCTEXT("FbxOptionWindow_CombineMeshes", "Combine Meshes").ToString())
+			.Text(LOCTEXT("FbxOptionWindow_CombineMeshes", "Combine Meshes"))
 		]
 	];
 
@@ -620,7 +655,7 @@ TSharedRef<SWidget> SFbxOptionWindow::ConstructStaticMeshAdvanced()
 		.OnCheckStateChanged(this, &SFbxOptionWindow::SetStaticMesh_ImportMeshLODs)
 		[
 			SNew(STextBlock)
-			.Text(LOCTEXT("FbxOptionWindow_ImportMeshLODs", "Import Mesh LODs").ToString())
+			.Text(LOCTEXT("FbxOptionWindow_ImportMeshLODs", "Import Mesh LODs"))
 		]
 	];
 
@@ -631,7 +666,7 @@ TSharedRef<SWidget> SFbxOptionWindow::ConstructStaticMeshAdvanced()
 		.OnCheckStateChanged(this, &SFbxOptionWindow::SetStaticMesh_ReplaceVertexColor)
 		[
 			SNew(STextBlock)
-			.Text(LOCTEXT("FbxOptionWindow_ReplaceVertexColors", "Replace Vertex Colors").ToString())
+			.Text(LOCTEXT("FbxOptionWindow_ReplaceVertexColors", "Replace Vertex Colors"))
 		]
 	];
 
@@ -642,7 +677,7 @@ TSharedRef<SWidget> SFbxOptionWindow::ConstructStaticMeshAdvanced()
 		.OnCheckStateChanged(this, &SFbxOptionWindow::SetStaticMesh_RemoveDegenerates)
 		[
 			SNew(STextBlock)
-			.Text(LOCTEXT("FbxOptionWindow_RemoveDegeneates", "Remove Degenerates").ToString())
+			.Text(LOCTEXT("FbxOptionWindow_RemoveDegeneates", "Remove Degenerates"))
 		]
 	];
 
@@ -653,7 +688,7 @@ TSharedRef<SWidget> SFbxOptionWindow::ConstructStaticMeshAdvanced()
 		.OnCheckStateChanged(this, &SFbxOptionWindow::SetStaticMesh_OneConvexHullPerUCX)
 		[
 			SNew(STextBlock)
-			.Text(LOCTEXT("FbxOptionWindow_OneConvexHullPerUCX", "One Convex Hull Per UCX").ToString())
+			.Text(LOCTEXT("FbxOptionWindow_OneConvexHullPerUCX", "One Convex Hull Per UCX"))
 		]
 	];
 
@@ -678,7 +713,7 @@ TSharedRef<SWidget> SFbxOptionWindow::ConstructSkeletalMeshBasic()
 		.OnCheckStateChanged(this, &SFbxOptionWindow::SetSkeletalMesh_ImportMorphTargets)
 		[
 			SNew(STextBlock)
-			.Text(LOCTEXT("FbxOptionWindow_ImportMorphTargets", "Import Morph Targets").ToString())
+			.Text(LOCTEXT("FbxOptionWindow_ImportMorphTargets", "Import Morph Targets"))
 		]
 	];
 
@@ -694,7 +729,7 @@ TSharedRef<SWidget> SFbxOptionWindow::ConstructSkeletalMeshBasic()
 			.OnCheckStateChanged(this, &SFbxOptionWindow::SetSkeletalMesh_ImportAnimation)
 			[
 				SNew(STextBlock)
-				.Text(LOCTEXT("FbxOptionWindow_ImportAnimation", "Import Animation").ToString())
+				.Text(LOCTEXT("FbxOptionWindow_ImportAnimation", "Import Animation"))
 			]
 		]
 		+SHorizontalBox::Slot()
@@ -703,7 +738,7 @@ TSharedRef<SWidget> SFbxOptionWindow::ConstructSkeletalMeshBasic()
 		[
 			SNew(SEditableTextBox)
 			.IsReadOnly(false)
-			.ToolTipText(LOCTEXT("FbxOptionWindow_ImportAnimationToolTip", "Type animation name if you're importing animation").ToString())
+			.ToolTipText(LOCTEXT("FbxOptionWindow_ImportAnimationToolTip", "Type animation name if you're importing animation"))
 			.Text( this, &SFbxOptionWindow::GetAnimationName)
 			.MinDesiredWidth(50)
 			.OnTextCommitted(this, &SFbxOptionWindow::SetSkeletalMesh_AnimationName)
@@ -718,7 +753,7 @@ TSharedRef<SWidget> SFbxOptionWindow::ConstructSkeletalMeshBasic()
 		.OnCheckStateChanged(this, &SFbxOptionWindow::SetSkeletalMesh_UpdateSkeletonRefPose)
 		[
 			SNew(STextBlock)
-			.Text(LOCTEXT("FbxOptionWindow_UpdateSkeletonRefPose", "Update Skeleton Reference Pose").ToString())
+			.Text(LOCTEXT("FbxOptionWindow_UpdateSkeletonRefPose", "Update Skeleton Reference Pose"))
 		]
 	];
 
@@ -734,7 +769,7 @@ TSharedRef<SWidget> SFbxOptionWindow::ConstructSkeletalMeshAdvanced()
 	[
 		SNew(STextBlock)
 		.TextStyle( FEditorStyle::Get(), "FBXSmallFont" ) 
-		.Text(LOCTEXT("FbxOptionWindow_SkeletalMesh", "Mesh").ToString())
+		.Text(LOCTEXT("FbxOptionWindow_SkeletalMesh", "Mesh"))
 	];
 
 	NewBox->AddSlot().AutoHeight() .Padding(2)
@@ -744,7 +779,7 @@ TSharedRef<SWidget> SFbxOptionWindow::ConstructSkeletalMeshAdvanced()
 		.OnCheckStateChanged(this, &SFbxOptionWindow::SetSkeletalMesh_ImportMeshLODs)
 		[
 			SNew(STextBlock)
-			.Text(LOCTEXT("FbxOptionWindow_ImportSkeltalMeshLODs", "Import Mesh LODs").ToString())
+			.Text(LOCTEXT("FbxOptionWindow_ImportSkeltalMeshLODs", "Import Mesh LODs"))
 		]
 	];
 
@@ -755,7 +790,7 @@ TSharedRef<SWidget> SFbxOptionWindow::ConstructSkeletalMeshAdvanced()
 		.OnCheckStateChanged(this, &SFbxOptionWindow::SetSkeletalMesh_ImportRigidMesh)
 		[
 			SNew(STextBlock)
-			.Text(LOCTEXT("FbxOptionWindow_ImportRigidMesh", "Import Rigid Mesh").ToString())
+			.Text(LOCTEXT("FbxOptionWindow_ImportRigidMesh", "Import Rigid Mesh"))
 		]
 	];
 
@@ -766,7 +801,7 @@ TSharedRef<SWidget> SFbxOptionWindow::ConstructSkeletalMeshAdvanced()
 		.OnCheckStateChanged(this, &SFbxOptionWindow::SetSkeletalMesh_UseT0AsRefPose)
 		[
 			SNew(STextBlock)
-			.Text(LOCTEXT("FbxOptionWindow_UseRefPose", "Use T0 as Ref Pose").ToString())
+			.Text(LOCTEXT("FbxOptionWindow_UseRefPose", "Use T0 as Ref Pose"))
 		]
 	];
 
@@ -777,7 +812,7 @@ TSharedRef<SWidget> SFbxOptionWindow::ConstructSkeletalMeshAdvanced()
 		.OnCheckStateChanged(this, &SFbxOptionWindow::SetSkeletalMesh_ReserveSmoothingGroups)
 		[
 			SNew(STextBlock)
-			.Text(LOCTEXT("FbxOptionWindow_PreserveSmoothingGroup", "Preserve Smoothing Groups").ToString())
+			.Text(LOCTEXT("FbxOptionWindow_PreserveSmoothingGroup", "Preserve Smoothing Groups"))
 		]
 	];
 
@@ -789,7 +824,7 @@ TSharedRef<SWidget> SFbxOptionWindow::ConstructSkeletalMeshAdvanced()
 		.OnCheckStateChanged(this, &SFbxOptionWindow::SetSkeletalMesh_KeepOverlappingVertices)
 		[
 			SNew(STextBlock)
-			.Text(LOCTEXT("FbxOptionWindow_KeepOverlappingVertices", "Keep Overlapping Vertices").ToString())
+			.Text(LOCTEXT("FbxOptionWindow_KeepOverlappingVertices", "Keep Overlapping Vertices"))
 		]
 	];
 	NewBox->AddSlot().AutoHeight() .Padding(2)
@@ -799,7 +834,7 @@ TSharedRef<SWidget> SFbxOptionWindow::ConstructSkeletalMeshAdvanced()
 		.OnCheckStateChanged(this, &SFbxOptionWindow::SetSkeletalMesh_ImportMeshesInBoneHierarchy)
 		[
 			SNew(STextBlock)
-			.Text(LOCTEXT("FbxOptionWindow_ImportMeshesInBoneHierarchy", "Import Meshes in Bone Hierarchy").ToString())
+			.Text(LOCTEXT("FbxOptionWindow_ImportMeshesInBoneHierarchy", "Import Meshes in Bone Hierarchy"))
 		]
 	];
 
@@ -807,7 +842,7 @@ TSharedRef<SWidget> SFbxOptionWindow::ConstructSkeletalMeshAdvanced()
 	[
 		SNew(STextBlock)
 		.TextStyle( FEditorStyle::Get(), "FBXSmallFont" ) 
-		.Text(LOCTEXT("FbxOptionWindow_PhysicsAsset", "Collision").ToString())
+		.Text(LOCTEXT("FbxOptionWindow_PhysicsAsset", "Collision"))
 	];
 
 	NewBox->AddSlot().AutoHeight() .Padding(2)
@@ -818,7 +853,7 @@ TSharedRef<SWidget> SFbxOptionWindow::ConstructSkeletalMeshAdvanced()
 		.IsEnabled(!bReimport)
 		[
 			SNew(STextBlock)
-			.Text(LOCTEXT("FbxOptionWindow_CreatePhysicsAsset", "Create Physics Asset").ToString())
+			.Text(LOCTEXT("FbxOptionWindow_CreatePhysicsAsset", "Create Physics Asset"))
 		]
 	];
 
@@ -830,7 +865,7 @@ TSharedRef<SWidget> SFbxOptionWindow::ConstructSkeletalMeshAdvanced()
 		.Padding(2) 
 		[
 			SNew(STextBlock)
-			.Text( LOCTEXT("FbxOptionWindow_SelectPhysicsAsset", "Select PhysicsAsset").ToString() )
+			.Text( LOCTEXT("FbxOptionWindow_SelectPhysicsAsset", "Select PhysicsAsset") )
 			.IsEnabled(this, &SFbxOptionWindow::ShouldShowPhysicsAssetPicker)
 		]
 		+SVerticalBox::Slot()
@@ -840,7 +875,7 @@ TSharedRef<SWidget> SFbxOptionWindow::ConstructSkeletalMeshAdvanced()
 			.ContentPadding(1.f)
 			.OnGetMenuContent( this, &SFbxOptionWindow::MakePhysicsAssetPickerMenu )
 			.HasDownArrow(true)
-			.ToolTipText(LOCTEXT("FBXOption", "Pick an skeleton from a popup menu").ToString())
+			.ToolTipText(LOCTEXT("FBXOption", "Pick an skeleton from a popup menu"))
 			.IsEnabled(this, &SFbxOptionWindow::ShouldShowPhysicsAssetPicker)
 			.ButtonContent()
 			[
@@ -862,7 +897,7 @@ TSharedRef<SWidget> SFbxOptionWindow::ConstructMaterialOption()
 	[
 		SNew(STextBlock)
 		.TextStyle( FEditorStyle::Get(), "FBXSmallFont" ) 
-		.Text( LOCTEXT("FbxOptionWindow_Material", "Material").ToString() )
+		.Text( LOCTEXT("FbxOptionWindow_Material", "Material") )
 	];
 
 	NewBox->AddSlot().AutoHeight() .Padding(2)
@@ -872,7 +907,7 @@ TSharedRef<SWidget> SFbxOptionWindow::ConstructMaterialOption()
 		.OnCheckStateChanged(this, &SFbxOptionWindow::SetMaterial_ImportMaterials)
 		[
 			SNew(STextBlock)
-			.Text(LOCTEXT("FbxOptionWindow_ImportMaterials", "Import Materials").ToString())
+			.Text(LOCTEXT("FbxOptionWindow_ImportMaterials", "Import Materials"))
 		]
 	];
 
@@ -883,7 +918,7 @@ TSharedRef<SWidget> SFbxOptionWindow::ConstructMaterialOption()
 		.OnCheckStateChanged(this, &SFbxOptionWindow::SetMaterial_ImportTextures)
 		[
 			SNew(STextBlock)
-			.Text(LOCTEXT("FbxOptionWindow_ImportTextures", "Import Textures").ToString())
+			.Text(LOCTEXT("FbxOptionWindow_ImportTextures", "Import Textures"))
 		]
 	];
 
@@ -894,7 +929,7 @@ TSharedRef<SWidget> SFbxOptionWindow::ConstructMaterialOption()
 		.OnCheckStateChanged(this, &SFbxOptionWindow::SetMaterial_InvertNormalMaps)
 		[
 			SNew(STextBlock)
-			.Text(LOCTEXT("FbxOptionWindow_InvertNormalMaps", "Invert Normal Maps").ToString())
+			.Text(LOCTEXT("FbxOptionWindow_InvertNormalMaps", "Invert Normal Maps"))
 		]
 	];
 
@@ -909,7 +944,7 @@ TSharedRef<SWidget> SFbxOptionWindow::ConstructAnimationOption()
 // 	[
 // 		SNew(STextBlock)
 // 		.TextStyle( FEditorStyle::Get(), "FBXSmallFont" ) 
-// 		.Text( LOCTEXT("FbxOptionWindow_Animation", "Animation").ToString() )
+// 		.Text( LOCTEXT("FbxOptionWindow_Animation", "Animation") )
 // 	];
 
 	NewBox->AddSlot().AutoHeight() .Padding(2, 4)
@@ -919,7 +954,7 @@ TSharedRef<SWidget> SFbxOptionWindow::ConstructAnimationOption()
 		.OnCheckStateChanged(this, &SFbxOptionWindow::SetSkeletalMesh_UseDefaultSampleRate)
 		[
 			SNew(STextBlock)
-			.Text(LOCTEXT("FbxOptionWindow_UseDefaultSampleRate", "Use Default Sample Rate").ToString())
+			.Text(LOCTEXT("FbxOptionWindow_UseDefaultSampleRate", "Use Default Sample Rate"))
 		]
 	];
 
@@ -930,14 +965,14 @@ TSharedRef<SWidget> SFbxOptionWindow::ConstructAnimationOption()
 		.OnCheckStateChanged(this, &SFbxOptionWindow::SetAnimation_ReserveLocalTransform)
 		[
 			SNew(STextBlock)
-			.Text(LOCTEXT("FbxOptionWindow_PreserveLocalTransform", "Preserve Local Transform").ToString())
+			.Text(LOCTEXT("FbxOptionWindow_PreserveLocalTransform", "Preserve Local Transform"))
 		]
 	];
 
 	NewBox->AddSlot().AutoHeight() .Padding(4)
 	[
 		SNew(STextBlock)
-		.Text( LOCTEXT("FbxOptionWindow_AnimationLength", "Animation Length").ToString() )
+		.Text( LOCTEXT("FbxOptionWindow_AnimationLength", "Animation Length") )
 		.TextStyle( FEditorStyle::Get(), "FBXSmallFont" ) 
 	];
 
@@ -1045,24 +1080,24 @@ TSharedRef<SWidget> SFbxOptionWindow::ConstructSkeletonOptionForAnim()
 	return NewBox;
 }
 
-FString SFbxOptionWindow::GetPhysicsAssetDisplay() const
+FText SFbxOptionWindow::GetPhysicsAssetDisplay() const
 {
 	if ( ImportUI->PhysicsAsset )
 	{
-		return ImportUI->PhysicsAsset->GetName();
+		return FText::FromString(ImportUI->PhysicsAsset->GetName());
 	}
 
-	return LOCTEXT("NoPhysicsAssetToDisplay", "None").ToString();
+	return LOCTEXT("NoPhysicsAssetToDisplay", "None");
 }
 
-FString SFbxOptionWindow::GetSkeletonDisplay() const
+FText SFbxOptionWindow::GetSkeletonDisplay() const
 {
 	if ( ImportUI->Skeleton )
 	{
-		return ImportUI->Skeleton->GetName();
+		return FText::FromString(ImportUI->Skeleton->GetName());
 	}
 
-	return LOCTEXT("NoSkeletonToDisplay", "None").ToString();
+	return LOCTEXT("NoSkeletonToDisplay", "None");
 }
 
 void SFbxOptionWindow::SetImportType(EFBXImportType ImportType)
@@ -1195,6 +1230,7 @@ void SFbxOptionWindow::SetAnimationRangeStart(const FText& Name, ETextCommit::Ty
 
 	if (ImportUI->AnimSequenceImportData->StartFrame >= ImportUI->AnimSequenceImportData->EndFrame)
 	{
+		// TODO: ErrorMessage no longer appears to be used?
 		ErrorMessage = LOCTEXT("FbxOptionWindow_InvalidStartFrame", "Invalid StartFrame").ToString();
 	}
 	else
@@ -1293,8 +1329,9 @@ TSharedRef<SWidget> SFbxOptionWindow::OnGenerateWidgetForComboItem( TSharedPtr<E
 	.ToolTipText( OnGenerateToolTipForImportMethod( *ImportMethod ) );
 }
 
-FString SFbxOptionWindow::OnGenerateTextForImportMethod( EFBXNormalImportMethod ImportMethod ) const
+FText SFbxOptionWindow::OnGenerateTextForImportMethod( EFBXNormalImportMethod ImportMethod ) const
 {
+	// TODO: change this to read the display name from the enum directly
 	if( ImportMethod == FBXNIM_MAX )
 	{
 		ImportMethod = GetCurrentNormalImportMethod();
@@ -1302,31 +1339,32 @@ FString SFbxOptionWindow::OnGenerateTextForImportMethod( EFBXNormalImportMethod
 
 	if( ImportMethod == FBXNIM_ComputeNormals ) 
 	{	
-		return LOCTEXT("FBXOptions_CalculateNormals", "Calculate Normals").ToString();
+		return LOCTEXT("FBXOptions_CalculateNormals", "Calculate Normals");
 	}
 	else if( ImportMethod == FBXNIM_ImportNormals )
 	{
-		return LOCTEXT("FBXOptions_ImportNormals", "Import Normals").ToString();
+		return LOCTEXT("FBXOptions_ImportNormals", "Import Normals");
 	}
 	else
 	{
-		return LOCTEXT("FBXOptions_ImportNormalsAndTangents", "Import Normals and Tangents").ToString();
+		return LOCTEXT("FBXOptions_ImportNormalsAndTangents", "Import Normals and Tangents");
 	}
 }
 
-FString SFbxOptionWindow::OnGenerateToolTipForImportMethod( EFBXNormalImportMethod ImportMethod ) const
+FText SFbxOptionWindow::OnGenerateToolTipForImportMethod( EFBXNormalImportMethod ImportMethod ) const
 {
-	if( ImportMethod == FBXNIM_ComputeNormals ) 
+	// TODO: change this to read the tooltip text from the enum directly
+	if (ImportMethod == FBXNIM_ComputeNormals)
 	{	
-		return LOCTEXT("FBXOptions_CalculateNormalsToolTip", "Let Unreal calculate normals and tangents. Ignores normals in the fbx file").ToString();
+		return LOCTEXT("FBXOptions_CalculateNormalsToolTip", "Let Unreal calculate normals and tangents. Ignores normals in the fbx file");
 	}
 	else if( ImportMethod == FBXNIM_ImportNormals )
 	{
-		return LOCTEXT("FBXOptions_ImportNormalsToolTip", "Import normals found in the fbx file.  Tangents are calculated by Unreal").ToString();
+		return LOCTEXT("FBXOptions_ImportNormalsToolTip", "Import normals found in the fbx file.  Tangents are calculated by Unreal");
 	}
 	else
 	{
-		return LOCTEXT("FBXOptions_ImportNormalsAndTangentsToolTip", "Import Normals and Tangents found in the fbx file instead of computing them").ToString();
+		return LOCTEXT("FBXOptions_ImportNormalsAndTangentsToolTip", "Import Normals and Tangents found in the fbx file instead of computing them");
 	}
 }
 
diff --git a/Engine/Source/Editor/UnrealEd/Private/FbxOptionWindow.h b/Engine/Source/Editor/UnrealEd/Private/FbxOptionWindow.h
--- a/Engine/Source/Editor/UnrealEd/Private/FbxOptionWindow.h
+++ b/Engine/Source/Editor/UnrealEd/Private/FbxOptionWindow.h
@@ -143,17 +143,17 @@ private:
 	TSharedRef<SWidget> MakeSkeletonPickerMenu();
 	void OnAssetSelectedFromSkeletonPicker(const FAssetData& AssetData);
 	TSharedPtr<SComboButton> SkeletonPickerComboButton;
-	FString GetSkeletonDisplay() const;
+	FText GetSkeletonDisplay() const;
 
 	TSharedRef<SWidget> MakePhysicsAssetPickerMenu();
 	void OnAssetSelectedFromPhysicsAssetPicker(const FAssetData& AssetData);
 	TSharedPtr<SComboButton> PhysicsAssetPickerComboButton;
-	FString GetPhysicsAssetDisplay() const;
+	FText GetPhysicsAssetDisplay() const;
 
 	TSharedRef<SWidget> OnGenerateWidgetForComboItem( TSharedPtr<EFBXNormalImportMethod> ImportMethod );
 	void OnNormalImportMethodChanged( TSharedPtr<EFBXNormalImportMethod> NewMethod, ESelectInfo::Type SelectionType );
-	FString OnGenerateTextForImportMethod( EFBXNormalImportMethod ImportMethod ) const;
-	FString OnGenerateToolTipForImportMethod( EFBXNormalImportMethod ImportMethod ) const;
+	FText OnGenerateTextForImportMethod( EFBXNormalImportMethod ImportMethod ) const;
+	FText OnGenerateToolTipForImportMethod( EFBXNormalImportMethod ImportMethod ) const;
 
 	FText GetAnimationName() const { return FText::FromString( ImportUI->AnimationName ); }
 	FString GetAnimationRangeStart() const { return FString::FromInt(ImportUI->AnimSequenceImportData->StartFrame); }
