Fixed the inconsistency of hover/drop behavior wheen asset skeleton does not match mesh skeleton.

The basic fix is that if asset skeleton is missing, we can't do much about it, so skip it.
If asset skeleton exists, try to switch up if it's not comptabile.

diff --git a/Engine/Source/Editor/UnrealEd/Private/LevelEditorViewport.cpp b/Engine/Source/Editor/UnrealEd/Private/LevelEditorViewport.cpp
--- a/Engine/Source/Editor/UnrealEd/Private/LevelEditorViewport.cpp
+++ b/Engine/Source/Editor/UnrealEd/Private/LevelEditorViewport.cpp
@@ -604,28 +604,39 @@ static bool AttemptApplyObjToActor( UObject* ObjToUse, AActor* ActorToApplyTo, b
 		UAnimBlueprint* DroppedObjAsAnimBlueprint = Cast<UAnimBlueprint>( ObjToUse );
 		if ( DroppedObjAsAnimBlueprint )
 		{
-			if ( bTest )
+			USkeleton* NeedsSkeleton = DroppedObjAsAnimBlueprint->TargetSkeleton;
+			if ( NeedsSkeleton )
 			{
-				if ( ActorToApplyTo->IsA(ASkeletalMeshActor::StaticClass()) )
+				if(bTest)
 				{
-					bResult = true;
-				}
-			}
-			else
-			{
-				if ( ASkeletalMeshActor* SkelMeshActor = Cast<ASkeletalMeshActor>(ActorToApplyTo) )
-				{
-					const FScopedTransaction Transaction( LOCTEXT( "DropAnimBlueprintOnObject", "Drop Anim Blueprint On Object" ) );
-					USkeletalMeshComponent* SkelMeshComponent = SkelMeshActor->SkeletalMeshComponent;
-					USkeleton* NeedsSkeleton = DroppedObjAsAnimBlueprint->TargetSkeleton;
-					if ( NeedsSkeleton && ! SkelMeshComponent->SkeletalMesh->Skeleton->IsCompatible(NeedsSkeleton) )
+					if(ActorToApplyTo->IsA(ASkeletalMeshActor::StaticClass()))
 					{
-						SkelMeshComponent->SetSkeletalMesh(NeedsSkeleton->GetPreviewMesh(true));
+						bResult = true;
 					}
-					if ( SkelMeshComponent && SkelMeshComponent->SkeletalMesh )
+				}
+				else
+				{
+					if(ASkeletalMeshActor* SkelMeshActor = Cast<ASkeletalMeshActor>(ActorToApplyTo))
 					{
-						SkelMeshComponent->SetAnimClass(DroppedObjAsAnimBlueprint->GeneratedClass);
-						bResult = true;
+						const FScopedTransaction Transaction(LOCTEXT("DropAnimBlueprintOnObject", "Drop Anim Blueprint On Object"));
+
+						USkeletalMeshComponent* SkelMeshComponent = SkelMeshActor->SkeletalMeshComponent;
+						// if anim blueprint skeleton and mesh skeleton does not match or component does not have any mesh, then change mesh
+						bool bShouldChangeMesh = (SkelMeshComponent->SkeletalMesh == NULL ||
+								!NeedsSkeleton->IsCompatible(SkelMeshComponent->SkeletalMesh->Skeleton));
+
+						if(bShouldChangeMesh)
+						{
+							SkelMeshComponent->SetSkeletalMesh(NeedsSkeleton->GetPreviewMesh(true));
+						}
+
+						// make sure if it's compabile now, if not we're not going to change anim blueprint
+						if(SkelMeshComponent->SkeletalMesh &&
+							NeedsSkeleton->IsCompatible(SkelMeshComponent->SkeletalMesh->Skeleton))
+						{
+							SkelMeshComponent->SetAnimClass(DroppedObjAsAnimBlueprint->GeneratedClass);
+							bResult = true;
+						}
 					}
 				}
 			}
@@ -645,46 +656,55 @@ static bool AttemptApplyObjToActor( UObject* ObjToUse, AActor* ActorToApplyTo, b
 		if ( DroppedObjAsAnimationAsset ||
 			 DroppedObjAsVertexAnimation)
 		{
-			if ( bTest )
-			{
-				if ( ActorToApplyTo->IsA(ASkeletalMeshActor::StaticClass()) )
-				{
-					bResult = true;
-				}
-			}
-			else
+			USkeleton* NeedsSkeleton = DroppedObjAsAnimationAsset? DroppedObjAsAnimationAsset->GetSkeleton() :
+				(DroppedObjAsVertexAnimation && DroppedObjAsVertexAnimation->BaseSkelMesh? DroppedObjAsVertexAnimation->BaseSkelMesh->Skeleton : NULL);
+
+			if (NeedsSkeleton)
 			{
-				if ( ASkeletalMeshActor* SkelMeshActor = Cast<ASkeletalMeshActor>(ActorToApplyTo) )
+				if(bTest)
 				{
-					const FScopedTransaction Transaction( LOCTEXT( "DropAnimationOnObject", "Drop Animation On Object" ) );	
-					USkeletalMeshComponent* SkelComponent = SkelMeshActor->SkeletalMeshComponent;
-					USkeleton* NeedsSkeleton = DroppedObjAsAnimationAsset? DroppedObjAsAnimationAsset->GetSkeleton() : 
-						(DroppedObjAsVertexAnimation && DroppedObjAsVertexAnimation->BaseSkelMesh? DroppedObjAsVertexAnimation->BaseSkelMesh->Skeleton : NULL);
-					SkelComponent->Modify();
-					if ( NeedsSkeleton && ! SkelComponent->SkeletalMesh->Skeleton->IsCompatible(NeedsSkeleton) )
+					if(ActorToApplyTo->IsA(ASkeletalMeshActor::StaticClass()))
 					{
-						SkelComponent->SetSkeletalMesh(NeedsSkeleton->GetPreviewMesh(true));
+						bResult = true;
 					}
-					if ( DroppedObjAsAnimationAsset )
+				}
+				else
+				{
+					if(ASkeletalMeshActor* SkelMeshActor = Cast<ASkeletalMeshActor>(ActorToApplyTo))
 					{
-						SkelComponent->SetAnimationMode(EAnimationMode::Type::AnimationSingleNode);
-						SkelComponent->AnimationData.AnimToPlay = DroppedObjAsAnimationAsset;
+						const FScopedTransaction Transaction(LOCTEXT("DropAnimationOnObject", "Drop Animation On Object"));
+						USkeletalMeshComponent* SkelComponent = SkelMeshActor->SkeletalMeshComponent;
+						SkelComponent->Modify();
+						// if asset skeleton and mesh skeleton does not match or component does not have any mesh, then change mesh
+						bool bShouldChangeMesh = SkelComponent->SkeletalMesh == NULL || 
+							!NeedsSkeleton->IsCompatible(SkelComponent->SkeletalMesh->Skeleton);
+
+						if(bShouldChangeMesh)
+						{
+							SkelComponent->SetSkeletalMesh(NeedsSkeleton->GetPreviewMesh(true));
+						}
 
-						// set runtime data
-						SkelComponent->SetAnimation(DroppedObjAsAnimationAsset);
-					}
-					if ( DroppedObjAsVertexAnimation )
-					{
-						SkelComponent->SetAnimationMode(EAnimationMode::Type::AnimationSingleNode);
-						SkelComponent->AnimationData.VertexAnimToPlay = DroppedObjAsVertexAnimation;
+						if(DroppedObjAsAnimationAsset)
+						{
+							SkelComponent->SetAnimationMode(EAnimationMode::Type::AnimationSingleNode);
+							SkelComponent->AnimationData.AnimToPlay = DroppedObjAsAnimationAsset;
 
-						// set runtime data
-						SkelComponent->SetVertexAnimation(DroppedObjAsVertexAnimation);
-					}
-					if ( SkelComponent && SkelComponent->SkeletalMesh )
-					{
-						bResult = true;
-						SkelComponent->InitAnim(true);
+							// set runtime data
+							SkelComponent->SetAnimation(DroppedObjAsAnimationAsset);
+						}
+						if(DroppedObjAsVertexAnimation)
+						{
+							SkelComponent->SetAnimationMode(EAnimationMode::Type::AnimationSingleNode);
+							SkelComponent->AnimationData.VertexAnimToPlay = DroppedObjAsVertexAnimation;
+
+							// set runtime data
+							SkelComponent->SetVertexAnimation(DroppedObjAsVertexAnimation);
+						}
+						if(SkelComponent && SkelComponent->SkeletalMesh)
+						{
+							bResult = true;
+							SkelComponent->InitAnim(true);
+						}
 					}
 				}
 			}
diff --git a/Engine/Source/Runtime/Engine/Classes/Animation/Skeleton.h b/Engine/Source/Runtime/Engine/Classes/Animation/Skeleton.h
--- a/Engine/Source/Runtime/Engine/Classes/Animation/Skeleton.h
+++ b/Engine/Source/Runtime/Engine/Classes/Animation/Skeleton.h
@@ -304,7 +304,7 @@ public:
 	/**
 	 *	Check if this skeleton may be used with other skeleton
 	 */
-	ENGINE_API bool IsCompatible(USkeleton const * InSkeleton) const { return this == InSkeleton; }
+	ENGINE_API bool IsCompatible(USkeleton const * InSkeleton) const { return (InSkeleton && this == InSkeleton); }
 
 	/** 
 	 * Indexing naming convention
