TTP# 331611: "Play" dropdown uses two different methods to denote the currently selected item

diff --git a/Engine/Source/Editor/UnrealEd/Private/Kismet2/DebuggerCommands.cpp b/Engine/Source/Editor/UnrealEd/Private/Kismet2/DebuggerCommands.cpp
--- a/Engine/Source/Editor/UnrealEd/Private/Kismet2/DebuggerCommands.cpp
+++ b/Engine/Source/Editor/UnrealEd/Private/Kismet2/DebuggerCommands.cpp
@@ -43,12 +43,11 @@ public:
 
 	static void PlayInViewport_Clicked();
 	static bool PlayInViewport_CanExecute();
-
 	static void PlayInEditorFloating_Clicked();
 	static bool PlayInEditorFloating_CanExecute();
-
 	static void PlayInNewProcess_Clicked( bool MobilePreview );
 	static bool PlayInNewProcess_CanExecute();
+	static bool PlayInModeIsChecked( EPlayModeType PlayMode );
 
 	static bool PlayInLocation_CanExecute( EPlayModeLocations Location );
 	static void PlayInLocation_Clicked( EPlayModeLocations Location );
@@ -190,10 +189,10 @@ void FPlayWorldCommands::RegisterCommands()
 
 	// PIE
 	UI_COMMAND( RepeatLastPlay, "Play", "Launches a game preview session in the same mode as the last game preview session launched from the Game Preview Modes dropdown next to the Play button on the level editor toolbar", EUserInterfaceActionType::Button, FInputGesture( EKeys::P, EModifierKey::Alt ) )
-	UI_COMMAND( PlayInViewport, "Selected Viewport", "Play this level in the active level editor viewport", EUserInterfaceActionType::Button, FInputGesture() );
-	UI_COMMAND( PlayInEditorFloating, "New Editor Window", "Play this level in a new window", EUserInterfaceActionType::Button, FInputGesture() );
-	UI_COMMAND( PlayInMobilePreview, "Mobile Preview", "Play this level as a mobile device preview (runs in its own process)", EUserInterfaceActionType::Button, FInputGesture() );
-	UI_COMMAND( PlayInNewProcess, "Standalone Game", "Play this level in a new window that runs in its own process", EUserInterfaceActionType::Button, FInputGesture() );
+	UI_COMMAND( PlayInViewport, "Selected Viewport", "Play this level in the active level editor viewport", EUserInterfaceActionType::Check, FInputGesture() );
+	UI_COMMAND( PlayInEditorFloating, "New Editor Window", "Play this level in a new window", EUserInterfaceActionType::Check, FInputGesture() );
+	UI_COMMAND( PlayInMobilePreview, "Mobile Preview", "Play this level as a mobile device preview (runs in its own process)", EUserInterfaceActionType::Check, FInputGesture() );
+	UI_COMMAND( PlayInNewProcess, "Standalone Game", "Play this level in a new window that runs in its own process", EUserInterfaceActionType::Check, FInputGesture() );
 	UI_COMMAND( PlayInCameraLocation, "Current Camera Location", "Spawn the player at the current camera location", EUserInterfaceActionType::RadioButton, FInputGesture() );
 	UI_COMMAND( PlayInDefaultPlayerStart, "Default Player Start", "Spawn the player at the map's default player start", EUserInterfaceActionType::RadioButton, FInputGesture() );
 	UI_COMMAND( PlayInNetworkSettings, "Network Settings...", "Open the settings for the 'Play In' feature", EUserInterfaceActionType::Button, FInputGesture() );
@@ -244,28 +243,28 @@ void FPlayWorldCommands::BindGlobalPlayWorldCommands()
 	ActionList.MapAction( Commands.PlayInViewport,
 		FExecuteAction::CreateStatic( &FInternalPlayWorldCommandCallbacks::PlayInViewport_Clicked ),
 		FCanExecuteAction::CreateStatic( &FInternalPlayWorldCommandCallbacks::PlayInViewport_CanExecute ),
-		FIsActionChecked(),
+		FIsActionChecked::CreateStatic( &FInternalPlayWorldCommandCallbacks::PlayInModeIsChecked, PlayMode_InViewPort ),
 		FIsActionButtonVisible::CreateStatic( &FInternalPlayWorldCommandCallbacks::CanShowNonPlayWorldOnlyActions )
 		);
 
 	ActionList.MapAction( Commands.PlayInEditorFloating,
 		FExecuteAction::CreateStatic( &FInternalPlayWorldCommandCallbacks::PlayInEditorFloating_Clicked ),
 		FCanExecuteAction::CreateStatic( &FInternalPlayWorldCommandCallbacks::PlayInEditorFloating_CanExecute ),
-		FIsActionChecked(),
+		FIsActionChecked::CreateStatic( &FInternalPlayWorldCommandCallbacks::PlayInModeIsChecked, PlayMode_InEditorFloating ),
 		FIsActionButtonVisible::CreateStatic( &FInternalPlayWorldCommandCallbacks::CanShowNonPlayWorldOnlyActions )
 		);
 
 	ActionList.MapAction( Commands.PlayInMobilePreview,
 		FExecuteAction::CreateStatic( &FInternalPlayWorldCommandCallbacks::PlayInNewProcess_Clicked, true ),
 		FCanExecuteAction::CreateStatic( &FInternalPlayWorldCommandCallbacks::PlayInNewProcess_CanExecute ),
-		FIsActionChecked(),
+		FIsActionChecked::CreateStatic( &FInternalPlayWorldCommandCallbacks::PlayInModeIsChecked, PlayMode_InMobilePreview ),
 		FIsActionButtonVisible::CreateStatic( &FInternalPlayWorldCommandCallbacks::CanShowNonPlayWorldOnlyActions )
 		);
 
 	ActionList.MapAction( Commands.PlayInNewProcess,
 		FExecuteAction::CreateStatic( &FInternalPlayWorldCommandCallbacks::PlayInNewProcess_Clicked, false ),
 		FCanExecuteAction::CreateStatic( &FInternalPlayWorldCommandCallbacks::PlayInNewProcess_CanExecute ),
-		FIsActionChecked(),
+		FIsActionChecked::CreateStatic( &FInternalPlayWorldCommandCallbacks::PlayInModeIsChecked, PlayMode_InNewProcess ),
 		FIsActionButtonVisible::CreateStatic( &FInternalPlayWorldCommandCallbacks::CanShowNonPlayWorldOnlyActions )
 		);
 
@@ -477,17 +476,7 @@ TSharedRef< SWidget > FPlayWorldCommands::GeneratePlayMenuContent( TSharedRef<FU
 
 			if (PlayModeCommand.IsValid())
 			{
-				if (PlayMode == GetDefault<ULevelEditorPlaySettings>()->LastExecutedPlayModeType)
-				{
-					FFormatNamedArguments Arguments;
-					Arguments.Add(TEXT("PlayMode"), PlayModeCommand->GetLabel());
-
-					MenuBuilder.AddMenuEntry(PlayModeCommand, NAME_None, FText::Format(LOCTEXT("ActivePlayModeLabel", "{PlayMode} (Active)"), Arguments));
-				}
-				else
-				{
-					MenuBuilder.AddMenuEntry(PlayModeCommand);
-				}
+				MenuBuilder.AddMenuEntry(PlayModeCommand);
 			}
 		}
 	};
@@ -627,16 +616,7 @@ TSharedRef< SWidget > FPlayWorldCommands::GenerateLaunchMenuContent( TSharedRef<
 						LabelArguments.Add(TEXT("HostUser"), FText::GetEmpty());
 					}
 
-					if (DeviceProxy->GetDeviceId() == GetDefault<ULevelEditorPlaySettings>()->LastExecutedLaunchDevice)
-					{
-						LabelArguments.Add(TEXT("Active"), LOCTEXT("ActiveHint", " (Active)"));
-					}
-					else
-					{
-						LabelArguments.Add(TEXT("Active"), FText::GetEmpty());
-					}
-
-					FText Label = FText::Format(LOCTEXT("LaunchDeviceLabel", "{DeviceName}{HostUser}{Active}"), LabelArguments);
+					FText Label = FText::Format(LOCTEXT("LaunchDeviceLabel", "{DeviceName}{HostUser}"), LabelArguments);
 
 					// ... generate tooltip text
 					FFormatNamedArguments TooltipArguments;
@@ -651,7 +631,7 @@ TSharedRef< SWidget > FPlayWorldCommands::GenerateLaunchMenuContent( TSharedRef<
 						FSlateIcon(FEditorStyle::GetStyleSetName(), *IconName),
 						LaunchDeviceAction,
 						NAME_None,
-						EUserInterfaceActionType::Button
+						EUserInterfaceActionType::Check
 					);
 				}
 			}
@@ -877,7 +857,7 @@ bool FInternalPlayWorldCommandCallbacks::Simulate_IsChecked()
 
 const TSharedRef < FUICommandInfo > GetLastPlaySessionCommand()
 {
-	ULevelEditorPlaySettings* PlaySettings = GetMutableDefault<ULevelEditorPlaySettings>();
+	const ULevelEditorPlaySettings* PlaySettings = GetDefault<ULevelEditorPlaySettings>();
 
 	const FPlayWorldCommands& Commands = FPlayWorldCommands::Get();
 	TSharedRef < FUICommandInfo > Command = Commands.PlayInViewport.ToSharedRef();
@@ -1157,6 +1137,12 @@ bool FInternalPlayWorldCommandCallbacks::PlayInNewProcess_CanExecute()
 }
 
 
+bool FInternalPlayWorldCommandCallbacks::PlayInModeIsChecked( EPlayModeType PlayMode )
+{
+	return (PlayMode == GetDefault<ULevelEditorPlaySettings>()->LastExecutedPlayModeType);
+}
+
+
 bool FInternalPlayWorldCommandCallbacks::PlayInLocation_CanExecute( EPlayModeLocations Location )
 {
 	switch (Location)
