* FDragDropOperation now inherits from TSharedFromThis, no need to do it on subclasses any more.  I've removed it from the subclasses that were also doing it.
* FDecoratedDragDropOp's Hover text is now FText, instead of FString.

diff --git a/Engine/Source/Developer/Profiler/Private/Widgets/StatDragDropOp.h b/Engine/Source/Developer/Profiler/Private/Widgets/StatDragDropOp.h
--- a/Engine/Source/Developer/Profiler/Private/Widgets/StatDragDropOp.h
+++ b/Engine/Source/Developer/Profiler/Private/Widgets/StatDragDropOp.h
@@ -3,9 +3,7 @@
 #pragma once
 
 /** Class that handles the drag dropping of the specified stat in the profiler. */
-class FStatIDDragDropOp 
-	: public FDragDropOperation
-	, public TSharedFromThis<FStatIDDragDropOp>
+class FStatIDDragDropOp : public FDragDropOperation
 {
 public:
 	DRAG_DROP_OPERATOR_TYPE(FStatIDDragDropOp, FDragDropOperation)
diff --git a/Engine/Source/Editor/ContentBrowser/Private/AssetViewWidgets.cpp b/Engine/Source/Editor/ContentBrowser/Private/AssetViewWidgets.cpp
--- a/Engine/Source/Editor/ContentBrowser/Private/AssetViewWidgets.cpp
+++ b/Engine/Source/Editor/ContentBrowser/Private/AssetViewWidgets.cpp
@@ -181,7 +181,7 @@ void SAssetViewItem::OnDragEnter( const FGeometry& MyGeometry, const FDragDropEv
 			if ( AssetDatas.Num() > 0 )
 			{
 				TSharedPtr< FAssetDragDropOp > DragDropOp = StaticCastSharedPtr< FAssetDragDropOp >( DragDropEvent.GetOperation() );	
-				DragDropOp->SetToolTip( LOCTEXT( "OnDragAssetsOverFolder", "Move or Copy Asset(s)" ).ToString(), FEditorStyle::GetBrush(TEXT("Graph.ConnectorFeedback.OK")) );
+				DragDropOp->SetToolTip( LOCTEXT( "OnDragAssetsOverFolder", "Move or Copy Asset(s)" ), FEditorStyle::GetBrush(TEXT("Graph.ConnectorFeedback.OK")) );
 			}
 			bDraggedOver = true;
 		}
diff --git a/Engine/Source/Editor/ContentBrowser/Private/SAssetView.cpp b/Engine/Source/Editor/ContentBrowser/Private/SAssetView.cpp
--- a/Engine/Source/Editor/ContentBrowser/Private/SAssetView.cpp
+++ b/Engine/Source/Editor/ContentBrowser/Private/SAssetView.cpp
@@ -1074,7 +1074,7 @@ FReply SAssetView::OnDragOver( const FGeometry& MyGeometry, const FDragDropEvent
 
 				if ( IsValidDrop )
 				{
-					DragAssetOp->SetToolTip( NSLOCTEXT( "AssetView", "OnDragOverCollection", "Add to Collection" ).ToString(), FEditorStyle::GetBrush(TEXT("Graph.ConnectorFeedback.OK"))) ;
+					DragAssetOp->SetToolTip( NSLOCTEXT( "AssetView", "OnDragOverCollection", "Add to Collection" ), FEditorStyle::GetBrush(TEXT("Graph.ConnectorFeedback.OK"))) ;
 				}
 			}
 
diff --git a/Engine/Source/Editor/ContentBrowser/Private/SourcesViewWidgets.cpp b/Engine/Source/Editor/ContentBrowser/Private/SourcesViewWidgets.cpp
--- a/Engine/Source/Editor/ContentBrowser/Private/SourcesViewWidgets.cpp
+++ b/Engine/Source/Editor/ContentBrowser/Private/SourcesViewWidgets.cpp
@@ -94,7 +94,7 @@ void SAssetTreeItem::OnDragEnter( const FGeometry& MyGeometry, const FDragDropEv
 
 			if (bCanDrop)
 			{
-				DragDropOp->SetToolTip( LOCTEXT( "OnDragAssetsOverFolder", "Move or Copy Asset(s)" ).ToString(), FEditorStyle::GetBrush(TEXT("Graph.ConnectorFeedback.OK")) );
+				DragDropOp->SetToolTip( LOCTEXT( "OnDragAssetsOverFolder", "Move or Copy Asset(s)" ), FEditorStyle::GetBrush(TEXT("Graph.ConnectorFeedback.OK")) );
 				bDraggedOver = true;
 			}
 		}
diff --git a/Engine/Source/Editor/GraphEditor/Private/SGraphNode.cpp b/Engine/Source/Editor/GraphEditor/Private/SGraphNode.cpp
--- a/Engine/Source/Editor/GraphEditor/Private/SGraphNode.cpp
+++ b/Engine/Source/Editor/GraphEditor/Private/SGraphNode.cpp
@@ -212,7 +212,7 @@ FReply SGraphNode::OnDragOver( const FGeometry& MyGeometry, const FDragDropEvent
 			FString TooltipText;
 			GraphNode->GetSchema()->GetAssetsNodeHoverMessage(AssetOp->AssetData, GraphNode, TooltipText, bOkIcon);
 			const FSlateBrush* TooltipIcon = bOkIcon ? FEditorStyle::GetBrush(TEXT("Graph.ConnectorFeedback.OK")) : FEditorStyle::GetBrush(TEXT("Graph.ConnectorFeedback.Error"));;
-			AssetOp->SetToolTip( TooltipText, TooltipIcon );
+			AssetOp->SetToolTip(FText::FromString(TooltipText), TooltipIcon);
 		}
 		return FReply::Handled();
 	}
diff --git a/Engine/Source/Editor/GraphEditor/Private/SGraphPanel.cpp b/Engine/Source/Editor/GraphEditor/Private/SGraphPanel.cpp
--- a/Engine/Source/Editor/GraphEditor/Private/SGraphPanel.cpp
+++ b/Engine/Source/Editor/GraphEditor/Private/SGraphPanel.cpp
@@ -673,7 +673,7 @@ FReply SGraphPanel::OnDragOver( const FGeometry& MyGeometry, const FDragDropEven
 			FString TooltipText;
 			GraphObj->GetSchema()->GetAssetsGraphHoverMessage(AssetOp->AssetData, GraphObj, TooltipText, bOkIcon);
 			const FSlateBrush* TooltipIcon = bOkIcon ? FEditorStyle::GetBrush(TEXT("Graph.ConnectorFeedback.OK")) : FEditorStyle::GetBrush(TEXT("Graph.ConnectorFeedback.Error"));;
-			AssetOp->SetToolTip(TooltipText, TooltipIcon);
+			AssetOp->SetToolTip(FText::FromString(TooltipText), TooltipIcon);
 		}
 		return FReply::Handled();
 	} 
diff --git a/Engine/Source/Editor/GraphEditor/Private/SGraphPin.cpp b/Engine/Source/Editor/GraphEditor/Private/SGraphPin.cpp
--- a/Engine/Source/Editor/GraphEditor/Private/SGraphPin.cpp
+++ b/Engine/Source/Editor/GraphEditor/Private/SGraphPin.cpp
@@ -526,7 +526,7 @@ void SGraphPin::OnDragEnter( const FGeometry& MyGeometry, const FDragDropEvent&
 				FString TooltipText;
 				Node->GetSchema()->GetAssetsPinHoverMessage(AssetOp->AssetData, GraphPinObj, TooltipText, bOkIcon);
 				const FSlateBrush* TooltipIcon = bOkIcon ? FEditorStyle::GetBrush(TEXT("Graph.ConnectorFeedback.OK")) : FEditorStyle::GetBrush(TEXT("Graph.ConnectorFeedback.Error"));;
-				AssetOp->SetToolTip(TooltipText, TooltipIcon);
+				AssetOp->SetToolTip(FText::FromString(TooltipText), TooltipIcon);
 			}
 		}
 	}
diff --git a/Engine/Source/Editor/Kismet/Private/SSCSEditor.cpp b/Engine/Source/Editor/Kismet/Private/SSCSEditor.cpp
--- a/Engine/Source/Editor/Kismet/Private/SSCSEditor.cpp
+++ b/Engine/Source/Editor/Kismet/Private/SSCSEditor.cpp
@@ -111,7 +111,7 @@ FReply SSCSEditorDragDropTree::OnDrop( const FGeometry& MyGeometry, const FDragD
 //////////////////////////////////////////////////////////////////////////
 //
 
-class FSCSRowDragDropOp : public FDragDropOperation, public TSharedFromThis<FSCSRowDragDropOp>
+class FSCSRowDragDropOp : public FDragDropOperation
 {
 public:
 	DRAG_DROP_OPERATOR_TYPE(FSCSRowDragDropOp, FDragDropOperation)
diff --git a/Engine/Source/Editor/Layers/Private/SLayerBrowser.h b/Engine/Source/Editor/Layers/Private/SLayerBrowser.h
--- a/Engine/Source/Editor/Layers/Private/SLayerBrowser.h
+++ b/Engine/Source/Editor/Layers/Private/SLayerBrowser.h
@@ -266,11 +266,11 @@ protected:
 
 		if ( bCanAssign )
 		{
-			DragActorOp->SetToolTip( FActorDragDropGraphEdOp::ToolTip_CompatibleGeneric, Message );
+			DragActorOp->SetToolTip( FActorDragDropGraphEdOp::ToolTip_CompatibleGeneric, FText::FromString(Message) );
 		}
 		else
 		{
-			DragActorOp->SetToolTip( FActorDragDropGraphEdOp::ToolTip_IncompatibleGeneric, Message );
+			DragActorOp->SetToolTip( FActorDragDropGraphEdOp::ToolTip_IncompatibleGeneric, FText::FromString(Message) );
 		}
 
 		// We leave the event unhandled so the children of the ListView get a chance to grab the drag/drop
diff --git a/Engine/Source/Editor/Layers/Private/SLayersView.h b/Engine/Source/Editor/Layers/Private/SLayersView.h
--- a/Engine/Source/Editor/Layers/Private/SLayersView.h
+++ b/Engine/Source/Editor/Layers/Private/SLayersView.h
@@ -176,7 +176,7 @@ protected:
 		}
 
 		TSharedPtr< FActorDragDropGraphEdOp > DragActorOp = StaticCastSharedPtr< FActorDragDropGraphEdOp >( DragDropEvent.GetOperation() );	
-		DragActorOp->SetToolTip( FActorDragDropGraphEdOp::ToolTip_CompatibleGeneric, LOCTEXT("OnDragOver", "Add Actors to New Layer").ToString() );
+		DragActorOp->SetToolTip( FActorDragDropGraphEdOp::ToolTip_CompatibleGeneric, LOCTEXT("OnDragOver", "Add Actors to New Layer") );
 
 		// We leave the event unhandled so the children of the ListView get a chance to grab the drag/drop
 		return FReply::Unhandled();
diff --git a/Engine/Source/Editor/Layers/Private/SLayersViewRow.h b/Engine/Source/Editor/Layers/Private/SLayersViewRow.h
--- a/Engine/Source/Editor/Layers/Private/SLayersViewRow.h
+++ b/Engine/Source/Editor/Layers/Private/SLayersViewRow.h
@@ -184,11 +184,11 @@ protected:
 
 		if ( bCanAssign )
 		{
-			DragActorOp->SetToolTip( FActorDragDropGraphEdOp::ToolTip_CompatibleGeneric, Message );
+			DragActorOp->SetToolTip( FActorDragDropGraphEdOp::ToolTip_CompatibleGeneric, FText::FromString(Message) );
 		}
 		else
 		{
-			DragActorOp->SetToolTip( FActorDragDropGraphEdOp::ToolTip_IncompatibleGeneric, Message );
+			DragActorOp->SetToolTip( FActorDragDropGraphEdOp::ToolTip_IncompatibleGeneric, FText::FromString(Message) );
 		}
 
 		return FReply::Handled();
diff --git a/Engine/Source/Editor/Levels/Private/SLevelsView.h b/Engine/Source/Editor/Levels/Private/SLevelsView.h
--- a/Engine/Source/Editor/Levels/Private/SLevelsView.h
+++ b/Engine/Source/Editor/Levels/Private/SLevelsView.h
@@ -243,7 +243,7 @@ protected:
 		}
 
 		TSharedPtr< FActorDragDropGraphEdOp > DragActorOp = StaticCastSharedPtr< FActorDragDropGraphEdOp >( DragDropEvent.GetOperation() );	
-		DragActorOp->SetToolTip( FActorDragDropGraphEdOp::ToolTip_CompatibleGeneric, LOCTEXT("OnDragOver", "Add Actors to New Level").ToString() );
+		DragActorOp->SetToolTip( FActorDragDropGraphEdOp::ToolTip_CompatibleGeneric, LOCTEXT("OnDragOver", "Add Actors to New Level") );
 
 		// We leave the event unhandled so the children of the ListView get a chance to grab the drag/drop
 		return FReply::Unhandled();
diff --git a/Engine/Source/Editor/Levels/Private/SLevelsViewRow.h b/Engine/Source/Editor/Levels/Private/SLevelsViewRow.h
--- a/Engine/Source/Editor/Levels/Private/SLevelsViewRow.h
+++ b/Engine/Source/Editor/Levels/Private/SLevelsViewRow.h
@@ -419,11 +419,11 @@ protected:
 
 		if ( bCanAssign )
 		{
-			DragActorOp->SetToolTip( FActorDragDropGraphEdOp::ToolTip_CompatibleGeneric, Message );
+			DragActorOp->SetToolTip( FActorDragDropGraphEdOp::ToolTip_CompatibleGeneric, FText::FromString(Message) );
 		}
 		else
 		{
-			DragActorOp->SetToolTip( FActorDragDropGraphEdOp::ToolTip_IncompatibleGeneric, Message );
+			DragActorOp->SetToolTip( FActorDragDropGraphEdOp::ToolTip_IncompatibleGeneric, FText::FromString(Message) );
 		}
 
 		return FReply::Handled();
diff --git a/Engine/Source/Editor/Persona/Private/SAnimationBlendSpaceBase.cpp b/Engine/Source/Editor/Persona/Private/SAnimationBlendSpaceBase.cpp
--- a/Engine/Source/Editor/Persona/Private/SAnimationBlendSpaceBase.cpp
+++ b/Engine/Source/Editor/Persona/Private/SAnimationBlendSpaceBase.cpp
@@ -47,7 +47,7 @@ struct FAnimWeightSorted
 
 ////////////////////////////////////////////////////////////////////////////////////////////////////////////////////////////////////////////////////////////////////
 
-class FSampleDragDropOp : public FDragDropOperation, public TSharedFromThis<FSampleDragDropOp>
+class FSampleDragDropOp : public FDragDropOperation
 {
 public:
 	DRAG_DROP_OPERATOR_TYPE(FSampleDragDropOp, FDragDropOperation)
diff --git a/Engine/Source/Editor/Persona/Public/BoneDragDropOp.h b/Engine/Source/Editor/Persona/Public/BoneDragDropOp.h
--- a/Engine/Source/Editor/Persona/Public/BoneDragDropOp.h
+++ b/Engine/Source/Editor/Persona/Public/BoneDragDropOp.h
@@ -3,7 +3,7 @@
 
 //////////////////////////////////////////////////////////////////////////
 // FBoneDragDropOp
-class FBoneDragDropOp : public FDragDropOperation, public TSharedFromThis<FBoneDragDropOp>
+class FBoneDragDropOp : public FDragDropOperation
 {
 public:	
 	DRAG_DROP_OPERATOR_TYPE(FBoneDragDropOp, FDragDropOperation)
diff --git a/Engine/Source/Editor/Persona/Public/SocketDragDropOp.h b/Engine/Source/Editor/Persona/Public/SocketDragDropOp.h
--- a/Engine/Source/Editor/Persona/Public/SocketDragDropOp.h
+++ b/Engine/Source/Editor/Persona/Public/SocketDragDropOp.h
@@ -3,7 +3,7 @@
 
 //////////////////////////////////////////////////////////////////////////
 // FSocketDragDropOp
-class FSocketDragDropOp : public FDragDropOperation, public TSharedFromThis<FSocketDragDropOp>
+class FSocketDragDropOp : public FDragDropOperation
 {
 public:	
 	
diff --git a/Engine/Source/Editor/PlacementMode/Private/PlacementMode.cpp b/Engine/Source/Editor/PlacementMode/Private/PlacementMode.cpp
--- a/Engine/Source/Editor/PlacementMode/Private/PlacementMode.cpp
+++ b/Engine/Source/Editor/PlacementMode/Private/PlacementMode.cpp
@@ -238,7 +238,9 @@ void FPlacementMode::UpdatePreviewActors( FLevelEditorViewportClient* ViewportCl
 			for ( auto AssetIt = Assets.CreateConstIterator(); AssetIt; ++AssetIt )
 			{
 				UObject* Asset = *AssetIt;
-				if ( !ViewportClient->CanDropObjectsAtCoordinates( x, y, FAssetData( Asset ) ) )
+
+				FDropQuery DropResult = ViewportClient->CanDropObjectsAtCoordinates( x, y, FAssetData( Asset ) );
+				if ( !DropResult.bCanDrop )
 				{
 					// At least one of the assets can't be dropped.
 					ViewportClient->DestroyDropPreviewActors();
@@ -378,7 +380,8 @@ bool FPlacementMode::HandleClick(FLevelEditorViewportClient* InViewportClient, H
 			for ( auto AssetIt = Assets.CreateConstIterator(); AssetIt; ++AssetIt )
 			{
 				UObject* Asset = *AssetIt;
-				if ( !InViewportClient->CanDropObjectsAtCoordinates( Viewport->GetMouseX(), Viewport->GetMouseY(), FAssetData( Asset ) ) )
+				FDropQuery DropResult = InViewportClient->CanDropObjectsAtCoordinates( Viewport->GetMouseX(), Viewport->GetMouseY(), FAssetData( Asset ) );
+				if ( !DropResult.bCanDrop )
 				{
 					// At least one of the assets can't be dropped.
 					InViewportClient->DestroyDropPreviewActors();
diff --git a/Engine/Source/Editor/SceneOutliner/Private/SOutlinerTreeView.cpp b/Engine/Source/Editor/SceneOutliner/Private/SOutlinerTreeView.cpp
--- a/Engine/Source/Editor/SceneOutliner/Private/SOutlinerTreeView.cpp
+++ b/Engine/Source/Editor/SceneOutliner/Private/SOutlinerTreeView.cpp
@@ -441,7 +441,7 @@ namespace SceneOutliner
 			{
 				// Generic messages can go through the decorated drag/drop operation base class
 				auto DecoratedDragOp = StaticCastSharedPtr<FDecoratedDragDropOp>(DragDropEvent.GetOperation());
-				DecoratedDragOp->CurrentHoverText = ValidationInfo.ValidationText.ToString();
+				DecoratedDragOp->CurrentHoverText = ValidationInfo.ValidationText;
 				if (ValidationInfo.TooltipType == FActorDragDropGraphEdOp::ToolTip_CompatibleGeneric)
 				{
 					DecoratedDragOp->CurrentIconBrush = FEditorStyle::GetBrush(TEXT("Graph.ConnectorFeedback.OK"));
@@ -455,7 +455,7 @@ namespace SceneOutliner
 			{
 				// Display actor specific messages
 				auto ActorDragOp = StaticCastSharedPtr<FActorDragDropGraphEdOp>(DragDropEvent.GetOperation());
-				ActorDragOp->SetToolTip(ValidationInfo.TooltipType, ValidationInfo.ValidationText.ToString());
+				ActorDragOp->SetToolTip(ValidationInfo.TooltipType, ValidationInfo.ValidationText);
 			}
 		}
 
@@ -654,12 +654,12 @@ namespace SceneOutliner
 		if (DragDrop::IsTypeMatch<FActorDragDropGraphEdOp>(DragDropEvent.GetOperation()))
 		{
 			auto ActorDragOp = StaticCastSharedPtr<FActorDragDropGraphEdOp>(DragDropEvent.GetOperation());
-			ActorDragOp->SetToolTip(ValidationInfo.TooltipType, ValidationInfo.ValidationText.ToString());
+			ActorDragOp->SetToolTip(ValidationInfo.TooltipType, ValidationInfo.ValidationText);
 		}
 		else if (DragDrop::IsTypeMatch<FDecoratedDragDropOp>(DragDropEvent.GetOperation()))
 		{
 			auto DecoratedDragOp = StaticCastSharedPtr<FDecoratedDragDropOp>(DragDropEvent.GetOperation());
-			DecoratedDragOp->CurrentHoverText = ValidationInfo.ValidationText.ToString();
+			DecoratedDragOp->CurrentHoverText = ValidationInfo.ValidationText;
 			DecoratedDragOp->CurrentIconBrush = ValidationInfo.IsValid() ? FEditorStyle::GetBrush(TEXT("Graph.ConnectorFeedback.OK")) : FEditorStyle::GetBrush(TEXT("Graph.ConnectorFeedback.Error"));
 		}
 	}
@@ -674,7 +674,7 @@ namespace SceneOutliner
 
 			// Attach each child
 			bool bAttached = false;
-			for(int32 i=0; i<ChildrenPtrs.Num(); i++)
+			for(int32 i = 0; i < ChildrenPtrs.Num(); i++)
 			{
 				AActor* ChildActor = ChildrenPtrs[i].Get();
 				if (GEditor->CanParentActors(ParentActor, ChildActor))
@@ -682,7 +682,6 @@ namespace SceneOutliner
 					GEditor->ParentActors(ParentActor, ChildActor, SocketName);
 					bAttached = true;
 				}
-
 			}
 
 			// refresh the tree, and ensure parent is expanded so we can still see the child if we attached something
diff --git a/Engine/Source/Editor/SceneOutliner/Private/SceneOutlinerTreeItems.cpp b/Engine/Source/Editor/SceneOutliner/Private/SceneOutlinerTreeItems.cpp
--- a/Engine/Source/Editor/SceneOutliner/Private/SceneOutlinerTreeItems.cpp
+++ b/Engine/Source/Editor/SceneOutliner/Private/SceneOutlinerTreeItems.cpp
@@ -99,22 +99,22 @@ namespace SceneOutliner
 				{
 					if (Actors.Num() == 1)
 					{
-						Operation->CurrentHoverText = FText::Format(LOCTEXT("DragActorAndFolder", "{NumActors} Actor, {NumFolders} Folder"), Args).ToString();
+						Operation->CurrentHoverText = FText::Format(LOCTEXT("DragActorAndFolder", "{NumActors} Actor, {NumFolders} Folder"), Args);
 					}
 					else
 					{
-						Operation->CurrentHoverText = FText::Format(LOCTEXT("DragActorsAndFolder", "{NumActors} Actors, {NumFolders} Folder"), Args).ToString();
+						Operation->CurrentHoverText = FText::Format(LOCTEXT("DragActorsAndFolder", "{NumActors} Actors, {NumFolders} Folder"), Args);
 					}
 				}
 				else
 				{
 					if (Actors.Num() == 1)
 					{
-						Operation->CurrentHoverText = FText::Format(LOCTEXT("DragActorAndFolders", "{NumActors} Actor, {NumFolders} Folders"), Args).ToString();
+						Operation->CurrentHoverText = FText::Format(LOCTEXT("DragActorAndFolders", "{NumActors} Actor, {NumFolders} Folders"), Args);
 					}
 					else
 					{
-						Operation->CurrentHoverText = FText::Format(LOCTEXT("DragActorsAndFolders", "{NumActors} Actor, {NumFolders} Folders"), Args).ToString();
+						Operation->CurrentHoverText = FText::Format(LOCTEXT("DragActorsAndFolders", "{NumActors} Actor, {NumFolders} Folders"), Args);
 					}
 				}
 			}
@@ -123,13 +123,13 @@ namespace SceneOutliner
 				auto TheFolder = Folders[0].Pin();
 				if (TheFolder.IsValid())
 				{
-					Operation->CurrentHoverText = FText::FromString(TheFolder->LeafName.ToString()).ToString();
+					Operation->CurrentHoverText = FText::FromString(TheFolder->LeafName.ToString());
 					Operation->CurrentIconBrush = FEditorStyle::Get().GetBrush(TEXT("SceneOutliner.FolderClosed"));
 				}
 			}
 			else
 			{
-				Operation->CurrentHoverText = FText::Format(LOCTEXT("DragFolders", "{NumFolders} Folders"), Args).ToString();
+				Operation->CurrentHoverText = FText::Format(LOCTEXT("DragFolders", "{NumFolders} Folders"), Args);
 			}
 		}
 
diff --git a/Engine/Source/Editor/SourceControlWindows/Private/SSourceControlHistory.cpp b/Engine/Source/Editor/SourceControlWindows/Private/SSourceControlHistory.cpp
--- a/Engine/Source/Editor/SourceControlWindows/Private/SSourceControlHistory.cpp
+++ b/Engine/Source/Editor/SourceControlWindows/Private/SSourceControlHistory.cpp
@@ -323,7 +323,7 @@ static bool DiffHistoryItems(TSharedPtr<FHistoryTreeItem> const FirstSelection,
 /**
 * A FDragDropOperation that represents dragging a source-control history tree item around.
 */
-class FSourceControlHistoryRowDragDropOp : public FDragDropOperation, public TSharedFromThis<FSourceControlHistoryRowDragDropOp>
+class FSourceControlHistoryRowDragDropOp : public FDragDropOperation
 {
 private:
 	/**
diff --git a/Engine/Source/Editor/UnrealEd/Private/DragAndDrop/AssetDragDropOp.cpp b/Engine/Source/Editor/UnrealEd/Private/DragAndDrop/AssetDragDropOp.cpp
new file mode 100644
--- /dev/null
+++ b/Engine/Source/Editor/UnrealEd/Private/DragAndDrop/AssetDragDropOp.cpp
@@ -0,0 +1,167 @@
+// Copyright 1998-2014 Epic Games, Inc. All Rights Reserved.
+
+#include "UnrealEd.h"
+#include "AssetDragDropOp.h"
+
+TSharedRef<FAssetDragDropOp> FAssetDragDropOp::New(const FAssetData& InAssetData, UActorFactory* ActorFactory /*= NULL*/)
+{
+	TArray<FAssetData> AssetDataArray;
+	AssetDataArray.Add(InAssetData);
+	return New(AssetDataArray, ActorFactory);
+}
+
+TSharedRef<FAssetDragDropOp> FAssetDragDropOp::New(const TArray<FAssetData>& InAssetData, UActorFactory* ActorFactory /*= NULL*/)
+{
+	TSharedRef<FAssetDragDropOp> Operation = MakeShareable(new FAssetDragDropOp);
+
+	Operation->MouseCursor = EMouseCursor::GrabHandClosed;
+
+	Operation->ThumbnailSize = 64;
+
+	Operation->AssetData = InAssetData;
+	Operation->ActorFactory = ActorFactory;
+
+	Operation->Init();
+
+	Operation->Construct();
+	return Operation;
+}
+
+FAssetDragDropOp::~FAssetDragDropOp()
+{
+	if ( ThumbnailPool.IsValid() )
+	{
+		// Release all rendering resources being held onto
+		ThumbnailPool->ReleaseResources();
+	}
+}
+
+TSharedPtr<SWidget> FAssetDragDropOp::GetDefaultDecorator() const
+{
+	TSharedPtr<SWidget> ThumbnailWidget;
+
+	if ( AssetThumbnail.IsValid() )
+	{
+		ThumbnailWidget = AssetThumbnail->MakeThumbnailWidget();
+	}
+	else
+	{
+		ThumbnailWidget = SNew(SImage) .Image( FEditorStyle::GetDefaultBrush() );
+	}
+		
+	const FSlateBrush* ActorTypeBrush = FEditorStyle::GetDefaultBrush();
+	if ( ActorFactory.IsValid() && AssetData.Num() > 0 )
+	{
+		AActor* DefaultActor = ActorFactory->GetDefaultActor( AssetData[0] );
+		ActorTypeBrush = FClassIconFinder::FindIconForActor( DefaultActor );
+	}
+
+	return 
+		SNew(SBorder)
+		.BorderImage(FEditorStyle::GetBrush("ContentBrowser.AssetDragDropTooltipBackground"))
+		.Content()
+		[
+			SNew(SHorizontalBox)
+
+			// Left slot is asset thumbnail
+			+ SHorizontalBox::Slot()
+			.AutoWidth()
+			.HAlign(HAlign_Left)
+			[
+				SNew(SBox) 
+				.WidthOverride(ThumbnailSize) 
+				.HeightOverride(ThumbnailSize)
+				.Content()
+				[
+					SNew(SOverlay)
+
+					+SOverlay::Slot()
+					[
+						ThumbnailWidget.ToSharedRef()
+					]
+
+					+SOverlay::Slot()
+					.HAlign(HAlign_Center)
+					.VAlign(VAlign_Top)
+					.Padding(FMargin(0, 4, 0, 0))
+					[
+						SNew(SBorder)
+						.BorderImage(FEditorStyle::GetBrush("Menu.Background"))
+						.Visibility(AssetData.Num() > 1 ? EVisibility::Visible : EVisibility::Collapsed)
+						.Content()
+						[
+							SNew(STextBlock) .Text(FText::AsNumber(AssetData.Num()))
+						]
+					]
+
+					+SOverlay::Slot()
+					.HAlign(HAlign_Right)
+					.VAlign(VAlign_Bottom)
+					.Padding(FMargin(4, 4))
+					[
+						SNew(SImage)
+						.Image( ActorTypeBrush )
+						.Visibility( (ActorTypeBrush != FEditorStyle::GetDefaultBrush()) ? EVisibility::Visible : EVisibility::Collapsed)
+					]
+				]
+			]
+
+			// Right slot is for optional tooltip
+			+ SHorizontalBox::Slot()
+			.AutoWidth()
+			.VAlign(VAlign_Fill)
+			[
+				SNew(SBox)
+				.Visibility(this, &FAssetDragDropOp::GetTooltipVisibility)
+				.Content()
+				[
+					SNew(SHorizontalBox)
+
+					+ SHorizontalBox::Slot()
+					.AutoWidth()
+					.Padding(3.0f)
+					[
+						SNew(SImage) 
+						.Image(this, &FAssetDragDropOp::GetIcon)
+					]
+
+					+SHorizontalBox::Slot()
+					.AutoWidth()
+					.Padding(0,0,3,0)
+					.VAlign(VAlign_Center)
+					[
+						SNew(STextBlock) 
+						.Text(this, &FAssetDragDropOp::GetHoverText)
+					]
+				]
+			]
+		];
+}
+
+void FAssetDragDropOp::Init()
+{
+	if ( AssetData.Num() > 0 && ThumbnailSize > 0 )
+	{
+		// Create a thumbnail pool to hold the single thumbnail rendered
+		ThumbnailPool = MakeShareable( new FAssetThumbnailPool(1, /*InAreRealTileThumbnailsAllowed=*/false) );
+
+		// Create the thumbnail handle
+		AssetThumbnail = MakeShareable( new FAssetThumbnail( AssetData[0], ThumbnailSize, ThumbnailSize, ThumbnailPool ) );
+
+		// Request the texture then tick the pool once to render the thumbnail
+		AssetThumbnail->GetViewportRenderTargetTexture();
+		ThumbnailPool->Tick(0);
+	}
+}
+
+EVisibility FAssetDragDropOp::GetTooltipVisibility() const
+{
+	if ( !CurrentHoverText.IsEmpty() )
+	{
+		return EVisibility::Visible;
+	}
+	else
+	{
+		return EVisibility::Collapsed;
+	}
+}
diff --git a/Engine/Source/Editor/UnrealEd/Public/DragAndDrop/ActorDragDropGraphEdOp.h b/Engine/Source/Editor/UnrealEd/Public/DragAndDrop/ActorDragDropGraphEdOp.h
--- a/Engine/Source/Editor/UnrealEd/Public/DragAndDrop/ActorDragDropGraphEdOp.h
+++ b/Engine/Source/Editor/UnrealEd/Public/DragAndDrop/ActorDragDropGraphEdOp.h
@@ -1,8 +1,6 @@
 // Copyright 1998-2014 Epic Games, Inc. All Rights Reserved.
 
-
-#ifndef __ActorDragDropGraphEdOp_h__
-#define __ActorDragDropGraphEdOp_h__
+#pragma once
 
 #include "ActorDragDropOp.h"
 
@@ -28,31 +26,31 @@ public:
 	};
 
 	/** Set the appropriate tool tip when dragging functionality is active*/
-	void SetToolTip(ToolTipTextType TextType, FString ParamText = FString(TEXT("")))
+	void SetToolTip(ToolTipTextType TextType, FText ParamText = FText::GetEmpty())
 	{
 		switch( TextType )
 		{
 		case ToolTip_Compatible:
 			{
-				CurrentHoverText = FString::Printf( *LOCTEXT("ToolTipCompatible", "'%s' is compatible to replace object reference").ToString(), *Actors[0].Get()->GetActorLabel( ) );
+				CurrentHoverText = FText::Format( LOCTEXT("ToolTipCompatible", "'{0}' is compatible to replace object reference"), FText::FromString(Actors[0].Get()->GetActorLabel()) );
 				CurrentIconBrush = FEditorStyle::GetBrush(TEXT("Graph.ConnectorFeedback.OK"));
 				break;
 			}
 		case ToolTip_Incompatible:
 			{
-				CurrentHoverText = FString::Printf( *LOCTEXT("ToolTipIncompatible", "'%s' is not compatible to replace object reference").ToString(), *Actors[0].Get()->GetActorLabel( ) );
+			CurrentHoverText = FText::Format( LOCTEXT("ToolTipIncompatible", "'{0}' is not compatible to replace object reference"), FText::FromString(Actors[0].Get()->GetActorLabel()) );
 				CurrentIconBrush = FEditorStyle::GetBrush(TEXT("Graph.ConnectorFeedback.Error"));
 				break;
 			}
 		case ToolTip_MultipleSelection_Incompatible:
 			{
-				CurrentHoverText = LOCTEXT("ToolTipMultipleSelectionIncompatible", "Cannot replace object reference when multiple objects are selected").ToString();
+				CurrentHoverText = LOCTEXT("ToolTipMultipleSelectionIncompatible", "Cannot replace object reference when multiple objects are selected");
 				CurrentIconBrush = FEditorStyle::GetBrush(TEXT("Graph.ConnectorFeedback.Error"));
 				break;
 			}
 		case ToolTip_CompatibleAttach:
 			{
-				CurrentHoverText = FString::Printf( *LOCTEXT("ToolTipCompatibleAttach", "Attach %s to %s").ToString(), *Actors[0].Get()->GetActorLabel( ), *ParamText);
+				CurrentHoverText = FText::Format( LOCTEXT("ToolTipCompatibleAttach", "Attach %s to %s"), FText::FromString(Actors[0].Get()->GetActorLabel()), ParamText);
 				CurrentIconBrush = FEditorStyle::GetBrush(TEXT("Graph.ConnectorFeedback.OK"));
 				break;
 			}
@@ -70,25 +68,25 @@ public:
 			}
 		case ToolTip_CompatibleMultipleAttach:
 			{
-				CurrentHoverText = FString::Printf( *LOCTEXT("ToolTipCompatibleMultipleAttach", "Attach multiple objects to %s").ToString(), *ParamText);
+				CurrentHoverText = FText::Format( LOCTEXT("ToolTipCompatibleMultipleAttach", "Attach multiple objects to %s"), ParamText);
 				CurrentIconBrush = FEditorStyle::GetBrush(TEXT("Graph.ConnectorFeedback.OK"));
 				break;
 			}
 		case ToolTip_IncompatibleMultipleAttach:
 			{
-				CurrentHoverText = FString::Printf( *LOCTEXT("ToolTipIncompatibleMultipleAttach", "Cannot attach multiple objects to %s").ToString(), *ParamText);
+				CurrentHoverText = FText::Format( LOCTEXT("ToolTipIncompatibleMultipleAttach", "Cannot attach multiple objects to %s"), ParamText);
 				CurrentIconBrush = FEditorStyle::GetBrush(TEXT("Graph.ConnectorFeedback.Error"));
 				break;
 			}
 		case ToolTip_CompatibleDetach:
 			{
-				CurrentHoverText = FString::Printf( *LOCTEXT("ToolTipCompatibleDetach", "Detach %s from %s").ToString(), *Actors[0].Get()->GetActorLabel( ), *ParamText);
+				CurrentHoverText = FText::Format( LOCTEXT("ToolTipCompatibleDetach", "Detach %s from %s"), FText::FromString(Actors[0].Get()->GetActorLabel()), ParamText);
 				CurrentIconBrush = FEditorStyle::GetBrush(TEXT("Graph.ConnectorFeedback.OK"));
 				break;
 			}
 		case ToolTip_CompatibleMultipleDetach:
 			{
-				CurrentHoverText = FString::Printf( *LOCTEXT("ToolTipCompatibleDetachMultiple", "Detach multiple objects from %s").ToString(), *ParamText);
+				CurrentHoverText = FText::Format( LOCTEXT("ToolTipCompatibleDetachMultiple", "Detach multiple objects from %s"), ParamText);
 				CurrentIconBrush = FEditorStyle::GetBrush(TEXT("Graph.ConnectorFeedback.OK"));
 				break;
 			}
@@ -108,4 +106,3 @@ public:
 };
 
 #undef LOCTEXT_NAMESPACE
-#endif // __ActorDragDropGraphEdOp_h__
diff --git a/Engine/Source/Editor/UnrealEd/Public/DragAndDrop/ActorDragDropOp.h b/Engine/Source/Editor/UnrealEd/Public/DragAndDrop/ActorDragDropOp.h
--- a/Engine/Source/Editor/UnrealEd/Public/DragAndDrop/ActorDragDropOp.h
+++ b/Engine/Source/Editor/UnrealEd/Public/DragAndDrop/ActorDragDropOp.h
@@ -1,11 +1,10 @@
 // Copyright 1998-2014 Epic Games, Inc. All Rights Reserved.
 
-
-#ifndef __ActorDragDropOp_h__
-#define __ActorDragDropOp_h__
+#pragma once
 
 #include "ClassIconFinder.h"
 #include "DecoratedDragDropOp.h"
+#include "Internationalization.h"
 
 class FActorDragDropOp : public FDecoratedDragDropOp
 {
@@ -30,19 +29,17 @@ public:
 		CurrentIconBrush = FClassIconFinder::FindIconForActors(Actors, CommonSelClass);
 		if(Actors.Num() == 0)
 		{
-			CurrentHoverText =  TEXT("None");
+			CurrentHoverText = NSLOCTEXT("FActorDragDropOp", "None", "None");
 		}
 		else if(Actors.Num() == 1)
 		{
 			// Find icon for actor
 			AActor* TheActor = Actors[0].Get();
-			CurrentHoverText = TheActor->GetActorLabel();
+			CurrentHoverText = FText::FromString(TheActor->GetActorLabel());
 		}
 		else
 		{
-			CurrentHoverText = FString::Printf(TEXT("%d Actors"), InActors.Num());
+			CurrentHoverText = FText::Format(NSLOCTEXT("FActorDragDropOp", "FormatActors", "{0} Actors"), FText::AsNumber(InActors.Num()));
 		}
 	}
 };
-
-#endif // __ActorDragDropOp_h__
diff --git a/Engine/Source/Editor/UnrealEd/Public/DragAndDrop/AssetDragDropOp.h b/Engine/Source/Editor/UnrealEd/Public/DragAndDrop/AssetDragDropOp.h
--- a/Engine/Source/Editor/UnrealEd/Public/DragAndDrop/AssetDragDropOp.h
+++ b/Engine/Source/Editor/UnrealEd/Public/DragAndDrop/AssetDragDropOp.h
@@ -24,171 +24,19 @@ public:
 	/** The actor factory to use if converting this asset to an actor */
 	TWeakObjectPtr< UActorFactory > ActorFactory;
 
-	static TSharedRef<FAssetDragDropOp> New(const FAssetData& InAssetData, UActorFactory* ActorFactory = NULL)
-	{
-		TArray<FAssetData> AssetDataArray;
-		AssetDataArray.Add(InAssetData);
-		return New(AssetDataArray, ActorFactory);
-	}
+	UNREALED_API static TSharedRef<FAssetDragDropOp> New(const FAssetData& InAssetData, UActorFactory* ActorFactory = NULL);
 
-	static TSharedRef<FAssetDragDropOp> New(const TArray<FAssetData>& InAssetData, UActorFactory* ActorFactory = NULL)
-	{
-		TSharedRef<FAssetDragDropOp> Operation = MakeShareable(new FAssetDragDropOp);
+	UNREALED_API static TSharedRef<FAssetDragDropOp> New(const TArray<FAssetData>& InAssetData, UActorFactory* ActorFactory = NULL);
 
-		Operation->MouseCursor = EMouseCursor::GrabHandClosed;
-
-		Operation->ThumbnailSize = 64;
+public:
+	UNREALED_API virtual ~FAssetDragDropOp();
 
-		Operation->AssetData = InAssetData;
-		Operation->ActorFactory = ActorFactory;
+	UNREALED_API virtual TSharedPtr<SWidget> GetDefaultDecorator() const OVERRIDE;
 
-		Operation->Init();
+	UNREALED_API void Init();
 
-		Operation->Construct();
-		return Operation;
-	}
+	UNREALED_API EVisibility GetTooltipVisibility() const;
 
 private:
 	int32 ThumbnailSize;
-	
-public:
-	virtual ~FAssetDragDropOp()
-	{
-		if ( ThumbnailPool.IsValid() )
-		{
-			// Release all rendering resources being held onto
-			ThumbnailPool->ReleaseResources();
-		}
-	}
-	
-	virtual TSharedPtr<SWidget> GetDefaultDecorator() const OVERRIDE
-	{
-		TSharedPtr<SWidget> ThumbnailWidget;
-
-		if ( AssetThumbnail.IsValid() )
-		{
-			ThumbnailWidget = AssetThumbnail->MakeThumbnailWidget();
-		}
-		else
-		{
-			ThumbnailWidget = SNew(SImage) .Image( FEditorStyle::GetDefaultBrush() );
-		}
-		
-		const FSlateBrush* ActorTypeBrush = FEditorStyle::GetDefaultBrush();
-		if ( ActorFactory.IsValid() && AssetData.Num() > 0 )
-		{
-			AActor* DefaultActor = ActorFactory->GetDefaultActor( AssetData[0] );
-			ActorTypeBrush = FClassIconFinder::FindIconForActor( DefaultActor );
-		}
-
-		return 
-			SNew(SBorder)
-			.BorderImage(FEditorStyle::GetBrush("ContentBrowser.AssetDragDropTooltipBackground"))
-			.Content()
-			[
-				SNew(SHorizontalBox)
-
-				// Left slot is asset thumbnail
-				+SHorizontalBox::Slot()
-				.AutoWidth()
-				.HAlign(HAlign_Left)
-				[
-					SNew(SBox) 
-					.WidthOverride(ThumbnailSize) 
-					.HeightOverride(ThumbnailSize)
-					.Content()
-					[
-						SNew(SOverlay)
-
-						+SOverlay::Slot()
-						[
-							ThumbnailWidget.ToSharedRef()
-						]
-
-						+SOverlay::Slot()
-						.HAlign(HAlign_Center)
-						.VAlign(VAlign_Top)
-						.Padding(FMargin(0, 4, 0, 0))
-						[
-							SNew(SBorder)
-							.BorderImage(FEditorStyle::GetBrush("Menu.Background"))
-							.Visibility(AssetData.Num() > 1 ? EVisibility::Visible : EVisibility::Collapsed)
-							.Content()
-							[
-								SNew(STextBlock) .Text(FText::AsNumber(AssetData.Num()))
-							]
-						]
-
-						+SOverlay::Slot()
-						.HAlign(HAlign_Right)
-						.VAlign(VAlign_Bottom)
-						.Padding(FMargin(4, 4))
-						[
-							SNew(SImage)
-							.Image( ActorTypeBrush )
-							.Visibility( (ActorTypeBrush != FEditorStyle::GetDefaultBrush()) ? EVisibility::Visible : EVisibility::Collapsed)
-						]
-					]
-				]
-
-				// Right slot is for optional tooltip
-				+SHorizontalBox::Slot()
-				.AutoWidth()
-				.VAlign(VAlign_Top)
-				[
-					SNew(SBox)
-					.Visibility(this, &FAssetDragDropOp::GetTooltipVisibility)
-					.Content()
-					[
-						SNew(SHorizontalBox)
-
-						+SHorizontalBox::Slot()
-						.AutoWidth()
-						.Padding(3.0f)
-						[
-							SNew(SImage) 
-							.Image(this, &FAssetDragDropOp::GetIcon)
-						]
-
-						+SHorizontalBox::Slot()
-						.AutoWidth()
-						.Padding(0,0,3,0)
-						.VAlign(VAlign_Center)
-						[
-							SNew(STextBlock) 
-							.Text(this, &FAssetDragDropOp::GetHoverText)
-						]
-					]
-				]
-
-			];
-	}
-
-	void Init()
-	{
-		if ( AssetData.Num() > 0 && ThumbnailSize > 0 )
-		{
-			// Create a thumbnail pool to hold the single thumbnail rendered
-			ThumbnailPool = MakeShareable( new FAssetThumbnailPool(1, /*InAreRealTileThumbnailsAllowed=*/false) );
-
-			// Create the thumbnail handle
-			AssetThumbnail = MakeShareable( new FAssetThumbnail( AssetData[0], ThumbnailSize, ThumbnailSize, ThumbnailPool ) );
-
-			// Request the texture then tick the pool once to render the thumbnail
-			AssetThumbnail->GetViewportRenderTargetTexture();
-			ThumbnailPool->Tick(0);
-		}
-	}
-
-	EVisibility GetTooltipVisibility() const
-	{
-		if(CurrentIconBrush != NULL && CurrentHoverText.Len() > 0)
-		{
-			return EVisibility::Visible;
-		}
-		else
-		{
-			return EVisibility::Collapsed;
-		}
-	}
 };
diff --git a/Engine/Source/Editor/UnrealEd/Public/DragAndDrop/ClassDragDropOp.h b/Engine/Source/Editor/UnrealEd/Public/DragAndDrop/ClassDragDropOp.h
--- a/Engine/Source/Editor/UnrealEd/Public/DragAndDrop/ClassDragDropOp.h
+++ b/Engine/Source/Editor/UnrealEd/Public/DragAndDrop/ClassDragDropOp.h
@@ -23,18 +23,20 @@ public:
 				.Content()
 				[			
 					SNew(SHorizontalBox)
-					+SHorizontalBox::Slot()
+					
+					+ SHorizontalBox::Slot()
 					.AutoWidth()
 					[
 						SNew( SImage )
 						.Image( ClassIcon )
 					]
-					+SHorizontalBox::Slot()
-						.AutoWidth()
-						[
-							SNew(STextBlock) 
-							.Text( ClassesToDrop[0]->GetName() )
-						]
+					
+					+ SHorizontalBox::Slot()
+					.AutoWidth()
+					[
+						SNew(STextBlock) 
+						.Text( ClassesToDrop[0]->GetName() )
+					]
 				];
 		}
 
@@ -43,12 +45,12 @@ public:
 			.Content()
 			[			
 				SNew(SHorizontalBox)
-				+SHorizontalBox::Slot()
+				+ SHorizontalBox::Slot()
 				.AutoWidth()
-					[
-						SNew(STextBlock) 
-						.Text( ClassesToDrop[0]->GetName() )
-					]
+				[
+					SNew(STextBlock) 
+					.Text( ClassesToDrop[0]->GetName() )
+				]
 			];
 	}
 
@@ -59,6 +61,10 @@ public:
 		Operation->Construct();
 		return Operation;
 	}
+
+protected:
+	/** Adding hint text to all drag drop operations that drop locations could set, and the decorator could report. */
+	FText HintText;
 };
 
 struct FClassPackageData
@@ -90,7 +96,7 @@ public:
 			.Content()
 			[			
 				SNew(SHorizontalBox)
-				+SHorizontalBox::Slot()
+				+ SHorizontalBox::Slot()
 				.AutoWidth()
 				[
 					SNew(STextBlock) 
diff --git a/Engine/Source/Editor/UnrealEd/Public/DragAndDrop/DecoratedDragDropOp.h b/Engine/Source/Editor/UnrealEd/Public/DragAndDrop/DecoratedDragDropOp.h
--- a/Engine/Source/Editor/UnrealEd/Public/DragAndDrop/DecoratedDragDropOp.h
+++ b/Engine/Source/Editor/UnrealEd/Public/DragAndDrop/DecoratedDragDropOp.h
@@ -4,13 +4,13 @@
 
 #define LOCTEXT_NAMESPACE "DecoratedDragDrop"
 
-class FDecoratedDragDropOp : public FDragDropOperation, public TSharedFromThis<FDecoratedDragDropOp>
+class FDecoratedDragDropOp : public FDragDropOperation
 {
 public:
 	DRAG_DROP_OPERATOR_TYPE(FDecoratedDragDropOp, FDragDropOperation)
 
 	/** String to show as hover text */
-	FString								CurrentHoverText;
+	FText								CurrentHoverText;
 
 	/** Icon to be displayed */
 	const FSlateBrush*					CurrentIconBrush;
@@ -58,7 +58,7 @@ public:
 			];
 	}
 
-	FString GetHoverText() const
+	FText GetHoverText() const
 	{
 		return CurrentHoverText;
 	}
@@ -69,7 +69,7 @@ public:
 	}
 
 	/** Set the text and icon for this tooltip */
-	void SetToolTip(const FString& Text, const FSlateBrush* Icon)
+	void SetToolTip(const FText& Text, const FSlateBrush* Icon)
 	{
 		CurrentHoverText = Text;
 		CurrentIconBrush = Icon;
@@ -85,7 +85,7 @@ public:
 protected:
 
 	/** Default string to show as hover text */
-	FString								DefaultHoverText;
+	FText								DefaultHoverText;
 
 	/** Default icon to be displayed */
 	const FSlateBrush*					DefaultHoverIcon;
diff --git a/Engine/Source/Runtime/Slate/Public/DragAndDrop.h b/Engine/Source/Runtime/Slate/Public/DragAndDrop.h
--- a/Engine/Source/Runtime/Slate/Public/DragAndDrop.h
+++ b/Engine/Source/Runtime/Slate/Public/DragAndDrop.h
@@ -9,7 +9,7 @@
  * a drag drop operation which keeps a decorator window (optionally)
  * Implement derived types with DRAG_DROP_OPERATOR_TYPE (see below)
  */
-class SLATE_API FDragDropOperation
+class SLATE_API FDragDropOperation : public TSharedFromThis<FDragDropOperation>
 {
 public:
 	FDragDropOperation()
@@ -42,7 +42,7 @@ public:
 	virtual FCursorReply OnCursorQuery();
 
 	/** Gets the widget that will serve as the decorator unless overridden. If you do not override, you will have no decorator */
-	virtual TSharedPtr<SWidget> GetDefaultDecorator() const {return TSharedPtr<SWidget>();}
+	virtual TSharedPtr<SWidget> GetDefaultDecorator() const { return TSharedPtr<SWidget>(); }
 	
 	/** Alters the visibility of the window */
 	void SetDecoratorVisibility(bool bVisible);
