TTP 322274

We now support applying the material to only to the submesh that you drag the material onto.  If you are holding down the [Shift] key we'll apply the material to every slot (the old behavior).

diff --git a/Engine/Source/Editor/LevelEditor/Private/SLevelViewport.cpp b/Engine/Source/Editor/LevelEditor/Private/SLevelViewport.cpp
--- a/Engine/Source/Editor/LevelEditor/Private/SLevelViewport.cpp
+++ b/Engine/Source/Editor/LevelEditor/Private/SLevelViewport.cpp
@@ -438,16 +438,40 @@ bool SLevelViewport::HandleDragObjects(const FGeometry& MyGeometry, const FDragD
 
 	DragDropEvent.GetOperation()->SetCursorOverride(TOptional<EMouseCursor::Type>());
 
+	FText HintText;
+
 	// Determine if we can drop the assets
 	for ( auto InfoIt = SelectedAssetDatas.CreateConstIterator(); InfoIt; ++InfoIt )
 	{
 		const FAssetData& AssetData = *InfoIt;
-		if ( AssetData.IsValid() && !LevelViewportClient->CanDropObjectsAtCoordinates( CachedOnDropLocalMousePos.X, CachedOnDropLocalMousePos.Y, AssetData ) )
+
+		// Ignore invalid assets
+		if ( !AssetData.IsValid() )
+		{
+			continue;
+		}
+		
+		FDropQuery DropResult = LevelViewportClient->CanDropObjectsAtCoordinates(CachedOnDropLocalMousePos.X, CachedOnDropLocalMousePos.Y, AssetData);
+		
+		if ( !DropResult.bCanDrop )
 		{
 			// At least one of the assets can't be dropped.
 			DragDropEvent.GetOperation()->SetCursorOverride(EMouseCursor::SlashedCircle);
 			return false;
 		}
+		else
+		{
+			if ( HintText.IsEmpty() )
+			{
+				HintText = DropResult.HintText;
+			}
+		}
+	}
+
+	if ( DragDrop::IsTypeMatch<FAssetDragDropOp>(DragDropEvent.GetOperation()) )
+	{
+		auto Operation = StaticCastSharedPtr<FAssetDragDropOp>(DragDropEvent.GetOperation());
+		Operation->SetToolTip(HintText, NULL);
 	}
 
 	return bValidDrag;
@@ -1698,7 +1722,7 @@ FLevelEditorViewportInstanceSettings SLevelViewport::LoadLegacyConfigFromIni(con
 			// Default to Lit for a perspective viewport
 			ViewportInstanceSettings.PerspViewModeIndex = VMI_Lit;
 		}
- 	}
+	}
 
 	if(!GConfig->GetInt(*IniSection, *(ConfigKey+TEXT(".OrthoViewModeIndex")), (int32&)ViewportInstanceSettings.OrthoViewModeIndex, GEditorUserSettingsIni))
 	{
diff --git a/Engine/Source/Editor/UnrealEd/Private/AssetSelection.cpp b/Engine/Source/Editor/UnrealEd/Private/AssetSelection.cpp
--- a/Engine/Source/Editor/UnrealEd/Private/AssetSelection.cpp
+++ b/Engine/Source/Editor/UnrealEd/Private/AssetSelection.cpp
@@ -664,10 +664,11 @@ bool FActorFactoryAssetProxy::IsActorValidForMaterialApplication( AActor* Target
 *
 * @param	TargetActor		the actor to apply the material to
 * @param	MaterialToApply	the material to apply to the actor
+* @param    OptionalMaterialSlot the material slot to apply to.
 *
 * @return	true if the material was successfully applied to the actor
 */
-bool FActorFactoryAssetProxy::ApplyMaterialToActor( AActor* TargetActor, UMaterialInterface* MaterialToApply )
+bool FActorFactoryAssetProxy::ApplyMaterialToActor( AActor* TargetActor, UMaterialInterface* MaterialToApply, int32 OptionalMaterialSlot )
 {
 	bool bResult = false;
 
@@ -685,7 +686,6 @@ bool FActorFactoryAssetProxy::ApplyMaterialToActor( AActor* TargetActor, UMateri
 		}
 		else
 		{
-
 			TArray<UActorComponent*> EditableComponents;
 			FActorEditorUtils::GetEditableComponents( TargetActor, EditableComponents );
 
@@ -737,10 +737,21 @@ bool FActorFactoryAssetProxy::ApplyMaterialToActor( AActor* TargetActor, UMateri
 						{
 							const FScopedTransaction Transaction( NSLOCTEXT("UnrealEd", "DropTarget_UndoSetActorMaterial", "Assign Material (Drag and Drop)") );
 							FoundMeshComponent->Modify();
-							for( int32 CurMaterialIndex = 0; CurMaterialIndex < MaterialCount; ++CurMaterialIndex )
+
+							// If the material slot is -1 then apply the material to every slot.
+							if ( OptionalMaterialSlot == -1 )
 							{
-								FoundMeshComponent->SetMaterial( CurMaterialIndex, MaterialToApply );
+								for ( int32 CurMaterialIndex = 0; CurMaterialIndex < MaterialCount; ++CurMaterialIndex )
+								{
+									FoundMeshComponent->SetMaterial(CurMaterialIndex, MaterialToApply);
+								}
 							}
+							else
+							{
+								check(OptionalMaterialSlot < MaterialCount);
+								FoundMeshComponent->SetMaterial(OptionalMaterialSlot, MaterialToApply);
+							}
+
 							TargetActor->MarkComponentsRenderStateDirty();
 							bResult = true;
 						}
diff --git a/Engine/Source/Editor/UnrealEd/Private/LevelEditorViewport.cpp b/Engine/Source/Editor/UnrealEd/Private/LevelEditorViewport.cpp
--- a/Engine/Source/Editor/UnrealEd/Private/LevelEditorViewport.cpp
+++ b/Engine/Source/Editor/UnrealEd/Private/LevelEditorViewport.cpp
@@ -531,10 +531,11 @@ static UObject* GetOrCreateMaterialFromTexture( UTexture* UnrealTexture )
  *
  * @param	ObjToUse				Object to attempt to apply as specific asset
  * @param	ActorToApplyTo			Actor to whom the asset should be applied
+ * @param   TargetMaterialSlot      When dealing with submeshes this will represent the target section/slot to apply materials to.
  *
  * @return	true if the provided object was successfully applied to the provided actor
  */
-static bool AttemptApplyObjToActor( UObject* ObjToUse, AActor* ActorToApplyTo, bool bTest = false )
+static bool AttemptApplyObjToActor( UObject* ObjToUse, AActor* ActorToApplyTo, int32 TargetMaterialSlot = -1, bool bTest = false )
 {
 	bool bResult = false;
 
@@ -565,7 +566,7 @@ static bool AttemptApplyObjToActor( UObject* ObjToUse, AActor* ActorToApplyTo, b
 			{
 				// Apply the material to the actor
 				FScopedTransaction Transaction( NSLOCTEXT("UnrealEd", "DragDrop_Transaction_ApplyMaterialToActor", "Apply Material to Actor") );
-				bResult = FActorFactoryAssetProxy::ApplyMaterialToActor( ActorToApplyTo, DroppedObjAsMaterial );
+				bResult = FActorFactoryAssetProxy::ApplyMaterialToActor( ActorToApplyTo, DroppedObjAsMaterial, TargetMaterialSlot );
 			}
 		}
 
@@ -854,7 +855,7 @@ bool FLevelEditorViewportClient::DropObjectsOnBackground( FViewportCursorLocatio
 	return bResult;
 }
 
-bool FLevelEditorViewportClient::DropObjectsOnActor(FViewportCursorLocation& Cursor, const TArray<UObject*>& DroppedObjects, AActor* DroppedUponActor, FVector* DroppedLocation, EObjectFlags ObjectFlags, TArray<AActor*>& OutNewActors, bool bUsedHitProxy, bool bSelectActors, UActorFactory* FactoryToUse)
+bool FLevelEditorViewportClient::DropObjectsOnActor(FViewportCursorLocation& Cursor, const TArray<UObject*>& DroppedObjects, AActor* DroppedUponActor, int32 DroppedUponSlot, FVector* DroppedLocation, EObjectFlags ObjectFlags, TArray<AActor*>& OutNewActors, bool bUsedHitProxy, bool bSelectActors, UActorFactory* FactoryToUse)
 {
 	bool bResult = false;
 
@@ -869,7 +870,7 @@ bool FLevelEditorViewportClient::DropObjectsOnActor(FViewportCursorLocation& Cur
 			ensure( AssetObj );
 
 			// Attempt to apply the dropped asset as a material to the actor
-			const bool bAppliedToActor = ( FactoryToUse == NULL ) ? AttemptApplyObjToActor( AssetObj, DroppedUponActor ) : false;
+			const bool bAppliedToActor = ( FactoryToUse == NULL ) ? AttemptApplyObjToActor( AssetObj, DroppedUponActor, DroppedUponSlot ) : false;
 			if ( !bAppliedToActor )
 			{
 				// Actor
@@ -1200,7 +1201,7 @@ bool FLevelEditorViewportClient::UpdateDropPreviewActors(int32 MouseX, int32 Mou
 						ensure( AssetObj );
 
 						// Attempt to apply the dropped asset as a material to the actor, just test if it is possible
-						if ( !AttemptApplyObjToActor( AssetObj, TargetActor, true ) )
+						if ( !AttemptApplyObjToActor( AssetObj, TargetActor, -1, true ) )
 						{
 							// hide all objects as we are applying them, not creating new object
 							out_bDroppedObjectsVisible = true;
@@ -1405,14 +1406,13 @@ void FLevelEditorViewportClient::DestroyDropPreviewActors()
 * @param MouseY			The position of the mouse's Y coordinate
 * @param AssetData			Asset in question to be dropped
 */
-bool FLevelEditorViewportClient::CanDropObjectsAtCoordinates(int32 MouseX, int32 MouseY, const FAssetData& AssetData)
+FDropQuery FLevelEditorViewportClient::CanDropObjectsAtCoordinates(int32 MouseX, int32 MouseY, const FAssetData& AssetData)
 {
-	// Initialize the result to wxDragNone in case we don't have any droppable assets
-	bool Result = false;
+	FDropQuery Result;
 
 	if ( !ObjectTools::IsAssetValidForPlacing( GetWorld(), AssetData.ObjectPath.ToString() ) )
 	{
-		return false;
+		return Result;
 	}
 
 	UObject* AssetObj = AssetData.GetAsset();
@@ -1430,12 +1430,12 @@ bool FLevelEditorViewportClient::CanDropObjectsAtCoordinates(int32 MouseX, int32
 
 		if ( AssetObj->IsA( AActor::StaticClass() ) || bHasActorFactory )
 		{
-			Result = true;
+			Result.bCanDrop = true;
 			bPivotMovedIndependantly = false;
 		}
 		else if( AssetObj->IsA( UBrushBuilder::StaticClass()) )
 		{
-			Result = true;
+			Result.bCanDrop = true;
 			bPivotMovedIndependantly = false;
 		}
 		else
@@ -1446,8 +1446,13 @@ bool FLevelEditorViewportClient::CanDropObjectsAtCoordinates(int32 MouseX, int32
 				if ( AssetObj->IsA( UMaterialInterface::StaticClass() ) || AssetObj->IsA( UTexture::StaticClass() ) )
 				{
 					// If our asset is a material and the target is a valid recipient
-					Result = true;
+					Result.bCanDrop = true;
 					bPivotMovedIndependantly = false;
+
+					//if ( HitProxy->IsA(HActor::StaticGetType()) )
+					//{
+					//	Result.HintText = LOCTEXT("Material_Shift_Hint", "Hold [Shift] to apply material to every slot");
+					//}
 				}
 			}
 		}
@@ -1500,11 +1505,14 @@ bool FLevelEditorViewportClient::DropObjectsAtCoordinates(int32 MouseX, int32 Mo
 			const bool bUsedHitProxy = true;
 			AActor* TargetActor = NULL;
 			FVector TargetLocation = FVector::ZeroVector;
+			int32 TargetMaterialSlot = -1;
+
 			if (HitProxy->IsA(HActor::StaticGetType()))
 			{
 				HActor* TargetProxy = static_cast<HActor*>(HitProxy);
 				TargetActor = TargetProxy->Actor;
 				TargetLocation = TargetActor ? TargetActor->GetActorLocation() : FVector::ZeroVector;
+				TargetMaterialSlot = TargetProxy->MaterialIndex;
 			}
 			else if (HitProxy->IsA(HBSPBrushVert::StaticGetType()))
 			{
@@ -1513,6 +1521,14 @@ bool FLevelEditorViewportClient::DropObjectsAtCoordinates(int32 MouseX, int32 Mo
 				TargetLocation = TargetActor ? (TargetProxy->Vertex ? TargetActor->ActorToWorld().TransformPosition(*TargetProxy->Vertex) : TargetActor->GetActorLocation()) : FVector::ZeroVector;
 			}
 
+			// If shift is pressed set the material slot to -1, so that it's applied to every slot.
+			// We have to request it from the platform application directly because IsShiftPressed gets 
+			// the cached state, when the viewport had focus
+			if ( FSlateApplication::Get().GetPlatformApplication()->GetModifierKeys().IsShiftDown() )
+			{
+				TargetMaterialSlot = -1;
+			}
+
 			if (TargetActor != NULL)
 			{
 				// if the target actor is selected, we should drop onto all selected items
@@ -1522,9 +1538,9 @@ bool FLevelEditorViewportClient::DropObjectsAtCoordinates(int32 MouseX, int32 Mo
 				if( !bDropOntoSelected || 
 					bOnlyDropOnTarget || 
 					FactoryToUse != NULL ||
-					!AttemptApplyObjToActor( DroppedObjects[0], TargetActor, true ) )
+					!AttemptApplyObjToActor(DroppedObjects[0], TargetActor, TargetMaterialSlot, true) )
 				{
-					bResult = DropObjectsOnActor( Cursor, DroppedObjects, TargetActor, &TargetLocation, ObjectFlags, OutNewActors, bUsedHitProxy, SelectActors, FactoryToUse );
+					bResult = DropObjectsOnActor( Cursor, DroppedObjects, TargetActor, TargetMaterialSlot, &TargetLocation, ObjectFlags, OutNewActors, bUsedHitProxy, SelectActors, FactoryToUse );
 				}
 				else
 				{
@@ -1534,7 +1550,7 @@ bool FLevelEditorViewportClient::DropObjectsAtCoordinates(int32 MouseX, int32 Mo
 						if( TargetActor )
 						{
 							TargetLocation = TargetActor->GetActorLocation();
-							DropObjectsOnActor( Cursor, DroppedObjects, TargetActor, &TargetLocation, ObjectFlags, OutNewActors, bUsedHitProxy, SelectActors, FactoryToUse );
+							DropObjectsOnActor( Cursor, DroppedObjects, TargetActor, TargetMaterialSlot, &TargetLocation, ObjectFlags, OutNewActors, bUsedHitProxy, SelectActors, FactoryToUse );
 							bResult = true;
 						}
 					}
diff --git a/Engine/Source/Editor/UnrealEd/Public/AssetSelection.h b/Engine/Source/Editor/UnrealEd/Public/AssetSelection.h
--- a/Engine/Source/Editor/UnrealEd/Public/AssetSelection.h
+++ b/Engine/Source/Editor/UnrealEd/Public/AssetSelection.h
@@ -287,10 +287,11 @@ public:
 	 *
 	 * @param	TargetActor		the actor to apply the material to
 	 * @param	MaterialToApply	the material to apply to the actor
+	 * @param   OptionalMaterialSlot the material slot to apply it to.
 	 *
 	 * @return	true if the material was successfully applied to the actor
 	 */
-	static bool ApplyMaterialToActor( AActor* TargetActor, UMaterialInterface* MaterialToApply );
+	static bool ApplyMaterialToActor( AActor* TargetActor, UMaterialInterface* MaterialToApply, int32 OptionalMaterialSlot = -1 );
 
 private:
 	/**
diff --git a/Engine/Source/Editor/UnrealEd/Public/LevelEditorViewport.h b/Engine/Source/Editor/UnrealEd/Public/LevelEditorViewport.h
--- a/Engine/Source/Editor/UnrealEd/Public/LevelEditorViewport.h
+++ b/Engine/Source/Editor/UnrealEd/Public/LevelEditorViewport.h
@@ -49,7 +49,6 @@ struct FViewportHoverTarget
 	{
 		return Key.HoveredActor ? GetTypeHash(Key.HoveredActor) : GetTypeHash(Key.HoveredModel)+Key.ModelSurfaceIndex;
 	}
-
 };
 
 struct FTrackingTransaction
@@ -101,6 +100,20 @@ private:
 };
 
 
+struct FDropQuery
+{
+	FDropQuery()
+		: bCanDrop(false)
+	{}
+
+	/** True if it's valid to drop the object at the location queried */
+	bool bCanDrop;
+
+	/** Optional hint text that may be returned to the user. */
+	FText HintText;
+};
+
+
 /** */
 class UNREALED_API FLevelEditorViewportClient : public FEditorViewportClient
 {
@@ -194,7 +207,7 @@ public:
 	/**
 	 * Updates the audio listener for this viewport 
 	 *
- 	 * @param View	The scene view to use when calculate the listener position
+	 * @param View	The scene view to use when calculate the listener position
 	 */
 	void UpdateAudioListener( const FSceneView& View );
 
@@ -231,7 +244,7 @@ public:
 	}
 
 	/** Called every time the viewport is ticked to update the view based on an actor that is
-	    driving the location, FOV and other settings for this view.  This is used for live camera PIP
+		driving the location, FOV and other settings for this view.  This is used for live camera PIP
 		preview within editor viewports */
 	void PushControllingActorDataToViewportClient();
 
@@ -370,7 +383,7 @@ public:
 	/**
 	 * If dragging an actor for placement, this function updates its position.
 	 *
- 	 * @param MouseX						The position of the mouse's X coordinate
+	 * @param MouseX						The position of the mouse's X coordinate
 	 * @param MouseY						The position of the mouse's Y coordinate
 	 * @param DroppedObjects				The Objects that were used to create preview objects
 	 * @param out_bDroppedObjectsVisible	Output, returns if preview objects are visible or not
@@ -387,16 +400,16 @@ public:
 	/**
 	 * Checks the viewport to see if the given object can be dropped using the given mouse coordinates local to this viewport
 	 *
- 	 * @param MouseX			The position of the mouse's X coordinate
+	 * @param MouseX			The position of the mouse's X coordinate
 	 * @param MouseY			The position of the mouse's Y coordinate
 	 * @param AssetInfo			Asset in question to be dropped
 	 */
-	bool CanDropObjectsAtCoordinates(int32 MouseX, int32 MouseY, const FAssetData& AssetInfo);
+	FDropQuery CanDropObjectsAtCoordinates(int32 MouseX, int32 MouseY, const FAssetData& AssetInfo);
 
 	/**
 	 * Attempts to intelligently drop the given objects in the viewport, using the given mouse coordinates local to this viewport
 	 *
- 	 * @param MouseX			 The position of the mouse's X coordinate
+	 * @param MouseX			 The position of the mouse's X coordinate
 	 * @param MouseY			 The position of the mouse's Y coordinate
 	 * @param DroppedObjects	 The Objects to be placed into the editor via this viewport
 	 * @param OutNewActors		 The new actor objects that were created
@@ -408,17 +421,17 @@ public:
 	bool DropObjectsAtCoordinates(int32 MouseX, int32 MouseY, const TArray<UObject*>& DroppedObjects, TArray<AActor*>& OutNewActors, bool bOnlyDropOnTarget = false, bool bCreateDropPreview = false, bool bSelectActors = true, class UActorFactory* FactoryToUse = NULL );
 
 	/**
- 	 * Sets GWorld to the appropriate world for this client
+	 * Sets GWorld to the appropriate world for this client
 	 * 
 	 * @return the previous GWorld
- 	 */
+	 */
 	virtual UWorld* ConditionalSetWorld() OVERRIDE;
 
 	/**
- 	 * Restores GWorld to InWorld
+	 * Restores GWorld to InWorld
 	 *
 	 * @param InWorld	The world to restore
- 	 */
+	 */
 	virtual void ConditionalRestoreWorld( UWorld* InWorld  ) OVERRIDE;
 
 	/**
@@ -584,6 +597,7 @@ private:
 	* @param	Cursor				Mouse cursor location
 	* @param	DroppedObjects		Array of objects dropped into the viewport
 	* @param	DroppedUponActor	The actor that we are dropping upon
+	* @param    DroppedUponSlot     The material slot/submesh that was identified as the drop location.  If unknown use -1.
 	* @param	DroppedLocation		The location that we're dropping the objects
 	* @param	ObjectFlags			The object flags to place on the actors that this function spawns.
 	* @param	OutNewActors		The list of actors created while dropping
@@ -593,7 +607,7 @@ private:
 	*
 	* @return	true if the drop operation was successfully handled; false otherwise
 	*/
-	bool DropObjectsOnActor(struct FViewportCursorLocation& Cursor, const TArray<UObject*>& DroppedObjects, AActor* DroppedUponActor, FVector* DroppedLocation, EObjectFlags ObjectFlags, TArray<AActor*>& OutNewActors, bool bUsedHitProxy = true, bool bSelectActors = true, class UActorFactory* FactoryToUse = NULL);
+	bool DropObjectsOnActor(struct FViewportCursorLocation& Cursor, const TArray<UObject*>& DroppedObjects, AActor* DroppedUponActor, int32 DroppedUponSlot, FVector* DroppedLocation, EObjectFlags ObjectFlags, TArray<AActor*>& OutNewActors, bool bUsedHitProxy = true, bool bSelectActors = true, class UActorFactory* FactoryToUse = NULL);
 
 	/**
 	 * Called when an asset is dropped upon a BSP surface.
@@ -691,7 +705,7 @@ private:
 	static TArray< TWeakObjectPtr< AActor > > DropPreviewActors;
 
 	/** Optional actor that is 'controlling' this view.  That is, the view's location, rotation, FOV and other settings
-	    should be pushed to this caInera every frame.  This is used by the editor's live viewport camera PIP feature */
+		should be pushed to this caInera every frame.  This is used by the editor's live viewport camera PIP feature */
 	TWeakObjectPtr< AActor > ControllingActor;
 
 	/** Bit array representing the visibility of every sprite category in the current viewport */
diff --git a/Engine/Source/Runtime/Engine/Private/SkeletalMesh.cpp b/Engine/Source/Runtime/Engine/Private/SkeletalMesh.cpp
--- a/Engine/Source/Runtime/Engine/Private/SkeletalMesh.cpp
+++ b/Engine/Source/Runtime/Engine/Private/SkeletalMesh.cpp
@@ -1159,7 +1159,7 @@ void FStaticLODModel::Serialize( FArchive& Ar, UObject* Owner, int32 Idx )
 	}
 
 	Ar << Sections;
-	MultiSizeIndexContainer.Serialize(Ar,bKeepBuffersInCPUMemory);
+	MultiSizeIndexContainer.Serialize(Ar, bKeepBuffersInCPUMemory);
 	Ar << ActiveBoneIndices;
 	Ar << Chunks;
 	Ar << Size;
@@ -3987,13 +3987,62 @@ private:
 #endif
 };
 
+#if WITH_EDITOR
+HHitProxy* FSkeletalMeshSceneProxy::CreateHitProxies(UPrimitiveComponent* Component, TArray<TRefCountPtr<HHitProxy> >& OutHitProxies)
+{
+	if ( Component->GetOwner() )
+	{
+		if ( LODSections.Num() > 0 )
+		{
+			for ( int32 LODIndex = 0; LODIndex < SkelMeshResource->LODModels.Num(); LODIndex++ )
+			{
+				const FStaticLODModel& LODModel = SkelMeshResource->LODModels[LODIndex];
+
+				FLODSectionElements& LODSection = LODSections[LODIndex];
+
+				check(LODSection.SectionElements.Num() == LODModel.Sections.Num());
+
+				for ( int32 SectionIndex = 0; SectionIndex < LODModel.Sections.Num(); SectionIndex++ )
+				{
+					HHitProxy* ActorHitProxy;
+
+					if ( Component->GetOwner()->IsA(ABrush::StaticClass()) && Component->IsA(UBrushComponent::StaticClass()) )
+					{
+						ActorHitProxy = new HActor(Component->GetOwner(), Component, HPP_Wireframe, SectionIndex);
+					}
+					else
+					{
+						ActorHitProxy = new HActor(Component->GetOwner(), Component, SectionIndex);
+					}
+
+					// Set the hitproxy.
+					check(LODSection.SectionElements[SectionIndex].HitProxy == NULL);
+					LODSection.SectionElements[SectionIndex].HitProxy = ActorHitProxy;
+					OutHitProxies.Add(ActorHitProxy);
+				}
+			}
+
+			return NULL;
+		}
+		else
+		{
+			return FPrimitiveSceneProxy::CreateHitProxies(Component, OutHitProxies);
+		}
+	}
+	else
+	{
+		return NULL;
+	}
+}
+#endif
+
 /** 
 * Draw the scene proxy as a dynamic element
 *
 * @param	PDI - draw interface to render to
 * @param	View - current view
 */
-void FSkeletalMeshSceneProxy::DrawDynamicElements(FPrimitiveDrawInterface* PDI,const FSceneView* View)
+void FSkeletalMeshSceneProxy::DrawDynamicElements(FPrimitiveDrawInterface* PDI, const FSceneView* View)
 {
 	QUICK_SCOPE_CYCLE_COUNTER( STAT_SkeletalMeshSceneProxy_DrawDynamicElements );
 
@@ -4010,9 +4059,9 @@ void FSkeletalMeshSceneProxy::DrawDynamicElements(FPrimitiveDrawInterface* PDI,c
 	{
 		const FLODSectionElements& LODSection = LODSections[LODIndex];
 
-		check(LODSection.SectionElements.Num()==LODModel.Sections.Num());
+		check(LODSection.SectionElements.Num() == LODModel.Sections.Num());
 
-		for (FSkeletalMeshSectionIter Iter(LODIndex,*MeshObject,LODModel,LODSection); Iter; ++Iter)
+		for (FSkeletalMeshSectionIter Iter(LODIndex, *MeshObject, LODModel, LODSection); Iter; ++Iter)
 		{
 			FSkelMeshSection Section = Iter.GetSection();
 			const FSkelMeshChunk& Chunk = Iter.GetChunk();
@@ -4024,7 +4073,7 @@ void FSkeletalMeshSceneProxy::DrawDynamicElements(FPrimitiveDrawInterface* PDI,c
 			Section.bSelected = (SkeletalMeshForDebug->SelectedEditorSection == SectionElementInfo.UseMaterialIndex);
 #endif
 			// If hidden skip the draw
-			if (MeshObject->IsMaterialHidden(LODIndex,SectionElementInfo.UseMaterialIndex))
+			if (MeshObject->IsMaterialHidden(LODIndex, SectionElementInfo.UseMaterialIndex))
 			{
 				continue;
 			}
@@ -4035,6 +4084,13 @@ void FSkeletalMeshSceneProxy::DrawDynamicElements(FPrimitiveDrawInterface* PDI,c
 				continue;
 			}
 
+#if WITH_EDITOR
+			if ( GIsEditor && PDI->IsHitTesting() )
+			{
+				PDI->SetHitProxy(SectionElementInfo.HitProxy);
+			}
+#endif
+
 			DrawDynamicElementsSection(PDI, View, LODModel, LODIndex, Section, Chunk, SectionElementInfo, CustomLeftRightVectors );
 		}
 	}
diff --git a/Engine/Source/Runtime/Engine/Private/StaticMeshRender.cpp b/Engine/Source/Runtime/Engine/Private/StaticMeshRender.cpp
--- a/Engine/Source/Runtime/Engine/Private/StaticMeshRender.cpp
+++ b/Engine/Source/Runtime/Engine/Private/StaticMeshRender.cpp
@@ -293,6 +293,49 @@ void FStaticMeshSceneProxy::SetIndexSource(int32 LODIndex, int32 SectionIndex, F
 }
 
 // FPrimitiveSceneProxy interface.
+#if WITH_EDITOR
+HHitProxy* FStaticMeshSceneProxy::CreateHitProxies(UPrimitiveComponent* Component, TArray<TRefCountPtr<HHitProxy> >& OutHitProxies)
+{
+	if ( Component->GetOwner() )
+	{
+		for ( int32 LODIndex = 0; LODIndex < RenderData->LODResources.Num(); LODIndex++ )
+		{
+			const FStaticMeshLODResources& LODModel = RenderData->LODResources[LODIndex];
+
+			check(LODs[LODIndex].Sections.Num() == LODModel.Sections.Num());
+
+			for ( int32 SectionIndex = 0; SectionIndex < LODModel.Sections.Num(); SectionIndex++ )
+			{
+				HHitProxy* ActorHitProxy;
+
+				if ( Component->GetOwner()->IsA(ABrush::StaticClass()) && Component->IsA(UBrushComponent::StaticClass()) )
+				{
+					ActorHitProxy = new HActor(Component->GetOwner(), Component, HPP_Wireframe, SectionIndex);
+				}
+				else
+				{
+					ActorHitProxy = new HActor(Component->GetOwner(), Component, SectionIndex);
+				}
+
+				FLODInfo::FSectionInfo& Section = LODs[LODIndex].Sections[SectionIndex];
+
+				// Set the hitproxy.
+				check(Section.HitProxy == NULL);
+				Section.HitProxy = ActorHitProxy;
+
+				OutHitProxies.Add(ActorHitProxy);
+			}
+		}
+		
+		return NULL;
+	}
+	else
+	{
+		return NULL;
+	}
+}
+#endif // WITH_EDITOR
+
 void FStaticMeshSceneProxy::DrawStaticElements(FStaticPrimitiveDrawInterface* PDI)
 {
 	checkSlow(IsInRenderingThread());
@@ -312,16 +355,20 @@ void FStaticMeshSceneProxy::DrawStaticElements(FStaticPrimitiveDrawInterface* PD
 			int32 LODIndex = FMath::Clamp(ForcedLodModel, 1, NumLODs) - 1;
 			const FStaticMeshLODResources& LODModel = RenderData->LODResources[LODIndex];
 			// Draw the static mesh elements.
-			for(int32 SectionIndex = 0;SectionIndex < LODModel.Sections.Num();SectionIndex++)
+			for(int32 SectionIndex = 0; SectionIndex < LODModel.Sections.Num(); SectionIndex++)
 			{
 #if WITH_EDITOR
 				if( GIsEditor )
 				{
-					bUseSelectedMaterial = LODs[LODIndex].Sections[SectionIndex].bSelected;
+					const FLODInfo::FSectionInfo& Section = LODs[LODIndex].Sections[SectionIndex];
+
+					bUseSelectedMaterial = Section.bSelected;
+					PDI->SetHitProxy(Section.HitProxy);
 				}
 #endif // WITH_EDITOR
+
 				FMeshBatch MeshElement;
-				if(GetMeshElement(LODIndex,SectionIndex,PrimitiveDPG,MeshElement,bUseSelectedMaterial,bUseHoveredMaterial))
+				if(GetMeshElement(LODIndex, SectionIndex, PrimitiveDPG, MeshElement, bUseSelectedMaterial, bUseHoveredMaterial))
 				{
 					PDI->DrawMesh(MeshElement, 0, FLT_MAX);
 				}
@@ -329,7 +376,7 @@ void FStaticMeshSceneProxy::DrawStaticElements(FStaticPrimitiveDrawInterface* PD
 		} 
 		else //no LOD is being forced, submit them all with appropriate cull distances
 		{
-			for(int32 LODIndex = 0;LODIndex < NumLODs;LODIndex++)
+			for(int32 LODIndex = 0; LODIndex < NumLODs; LODIndex++)
 			{
 				const FStaticMeshLODResources& LODModel = RenderData->LODResources[LODIndex];
 				float MinDist = GetMinLODDist(LODIndex);
@@ -376,11 +423,15 @@ void FStaticMeshSceneProxy::DrawStaticElements(FStaticPrimitiveDrawInterface* PD
 #if WITH_EDITOR
 					if( GIsEditor )
 					{
-						bUseSelectedMaterial = LODs[LODIndex].Sections[SectionIndex].bSelected;
+						const FLODInfo::FSectionInfo& Section = LODs[LODIndex].Sections[SectionIndex];
+
+						bUseSelectedMaterial = Section.bSelected;
+						PDI->SetHitProxy(Section.HitProxy);
 					}
 #endif // WITH_EDITOR
+
 					FMeshBatch MeshElement;
-					if(GetMeshElement(LODIndex,SectionIndex,PrimitiveDPG,MeshElement,bUseSelectedMaterial,bUseHoveredMaterial))
+					if(GetMeshElement(LODIndex, SectionIndex, PrimitiveDPG, MeshElement, bUseSelectedMaterial, bUseHoveredMaterial))
 					{
 						// If we have submitted an optimized shadow-only mesh, remaining mesh elements must not cast shadows.
 						MeshElement.CastShadow = MeshElement.CastShadow && !bHaveShadowOnlyMesh;
@@ -496,16 +547,20 @@ void FStaticMeshSceneProxy::DrawDynamicElements(FPrimitiveDrawInterface* PDI,con
 				const FLinearColor UtilColor( LevelColor );
 
 				// Draw the static mesh sections.
-				for(int32 SectionIndex = 0;SectionIndex < LODModel.Sections.Num();SectionIndex++)
+				for(int32 SectionIndex = 0; SectionIndex < LODModel.Sections.Num(); SectionIndex++)
 				{
 					bool bSectionIsSelected = false;
 
 #if WITH_EDITOR
 					if( GIsEditor )
 					{
-						bSectionIsSelected = LODs[LODIndex].Sections[SectionIndex].bSelected;
+						const FLODInfo::FSectionInfo& Section = LODs[LODIndex].Sections[SectionIndex];
+
+						bSectionIsSelected = Section.bSelected;
+						PDI->SetHitProxy(Section.HitProxy);
 					}
 #endif // WITH_EDITOR
+
 					FMeshBatch MeshElement;
 					if(GetMeshElement(LODIndex,SectionIndex,GetDepthPriorityGroup(View),MeshElement,bSectionIsSelected,IsHovered()))
 					{
diff --git a/Engine/Source/Runtime/Engine/Public/EngineUtils.h b/Engine/Source/Runtime/Engine/Public/EngineUtils.h
--- a/Engine/Source/Runtime/Engine/Public/EngineUtils.h
+++ b/Engine/Source/Runtime/Engine/Public/EngineUtils.h
@@ -17,16 +17,32 @@ struct HActor : public HHitProxy
 	DECLARE_HIT_PROXY( ENGINE_API )
 	AActor* Actor;
 	const UPrimitiveComponent* PrimComponent;
+	int32 MaterialIndex;
 
-	HActor( AActor* InActor, const UPrimitiveComponent* InPrimComponent ) 
+	HActor( AActor* InActor, const UPrimitiveComponent* InPrimComponent )
 		: Actor( InActor ) 
 		, PrimComponent( InPrimComponent )
+		, MaterialIndex( -1 )
+		{}
+
+	HActor(AActor* InActor, const UPrimitiveComponent* InPrimComponent, int32 MaterialIndex)
+		: Actor( InActor )
+		, PrimComponent( InPrimComponent )
+		, MaterialIndex( MaterialIndex )
 		{}
 
 	HActor( AActor* InActor, const UPrimitiveComponent* InPrimComponent, EHitProxyPriority InPriority) 
-		: HHitProxy(InPriority)
-		, Actor(InActor) 
+		: HHitProxy( InPriority )
+		, Actor( InActor )
 		, PrimComponent( InPrimComponent )
+		, MaterialIndex( -1 )
+		{}
+
+	HActor(AActor* InActor, const UPrimitiveComponent* InPrimComponent, EHitProxyPriority InPriority, int32 MaterialIndex)
+		: HHitProxy(InPriority)
+		, Actor(InActor)
+		, PrimComponent(InPrimComponent)
+		, MaterialIndex(MaterialIndex)
 		{}
 
 	virtual void AddReferencedObjects( FReferenceCollector& Collector ) OVERRIDE
diff --git a/Engine/Source/Runtime/Engine/Public/SkeletalMeshTypes.h b/Engine/Source/Runtime/Engine/Public/SkeletalMeshTypes.h
--- a/Engine/Source/Runtime/Engine/Public/SkeletalMeshTypes.h
+++ b/Engine/Source/Runtime/Engine/Public/SkeletalMeshTypes.h
@@ -584,7 +584,7 @@ struct FApexClothPhysToRenderVertData
 {
 	/**
 	\brief xyz : Barycentric coordinates of the graphical vertex relative to the simulated triangle.
-	       w : distance from the mesh
+		   w : distance from the mesh
 
 	\note PX_MAX_F32 values represent invalid coordinates.
 	*/
@@ -773,7 +773,7 @@ struct FSkelMeshChunk
 	*/
 	FORCEINLINE bool HasApexClothData() const
 	{
-        return (ApexClothMappingData.Num() > 0);
+		return (ApexClothMappingData.Num() > 0);
 	}
 
 	FORCEINLINE void SetClothSubmeshIndex(int16 AssetIndex, int16 AssetSubmeshIndex)
@@ -847,8 +847,8 @@ struct FSkelMeshSection
 	/** This Section can be disabled for cloth simulation and corresponding Cloth Section will be enabled*/
 	bool bDisabled;
 	/** Corresponding Section Index will be enabled when this section is disabled 
-	    because corresponding cloth section will be showed instead of this
-	    or disabled section index when this section is enabled for cloth simulation
+		because corresponding cloth section will be showed instead of this
+		or disabled section index when this section is enabled for cloth simulation
 	*/
 	int16 CorrespondClothSectionIndex;
 
@@ -863,12 +863,12 @@ struct FSkelMeshSection
 		, TriangleSorting(0)
 		, bSelected(false)
 		, bDisabled(false)
-		,CorrespondClothSectionIndex(-1)
-		,bEnableClothLOD(true)
+		, CorrespondClothSectionIndex(-1)
+		, bEnableClothLOD(true)
 	{}
 
 	// Serialization.
-	friend FArchive& operator<<(FArchive& Ar,FSkelMeshSection& S);
+	friend FArchive& operator<<(FArchive& Ar, FSkelMeshSection& S);
 };
 template <> struct TIsPODType<FSkelMeshSection> { enum { Value = true }; };
 
@@ -2000,7 +2000,7 @@ public:
 				return true;
 			}
 		}
-        return false;
+		return false;
 	}
 
 	int32 GetApexClothCunkIndex(TArray<int32>& ChunkIndices) const
@@ -2128,7 +2128,10 @@ public:
 	FSkeletalMeshSceneProxy(const USkinnedMeshComponent* Component, FSkeletalMeshResource* InSkelMeshResource);
 
 	// FPrimitiveSceneProxy interface.
-	virtual void DrawDynamicElements(FPrimitiveDrawInterface* PDI,const FSceneView* View) OVERRIDE;
+#if WITH_EDITOR
+	virtual HHitProxy* CreateHitProxies(UPrimitiveComponent* Component, TArray<TRefCountPtr<HHitProxy> >& OutHitProxies) OVERRIDE;
+#endif
+	virtual void DrawDynamicElements(FPrimitiveDrawInterface* PDI, const FSceneView* View) OVERRIDE;
 	virtual FPrimitiveViewRelevance GetViewRelevance(const FSceneView* View) OVERRIDE;
 	virtual bool CanBeOccluded() const OVERRIDE;
 	virtual void PreRenderView(const FSceneViewFamily* ViewFamily, const uint32 VisibilityMap, int32 FrameNumber) OVERRIDE;
@@ -2197,22 +2200,27 @@ protected:
 	/** info for section element in an LOD */
 	struct FSectionElementInfo
 	{
-		/*
-		FSectionElementInfo() 
-		:	Material(NULL)
-		,	bEnableShadowCasting(true)
-		{}
-		*/
 		FSectionElementInfo(UMaterialInterface* InMaterial, bool bInEnableShadowCasting, int32 InUseMaterialIndex)
-		:	Material( InMaterial )
-		,	bEnableShadowCasting( bInEnableShadowCasting )
-		,	UseMaterialIndex( InUseMaterialIndex )
+		: Material( InMaterial )
+		, bEnableShadowCasting( bInEnableShadowCasting )
+		, UseMaterialIndex( InUseMaterialIndex )
+#if WITH_EDITOR
+		, HitProxy(NULL)
+#endif
 		{}
+		
 		UMaterialInterface* Material;
+		
 		/** Whether shadow casting is enabled for this section. */
 		bool bEnableShadowCasting;
+		
 		/** Index into the materials array of the skel mesh or the component after LOD mapping */
 		int32 UseMaterialIndex;
+
+#if WITH_EDITOR
+		/** The editor needs to be able to individual sub-mesh hit detection, so we store a hit proxy on each mesh. */
+		HHitProxy* HitProxy;
+#endif
 	};
 
 	/** Section elements for a particular LOD */
diff --git a/Engine/Source/Runtime/Engine/Public/StaticMeshResources.h b/Engine/Source/Runtime/Engine/Public/StaticMeshResources.h
--- a/Engine/Source/Runtime/Engine/Public/StaticMeshResources.h
+++ b/Engine/Source/Runtime/Engine/Public/StaticMeshResources.h
@@ -749,6 +749,9 @@ protected:
 
 public:
 	// FPrimitiveSceneProxy interface.
+#if WITH_EDITOR
+	virtual HHitProxy* CreateHitProxies(UPrimitiveComponent* Component, TArray<TRefCountPtr<HHitProxy> >& OutHitProxies) OVERRIDE;
+#endif
 	virtual void DrawStaticElements(FStaticPrimitiveDrawInterface* PDI) OVERRIDE;
 	virtual void DrawDynamicElements(FPrimitiveDrawInterface* PDI,const FSceneView* View) OVERRIDE { DrawDynamicElements( PDI, View, 0 ); }
 	virtual void DrawDynamicElements(FPrimitiveDrawInterface* PDI,const FSceneView* View, uint32 DrawDynamicFlags ) OVERRIDE;
@@ -770,17 +773,27 @@ protected:
 		/** Information about an element of a LOD. */
 		struct FSectionInfo
 		{
+			/** Default constructor. */
+			FSectionInfo()
+				: Material(NULL)
+				, bSelected(false)
+#if WITH_EDITOR
+				, HitProxy(NULL)
+#endif
+			{}
+
 			/** The material with which to render this section. */
 			UMaterialInterface* Material;
+
 			/** True if this section should be rendered as selected (editor only). */
 			bool bSelected;
 
-			/** Default constructor. */
-			FSectionInfo() :
-				Material(NULL),
-				bSelected(false)
-			{}
+#if WITH_EDITOR
+			/** The editor needs to be able to individual sub-mesh hit detection, so we store a hit proxy on each mesh. */
+			HHitProxy* HitProxy;
+#endif
 		};
+
 		/** Per-section information. */
 		TArray<FSectionInfo> Sections;
 
diff --git a/Engine/Source/Runtime/Renderer/Private/SceneHitProxyRendering.cpp b/Engine/Source/Runtime/Renderer/Private/SceneHitProxyRendering.cpp
--- a/Engine/Source/Runtime/Renderer/Private/SceneHitProxyRendering.cpp
+++ b/Engine/Source/Runtime/Renderer/Private/SceneHitProxyRendering.cpp
@@ -99,7 +99,7 @@ IMPLEMENT_MATERIAL_SHADER_TYPE(,FHitProxyHS,TEXT("HitProxyVertexShader"),TEXT("M
 IMPLEMENT_MATERIAL_SHADER_TYPE(,FHitProxyDS,TEXT("HitProxyVertexShader"),TEXT("MainDomain"),SF_Domain);
 
 /**
- * A pixel shader for rendering the depth of a mesh.
+ * A pixel shader for rendering the HitProxyId of an object as a unique color in the scene.
  */
 class FHitProxyPS : public FMeshMaterialShader
 {
@@ -135,7 +135,7 @@ public:
 
 	void SetHitProxyId(FHitProxyId HitProxyIdValue)
 	{
-		SetShaderValue(GetPixelShader(),HitProxyId,HitProxyIdValue.GetColor().ReinterpretAsLinear());
+		SetShaderValue(GetPixelShader(), HitProxyId, HitProxyIdValue.GetColor().ReinterpretAsLinear());
 	}
 
 	virtual bool Serialize(FArchive& Ar)
@@ -155,7 +155,7 @@ FHitProxyDrawingPolicy::FHitProxyDrawingPolicy(
 	const FVertexFactory* InVertexFactory,
 	const FMaterialRenderProxy* InMaterialRenderProxy
 	):
-	FMeshDrawingPolicy(InVertexFactory,InMaterialRenderProxy,*InMaterialRenderProxy->GetMaterial(GRHIFeatureLevel))
+	FMeshDrawingPolicy(InVertexFactory, InMaterialRenderProxy, *InMaterialRenderProxy->GetMaterial(GRHIFeatureLevel))
 {
 	HullShader = NULL;
 	DomainShader = NULL;
@@ -172,7 +172,7 @@ FHitProxyDrawingPolicy::FHitProxyDrawingPolicy(
 	PixelShader = MaterialResource->GetShader<FHitProxyPS>(InVertexFactory->GetType());
 }
 
-void FHitProxyDrawingPolicy::DrawShared(const FSceneView* View,FBoundShaderStateRHIParamRef BoundShaderState) const
+void FHitProxyDrawingPolicy::DrawShared(const FSceneView* View, FBoundShaderStateRHIParamRef BoundShaderState) const
 {
 	// Set the actual shader & vertex declaration state
 	RHISetBoundShaderState( BoundShaderState);
@@ -224,15 +224,15 @@ void FHitProxyDrawingPolicy::SetMeshRenderState(
 
 	const FMeshBatchElement& BatchElement = Mesh.Elements[BatchElementIndex];
 
-	VertexShader->SetMesh(VertexFactory,View,PrimitiveSceneProxy,BatchElement);
+	VertexShader->SetMesh(VertexFactory, View, PrimitiveSceneProxy, BatchElement);
 
 	if(HullShader && DomainShader)
 	{
-		HullShader->SetMesh(VertexFactory,View,PrimitiveSceneProxy,BatchElement);
-		DomainShader->SetMesh(VertexFactory,View,PrimitiveSceneProxy,BatchElement);
+		HullShader->SetMesh(VertexFactory, View, PrimitiveSceneProxy, BatchElement);
+		DomainShader->SetMesh(VertexFactory, View, PrimitiveSceneProxy, BatchElement);
 	}
 
-	PixelShader->SetMesh(VertexFactory,View,PrimitiveSceneProxy,BatchElement);
+	PixelShader->SetMesh(VertexFactory, View, PrimitiveSceneProxy, BatchElement);
 	// Per-instance hitproxies are supplied by the vertex factory
 	if( PrimitiveSceneProxy && PrimitiveSceneProxy->HasPerInstanceHitProxies() )
 	{
@@ -311,10 +311,10 @@ bool FHitProxyDrawingPolicyFactory::DrawDynamicMesh(
 				MaterialRenderProxy = UMaterial::GetDefaultMaterial(MD_Surface)->GetRenderProxy(false);
 			}
 			FHitProxyDrawingPolicy DrawingPolicy( Mesh.VertexFactory, MaterialRenderProxy );
-			DrawingPolicy.DrawShared(&View,DrawingPolicy.CreateBoundShaderState());
-			for( int32 BatchElementIndex=0;BatchElementIndex<Mesh.Elements.Num();BatchElementIndex++ )
+			DrawingPolicy.DrawShared(&View, DrawingPolicy.CreateBoundShaderState());
+			for( int32 BatchElementIndex = 0; BatchElementIndex < Mesh.Elements.Num(); BatchElementIndex++ )
 			{
-				DrawingPolicy.SetMeshRenderState(View,PrimitiveSceneProxy,Mesh,BatchElementIndex,bBackFace,HitProxyId);
+				DrawingPolicy.SetMeshRenderState(View, PrimitiveSceneProxy, Mesh, BatchElementIndex, bBackFace, HitProxyId);
 				DrawingPolicy.DrawMesh(Mesh,BatchElementIndex);
 			}
 			return true;
