#ttp 331130 CRITICAL: Regression: Networking: Multiplayer Options: Number of Players only launches 1 instance
#branch UE4
#project Editor
#summary Launch the correct number of separate clients when testing multiplayer games without using a single instance.
#change Make sure to count the editor as a client when set to be a client.
#change When MapNameOverride isn't empty assign it as the only element in SavedMapNames so that secondary instances actually launch.

diff --git a/Engine/Source/Editor/UnrealEd/Private/PlayLevel.cpp b/Engine/Source/Editor/UnrealEd/Private/PlayLevel.cpp
--- a/Engine/Source/Editor/UnrealEd/Private/PlayLevel.cpp
+++ b/Engine/Source/Editor/UnrealEd/Private/PlayLevel.cpp
@@ -712,6 +712,12 @@ void UEditorEngine::StartQueuedPlayMapRequest()
 				NumClients++;
 			}
 		}
+		
+		if(!bPlayOnLocalPcSession && !bStartMovieCapture && PlayInSettings->PlayNetMode == PIE_Client)
+		{
+			// Editor counts as a client
+			NumClients++;
+		}
 
 		// Spawn number of clients
 		for (int32 i=NumClients; i < PlayInSettings->PlayNumberOfClients; ++i)
@@ -908,7 +914,7 @@ void UEditorEngine::PlayOnLocalPc(FString MapNameOverride, FString URLParms, FSt
 	TArray<FString> SavedMapNames;
 	if (MapNameOverride.IsEmpty())
 	{
-	FWorldContext & EditorContext = GetEditorWorldContext();
+		FWorldContext & EditorContext = GetEditorWorldContext();
 		if (EditorContext.World()->WorldComposition)
 		{
 			// Open world composition from original folder
@@ -921,6 +927,10 @@ void UEditorEngine::PlayOnLocalPc(FString MapNameOverride, FString URLParms, FSt
 			SaveWorldForPlay(SavedMapNames);
 		}
 	}
+	else
+	{
+		SavedMapNames.Add(MapNameOverride);
+	}
 
 	if (SavedMapNames.Num() == 0)
 	{
