#ttp 331904 [GitHub] Fix case for WIndows, from dotneB -https://github.com/EpicGames/UnrealEngine/pull/41/files

#summary Fix integrated from GitHub, along with another instance of "WIndows" I found.

diff --git a/Engine/Source/Editor/UnrealEd/Classes/Editor/EditorEngine.h b/Engine/Source/Editor/UnrealEd/Classes/Editor/EditorEngine.h
--- a/Engine/Source/Editor/UnrealEd/Classes/Editor/EditorEngine.h
+++ b/Engine/Source/Editor/UnrealEd/Classes/Editor/EditorEngine.h
@@ -2434,7 +2434,7 @@ private:
 
 	virtual void VerifyLoadMapWorldCleanup() OVERRIDE;
 
-	FPIEInstanceWindowSwitch PIEInstanceWIndowSwitchDelegate;
+	FPIEInstanceWindowSwitch PIEInstanceWindowSwitchDelegate;
 
 protected:
 
diff --git a/Engine/Source/Editor/UnrealEd/Private/PlayLevel.cpp b/Engine/Source/Editor/UnrealEd/Private/PlayLevel.cpp
--- a/Engine/Source/Editor/UnrealEd/Private/PlayLevel.cpp
+++ b/Engine/Source/Editor/UnrealEd/Private/PlayLevel.cpp
@@ -2697,7 +2697,7 @@ bool UEditorEngine::WorldIsPIEInNewViewport(UWorld *InWorld)
 
 void UEditorEngine::SetPIEInstanceWindowSwitchDelegate(FPIEInstanceWindowSwitch InSwitchDelegate)
 {
-	PIEInstanceWIndowSwitchDelegate = InSwitchDelegate;
+	PIEInstanceWindowSwitchDelegate = InSwitchDelegate;
 }
 
 void UEditorEngine::FocusNextPIEWorld(UWorld *CurrentPieWorld, bool previous)
@@ -2739,7 +2739,7 @@ void UEditorEngine::FocusNextPIEWorld(UWorld *CurrentPieWorld, bool previous)
 			FSlateApplication::Get().SetKeyboardFocus(SlateInfoPtr->SlatePlayInEditorWindowViewport->GetViewportWidget().Pin(), EKeyboardFocusCause::Keyboard);
 
 			// Execute notifcation delegate incase game code has to do anything else
-			PIEInstanceWIndowSwitchDelegate.ExecuteIfBound();
+			PIEInstanceWindowSwitchDelegate.ExecuteIfBound();
 		}
 	}
 }
diff --git a/Engine/Source/Runtime/Core/Public/HAL/Platform.h b/Engine/Source/Runtime/Core/Public/HAL/Platform.h
--- a/Engine/Source/Runtime/Core/Public/HAL/Platform.h
+++ b/Engine/Source/Runtime/Core/Public/HAL/Platform.h
@@ -17,7 +17,7 @@
 
 //@port Identify the platform here and include the platform header to setup the platform types, etc
 #if PLATFORM_WINDOWS
-	#include "Windows/WIndowsPlatform.h"
+	#include "Windows/WindowsPlatform.h"
 #elif PLATFORM_PS4
 	#include "PS4/PS4Platform.h"
 #elif PLATFORM_XBOXONE
