GitHub pull request

https://github.com/EpicGames/UnrealEngine/pull/48/

Fix spelling of "dependency" in logs.

Fix from GitHub user "AKPWebDesign"

diff --git a/Engine/Source/Editor/UnrealEd/Private/Commandlets/CookCommandlet.cpp b/Engine/Source/Editor/UnrealEd/Private/Commandlets/CookCommandlet.cpp
--- a/Engine/Source/Editor/UnrealEd/Private/Commandlets/CookCommandlet.cpp
+++ b/Engine/Source/Editor/UnrealEd/Private/Commandlets/CookCommandlet.cpp
@@ -497,7 +497,7 @@ bool UCookCommandlet::ShouldCook(const FString& InFileName, const FString &InPla
 
 		if (GetPackageTimestamp(FPaths::GetBaseFilename(PkgFilename, false), DependentTimeStamp) == false)
 		{
-			UE_LOG(LogCookCommandlet, Display, TEXT("Failed to find depedency timestamp for: %s"), *PkgFilename);
+			UE_LOG(LogCookCommandlet, Display, TEXT("Failed to find dependency timestamp for: %s"), *PkgFilename);
 		}
 	}
 
@@ -569,7 +569,7 @@ bool UCookCommandlet::SaveCookedPackage( UPackage* Package, uint32 SaveFlags, bo
 
 			if (GetPackageTimestamp(FPaths::GetBaseFilename(PkgFilename, false), DependentTimeStamp) == false)
 			{
-				UE_LOG(LogCookCommandlet, Display, TEXT("Failed to find depedency timestamp for: %s"), *PkgFilename);
+				UE_LOG(LogCookCommandlet, Display, TEXT("Failed to find dependency timestamp for: %s"), *PkgFilename);
 			}
 		}
 
