[GitHub] Removed UEditorEngine::GetContentBrowserAssetTreePath
https://github.com/EpicGames/UnrealEngine/pull/34

diff --git a/Engine/Source/Editor/Kismet/Private/STimelineEditor.cpp b/Engine/Source/Editor/Kismet/Private/STimelineEditor.cpp
--- a/Engine/Source/Editor/Kismet/Private/STimelineEditor.cpp
+++ b/Engine/Source/Editor/Kismet/Private/STimelineEditor.cpp
@@ -333,9 +333,6 @@ void STimelineEdTrack::InitCurveAssetPath()
 	//Default path
 	SaveCurveAssetPath = FString(TEXT( "/Game/Unsorted" ));
 
-	//Get the currently selected 
-	GEditor->GetContentBrowserAssetTreePath( SaveCurveAssetPath );
-
 	TSharedRef<STimelineEditor> TimelineRef = TimelineEdPtr.Pin().ToSharedRef();
 
 	//Get curve name from editable text box
diff --git a/Engine/Source/Editor/UnrealEd/Classes/Editor/EditorEngine.h b/Engine/Source/Editor/UnrealEd/Classes/Editor/EditorEngine.h
--- a/Engine/Source/Editor/UnrealEd/Classes/Editor/EditorEngine.h
+++ b/Engine/Source/Editor/UnrealEd/Classes/Editor/EditorEngine.h
@@ -1565,9 +1565,6 @@ public:
 	/** Function to return list of assets currently selected in content browser (loaded/not loaded) */
 	void GetContentBrowserSelections( TArray<UClass*>& Selection ) const;
 
-	/** Function to fetch the currently selected asset tree path in content browser*/
-	void GetContentBrowserAssetTreePath( FString& SelectedPath ) const;
-
 	/**
 	 * Returns an FSelectionIterator that iterates over the set of selected actors.
 	 */
diff --git a/Engine/Source/Editor/UnrealEd/Private/Editor.cpp b/Engine/Source/Editor/UnrealEd/Private/Editor.cpp
--- a/Engine/Source/Editor/UnrealEd/Private/Editor.cpp
+++ b/Engine/Source/Editor/UnrealEd/Private/Editor.cpp
@@ -269,11 +269,6 @@ void UEditorEngine::GetContentBrowserSelections( TArray<UClass*>& Selection ) co
 	}
 }
 
-void UEditorEngine::GetContentBrowserAssetTreePath( FString& SelectedPath ) const
-{
-	// @todo: Remove this function
-}
-
 USelection* UEditorEngine::GetSelectedSet( const UClass* Class ) const
 {
 	USelection* SelectedSet = GetSelectedActors();
