TTP #331746: Blueprints that were loaded by a previous automation test have to be unloaded to properly test compile-on-load.

diff --git a/Engine/Source/Editor/UnrealEd/Private/Tests/BlueprintAutomationTests.cpp b/Engine/Source/Editor/UnrealEd/Private/Tests/BlueprintAutomationTests.cpp
--- a/Engine/Source/Editor/UnrealEd/Private/Tests/BlueprintAutomationTests.cpp
+++ b/Engine/Source/Editor/UnrealEd/Private/Tests/BlueprintAutomationTests.cpp
@@ -388,7 +388,7 @@ public:
 	 * @param  AssetPath	A path detailing the asset in question (in the form of <PackageName>.<AssetName>)
 	 * @return True if a blueprint can be found with a matching package/name, false if no corresponding blueprint was found.
 	 */
-	static bool IsBlueprintLoaded(FString const& AssetPath)
+	static bool IsBlueprintLoaded(FString const& AssetPath, UBlueprint** BlueprintOut = nullptr)
 	{
 		bool bIsLoaded = false;
 
@@ -400,6 +400,10 @@ public:
 			if (UBlueprint* ExistingBP = Cast<UBlueprint>(StaticFindObject(UBlueprint::StaticClass(), ExistingPackage, *AssetName)))
 			{
 				bIsLoaded = true;
+				if (BlueprintOut != nullptr)
+				{
+					*BlueprintOut = ExistingBP;
+				}
 			}
 		}
 
@@ -700,9 +704,10 @@ bool FBlueprintCompileOnLoadTest::RunTest(const FString& BlueprintAssetPath)
 {
 	FCompilerResultsLog Results;
 
+	UBlueprint* ExistingBP = nullptr;
 	// if this blueprint was already loaded, then these tests are invalidated 
 	// (because dependencies have already been loaded)
-	if (FBlueprintAutomationTestUtilities::IsBlueprintLoaded(BlueprintAssetPath))
+	if (FBlueprintAutomationTestUtilities::IsBlueprintLoaded(BlueprintAssetPath, &ExistingBP))
 	{
 		if (FBlueprintAutomationTestUtilities::IsAssetUnsaved(BlueprintAssetPath))
 		{
@@ -711,7 +716,8 @@ bool FBlueprintCompileOnLoadTest::RunTest(const FString& BlueprintAssetPath)
 		}
 		else 
 		{
-			AddWarning(FString::Printf(TEXT("Invalid test, blueprint has previously been loaded: '%s'"), *BlueprintAssetPath));
+			AddWarning(FString::Printf(TEXT("Test may be invalid (the blueprint is already loaded): '%s'"), *BlueprintAssetPath));
+			FBlueprintAutomationTestUtilities::UnloadBlueprint(ExistingBP);
 		}		
 	}
 
diff --git a/Engine/Source/Runtime/Engine/Private/Blueprint.cpp b/Engine/Source/Runtime/Engine/Private/Blueprint.cpp
--- a/Engine/Source/Runtime/Engine/Private/Blueprint.cpp
+++ b/Engine/Source/Runtime/Engine/Private/Blueprint.cpp
@@ -181,7 +181,11 @@ void UBlueprint::Serialize(FArchive& Ar)
 	{
 		for (UClass* ClassIt = ParentClass; (ClassIt != NULL) && !(ClassIt->HasAnyClassFlags(CLASS_Native)); ClassIt = ClassIt->GetSuperClass())
 		{
-			if (ClassIt->ClassGeneratedBy->HasAnyFlags(RF_NeedLoad))
+			if (!ensure(ClassIt->ClassGeneratedBy != nullptr))
+			{
+				UE_LOG(LogBlueprint, Error, TEXT("Cannot preload parent blueprint from null ClassGeneratedBy field (for '%s')"), *ClassIt->GetName());
+			}
+			else if (ClassIt->ClassGeneratedBy->HasAnyFlags(RF_NeedLoad))
 			{
 				ClassIt->ClassGeneratedBy->GetLinker()->Preload(ClassIt->ClassGeneratedBy);
 			}
