#UE4 - Refactor of FBX material importing to not look at FBX material name to determine if a material is unique. Resolves issues where materials are not properly applied to meshes on import if they are the same name as another material in the fbx file.
- We now support tiling parameters on fbx materials and generate the proper node in the Unreal material on import
- Cleanup of some of the SkinXX syntax

diff --git a/Engine/Source/Editor/UnrealEd/Private/Factories/ApexDestructibleAssetImport.cpp b/Engine/Source/Editor/UnrealEd/Private/Factories/ApexDestructibleAssetImport.cpp
--- a/Engine/Source/Editor/UnrealEd/Private/Factories/ApexDestructibleAssetImport.cpp
+++ b/Engine/Source/Editor/UnrealEd/Private/Factories/ApexDestructibleAssetImport.cpp
@@ -319,7 +319,8 @@ static void ImportMaterialsForSkelMesh(FSkeletalMeshImportData &ImportData, cons
 		{
 			ImportData.Materials.Add( VMaterial() );
 
-			ImportData.Materials.Last().MaterialName = DefaultMaterial->GetName();
+			ImportData.Materials.Last().Material = DefaultMaterial;
+			ImportData.Materials.Last().MaterialImportName = DefaultMaterial->GetName();
 		}
 	}
 }
@@ -483,7 +484,7 @@ static bool FillSkelMeshImporterFromApexDestructibleAsset(FSkeletalMeshImportDat
 
 		for (int32 CurUniqueMatIdx=0; CurUniqueMatIdx < UniqueMaterials.Num(); ++CurUniqueMatIdx)
 		{
-			if (ImportData.Materials[CurMatIdx].MaterialName == UniqueMaterials[CurUniqueMatIdx].MaterialName)
+			if (ImportData.Materials[CurMatIdx].MaterialImportName == UniqueMaterials[CurUniqueMatIdx].MaterialImportName)
 			{
 				bHasMaterial = true;
 				break;
@@ -619,11 +620,11 @@ static bool FillSkelMeshImporterFromApexDestructibleAsset(FSkeletalMeshImportDat
 				
 				if (SubmeshIndex < (uint32)ImportData.Materials.Num())
 				{
-					const FString& MaterialName = ImportData.Materials[SubmeshIndex].MaterialName;
+					const FString& MaterialName = ImportData.Materials[SubmeshIndex].MaterialImportName;
 
 					for (int32 UniqueMaterialIdx = 0; UniqueMaterialIdx < UniqueMaterials.Num(); ++UniqueMaterialIdx)
 					{
-						if (UniqueMaterials[UniqueMaterialIdx].MaterialName == MaterialName)
+						if (UniqueMaterials[UniqueMaterialIdx].MaterialImportName == MaterialName)
 						{
 							MatIdx = UniqueMaterialIdx;
 							break;
@@ -851,7 +852,8 @@ bool SetApexDestructibleAsset(UDestructibleMesh& DestructibleMesh, NxDestructibl
 			UMaterialInterface* MI = DestructibleMesh.Materials[i].MaterialInterface;
 			VMaterial NewMat;
 			
-			NewMat.MaterialName = MI != NULL ? MI->GetName() : TEXT("");
+			NewMat.Material = MI;
+			NewMat.MaterialImportName = MI != NULL ? MI->GetName() : TEXT("");
 			SkelMeshImportDataPtr->Materials.Add(NewMat);
 		}
 	}
diff --git a/Engine/Source/Editor/UnrealEd/Private/Factories/SkeletalMeshImport.cpp b/Engine/Source/Editor/UnrealEd/Private/Factories/SkeletalMeshImport.cpp
--- a/Engine/Source/Editor/UnrealEd/Private/Factories/SkeletalMeshImport.cpp
+++ b/Engine/Source/Editor/UnrealEd/Private/Factories/SkeletalMeshImport.cpp
@@ -150,26 +150,32 @@ void FSkeletalMeshImportData::CopyLODImportData(
 */
 void ProcessImportMeshMaterials(TArray<FSkeletalMaterial>& Materials, FSkeletalMeshImportData& ImportData )
 {
-	TArray <VMaterial>&	MaterialsBinary = ImportData.Materials;
+	TArray <VMaterial>&	ImportedMaterials = ImportData.Materials;
 
 	// If direct linkup of materials is requested, try to find them here - to get a texture name from a 
 	// material name, cut off anything in front of the dot (beyond are special flags).
 	Materials.Empty();
-	for( int32 m=0; m < MaterialsBinary.Num(); m++)
+	for( int32 MatIndex=0; MatIndex < ImportedMaterials.Num(); ++MatIndex)
 	{			
-		const FString& MaterialName = MaterialsBinary[m].MaterialName;
+		const VMaterial& ImportedMaterial = ImportedMaterials[MatIndex];
 
-		UMaterialInterface* Material = FindObject<UMaterialInterface>( ANY_PACKAGE, *MaterialName );
-		Materials.Add( FSkeletalMaterial( Material, true ) );	// I guess true is a sensible default for shadow casting
-
-		if( !Material )
+		UMaterialInterface* Material = NULL;
+		if( ImportedMaterial.Material.IsValid() )
 		{
-			UE_LOG(LogSkeletalMeshImport, Log, TEXT(" Mesh material not found among currently loaded ones: %s"), *MaterialName );
+			Material = ImportedMaterial.Material.Get();
 		}
+		else
+		{
+			const FString& MaterialName = ImportedMaterial.MaterialImportName;
+			Material = FindObject<UMaterialInterface>(ANY_PACKAGE, *MaterialName);
+		}
+
+		const bool bEnableShadowCasting = true;
+		Materials.Add( FSkeletalMaterial( Material, bEnableShadowCasting ) );
 	}
 
 	// Pad the material pointers.
-	while( MaterialsBinary.Num() > Materials.Num() )
+	while( ImportedMaterials.Num() > Materials.Num() )
 	{
 		Materials.Add( FSkeletalMaterial( NULL, true ) );
 	}
diff --git a/Engine/Source/Editor/UnrealEd/Private/Fbx/FbxFactory.cpp b/Engine/Source/Editor/UnrealEd/Private/Fbx/FbxFactory.cpp
--- a/Engine/Source/Editor/UnrealEd/Private/Fbx/FbxFactory.cpp
+++ b/Engine/Source/Editor/UnrealEd/Private/Fbx/FbxFactory.cpp
@@ -184,12 +184,13 @@ UObject* UFbxFactory::FactoryCreateBinary
 	}
 	// logger for all error/warnings
 	// this one prints all messages that are stored in FFbxImporter
-	UnFbx::FFbxImporter* FFbxImporter = UnFbx::FFbxImporter::GetInstance();
-	UnFbx::FFbxLoggerSetter Logger(FFbxImporter);
+	UnFbx::FFbxImporter* FbxImporter = UnFbx::FFbxImporter::GetInstance();
+	
+	UnFbx::FFbxLoggerSetter Logger(FbxImporter);
 
 	bool bShowImportDialog = bShowOption && !GIsAutomationTesting;
 	bool bImportAll = false;
-	UnFbx::FBXImportOptions* ImportOptions = GetImportOptions(FFbxImporter, ImportUI, bShowImportDialog, InParent->GetPathName(), bOperationCanceled, bImportAll);
+	UnFbx::FBXImportOptions* ImportOptions = GetImportOptions(FbxImporter, ImportUI, bShowImportDialog, InParent->GetPathName(), bOperationCanceled, bImportAll);
 	bOutOperationCanceled = bOperationCanceled;
 	
 	if( bImportAll )
@@ -204,22 +205,22 @@ UObject* UFbxFactory::FactoryCreateBinary
 	if (ImportOptions)
 	{
 		Warn->BeginSlowTask( NSLOCTEXT("FbxFactory", "BeginImportingFbxMeshTask", "Importing FBX mesh"), true );
-		if ( !FFbxImporter->ImportFromFile( *UFactory::CurrentFilename ) )
+		if ( !FbxImporter->ImportFromFile( *UFactory::CurrentFilename ) )
 		{
 			// Log the error message and fail the import.
-			Warn->Log(ELogVerbosity::Error, FFbxImporter->GetErrorMessage() );
+			Warn->Log(ELogVerbosity::Error, FbxImporter->GetErrorMessage() );
 		}
 		else
 		{
 			// Log the import message and import the mesh.
-			const TCHAR* errorMessage = FFbxImporter->GetErrorMessage();
+			const TCHAR* errorMessage = FbxImporter->GetErrorMessage();
 			if (errorMessage[0] != NULL)
 			{
 				Warn->Log( errorMessage );
 			}
 
 			FbxNode* RootNodeToImport = NULL;
-			RootNodeToImport = FFbxImporter->Scene->GetRootNode();
+			RootNodeToImport = FbxImporter->Scene->GetRootNode();
 
 			// For animation and static mesh we assume there is at lease one interesting node by default
 			int32 InterestingNodeCount = 1;
@@ -229,7 +230,7 @@ UObject* UFbxFactory::FactoryCreateBinary
 			bool bCombineMeshes = ImportUI->bCombineMeshes;
 			if ( ImportUI->MeshTypeToImport == FBXIT_SkeletalMesh )
 			{
-				FFbxImporter->FillFbxSkelMeshArrayInScene(RootNodeToImport, SkelMeshArray, false);
+				FbxImporter->FillFbxSkelMeshArrayInScene(RootNodeToImport, SkelMeshArray, false);
 				InterestingNodeCount = SkelMeshArray.Num();
 			}
 			else if( ImportUI->MeshTypeToImport == FBXIT_StaticMesh )
@@ -244,7 +245,7 @@ UObject* UFbxFactory::FactoryCreateBinary
 					// count meshes in lod groups if we dont care about importing LODs
 					bool bCountLODGroupMeshes = !bImportStaticMeshLODs;
 					int32 NumLODGroups = 0;
-					InterestingNodeCount = FFbxImporter->GetFbxMeshCount(RootNodeToImport,bCountLODGroupMeshes,NumLODGroups);
+					InterestingNodeCount = FbxImporter->GetFbxMeshCount(RootNodeToImport,bCountLODGroupMeshes,NumLODGroups);
 
 					// if there were LODs in the file, do not combine meshes even if requested
 					if( bImportStaticMeshLODs && bCombineMeshes )
@@ -273,10 +274,10 @@ UObject* UFbxFactory::FactoryCreateBinary
 					if (bCombineMeshes)
 					{
 						TArray<FbxNode*> FbxMeshArray;
-						FFbxImporter->FillFbxMeshArray(RootNodeToImport, FbxMeshArray, FFbxImporter);
+						FbxImporter->FillFbxMeshArray(RootNodeToImport, FbxMeshArray, FbxImporter);
 						if (FbxMeshArray.Num() > 0)
 						{
-							NewStaticMesh = FFbxImporter->ImportStaticMeshAsSingle(InParent, FbxMeshArray, Name, Flags, ImportUI->StaticMeshImportData, NULL, 0);
+							NewStaticMesh = FbxImporter->ImportStaticMeshAsSingle(InParent, FbxMeshArray, Name, Flags, ImportUI->StaticMeshImportData, NULL, 0);
 						}
 
 						ImportedMeshCount = NewStaticMesh ? 1 : 0;
@@ -284,7 +285,7 @@ UObject* UFbxFactory::FactoryCreateBinary
 					else
 					{
 						TArray<UObject*> AllNewAssets;
-						UObject* Object = RecursiveImportNode(FFbxImporter,RootNodeToImport,InParent,Name,Flags,NodeIndex,InterestingNodeCount, AllNewAssets);
+						UObject* Object = RecursiveImportNode(FbxImporter,RootNodeToImport,InParent,Name,Flags,NodeIndex,InterestingNodeCount, AllNewAssets);
 
 						NewStaticMesh = Cast<UStaticMesh>( Object );
 
@@ -305,7 +306,7 @@ UObject* UFbxFactory::FactoryCreateBinary
 					// Importing static mesh sockets only works if one mesh is being imported
 					if( ImportedMeshCount == 1 && NewStaticMesh )
 					{
-						FFbxImporter->ImportStaticMeshSockets( NewStaticMesh );
+						FbxImporter->ImportStaticMeshSockets( NewStaticMesh );
 					}
 
 					NewObject = NewStaticMesh;
@@ -367,32 +368,23 @@ UObject* UFbxFactory::FactoryCreateBinary
 						
 							if (LODIndex == 0 && SkelMeshNodeArray.Num() != 0)
 							{
-								FName OutputName = FFbxImporter->MakeNameForMesh(Name.ToString(), SkelMeshNodeArray[0]);
+								FName OutputName = FbxImporter->MakeNameForMesh(Name.ToString(), SkelMeshNodeArray[0]);
 
-								USkeletalMesh* NewMesh = FFbxImporter->ImportSkeletalMesh( InParent, SkelMeshNodeArray, OutputName, Flags, ImportUI->SkeletalMeshImportData, FPaths::GetBaseFilename(Filename) );
+								USkeletalMesh* NewMesh = FbxImporter->ImportSkeletalMesh( InParent, SkelMeshNodeArray, OutputName, Flags, ImportUI->SkeletalMeshImportData, FPaths::GetBaseFilename(Filename) );
 								NewObject = NewMesh;
 
 								if ( NewMesh && ImportUI->bImportAnimations )
 								{
-									FFbxImporter->SetupAnimationDataFromMesh(NewMesh, InParent, SkelMeshNodeArray, ImportUI->AnimSequenceImportData, OutputName.ToString());
+									FbxImporter->SetupAnimationDataFromMesh(NewMesh, InParent, SkelMeshNodeArray, ImportUI->AnimSequenceImportData, OutputName.ToString());
 								}
 							}
 							else if (NewObject) // the base skeletal mesh is imported successfully
 							{
 								USkeletalMesh* BaseSkeletalMesh = Cast<USkeletalMesh>(NewObject);
 								FName LODObjectName = NAME_None;
-								USkeletalMesh *LODObject = FFbxImporter->ImportSkeletalMesh( GetTransientPackage(), SkelMeshNodeArray, LODObjectName, RF_NoFlags, ImportUI->SkeletalMeshImportData, FPaths::GetBaseFilename(Filename) );
-								FFbxImporter->ImportSkeletalMeshLOD( LODObject, BaseSkeletalMesh, LODIndex);
-								// Set LOD Model's DisplayFactor
-								/*
-								fbxDistance Threshold;
+								USkeletalMesh *LODObject = FbxImporter->ImportSkeletalMesh( GetTransientPackage(), SkelMeshNodeArray, LODObjectName, RF_NoFlags, ImportUI->SkeletalMeshImportData, FPaths::GetBaseFilename(Filename) );
+								FbxImporter->ImportSkeletalMeshLOD( LODObject, BaseSkeletalMesh, LODIndex);
 
-								if ( LodGroup->GetThreshold(i-1, Threshold) )
-								{
-								BaseSkeletalMesh->LODInfo(i).DisplayFactor = (float)((Threshold.value() - LodGroup->MinDistance.Get()) /
-								(LodGroup->MaxDistance.Get() - LodGroup->MinDistance.Get()));
-								}
-								*/
 								BaseSkeletalMesh->LODInfo[LODIndex].DisplayFactor = 1.0f / (MaxLODLevel * LODIndex);
 							}
 						
@@ -403,7 +395,7 @@ UObject* UFbxFactory::FactoryCreateBinary
 								uint32 bImportMaterials = ImportOptions->bImportMaterials;
 								ImportOptions->bImportMaterials = 0;
 
-								FFbxImporter->ImportFbxMorphTarget(SkelMeshNodeArray, Cast<USkeletalMesh>(NewObject), InParent, Filename, LODIndex);
+								FbxImporter->ImportFbxMorphTarget(SkelMeshNodeArray, Cast<USkeletalMesh>(NewObject), InParent, Filename, LODIndex);
 							
 								ImportOptions->bImportMaterials = !!bImportMaterials;
 							}
@@ -427,7 +419,7 @@ UObject* UFbxFactory::FactoryCreateBinary
 					// if total nodes we found is 0, we didn't find anything. 
 					if (TotalNumNodes == 0)
 					{
-						FFbxImporter->AddTokenizedErrorMessage(FTokenizedMessage::Create(EMessageSeverity::Error, LOCTEXT("FailedToImport_NoMeshFoundOnRoot", "Could not find any valid mesh on the root hierarchy. If you have mesh in the sub hierarchy, please enable option of [Import Meshes In Bone Hierarchy] when import.")));
+						FbxImporter->AddTokenizedErrorMessage(FTokenizedMessage::Create(EMessageSeverity::Error, LOCTEXT("FailedToImport_NoMeshFoundOnRoot", "Could not find any valid mesh on the root hierarchy. If you have mesh in the sub hierarchy, please enable option of [Import Meshes In Bone Hierarchy] when import.")));
 					}
 				}
 				else if ( ImportUI->MeshTypeToImport == FBXIT_Animation )// animation
@@ -443,25 +435,25 @@ UObject* UFbxFactory::FactoryCreateBinary
 			{
 				if (RootNodeToImport == NULL)
 				{
-					FFbxImporter->AddTokenizedErrorMessage(FTokenizedMessage::Create(EMessageSeverity::Error, LOCTEXT("FailedToImport_InvalidRoot", "Could not find root node.")));
+					FbxImporter->AddTokenizedErrorMessage(FTokenizedMessage::Create(EMessageSeverity::Error, LOCTEXT("FailedToImport_InvalidRoot", "Could not find root node.")));
 				}
 				else if (ImportUI->MeshTypeToImport == FBXIT_SkeletalMesh)
 				{
-					FFbxImporter->AddTokenizedErrorMessage(FTokenizedMessage::Create(EMessageSeverity::Error, LOCTEXT("FailedToImport_InvalidBone", "Failed to find any bone hierarchy. Try to import as Rigid Mesh option enabled. ")));
+					FbxImporter->AddTokenizedErrorMessage(FTokenizedMessage::Create(EMessageSeverity::Error, LOCTEXT("FailedToImport_InvalidBone", "Failed to find any bone hierarchy. Try to import as Rigid Mesh option enabled. ")));
 				}
 				else
 				{
-					FFbxImporter->AddTokenizedErrorMessage(FTokenizedMessage::Create(EMessageSeverity::Error, LOCTEXT("FailedToImport_InvalidNode", "Could not find any node.")));
+					FbxImporter->AddTokenizedErrorMessage(FTokenizedMessage::Create(EMessageSeverity::Error, LOCTEXT("FailedToImport_InvalidNode", "Could not find any node.")));
 				}
 			}
 		}
 
 		if (NewObject == NULL)
 		{
-			FFbxImporter->AddTokenizedErrorMessage(FTokenizedMessage::Create(EMessageSeverity::Error, LOCTEXT("FailedToImport_NoObject", "Import failed.")));
+			FbxImporter->AddTokenizedErrorMessage(FTokenizedMessage::Create(EMessageSeverity::Error, LOCTEXT("FailedToImport_NoObject", "Import failed.")));
 		}
 
-		FFbxImporter->ReleaseScene();
+		FbxImporter->ReleaseScene();
 		Warn->EndSlowTask();
 	}
 
diff --git a/Engine/Source/Editor/UnrealEd/Private/Fbx/FbxMainImport.cpp b/Engine/Source/Editor/UnrealEd/Private/Fbx/FbxMainImport.cpp
--- a/Engine/Source/Editor/UnrealEd/Private/Fbx/FbxMainImport.cpp
+++ b/Engine/Source/Editor/UnrealEd/Private/Fbx/FbxMainImport.cpp
@@ -197,9 +197,30 @@ void ApplyImportUIToImportOptions(UFbxImportUI* ImportUI, FBXImportOptions& InOu
 	InOutImportOptions.bPreserveLocalTransform = ImportUI->bPreserveLocalTransform;
 }
 
-//-------------------------------------------------------------------------
-//
-//-------------------------------------------------------------------------
+void FImportedMaterialData::AddImportedMaterial( FbxSurfaceMaterial& FbxMaterial, UMaterialInterface& UnrealMaterial )
+{
+	FbxToUnrealMaterialMap.Add( &FbxMaterial, &UnrealMaterial );
+	ImportedMaterialNames.Add( *UnrealMaterial.GetPathName() );
+}
+
+bool FImportedMaterialData::IsUnique( FbxSurfaceMaterial& FbxMaterial, FName ImportedMaterialName ) const
+{
+	UMaterialInterface* FoundMaterial = GetUnrealMaterial( FbxMaterial );
+
+	return FoundMaterial != NULL || ImportedMaterialNames.Contains( ImportedMaterialName );
+}
+
+UMaterialInterface* FImportedMaterialData::GetUnrealMaterial( const FbxSurfaceMaterial& FbxMaterial ) const
+{
+	return FbxToUnrealMaterialMap.FindRef( &FbxMaterial ).Get();
+}
+
+void FImportedMaterialData::Clear()
+{
+	FbxToUnrealMaterialMap.Empty();
+	ImportedMaterialNames.Empty();
+}
+
 FFbxImporter::FFbxImporter()
 	: bFirstMesh(true)
 	, Importer( NULL )
@@ -234,8 +255,6 @@ FFbxImporter::~FFbxImporter()
 	CleanUp();
 }
 
-const float FFbxImporter::SCALE_TOLERANCE = 0.000001;
-
 //-------------------------------------------------------------------------
 //
 //-------------------------------------------------------------------------
@@ -291,6 +310,8 @@ void FFbxImporter::ReleaseScene()
 		Scene = NULL;
 	}
 	
+	ImportedMaterialData.Clear();
+
 	// reset
 	CollisionModels.Clear();
 	CurPhase = NOTSTARTED;
diff --git a/Engine/Source/Editor/UnrealEd/Private/Fbx/FbxMaterialImport.cpp b/Engine/Source/Editor/UnrealEd/Private/Fbx/FbxMaterialImport.cpp
--- a/Engine/Source/Editor/UnrealEd/Private/Fbx/FbxMaterialImport.cpp
+++ b/Engine/Source/Editor/UnrealEd/Private/Fbx/FbxMaterialImport.cpp
@@ -248,6 +248,9 @@ bool UnFbx::FFbxImporter::CreateAndLinkExpressionForMaterialProperty(
 				
 					if (UnrealTexture)
 					{
+						float ScaleU = FbxTexture->GetScaleU();
+						float ScaleV = FbxTexture->GetScaleV();
+
 						// and link it to the material 
 						UMaterialExpressionTextureSample* UnrealTextureExpression = ConstructObject<UMaterialExpressionTextureSample>( UMaterialExpressionTextureSample::StaticClass(), UnrealMaterial );
 						UnrealMaterial->Expressions.Add( UnrealTextureExpression );
@@ -261,15 +264,18 @@ bool UnFbx::FFbxImporter::CreateAndLinkExpressionForMaterialProperty(
 						FbxString UVSetName = FbxTexture->UVSet.Get();
 						FString LocalUVSetName = ANSI_TO_TCHAR(UVSetName.Buffer());
 						int32 SetIndex = UVSet.Find(LocalUVSetName);
-						if( SetIndex != 0 && SetIndex != INDEX_NONE )
+						if( (SetIndex != 0 && SetIndex != INDEX_NONE) || ScaleU != 1.0f || ScaleV != 1.0f )
 						{
 							// Create a texture coord node for the texture sample
 							UMaterialExpressionTextureCoordinate* MyCoordExpression = ConstructObject<UMaterialExpressionTextureCoordinate>(UMaterialExpressionTextureCoordinate::StaticClass(), UnrealMaterial);
 							UnrealMaterial->Expressions.Add(MyCoordExpression);
 							MyCoordExpression->CoordinateIndex = (SetIndex >= 0) ? SetIndex : 0;
+							MyCoordExpression->UTiling = ScaleU;
+							MyCoordExpression->VTiling = ScaleV;
 							UnrealTextureExpression->Coordinates.Expression = MyCoordExpression;
 							MyCoordExpression->MaterialExpressionEditorX = FMath::Trunc(Location.X+175);
 							MyCoordExpression->MaterialExpressionEditorY = FMath::Trunc(Location.Y);
+
 						}
 
 						bCreated = true;
@@ -349,24 +355,22 @@ void UnFbx::FFbxImporter::FixupMaterial( FbxSurfaceMaterial& FbxMaterial, UMater
 //
 //-------------------------------------------------------------------------
 
-void UnFbx::FFbxImporter::CreateUnrealMaterial(FbxSurfaceMaterial* FbxMaterial, TArray<UMaterialInterface*>& OutMaterials, TArray<FString>& UVSets)
+void UnFbx::FFbxImporter::CreateUnrealMaterial(FbxSurfaceMaterial& FbxMaterial, TArray<UMaterialInterface*>& OutMaterials, TArray<FString>& UVSets)
 {
-	FString MaterialFullName = ANSI_TO_TCHAR(MakeName(FbxMaterial->GetName()));
+	FString MaterialFullName = ANSI_TO_TCHAR(MakeName(FbxMaterial.GetName()));
 
-	// check for a 'skinXX' suffix in the material name
-	int32 MaterialNameLen = FCString::Strlen(*MaterialFullName) + 1;
-	char* MaterialNameANSI = new char[MaterialNameLen];
-	FCStringAnsi::Strcpy(MaterialNameANSI, MaterialNameLen, TCHAR_TO_ANSI(*MaterialFullName));
-	if (FCStringAnsi::Strlen(MaterialNameANSI) > 6)
+	if (MaterialFullName.Len() > 6)
 	{
-		const char* SkinXX = MaterialNameANSI + FCStringAnsi::Strlen(MaterialNameANSI) - 6;
-		if (FCharAnsi::ToUpper(*SkinXX) == 'S' && FCharAnsi::ToUpper(*(SkinXX+1)) == 'K' && 
-			FCharAnsi::ToUpper(*(SkinXX+2)) == 'I' && FCharAnsi::ToUpper(*(SkinXX+3)) == 'N')
+		int32 Offset = MaterialFullName.Find(TEXT("_SKIN"), ESearchCase::IgnoreCase, ESearchDir::FromEnd);
+		if (Offset != INDEX_NONE)
 		{
-			if (FCharAnsi::IsDigit(*(SkinXX+4)) && FCharAnsi::IsDigit(*(SkinXX+5)))
+			// Chop off the material name so we are left with the number in _SKINXX
+			FString SkinXXNumber = MaterialFullName.Right(MaterialFullName.Len() - (Offset + 1)).RightChop(4);
+
+			if (SkinXXNumber.IsNumeric())
 			{
-				// remove the 'skinXX' suffix from the material name
-				MaterialFullName = MaterialFullName.Left(MaterialNameLen - 7);
+				// remove the '_skinXX' suffix from the material name					
+				MaterialFullName = MaterialFullName.LeftChop(Offset+1);
 			}
 		}
 	}
@@ -382,18 +386,30 @@ void UnFbx::FFbxImporter::CreateUnrealMaterial(FbxSurfaceMaterial* FbxMaterial,
 	FString BasePackageName = FPackageName::GetLongPackagePath(Parent->GetOutermost()->GetName()) / MaterialFullName;
 	BasePackageName = PackageTools::SanitizePackageName(BasePackageName);
 
-	// First check if the asset already exists.  We do not override existing materials
+	// The material could already exist in the project
+	FName ObjectPath = *(BasePackageName + TEXT(".") + MaterialFullName);
+
+	if( ImportedMaterialData.IsUnique( FbxMaterial, ObjectPath ) )
 	{
-		FString ObjectPath = BasePackageName + TEXT(".") + MaterialFullName;
-		UMaterialInterface* FoundMaterial = LoadObject<UMaterialInterface>( NULL, *ObjectPath );
+		UMaterialInterface* FoundMaterial = ImportedMaterialData.GetUnrealMaterial( FbxMaterial );
+		if (FoundMaterial)
+		{
+			// The material was imported from this FBX.  Reuse it
+			OutMaterials.Add(FoundMaterial);
+			return;
+		}
+	}
+	else
+	{
+		UMaterialInterface* FoundMaterial = LoadObject<UMaterialInterface>(NULL, *ObjectPath.ToString());
 		// do not override existing materials
-		if (FoundMaterial != NULL)
+		if (FoundMaterial)
 		{
+			ImportedMaterialData.AddImportedMaterial( FbxMaterial, *FoundMaterial );
 			OutMaterials.Add(FoundMaterial);
 			return;
 		}
 	}
-
 	
 	const FString Suffix(TEXT(""));
 	FAssetToolsModule& AssetToolsModule = FModuleManager::LoadModuleChecked<FAssetToolsModule>("AssetTools");
@@ -408,7 +424,6 @@ void UnFbx::FFbxImporter::CreateUnrealMaterial(FbxSurfaceMaterial* FbxMaterial,
 	
 	UMaterial* UnrealMaterial = (UMaterial*)MaterialFactory->FactoryCreateNew(
 		UMaterial::StaticClass(), Package, *MaterialFullName, RF_Standalone|RF_Public, NULL, GWarn );
-	// TODO :  need this ? UnrealMaterial->bUsedWithStaticLighting = true;
 
 	if ( UnrealMaterial != NULL )
 	{
@@ -420,18 +435,18 @@ void UnFbx::FFbxImporter::CreateUnrealMaterial(FbxSurfaceMaterial* FbxMaterial,
 	}
 
 	// textures and properties
-	CreateAndLinkExpressionForMaterialProperty( *FbxMaterial, UnrealMaterial, FbxSurfaceMaterial::sDiffuse, UnrealMaterial->BaseColor, false, UVSets, FVector2D(240,-320) );
-	CreateAndLinkExpressionForMaterialProperty( *FbxMaterial, UnrealMaterial, FbxSurfaceMaterial::sEmissive, UnrealMaterial->EmissiveColor, false, UVSets, FVector2D(240,-64) );
-	CreateAndLinkExpressionForMaterialProperty( *FbxMaterial, UnrealMaterial, FbxSurfaceMaterial::sSpecular, UnrealMaterial->Specular, false, UVSets, FVector2D(240, -128) );
+	CreateAndLinkExpressionForMaterialProperty( FbxMaterial, UnrealMaterial, FbxSurfaceMaterial::sDiffuse, UnrealMaterial->BaseColor, false, UVSets, FVector2D(240,-320) );
+	CreateAndLinkExpressionForMaterialProperty( FbxMaterial, UnrealMaterial, FbxSurfaceMaterial::sEmissive, UnrealMaterial->EmissiveColor, false, UVSets, FVector2D(240,-64) );
+	CreateAndLinkExpressionForMaterialProperty( FbxMaterial, UnrealMaterial, FbxSurfaceMaterial::sSpecular, UnrealMaterial->Specular, false, UVSets, FVector2D(240, -128) );
 	//CreateAndLinkExpressionForMaterialProperty( *FbxMaterial, UnrealMaterial, FbxSurfaceMaterial::sSpecularFactor, UnrealMaterial->SpecularColor, false, UVSets); // SpecularFactor modulates the SpecularColor value if there's one
 	//CreateAndLinkExpressionForMaterialProperty( *FbxMaterial, UnrealMaterial, FbxSurfaceMaterial::sShininess, UnrealMaterial->SpecularPower, false, UVSets);
-	if (!CreateAndLinkExpressionForMaterialProperty( *FbxMaterial, UnrealMaterial, FbxSurfaceMaterial::sNormalMap, UnrealMaterial->Normal, true, UVSets, FVector2D(240,256) ) )
+	if (!CreateAndLinkExpressionForMaterialProperty( FbxMaterial, UnrealMaterial, FbxSurfaceMaterial::sNormalMap, UnrealMaterial->Normal, true, UVSets, FVector2D(240,256) ) )
 	{
-		CreateAndLinkExpressionForMaterialProperty( *FbxMaterial, UnrealMaterial, FbxSurfaceMaterial::sBump, UnrealMaterial->Normal, true, UVSets, FVector2D(240,256) ); // no bump in unreal, use as normal map
+		CreateAndLinkExpressionForMaterialProperty( FbxMaterial, UnrealMaterial, FbxSurfaceMaterial::sBump, UnrealMaterial->Normal, true, UVSets, FVector2D(240,256) ); // no bump in unreal, use as normal map
 	}
 	//CreateAndLinkExpressionForMaterialProperty( *FbxMaterial, UnrealMaterial, KFbxSurfaceMaterial::sTransparentColor, UnrealMaterial->Opacity, false, UVSets);
 	//CreateAndLinkExpressionForMaterialProperty( *FbxMaterial, UnrealMaterial, KFbxSurfaceMaterial::sTransparencyFactor, UnrealMaterial->OpacityMask, false, UVSets);
-	FixupMaterial( *FbxMaterial, UnrealMaterial); // add random diffuse if none exists
+	FixupMaterial( FbxMaterial, UnrealMaterial); // add random diffuse if none exists
 
 	// compile shaders for PC (from UPrecompileShadersCommandlet::ProcessMaterial
 	// and FMaterialEditor::UpdateOriginalMaterial)
@@ -444,6 +459,8 @@ void UnFbx::FFbxImporter::CreateUnrealMaterial(FbxSurfaceMaterial* FbxMaterial,
 	UnrealMaterial->PreEditChange(NULL);
 		UnrealMaterial->PostEditChange();
 	
+	ImportedMaterialData.AddImportedMaterial( FbxMaterial, *UnrealMaterial );
+
 	OutMaterials.Add(UnrealMaterial);
 }
 
@@ -454,7 +471,10 @@ int32 UnFbx::FFbxImporter::CreateNodeMaterials(FbxNode* FbxNode, TArray<UMateria
 	{
 		FbxSurfaceMaterial *FbxMaterial = FbxNode->GetMaterial(MaterialIndex);
 
-		CreateUnrealMaterial(FbxMaterial, OutMaterials, UVSets);
+		if( FbxMaterial )
+		{
+			CreateUnrealMaterial(*FbxMaterial, OutMaterials, UVSets);
+		}
 	}
 	return MaterialCount;
 }
diff --git a/Engine/Source/Editor/UnrealEd/Private/Fbx/FbxMatineeImport.cpp b/Engine/Source/Editor/UnrealEd/Private/Fbx/FbxMatineeImport.cpp
--- a/Engine/Source/Editor/UnrealEd/Private/Fbx/FbxMatineeImport.cpp
+++ b/Engine/Source/Editor/UnrealEd/Private/Fbx/FbxMatineeImport.cpp
@@ -28,10 +28,6 @@
 * and restrictions as described herein.
 */
 
-/*=============================================================================
-	FBX importer for Unreal Engine 3.
-==============================================================================*/
-
 #include "UnrealEd.h"
 
 #include "EngineInterpolationClasses.h"
diff --git a/Engine/Source/Editor/UnrealEd/Private/Fbx/FbxSkeletalMeshImport.cpp b/Engine/Source/Editor/UnrealEd/Private/Fbx/FbxSkeletalMeshImport.cpp
--- a/Engine/Source/Editor/UnrealEd/Private/Fbx/FbxSkeletalMeshImport.cpp
+++ b/Engine/Source/Editor/UnrealEd/Private/Fbx/FbxSkeletalMeshImport.cpp
@@ -931,89 +931,6 @@ bool UnFbx::FFbxImporter::ImportBone(TArray<FbxNode*>& NodeArray, FSkeletalMeshI
 	return true;
 }
 
-void UnFbx::FFbxImporter::ImportMaterialsForSkelMesh(TArray<FbxNode*>& NodeArray, FSkeletalMeshImportData &ImportData, TArray<FString>& OutFbxMatList)
-{
-	FbxNode* Node;
-	int32 MaterialCount;
-	
-	// visit all FBX node to get all materials
-	for (int32 NodeIndex = 0; NodeIndex < NodeArray.Num(); NodeIndex++)
-	{
-		Node = NodeArray[NodeIndex];
-		
-		MaterialCount = Node->GetMaterialCount();
-		int32 ExistMatNum = OutFbxMatList.Num();
-
-		for(int32 MaterialIndex=0; MaterialIndex < MaterialCount; ++MaterialIndex)
-		{
-			FbxSurfaceMaterial *FbxMaterial = Node->GetMaterial(MaterialIndex);
-
-			FString MaterialName = ANSI_TO_TCHAR( FbxMaterial->GetName() );
-
-			int32 ExistMatIndex;
-			// check if the material has been recorded to the OurFbxMatList
-			for ( ExistMatIndex = 0; ExistMatIndex<ExistMatNum; ExistMatIndex++)
-			{
-				if ( OutFbxMatList[ExistMatIndex] == MaterialName )
-				{
-					break;
-				}
-			}
-
-			// the material has not been imported for this skeletal mesh, import this FBX material
-			if (ExistMatIndex == ExistMatNum)
-			{
-				OutFbxMatList.Add( MaterialName );
-
-				// Add material slot for the skeletal mesh
-				// We will create the unreal material later when import mesh data
-				if ( ImportOptions->bImportMaterials )
-				{
-					ImportData.Materials.Add( VMaterial() );
-
-					ImportData.Materials.Last().MaterialName = ANSI_TO_TCHAR( MakeName(FbxMaterial->GetName() ) );
-				}
-			}
-		}
-
-		// import texture only if UI option is checked
-		if ( !ImportOptions->bImportMaterials && ImportOptions->bImportTextures)
-		{
-			ImportTexturesFromNode(Node);
-		}
-	}
-	
-
-	// no material info, create a default material slot
-	if ( OutFbxMatList.Num() == 0 )
-	{
-		AddTokenizedErrorMessage(FTokenizedMessage::Create(EMessageSeverity::Warning, LOCTEXT("FbxSkeletaLMeshimport_NoMaterialAssociatedWithMesh", "No material associated with skeletal mesh - using default")));
-		UMaterial* DefaultMaterial = UMaterial::GetDefaultMaterial(MD_Surface);
-		if (DefaultMaterial)
-		{
-			ImportData.Materials.Add( VMaterial() );
-
-			ImportData.Materials.Last().MaterialName = DefaultMaterial->GetName();
-		}
-	}
-	else if (!ImportOptions->bImportMaterials)
-	{
-		// if we don't import FBX materials, we will create material slots as FBX objects have
-		AddTokenizedErrorMessage(FTokenizedMessage::Create(EMessageSeverity::Warning, LOCTEXT("FbxSkeletaLMeshimport_NoimportMaterial", "Import Material isn't set. Using default materials for material slot")));
-		UMaterial* DefaultMaterial = UMaterial::GetDefaultMaterial(MD_Surface);
-		if (DefaultMaterial)
-		{
-			for (int32 MatIndex = 0; MatIndex < OutFbxMatList.Num(); MatIndex++)
-			{
-				VMaterial NewMat;
-				NewMat.MaterialName = DefaultMaterial->GetName();
-
-				ImportData.Materials.Add(NewMat);
-			}
-		}
-	}
-}
-
 USkeletalMesh* UnFbx::FFbxImporter::ImportSkeletalMesh(UObject* InParent, TArray<FbxNode*>& NodeArray, const FName& Name, EObjectFlags Flags, UFbxSkeletalMeshImportData* TemplateImportData, FString Filename, TArray<FbxShape*> *FbxShapeArray, FSkeletalMeshImportData* OutData, bool bCreateRenderData )
 {
 	if (NodeArray.Num() == 0)
@@ -1081,11 +998,6 @@ USkeletalMesh* UnFbx::FFbxImporter::ImportSkeletalMesh(UObject* InParent, TArray
 		SkelMeshImportDataPtr = OutData;
 	}
 
-	// Get all FBX material names here, 
-	// then we can set material index mapping correctly in FillSkelMeshImporterFromFbx() even if we do not import FBX materials
-	TArray<FString> FbxMatList;
-	ImportMaterialsForSkelMesh(NodeArray, *SkelMeshImportDataPtr, FbxMatList);
-	
 	bool bDiffPose;
 	TArray<FbxNode*> SortedLinkArray;
 	FbxArray<FbxAMatrix> GlobalsPerLink;
@@ -1102,19 +1014,40 @@ USkeletalMesh* UnFbx::FFbxImporter::ImportSkeletalMesh(UObject* InParent, TArray
 		return NULL;
 	}
 
-	for ( int32 i = 0; i < NodeArray.Num(); i++)
+	// Create a list of all unique fbx materials.  This needs to be done as a separate pass before reading geometry
+	// so that we know about all possible materials before assigning material indices to each triangle
+	TArray<FbxSurfaceMaterial*> FbxMaterials;
+	for( int32 NodeIndex = 0; NodeIndex < NodeArray.Num(); ++NodeIndex )
+	{
+		FbxNode* Node = NodeArray[NodeIndex];
+
+		int32 MaterialCount = Node->GetMaterialCount();
+
+		for (int32 MaterialIndex = 0; MaterialIndex < MaterialCount; ++MaterialIndex)
+		{
+			FbxSurfaceMaterial* FbxMaterial = Node->GetMaterial( MaterialIndex );
+			if( !FbxMaterials.Contains( FbxMaterial ) )
+			{
+				FbxMaterials.Add( FbxMaterial );
+				// Add an entry for each unique material
+				SkelMeshImportDataPtr->Materials.Add( VMaterial() );
+			}
+		}
+	}
+
+	for ( int32 NodeIndex = 0; NodeIndex < NodeArray.Num(); ++NodeIndex)
 	{
-		Node = NodeArray[i];
+		Node = NodeArray[NodeIndex];
 		FbxMesh = Node->GetMesh();
 		FbxSkin* Skin = (FbxSkin*)FbxMesh->GetDeformer(0, FbxDeformer::eSkin);
 		FbxShape* FbxShape = NULL;
 		if (FbxShapeArray)
 		{
-			FbxShape = (*FbxShapeArray)[i];
+			FbxShape = (*FbxShapeArray)[NodeIndex];
 		}
 
 		// NOTE: This function may invalidate FbxMesh and set it to point to a an updated version
-		if (!FillSkelMeshImporterFromFbx(*SkelMeshImportDataPtr, FbxMesh, Skin, FbxShape, SortedLinkArray, FbxMatList))
+		if (!FillSkelMeshImporterFromFbx( *SkelMeshImportDataPtr, FbxMesh, Skin, FbxShape, SortedLinkArray, FbxMaterials ) )
 		{
 			// I can't delete object here since this is middle of import
 			// but I can move to transient package, and GC will automatically collect it
@@ -1130,11 +1063,8 @@ USkeletalMesh* UnFbx::FFbxImporter::ImportSkeletalMesh(UObject* InParent, TArray
 		}
 	}
 	
-	if( SkelMeshImportDataPtr->Materials.Num() == FbxMatList.Num() )
-	{
-		// reorder material according to "SKinXX" in material name
-		SetMaterialSkinXXOrder(*SkelMeshImportDataPtr, FbxMatList );
-	}
+	// reorder material according to "SKinXX" in material name
+	SetMaterialSkinXXOrder(*SkelMeshImportDataPtr );
 	
 	if( ImportOptions->bPreserveSmoothingGroups )
 	{
@@ -1522,11 +1452,9 @@ USkeletalMesh* UnFbx::FFbxImporter::ReimportSkeletalMesh(USkeletalMesh* Mesh, UF
 	return NewMesh;
 }
 
-void UnFbx::FFbxImporter::SetMaterialSkinXXOrder(FSkeletalMeshImportData& ImportData, const TArray<FString>& FbxMatList )
+void UnFbx::FFbxImporter::SetMaterialSkinXXOrder(FSkeletalMeshImportData& ImportData )
 {
 	TArray<int32> MaterialIndexToSkinIndex;
-
-	if( ImportData.Materials.Num() == FbxMatList.Num() )
 	{
 		int32 MaterialCount = ImportData.Materials.Num();
 		MaterialIndexToSkinIndex.AddZeroed(MaterialCount);
@@ -1535,7 +1463,7 @@ void UnFbx::FFbxImporter::SetMaterialSkinXXOrder(FSkeletalMeshImportData& Import
 		for(int32 MaterialIndex=0; MaterialIndex < MaterialCount; ++MaterialIndex)
 		{
 			// get skin index
-			FString MatName = FbxMatList[MaterialIndex];
+			FString MatName = ImportData.Materials[MaterialIndex].MaterialImportName;
 
 			if (MatName.Len() > 6)
 			{
@@ -1553,7 +1481,7 @@ void UnFbx::FFbxImporter::SetMaterialSkinXXOrder(FSkeletalMeshImportData& Import
 						MaterialIndexToSkinIndex[MaterialIndex] = TmpIndex;
 
 						// remove the 'skinXX' suffix from the material name					
-						ImportData.Materials[MaterialIndex].MaterialName.LeftChop( Offset );
+						ImportData.Materials[MaterialIndex].MaterialImportName.LeftChop( Offset );
 					}
 				}
 			}
@@ -1562,19 +1490,14 @@ void UnFbx::FFbxImporter::SetMaterialSkinXXOrder(FSkeletalMeshImportData& Import
 		if (bNeedsReorder)
 		{
 			// re-order the materials
-			TArray< FString > MatNames;
-
-			for(int32 MaterialIndex=0; MaterialIndex < MaterialCount; ++MaterialIndex)
-			{
-				MatNames.Add( ImportData.Materials[MaterialIndex].MaterialName );
-			}
+			TArray< VMaterial > ExistingMatList = ImportData.Materials;
 
 			for(int32 MaterialIndex=0; MaterialIndex < MaterialCount; ++MaterialIndex)
 			{
 				int32 NewIndex = MaterialIndexToSkinIndex[MaterialIndex];
-				if( MatNames.IsValidIndex( NewIndex ) )
+				if( ExistingMatList.IsValidIndex( NewIndex ) )
 				{
-					ImportData.Materials[MaterialIndex].MaterialName =  MatNames[ NewIndex ];
+					ImportData.Materials[MaterialIndex] = ExistingMatList[ NewIndex ];
 				}
 			}
 
@@ -1589,7 +1512,7 @@ void UnFbx::FFbxImporter::SetMaterialSkinXXOrder(FSkeletalMeshImportData& Import
 	}
 }
 
-bool UnFbx::FFbxImporter::FillSkelMeshImporterFromFbx( FSkeletalMeshImportData& ImportData, FbxMesh*& Mesh, FbxSkin* Skin, FbxShape* FbxShape, TArray<FbxNode*> &SortedLinks, TArray<FString>& FbxMatList)
+bool UnFbx::FFbxImporter::FillSkelMeshImporterFromFbx( FSkeletalMeshImportData& ImportData, FbxMesh*& Mesh, FbxSkin* Skin, FbxShape* FbxShape, TArray<FbxNode*> &SortedLinks, const TArray<FbxSurfaceMaterial*>& FbxMaterials )
 {
 	FbxNode* Node = Mesh->GetNode();
 
@@ -1671,38 +1594,45 @@ bool UnFbx::FFbxImporter::FillSkelMeshImporterFromFbx( FSkeletalMeshImportData&
 		}
 	}
 
-	//
-	// setup materials index mapping
-	//
-	int32 MaterialCount = Node->GetMaterialCount();
-	int32 ExistedMatNum = ImportData.Materials.Num();
-	// for skeletal mesh, only one UVSet is used, so the variable UVSets is useless, just used as parameter
+
+	TArray<UMaterialInterface*> Materials;
+	if (ImportOptions->bImportMaterials)
+	{
+		CreateNodeMaterials(Node, Materials, UVSets);
+	}
+	else if (ImportOptions->bImportTextures)
+	{
+		ImportTexturesFromNode(Node);
+	}
+
+	// Maps local mesh material index to global material index
 	TArray<int32> MaterialMapping;
 
+	int32 MaterialCount = Node->GetMaterialCount();
+
 	MaterialMapping.AddUninitialized(MaterialCount);
 
-	for (int32 MaterialIndex=0; MaterialIndex < MaterialCount; ++MaterialIndex)
+	for( int32 MaterialIndex = 0; MaterialIndex < MaterialCount; ++MaterialIndex )
 	{
-		FbxSurfaceMaterial *FbxMaterial = Node->GetMaterial(MaterialIndex);
+		FbxSurfaceMaterial* FbxMaterial = Node->GetMaterial( MaterialIndex );
 
-		// create unreal materials here because we get UVset data now
-		if (ImportOptions->bImportMaterials)
+		int32 ExistingMatIndex = INDEX_NONE;
+		FbxMaterials.Find( FbxMaterial, ExistingMatIndex ); 
+		if( ExistingMatIndex != INDEX_NONE )
 		{
-			TArray<UMaterialInterface*> Materials;
-			Materials.Empty();
-
-			CreateUnrealMaterial(FbxMaterial,Materials,UVSets);
-		}
+			// Reuse existing material
+			MaterialMapping[MaterialIndex] = ExistingMatIndex;
 
-		int32 ExistMatIndex;
-		for ( ExistMatIndex = 0; ExistMatIndex<ExistedMatNum; ExistMatIndex++)
-		{
-			if ( FbxMatList[ExistMatIndex] == ANSI_TO_TCHAR(FbxMaterial->GetName()) )
+			if (ImportOptions->bImportMaterials && Materials.IsValidIndex(MaterialIndex) )
 			{
-				MaterialMapping[MaterialIndex] = ExistMatIndex;
-				break;
+				ImportData.Materials[ExistingMatIndex].MaterialImportName = ANSI_TO_TCHAR(MakeName(FbxMaterial->GetName()));
+				ImportData.Materials[ExistingMatIndex].Material = Materials[MaterialIndex];
 			}
 		}
+		else
+		{
+			MaterialMapping[MaterialIndex] = 0;
+		}
 	}
 
 	if( ImportOptions->bPreserveSmoothingGroups )
@@ -2037,13 +1967,13 @@ bool UnFbx::FFbxImporter::FillSkelMeshImporterFromFbx( FSkeletalMeshImportData&
 				// material index is stored in the IndexArray, not the DirectArray (which is irrelevant with 2009.1)
 				case FbxLayerElement::eAllSame:
 					{	
-						Triangle.MatIndex = MaterialMapping[LayerElementMaterial->GetIndexArray().GetAt(0)];
+						Triangle.MatIndex = MaterialMapping[ LayerElementMaterial->GetIndexArray().GetAt(0) ];
 					}
 					break;
 				case FbxLayerElement::eByPolygon:
 					{	
-						int32 Index = LayerElementMaterial->GetIndexArray().GetAt(LocalIndex);
-						if( !MaterialMapping.IsValidIndex(Index) )
+						int32 Index = LayerElementMaterial->GetIndexArray().GetAt(LocalIndex);							
+						if (!MaterialMapping.IsValidIndex(Index))
 						{
 							AddTokenizedErrorMessage(FTokenizedMessage::Create(EMessageSeverity::Warning, LOCTEXT("FbxSkeletaLMeshimport_MaterialIndexInconsistency", "Face material index inconsistency - forcing to 0")));
 						}
@@ -2055,14 +1985,16 @@ bool UnFbx::FFbxImporter::FillSkelMeshImporterFromFbx( FSkeletalMeshImportData&
 					break;
 				}
 			}
+
 			// When import morph, we don't check the material index 
 			// because we don't import material for morph, so the ImportData.Materials contains zero material
-			if ( !FbxShape && (Triangle.MatIndex < 0 ||  Triangle.MatIndex >= ImportData.Materials.Num() ) )
+			if ( !FbxShape && (Triangle.MatIndex < 0 ||  Triangle.MatIndex >= FbxMaterials.Num() ) )
 			{
 				AddTokenizedErrorMessage(FTokenizedMessage::Create(EMessageSeverity::Warning, LOCTEXT("FbxSkeletaLMeshimport_MaterialIndexInconsistency", "Face material index inconsistency - forcing to 0")));
 				Triangle.MatIndex = 0;
 			}
 		}
+	
 		Triangle.AuxMatIndex = 0;
 		for (VertexIndex=0; VertexIndex<3; VertexIndex++)
 		{
diff --git a/Engine/Source/Editor/UnrealEd/Private/Fbx/FbxStaticMeshImport.cpp b/Engine/Source/Editor/UnrealEd/Private/Fbx/FbxStaticMeshImport.cpp
--- a/Engine/Source/Editor/UnrealEd/Private/Fbx/FbxStaticMeshImport.cpp
+++ b/Engine/Source/Editor/UnrealEd/Private/Fbx/FbxStaticMeshImport.cpp
@@ -184,11 +184,11 @@ bool UnFbx::FFbxImporter::BuildStaticMeshFromGeometry(FbxMesh* Mesh, UStaticMesh
 	//
 	int32 MaterialCount = 0;
 	TArray<UMaterialInterface*> Materials;
-	if ( ImportOptions->bImportMaterials )
+	if (ImportOptions->bImportMaterials)
 	{
-		CreateNodeMaterials(Node,Materials,UVSets);
+		CreateNodeMaterials(Node, Materials, UVSets);
 	}
-	else if ( ImportOptions->bImportTextures )
+	else if (ImportOptions->bImportTextures)
 	{
 		ImportTexturesFromNode(Node);
 	}
@@ -203,14 +203,15 @@ bool UnFbx::FFbxImporter::BuildStaticMeshFromGeometry(FbxMesh* Mesh, UStaticMesh
 	{
 		FFbxMaterial* NewMaterial = new(MeshMaterials) FFbxMaterial;
 		FbxSurfaceMaterial *FbxMaterial = Node->GetMaterial(MaterialIndex);
-		NewMaterial->Name = ANSI_TO_TCHAR(MakeName(FbxMaterial->GetName()));
+		FString MaterialName = ANSI_TO_TCHAR(MakeName(FbxMaterial->GetName()));
+		NewMaterial->FbxMaterial = FbxMaterial;
 		if (ImportOptions->bImportMaterials)
 		{
 			NewMaterial->Material = Materials[MaterialIndex];
 		}
 		else
 		{
-			UMaterialInterface* UnrealMaterialInterface = FindObject<UMaterialInterface>(Parent,*NewMaterial->Name);
+			UMaterialInterface* UnrealMaterialInterface = FindObject<UMaterialInterface>(NULL,*MaterialName);
 			if (UnrealMaterialInterface == NULL)
 			{
 				UnrealMaterialInterface = UMaterial::GetDefaultMaterial(MD_Surface);
@@ -225,6 +226,7 @@ bool UnFbx::FFbxImporter::BuildStaticMeshFromGeometry(FbxMesh* Mesh, UStaticMesh
 		check(DefaultMaterial);
 		FFbxMaterial* NewMaterial = new(MeshMaterials) FFbxMaterial;
 		NewMaterial->Material = DefaultMaterial;
+		NewMaterial->FbxMaterial = NULL;
 		MaterialCount = 1;
 	}
 
@@ -963,7 +965,7 @@ UStaticMesh* UnFbx::FFbxImporter::ImportStaticMeshAsSingle(UObject* InParent, TA
 		UE_LOG(LogFbx,Log,TEXT("== Initial material list:"));
 		for (int32 MaterialIndex = 0; MaterialIndex < MeshMaterials.Num(); ++MaterialIndex)
 		{
-			UE_LOG(LogFbx,Log,TEXT("%d: %s"),MaterialIndex,*MeshMaterials[MaterialIndex].Name);
+			UE_LOG(LogFbx,Log,TEXT("%d: %s"),MaterialIndex,*MeshMaterials[MaterialIndex].GetName() );
 		}
 
 		// Compress the materials array by removing any duplicates.
@@ -975,7 +977,7 @@ UStaticMesh* UnFbx::FFbxImporter::ImportStaticMeshAsSingle(UObject* InParent, TA
 			bool bUnique = true;
 			for (int32 OtherMaterialIndex = MaterialIndex - 1; OtherMaterialIndex >= 0; --OtherMaterialIndex)
 			{
-				if (MeshMaterials[MaterialIndex].Name == MeshMaterials[OtherMaterialIndex].Name)
+				if (MeshMaterials[MaterialIndex].FbxMaterial == MeshMaterials[OtherMaterialIndex].FbxMaterial)
 				{
 					int32 UniqueIndex = MaterialMap[OtherMaterialIndex];
 					MaterialMap.Add(UniqueIndex);
@@ -1002,22 +1004,21 @@ UStaticMesh* UnFbx::FFbxImporter::ImportStaticMeshAsSingle(UObject* InParent, TA
 			int32 RemappedIndex = MaterialMap[MaterialIndex];
 			if (!SortedMaterialIndex.IsValidIndex(RemappedIndex))
 			{
-				int32 MatNameLen = FCString::Strlen(*MeshMaterials[RemappedIndex].Name) + 1;
-				char* MatName = new char[MatNameLen];
-				FCStringAnsi::Strcpy(MatName, MatNameLen, TCHAR_TO_ANSI(*MeshMaterials[RemappedIndex].Name));
-				if (strlen(MatName) > 6)
+				FString FbxMatName = MeshMaterials[RemappedIndex].GetName();
+
+				int32 Offset = FbxMatName.Find(TEXT("_SKIN"), ESearchCase::IgnoreCase, ESearchDir::FromEnd);
+				if (Offset != INDEX_NONE)
 				{
-					const char* SkinXX = MatName + strlen(MatName) - 6;
-					if (toupper(*SkinXX) == 'S' && toupper(*(SkinXX+1)) == 'K' && toupper(*(SkinXX+2)) == 'I' && toupper(*(SkinXX+3)) == 'N')
+					// Chop off the material name so we are left with the number in _SKINXX
+					FString SkinXXNumber = FbxMatName.Right(FbxMatName.Len() - (Offset + 1)).RightChop(4);
+
+					if (SkinXXNumber.IsNumeric())
 					{
-						if (isdigit(*(SkinXX+4)) && isdigit(*(SkinXX+5)))
-						{
-							SkinIndex = (*(SkinXX+4) - 0x30) * 10 + (*(SkinXX+5) - 0x30);
-							bDoRemap = true;
-						}
+						SkinIndex = FPlatformString::Atoi( *SkinXXNumber );
+						bDoRemap = true;
 					}
 				}
-				delete [] MatName;
+
 				SortedMaterialIndex.Add(((uint32)SkinIndex << 16) | ((uint32)RemappedIndex & 0xffff));
 			}
 		}
@@ -1029,7 +1030,7 @@ UStaticMesh* UnFbx::FFbxImporter::ImportStaticMeshAsSingle(UObject* InParent, TA
 		{
 			int32 RemappedIndex = SortedMaterialIndex[SortedIndex] & 0xffff;
 			SortedMaterials.Add(UniqueMaterials[RemappedIndex].Material);
-			UE_LOG(LogFbx,Log,TEXT("%d: %s"),SortedIndex,*UniqueMaterials[RemappedIndex].Name);
+			UE_LOG(LogFbx,Log,TEXT("%d: %s"),SortedIndex,*UniqueMaterials[RemappedIndex].GetName());
 		}
 		UE_LOG(LogFbx,Log,TEXT("== Mapping table:"));
 		for (int32 MaterialIndex = 0; MaterialIndex < MaterialMap.Num(); ++MaterialIndex)
diff --git a/Engine/Source/Editor/UnrealEd/Private/FbxImporter.h b/Engine/Source/Editor/UnrealEd/Private/FbxImporter.h
--- a/Engine/Source/Editor/UnrealEd/Private/FbxImporter.h
+++ b/Engine/Source/Editor/UnrealEd/Private/FbxImporter.h
@@ -220,6 +220,19 @@ public:
 FBXImportOptions* GetImportOptions( class FFbxImporter* FbxImporter, UFbxImportUI* ImportUI, bool bShowOptionDialog, const FString& FullPath, bool& OutOperationCanceled, bool& OutImportAll, bool bForceImportType = false, EFBXImportType ImportType = FBXIT_StaticMesh );
 void ApplyImportUIToImportOptions(UFbxImportUI* ImportUI, FBXImportOptions& InOutImportOptions);
 
+struct FImportedMaterialData
+{
+public:
+	void AddImportedMaterial( FbxSurfaceMaterial& FbxMaterial, UMaterialInterface& UnrealMaterial );
+	bool IsUnique( FbxSurfaceMaterial& FbxMaterial, FName ImportedMaterialName ) const;
+	UMaterialInterface* GetUnrealMaterial( const FbxSurfaceMaterial& FbxMaterial ) const;
+	void Clear();
+private:
+	/** Mapping of FBX material to Unreal material.  Some materials in FBX have the same name so we use this map to determine if materials are unique */
+	TMap<FbxSurfaceMaterial*, TWeakObjectPtr<UMaterialInterface> > FbxToUnrealMaterialMap;
+	TSet<FName> ImportedMaterialNames;
+};
+
 /**
  * Main FBX Importer class.
  */
@@ -586,12 +599,13 @@ protected:
 	};
 	
 	static TSharedPtr<FFbxImporter> StaticInstance;
-	static const float SCALE_TOLERANCE;
 
 	struct FFbxMaterial
 	{
-		FString Name;
+		FbxSurfaceMaterial* FbxMaterial;
 		UMaterialInterface* Material;
+
+		FString GetName() const { return FbxMaterial ? ANSI_TO_TCHAR( FbxMaterial->GetName() ) : TEXT("None"); }
 	};
 	
 	// scene management
@@ -603,7 +617,7 @@ protected:
 	FString ErrorMessage;
 	// base path of fbx file
 	FString FileBasePath;
-	UObject* Parent;
+	TWeakObjectPtr<UObject> Parent;
 	// Flag that the mesh is the first mesh to import in current FBX scene
 	// FBX scene may contain multiple meshes, importer can import them at one time.
 	// Initialized as true when start to import a FBX scene
@@ -628,18 +642,11 @@ protected:
 	 */
 	bool BuildStaticMeshFromGeometry(FbxMesh* FbxMesh, UStaticMesh* StaticMesh, TArray<FFbxMaterial>& MeshMaterials, int LODIndex, TMap<FVector, FColor>* ExistingVertexColorData);
 	
-	/**
-	 * Creates a Matinee group for a given actor within a given Matinee sequence.
-	 */
-	//UInterpGroupInst* CreateMatineeGroup(USeqAct_Interp* Sequence, AActor* Actor);
-
 	/**
 	 * Clean up for destroy the Importer.
 	 */
 	void CleanUp();
-	
-	//UObject* CreateObjectFromNode(FbxNode* Node);
-	//void PlaceActor(AActor* Actor, FbxNode* Node, bool bInvertOrientation = false);
+
 	/**
 	* Compute the global matrix for Fbx Node
 	*
@@ -695,17 +702,7 @@ protected:
 	* @returns bool*	true if import successfully.
 	*/
     bool FillSkelMeshImporterFromFbx(FSkeletalMeshImportData& ImportData, FbxMesh*& Mesh, FbxSkin* Skin, 
-										FbxShape* Shape, TArray<FbxNode*> &SortedLinks, TArray<FString>& FbxMatList);
-
-	/**
-	* Fill material name list data from Fbx Nodes.
-	* If import material option is OFF, the API will create default material slots for the skeletal mesh.
-	*
-	* @param NodeArray   Fbx Nodes to import, they are bound to the same skeleton system.
-	* @param ImportData object to store skeletal mesh data.
-	* @param FbxMatList  All material names of the skeletal mesh
-	*/
-	void ImportMaterialsForSkelMesh(TArray<FbxNode*>& NodeArray, FSkeletalMeshImportData &ImportData, TArray<FString>& OutFbxMatList);
+										FbxShape* Shape, TArray<FbxNode*> &SortedLinks, const TArray<FbxSurfaceMaterial*>& FbxMaterials );
 
 	/**
 	 * Import bones from skeletons that NodeArray bind to.
@@ -781,7 +778,7 @@ protected:
 	 *
 	 * @param FSkeletalMeshBinaryImport& The unreal skeletal mesh.
 	 */
-	void SetMaterialSkinXXOrder(FSkeletalMeshImportData& ImportData, const TArray<FString>& FbxMatList );
+	void SetMaterialSkinXXOrder(FSkeletalMeshImportData& ImportData);
 	
 	/**
 	 * Create materials from Fbx node.
@@ -804,7 +801,7 @@ protected:
 	* @param outMaterials Unreal Materials we created
 	* @param outUVSets
 	 */
-	void CreateUnrealMaterial(FbxSurfaceMaterial* FbxMaterial, TArray<UMaterialInterface*>& OutMaterials, TArray<FString>& UVSets);
+	void CreateUnrealMaterial(FbxSurfaceMaterial& FbxMaterial, TArray<UMaterialInterface*>& OutMaterials, TArray<FString>& UVSets);
 	
 	/**
 	 * Visit all materials of one node, import textures from materials.
@@ -970,6 +967,8 @@ private:
 	class FFbxLogger * Logger;
 	void SetLogger(class FFbxLogger * InLogger);
 	void ClearLogger();
+
+	FImportedMaterialData ImportedMaterialData;
 };
 
 
diff --git a/Engine/Source/Editor/UnrealEd/Public/SkelImport.h b/Engine/Source/Editor/UnrealEd/Public/SkelImport.h
--- a/Engine/Source/Editor/UnrealEd/Public/SkelImport.h
+++ b/Engine/Source/Editor/UnrealEd/Public/SkelImport.h
@@ -4,14 +4,15 @@
 	Data structures only used for importing skeletal meshes and animations.
 -----------------------------------------------------------------------------*/
 
-#ifndef __SKELIMPORT_H__
-#define __SKELIMPORT_H__
-
+#pragma once
 
 // Raw data material.
 struct VMaterial
 {
-	FString            MaterialName;
+	/** The actual material created on import or found among existing materials */
+	TWeakObjectPtr<UMaterialInterface> Material;
+	/** The material name found by the importer */
+	FString MaterialImportName;
 };
 
 
@@ -211,5 +212,3 @@ public:
 	FSkeletalMeshImportData RawMeshInfluencesData;
 	int32 MaxBoneCountPerChunk;
 };
-
-#endif
