#UE4 - Making sure OnlineBlueprintSupport is loaded when cooking (required by TappyChicken)

diff --git a/Engine/Source/Editor/UnrealEd/Private/Editor.cpp b/Engine/Source/Editor/UnrealEd/Private/Editor.cpp
--- a/Engine/Source/Editor/UnrealEd/Private/Editor.cpp
+++ b/Engine/Source/Editor/UnrealEd/Private/Editor.cpp
@@ -406,6 +406,9 @@ void UEditorEngine::InitEditor(IEngineLoop* InEngineLoop)
 
 	// Update the auto-load project
 	UpdateAutoLoadProject();
+
+	// Load any modules that might be required by commandlets
+	FModuleManager::Get().LoadModule(TEXT("OnlineBlueprintSupport"));
 }
 
 void UEditorEngine::HandleSettingChanged( FName Name )
