#ue4 replacing blueprints with classes (in Blueprints) - WIP

diff --git a/Engine/Config/BaseEditor.ini b/Engine/Config/BaseEditor.ini
--- a/Engine/Config/BaseEditor.ini
+++ b/Engine/Config/BaseEditor.ini
@@ -250,4 +250,7 @@ SupportWebsiteURL="http://unrealengine.com/support"
 bUseUserDefinedStructure=false
 
 [EditoronlyBP]
-bAllowClassObjAndBluepirntPinMatching=false
\ No newline at end of file
+bAllowClassAndBlueprintPinMatching=false
+bReplaceBlueprintWithClass=false
+bDontLoadBlueprintOutsideEditor=false
+bBlueprintIsNotBlueprintType=false
\ No newline at end of file
diff --git a/Engine/Source/Editor/BlueprintGraph/Private/EdGraphSchema_K2.cpp b/Engine/Source/Editor/BlueprintGraph/Private/EdGraphSchema_K2.cpp
--- a/Engine/Source/Editor/BlueprintGraph/Private/EdGraphSchema_K2.cpp
+++ b/Engine/Source/Editor/BlueprintGraph/Private/EdGraphSchema_K2.cpp
@@ -2495,7 +2495,7 @@ bool UEdGraphSchema_K2::ArePinTypesCompatible(const FEdGraphPinType& Output, con
 		return OutputClass->ImplementsInterface(InterfaceClass) || OutputClass->IsChildOf(InterfaceClass);
 	}
 
-	// Pins representing BLueprint objects and subclass of UObject can match when EditoronlyBP.bAllowClassObjAndBluepirntPinMatching=true (BaseEngine.ini)
+	// Pins representing BLueprint objects and subclass of UObject can match when EditoronlyBP.bAllowClassAndBlueprintPinMatching=true (BaseEngine.ini)
 	// It's required for converting all UBlueprint references into UClass.
 	struct FObjClassAndBlueprintHelper
 	{
@@ -2504,7 +2504,7 @@ bool UEdGraphSchema_K2::ArePinTypesCompatible(const FEdGraphPinType& Output, con
 	public:
 		FObjClassAndBlueprintHelper() : bAllow(false)
 		{
-			GConfig->GetBool(TEXT("EditoronlyBP"), TEXT("bAllowClassObjAndBluepirntPinMatching"), bAllow, GEditorIni);
+			GConfig->GetBool(TEXT("EditoronlyBP"), TEXT("bAllowClassAndBlueprintPinMatching"), bAllow, GEditorIni);
 		}
 
 		bool Match(const FEdGraphPinType& A, const FEdGraphPinType& B, const UEdGraphSchema_K2& Schema) const
@@ -3329,6 +3329,16 @@ struct FBackwardCompatibilityConversionHelper
 					return false;
 				}
 
+				auto InValidCastPin = CastNode->GetInvalidCastPin();
+				check(InValidCastPin);
+				if (!Schema.TryCreateConnection(InValidCastPin, ExecPin))
+				{
+					UE_LOG(LogBlueprint, Warning, TEXT("BackwardCompatibilityNodeConversion Error 'cannot connect' in blueprint: %s, pin: %s"),
+						Blueprint ? *Blueprint->GetName() : TEXT("Unknown"),
+						*InValidCastPin->PinName);
+					return false;
+				}
+
 				auto CastSourcePin = CastNode->GetCastSourcePin();
 				check(CastSourcePin);
 				if (!Schema.MovePinLinks(*OldBlueprintPin, *CastSourcePin).CanSafeConnect())
diff --git a/Engine/Source/Editor/UnrealEd/Private/Kismet2/BlueprintEditorUtils.cpp b/Engine/Source/Editor/UnrealEd/Private/Kismet2/BlueprintEditorUtils.cpp
--- a/Engine/Source/Editor/UnrealEd/Private/Kismet2/BlueprintEditorUtils.cpp
+++ b/Engine/Source/Editor/UnrealEd/Private/Kismet2/BlueprintEditorUtils.cpp
@@ -714,6 +714,179 @@ void FBlueprintEditorUtils::RefreshInputDelegatePins(UBlueprint* Blueprint)
 	}
 }
 
+struct FEditoronlyBlueprintHelper
+{
+	static bool IsUnwantedType(const FEdGraphPinType& Type)
+	{
+		if (const UClass* BPClass = Cast<const UClass>(Type.PinSubCategoryObject.Get()))
+		{
+			if (BPClass->IsChildOf(UBlueprint::StaticClass()))
+			{
+				return true;
+			}
+		}
+
+		return false;
+	}
+
+	static bool IsUnwantedDefaultObject(const UObject* Obj)
+	{
+		return Obj && Obj->IsA<UBlueprint>();
+	}
+
+	static void ChangePinType(FEdGraphPinType& Type)
+	{
+		const UEdGraphSchema_K2* Schema = GetDefault<UEdGraphSchema_K2>();
+
+		Type.PinCategory = Schema->PC_Class;
+		Type.PinSubCategoryObject = UObject::StaticClass();
+	}
+
+	static bool ShouldBeFixed(const UBlueprint* Blueprint, bool bLogWhy)
+	{
+		check(Blueprint);
+		//any blueprint variable
+		for (const auto& VarDesc : Blueprint->NewVariables)
+		{
+			if (IsUnwantedType(VarDesc.VarType))
+			{
+				if (bLogWhy)
+				{
+					UE_LOG(LogBlueprint, Warning, TEXT("FEditoronlyBlueprintHelper::ShouldBeFixed"));
+				}
+				return true;
+			}
+		}
+
+		//anything on graph
+		TArray<UEdGraph*> Graphs;
+		Blueprint->GetAllGraphs(Graphs);
+		for (const auto Graph : Graphs)
+		{
+			TArray<UK2Node*> Nodes;
+			Graph->GetNodesOfClass(Nodes);
+			for (const auto Node : Nodes)
+			{
+				for (const auto Pin : Node->Pins)
+				{
+					if (Pin && (IsUnwantedType(Pin->PinType) || IsUnwantedDefaultObject(Pin->DefaultObject)))
+					{
+						if (bLogWhy)
+						{
+							UE_LOG(LogBlueprint, Warning, TEXT("FEditoronlyBlueprintHelper::ShouldBeFixed"));
+						}
+						return true;
+					}
+				}
+			}
+		}
+
+		return false;
+	}
+
+	static void HandleEditablePinNode(UK2Node_EditablePinBase* Node)
+	{
+		check(Node);
+
+		const UEdGraphSchema_K2* Schema = GetDefault<UEdGraphSchema_K2>();
+
+		for (auto UserPinInfo : Node->UserDefinedPins)
+		{
+			if (UserPinInfo.IsValid() && IsUnwantedType(UserPinInfo->PinType))
+			{
+				UserPinInfo->PinType.PinCategory = Schema->PC_Class;
+				UserPinInfo->PinType.PinSubCategoryObject = UObject::StaticClass();
+
+				UserPinInfo->PinDefaultValue.Empty();
+			}
+		}
+
+		for (auto Pin : Node->Pins)
+		{
+			if (IsUnwantedType(Pin->PinType))
+			{
+				ChangePinType(Pin->PinType);
+
+				if (Pin->DefaultObject)
+				{
+					UBlueprint* DefaultBlueprint = Cast<UBlueprint>(Pin->DefaultObject);
+					Pin->DefaultObject = DefaultBlueprint ? *DefaultBlueprint->GeneratedClass : NULL;
+				}
+			}
+
+		}
+	}
+
+	static void HandleDefaultObjects(UK2Node* Node)
+	{
+		if (Node)
+		{
+			for (auto Pin : Node->Pins)
+			{
+				if (Pin && IsUnwantedDefaultObject(Pin->DefaultObject))
+				{
+					const UBlueprint* DefaultBlueprint = Cast<const UBlueprint>(Pin->DefaultObject);
+					if (DefaultBlueprint)
+					{
+						Pin->DefaultObject = *DefaultBlueprint->GeneratedClass;
+					}
+				}
+			}
+		}
+	}
+
+	static bool ChangeBlueprint(UBlueprint* Blueprint)
+	{
+		if (ShouldBeFixed(Blueprint, false))
+		{
+			UE_LOG(LogBlueprint, Log, TEXT("Bluepirnt references will b e removed from '%s'"), *Blueprint->GetName());
+
+			const UEdGraphSchema_K2* Schema = GetDefault<UEdGraphSchema_K2>();
+
+			//any native function call with blueprint parameter
+			TArray<UEdGraph*> Graphs;
+			Blueprint->GetAllGraphs(Graphs);
+			for (const auto Graph : Graphs)
+			{
+				Schema->BackwardCompatibilityNodeConversion(Graph, false);
+			}
+
+			//any blueprint pin (function/event/macro parameter
+			for (const auto Graph : Graphs)
+			{
+				TArray<UK2Node_EditablePinBase*> EditablePinNodes;
+				Graph->GetNodesOfClass(EditablePinNodes);
+				for (const auto Node : EditablePinNodes)
+				{
+					HandleEditablePinNode(Node);
+				}
+			}
+
+			// change variables
+			for (auto& VarDesc : Blueprint->NewVariables)
+			{
+				if (IsUnwantedType(VarDesc.VarType))
+				{
+					ChangePinType(VarDesc.VarType);
+				}
+			}
+
+			for (const auto Graph : Graphs)
+			{
+				TArray<UK2Node*> Nodes;
+				Graph->GetNodesOfClass(Nodes);
+				for (const auto Node : Nodes)
+				{
+					HandleDefaultObjects(Node);
+				}
+			}
+
+			return true;
+		}
+		return false;
+	}
+};
+
 UClass* FBlueprintEditorUtils::RegenerateBlueprintClass(UBlueprint* Blueprint, UClass* ClassToRegenerate, UObject* PreviousCDO, TArray<UObject*>& ObjLoaded)
 {
 	bool bRegenerated = false;
@@ -797,6 +970,23 @@ UClass* FBlueprintEditorUtils::RegenerateBlueprintClass(UBlueprint* Blueprint, U
 		FRegenerationHelper::LinkExternalDependencies(Blueprint);
 		FBlueprintEditorUtils::RefreshInputDelegatePins(Blueprint);
 
+		struct FReplaceBlueprintWithClassHelper
+		{
+			bool bShouldReplace;
+
+			FReplaceBlueprintWithClassHelper() : bShouldReplace(false)
+			{
+				GConfig->GetBool(TEXT("EditoronlyBP"), TEXT("bReplaceBlueprintWithClass"), bShouldReplace, GEditorIni);
+			}
+		};
+		static FReplaceBlueprintWithClassHelper ReplaceBlueprintWithClassHelper;
+		const bool bReplaceBlueprintWithClass = ReplaceBlueprintWithClassHelper.bShouldReplace;
+
+		if (bReplaceBlueprintWithClass)
+		{
+			FEditoronlyBlueprintHelper::ChangeBlueprint(Blueprint);
+		}
+
 		if( bHasCode )
 		{
 			// Make sure parent function calls are up to date
@@ -841,6 +1031,11 @@ UClass* FBlueprintEditorUtils::RegenerateBlueprintClass(UBlueprint* Blueprint, U
 			// Flag data only blueprints as being up-to-date
 			Blueprint->Status = BS_UpToDate;
 		}
+
+		if (bReplaceBlueprintWithClass)
+		{
+			FEditoronlyBlueprintHelper::ShouldBeFixed(Blueprint, true);
+		}
 		
 		// Patch the new CDOs to the old indices in the linker
 		if( Blueprint->SkeletonGeneratedClass )
