#UE4 - Fix bogus material index warning when importing FBX static meshes

diff --git a/Engine/Source/Editor/UnrealEd/Private/Fbx/FbxStaticMeshImport.cpp b/Engine/Source/Editor/UnrealEd/Private/Fbx/FbxStaticMeshImport.cpp
--- a/Engine/Source/Editor/UnrealEd/Private/Fbx/FbxStaticMeshImport.cpp
+++ b/Engine/Source/Editor/UnrealEd/Private/Fbx/FbxStaticMeshImport.cpp
@@ -654,22 +654,10 @@ bool UnFbx::FFbxImporter::BuildStaticMeshFromGeometry(FbxMesh* Mesh, UStaticMesh
 			AddTokenizedErrorMessage(FTokenizedMessage::Create(EMessageSeverity::Warning, LOCTEXT("Error_MaterialIndexInconsistency", "Face material index inconsistency - forcing to 0")));
 			MaterialIndex = 0;
 		}
-		MaxMaterialIndex = FMath::Max( MaxMaterialIndex, MaterialIndex );
-
+	
 		RawMesh.FaceMaterialIndices[DestTriangleIndex] = FMath::Min( MaterialIndex, MAX_MESH_MATERIAL_INDEX );
 	}
 
-	if( MaxMaterialIndex > MAX_MESH_MATERIAL_INDEX )
-	{
-		AddTokenizedErrorMessage( 
-			FTokenizedMessage::Create(
-				EMessageSeverity::Warning, 
-				FText::Format( LOCTEXT("Error_TooManyMaterials", "Mesh '{0}' has too many({2}) materials. Clamping materials to {1}."), 
-				FText::FromString(Mesh->GetName()),
-				FText::AsNumber(MAX_MESH_MATERIAL_INDEX),
-				FText::AsNumber(MaxMaterialIndex+1)
-				) ) );
-	}
 
 	// Store the new raw mesh.
 	SrcModel.RawMeshBulkData->SaveRawMesh(RawMesh);
@@ -980,6 +968,12 @@ UStaticMesh* UnFbx::FFbxImporter::ImportStaticMeshAsSingle(UObject* InParent, TA
 				if (MeshMaterials[MaterialIndex].FbxMaterial == MeshMaterials[OtherMaterialIndex].FbxMaterial)
 				{
 					int32 UniqueIndex = MaterialMap[OtherMaterialIndex];
+
+					if( UniqueIndex > MAX_MESH_MATERIAL_INDEX )
+					{
+						UniqueIndex = MAX_MESH_MATERIAL_INDEX;
+					}
+
 					MaterialMap.Add(UniqueIndex);
 					bDoRemap = true;
 					bUnique = false;
@@ -988,7 +982,14 @@ UStaticMesh* UnFbx::FFbxImporter::ImportStaticMeshAsSingle(UObject* InParent, TA
 			}
 			if (bUnique)
 			{
-				MaterialMap.Add(UniqueMaterials.Add(MeshMaterials[MaterialIndex]));
+				int32 UniqueIndex = UniqueMaterials.Add(MeshMaterials[MaterialIndex]);
+
+				if (UniqueIndex > MAX_MESH_MATERIAL_INDEX)
+				{
+					UniqueIndex = MAX_MESH_MATERIAL_INDEX;
+				}
+
+				MaterialMap.Add( UniqueIndex );
 			}
 			else
 			{
@@ -996,6 +997,17 @@ UStaticMesh* UnFbx::FFbxImporter::ImportStaticMeshAsSingle(UObject* InParent, TA
 			}
 		}
 
+		if (UniqueMaterials.Num() > MAX_MESH_MATERIAL_INDEX)
+		{
+			AddTokenizedErrorMessage(
+				FTokenizedMessage::Create(
+				EMessageSeverity::Warning,
+				FText::Format(LOCTEXT("Error_TooManyMaterials", "StaticMesh has too many({1}) materials. Clamping materials to {0} which may produce unexpected results. Break apart your mesh into multiple pieces to fix this."),
+				FText::AsNumber(MAX_MESH_MATERIAL_INDEX),
+				FText::AsNumber(UniqueMaterials.Num())
+				)));
+		}
+
 		// Sort materials based on _SkinXX in the name.
 		TArray<uint32> SortedMaterialIndex;
 		for (int32 MaterialIndex = 0; MaterialIndex < MeshMaterials.Num(); ++MaterialIndex)
