TTP #330690: Adding the ability to drag and drop a blueprint asset into another blueprint graph (adds a child component node).

diff --git a/Engine/Plugins/Experimental/Paper2D/Source/Paper2DEditor/Private/Paper2DEditorModule.cpp b/Engine/Plugins/Experimental/Paper2D/Source/Paper2DEditor/Private/Paper2DEditorModule.cpp
--- a/Engine/Plugins/Experimental/Paper2D/Source/Paper2DEditor/Private/Paper2DEditorModule.cpp
+++ b/Engine/Plugins/Experimental/Paper2D/Source/Paper2DEditor/Private/Paper2DEditorModule.cpp
@@ -71,10 +71,10 @@ public:
 		RegisterAssetTypeAction(AssetTools, MakeShareable(new FTileSetAssetTypeActions));
 
 		PaperSpriteBroker = MakeShareable(new FPaperSpriteAssetBroker);
-		FComponentAssetBrokerage::RegisterBroker(PaperSpriteBroker, UPaperRenderComponent::StaticClass(), true);
+		FComponentAssetBrokerage::RegisterBroker(PaperSpriteBroker, UPaperRenderComponent::StaticClass(), true, true);
 
 		PaperFlipbookBroker = MakeShareable(new FPaperFlipbookAssetBroker);
-		FComponentAssetBrokerage::RegisterBroker(PaperFlipbookBroker, UPaperAnimatedRenderComponent::StaticClass(), true);
+		FComponentAssetBrokerage::RegisterBroker(PaperFlipbookBroker, UPaperAnimatedRenderComponent::StaticClass(), true, true);
 
 		// Register the details customizations
 		{
diff --git a/Engine/Source/Editor/BlueprintGraph/Private/EdGraphSchema_K2.cpp b/Engine/Source/Editor/BlueprintGraph/Private/EdGraphSchema_K2.cpp
--- a/Engine/Source/Editor/BlueprintGraph/Private/EdGraphSchema_K2.cpp
+++ b/Engine/Source/Editor/BlueprintGraph/Private/EdGraphSchema_K2.cpp
@@ -3104,7 +3104,17 @@ void UEdGraphSchema_K2::DroppedAssetsOnGraph(const TArray<FAssetData>& Assets, c
 
 			UObject* Asset = Assets[AssetIdx].GetAsset();
 
-			TSubclassOf<UActorComponent> DestinationComponentType = FComponentAssetBrokerage::GetPrimaryComponentForAsset(Asset->GetClass());
+			UClass* AssetClass = Asset->GetClass();		
+			if (UBlueprint* BlueprintAsset = Cast<UBlueprint>(Asset))
+			{
+				AssetClass = BlueprintAsset->GeneratedClass;
+			}
+
+			TSubclassOf<UActorComponent> DestinationComponentType = FComponentAssetBrokerage::GetPrimaryComponentForAsset(AssetClass);
+			if ((DestinationComponentType == NULL) && AssetClass->IsChildOf(AActor::StaticClass()))
+			{
+				DestinationComponentType = UChildActorComponent::StaticClass();
+			}
 
 			// Make sure we have an asset type that's registered with the component list
 			if (DestinationComponentType != NULL)
diff --git a/Engine/Source/Editor/BlueprintGraph/Private/K2Node_AddComponent.cpp b/Engine/Source/Editor/BlueprintGraph/Private/K2Node_AddComponent.cpp
--- a/Engine/Source/Editor/BlueprintGraph/Private/K2Node_AddComponent.cpp
+++ b/Engine/Source/Editor/BlueprintGraph/Private/K2Node_AddComponent.cpp
@@ -294,6 +294,7 @@ FString UK2Node_AddComponent::GetNodeTitle(ENodeTitleType::Type TitleType) const
 			UStaticMeshComponent* StaticMeshComp = Cast<UStaticMeshComponent>(SourceTemplate);
 			USkeletalMeshComponent* SkelMeshComp = Cast<USkeletalMeshComponent>(SourceTemplate);
 			UParticleSystemComponent* PSysComp = Cast<UParticleSystemComponent>(SourceTemplate);
+			UChildActorComponent* SubActorComp = Cast<UChildActorComponent>(SourceTemplate);
 
 			if(StaticMeshComp != NULL && StaticMeshComp->StaticMesh != NULL)
 			{
@@ -307,6 +308,10 @@ FString UK2Node_AddComponent::GetNodeTitle(ENodeTitleType::Type TitleType) const
 			{
 				return FString(TEXT("Add ParticleSystem ")) + PSysComp->Template->GetName();		
 			}
+			else if (SubActorComp && SubActorComp->ChildActorClass)
+			{
+				return FString(TEXT("Add ChildActorComponent ")) + SubActorComp->ChildActorClass->GetName();
+			}
 			else
 			{
 				return FString(TEXT("Add ")) + SourceTemplate->GetClass()->GetName();		
diff --git a/Engine/Source/Editor/UnrealEd/Private/ComponentAssetBroker.cpp b/Engine/Source/Editor/UnrealEd/Private/ComponentAssetBroker.cpp
--- a/Engine/Source/Editor/UnrealEd/Private/ComponentAssetBroker.cpp
+++ b/Engine/Source/Editor/UnrealEd/Private/ComponentAssetBroker.cpp
@@ -347,16 +347,16 @@ void FComponentAssetBrokerage::InitializeMap()
 	{
 		bInitializedBuiltinMap = true;
 
-		RegisterBroker(MakeShareable(new FStaticMeshComponentBroker), UStaticMeshComponent::StaticClass(), true);
-		RegisterBroker(MakeShareable(new FSkeletalMeshComponentBroker), USkeletalMeshComponent::StaticClass(), true);
-		RegisterBroker(MakeShareable(new FDestructableMeshComponentBroker), UDestructibleComponent::StaticClass(), false);
-		RegisterBroker(MakeShareable(new FParticleSystemComponentBroker), UParticleSystemComponent::StaticClass(), true);
-		RegisterBroker(MakeShareable(new FAudioComponentBroker), UAudioComponent::StaticClass(), true);
-		//RegisterBroker(MakeShareable(new FChildActorComponentBroker), UChildActorComponent::StaticClass(), true);
+		RegisterBroker(MakeShareable(new FStaticMeshComponentBroker), UStaticMeshComponent::StaticClass(), true, true);
+		RegisterBroker(MakeShareable(new FSkeletalMeshComponentBroker), USkeletalMeshComponent::StaticClass(), true, true);
+		RegisterBroker(MakeShareable(new FDestructableMeshComponentBroker), UDestructibleComponent::StaticClass(), false, true);
+		RegisterBroker(MakeShareable(new FParticleSystemComponentBroker), UParticleSystemComponent::StaticClass(), true, true);
+		RegisterBroker(MakeShareable(new FAudioComponentBroker), UAudioComponent::StaticClass(), true, true);
+		RegisterBroker(MakeShareable(new FChildActorComponentBroker), UChildActorComponent::StaticClass(), true, false);
 	}
 }
 
-void FComponentAssetBrokerage::RegisterBroker(TSharedPtr<IComponentAssetBroker> Broker, TSubclassOf<UActorComponent> InComponentClass, bool bSetAsPrimary)
+void FComponentAssetBrokerage::RegisterBroker(TSharedPtr<IComponentAssetBroker> Broker, TSubclassOf<UActorComponent> InComponentClass, bool bSetAsPrimary, bool bMapCompnentForAssets)
 {
 	InitializeMap();
 
@@ -368,17 +368,27 @@ void FComponentAssetBrokerage::RegisterBroker(TSharedPtr<IComponentAssetBroker>
 	checkf(!ComponentToBrokerMap.Contains(InComponentClass), TEXT("Component class already has a registered broker; you have to chain them yourself."));
 	ComponentToBrokerMap.Add(InComponentClass, Broker);
 
-	FComponentClassList& ValidComponentTypes = AssetToComponentClassMap.FindOrAdd(AssetClass);
 	if (bSetAsPrimary)
 	{
-		ValidComponentTypes.Insert(InComponentClass, 0);
 		AssetToBrokerMap.FindOrAdd(AssetClass).Insert(Broker, 0);
 	}
 	else
 	{
-		ValidComponentTypes.Add(InComponentClass);
 		AssetToBrokerMap.FindOrAdd(AssetClass).Add(Broker);
 	}
+
+	if (bMapCompnentForAssets)
+	{
+		FComponentClassList& ValidComponentTypes = AssetToComponentClassMap.FindOrAdd(AssetClass);
+		if (bSetAsPrimary)
+		{
+			ValidComponentTypes.Insert(InComponentClass, 0);
+		}
+		else 
+		{
+			ValidComponentTypes.Add(InComponentClass);
+		}
+	}
 }
 
 void FComponentAssetBrokerage::UnregisterBroker(TSharedPtr<IComponentAssetBroker> Broker)
diff --git a/Engine/Source/Editor/UnrealEd/Public/ComponentAssetBroker.h b/Engine/Source/Editor/UnrealEd/Public/ComponentAssetBroker.h
--- a/Engine/Source/Editor/UnrealEd/Public/ComponentAssetBroker.h
+++ b/Engine/Source/Editor/UnrealEd/Public/ComponentAssetBroker.h
@@ -73,7 +73,7 @@ public:
 
 
 	/** Register a component class for a specified asset class */
-	static void RegisterBroker(TSharedPtr<IComponentAssetBroker> Broker, TSubclassOf<UActorComponent> InComponentClass, bool bSetAsPrimary);
+	static void RegisterBroker(TSharedPtr<IComponentAssetBroker> Broker, TSubclassOf<UActorComponent> InComponentClass, bool bSetAsPrimary, bool bMapCompnentForAssets);
 
 	/** Unregister a component type for a specified asset class */
 	static void UnregisterBroker(TSharedPtr<IComponentAssetBroker> Broker);
