When redirectors are added to the delete model we now also do similar logic to that found in the fixup code to try and allow redirectors to be included in the set of objects to be deleted, assuming they are not referenced by anything else.

diff --git a/Engine/Source/Editor/UnrealEd/Private/AssetDeleteModel.cpp b/Engine/Source/Editor/UnrealEd/Private/AssetDeleteModel.cpp
--- a/Engine/Source/Editor/UnrealEd/Private/AssetDeleteModel.cpp
+++ b/Engine/Source/Editor/UnrealEd/Private/AssetDeleteModel.cpp
@@ -36,6 +36,8 @@ void FAssetDeleteModel::AddObjectToDelete( UObject* InObject )
 {
 	PendingDeletes.Add(MakeShareable(new FPendingDelete(InObject)));
 
+	PrepareToDelete(InObject);
+
 	SetState(StartScanning);
 }
 
@@ -358,6 +360,54 @@ int32 FAssetDeleteModel::GetDeletedObjectCount() const
 	return ObjectsDeleted;
 }
 
+void FAssetDeleteModel::PrepareToDelete(UObject* InObject)
+{
+	if ( InObject->IsA<UObjectRedirector>() )
+	{
+		// Add all redirectors found in this package to the redirectors to delete list.
+		// All redirectors in this package should be fixed up.
+		UPackage* RedirectorPackage = InObject->GetOutermost();
+		TArray<UObject*> AssetsInRedirectorPackage;
+		
+		GetObjectsWithOuter(RedirectorPackage, AssetsInRedirectorPackage, /*bIncludeNestedObjects=*/false);
+		UMetaData* PackageMetaData = NULL;
+		bool bContainsAtLeastOneOtherAsset = false;
+
+		for ( auto ObjIt = AssetsInRedirectorPackage.CreateConstIterator(); ObjIt; ++ObjIt )
+		{
+			if ( UObjectRedirector* Redirector = Cast<UObjectRedirector>(*ObjIt) )
+			{
+				Redirector->RemoveFromRoot();
+			}
+			else if ( UMetaData* MetaData = Cast<UMetaData>(*ObjIt) )
+			{
+				PackageMetaData = MetaData;
+			}
+			else
+			{
+				bContainsAtLeastOneOtherAsset = true;
+			}
+		}
+
+		if ( !bContainsAtLeastOneOtherAsset )
+		{
+			RedirectorPackage->RemoveFromRoot();
+			ULinkerLoad* Linker = ULinkerLoad::FindExistingLinkerForPackage(RedirectorPackage);
+			if ( Linker )
+			{
+				Linker->RemoveFromRoot();
+			}
+
+			// @todo we shouldnt be worrying about metadata objects here, ObjectTools::CleanUpAfterSuccessfulDelete should
+			if ( PackageMetaData )
+			{
+				PackageMetaData->RemoveFromRoot();
+				PendingDeletes.AddUnique(MakeShareable(new FPendingDelete(PackageMetaData)));
+			}
+		}
+	}
+}
+
 // FPendingDelete
 //-----------------------------------------------------------------
 
@@ -474,4 +524,9 @@ void FPendingDelete::CheckForReferences()
 	}
 }
 
+bool FPendingDelete::operator == ( const FPendingDelete& Other ) const
+{
+	return Object == Other.Object;
+}
+
 #undef LOCTEXT_NAMESPACE
\ No newline at end of file
diff --git a/Engine/Source/Editor/UnrealEd/Private/ObjectTools.cpp b/Engine/Source/Editor/UnrealEd/Private/ObjectTools.cpp
--- a/Engine/Source/Editor/UnrealEd/Private/ObjectTools.cpp
+++ b/Engine/Source/Editor/UnrealEd/Private/ObjectTools.cpp
@@ -1302,7 +1302,7 @@ namespace ObjectTools
 		}
 	}
 
-	void CleanupAfterSuccessfulDelete (const TArray<UPackage*>& PotentialPackagesToDelete)
+	void CleanupAfterSuccessfulDelete (const TArray<UPackage*>& PotentialPackagesToDelete, bool bPerformReferenceCheck)
 	{
 		TArray<UPackage*> PackagesToDelete = PotentialPackagesToDelete;
 		TArray<FString> PackageFilesToDelete;
@@ -1318,26 +1318,31 @@ namespace ObjectTools
 			GWarn->StatusUpdate(OriginalNumPackagesToDelete - PackageIdx, OriginalNumPackagesToDelete, NSLOCTEXT("ObjectTools", "OldPackageCleanupSlowTask", "Cleaning Up Old Assets"));
 			UObject* Package = PackagesToDelete[PackageIdx];
 
-			FReferencerInformationList FoundReferences;
-			bool bIsReferenced = IsReferenced( Package, GARBAGE_COLLECTION_KEEPFLAGS, true, &FoundReferences );
-			if ( bIsReferenced )
+			bool bIsReferenced = false;
+			
+			if ( bPerformReferenceCheck )
 			{
-				// determine whether the transaction buffer is the only thing holding a reference to the object
-				// and if so, offer the user the option to reset the transaction buffer.
-				GEditor->Trans->DisableObjectSerialization();
+				FReferencerInformationList FoundReferences;
 				bIsReferenced = IsReferenced(Package, GARBAGE_COLLECTION_KEEPFLAGS, true, &FoundReferences);
-				GEditor->Trans->EnableObjectSerialization();
-
-				// only ref to this object is the transaction buffer - let the user choose whether to clear the undo buffer
-				if ( !bIsReferenced )
+				if ( bIsReferenced )
 				{
-					if ( EAppReturnType::Yes == FMessageDialog::Open( EAppMsgType::YesNo, NSLOCTEXT("UnrealEd", "ResetUndoBufferForObjectDeletionPrompt", "The only reference to this object is the undo history.  In order to delete this object, you must clear all undo history - would you like to clear undo history?")) )
-					{
-						GEditor->Trans->Reset( NSLOCTEXT("UnrealEd", "DeleteSelectedItem", "Delete Selected Item") );
-					}
-					else
+					// determine whether the transaction buffer is the only thing holding a reference to the object
+					// and if so, offer the user the option to reset the transaction buffer.
+					GEditor->Trans->DisableObjectSerialization();
+					bIsReferenced = IsReferenced(Package, GARBAGE_COLLECTION_KEEPFLAGS, true, &FoundReferences);
+					GEditor->Trans->EnableObjectSerialization();
+
+					// only ref to this object is the transaction buffer - let the user choose whether to clear the undo buffer
+					if ( !bIsReferenced )
 					{
-						bIsReferenced = true;
+						if ( EAppReturnType::Yes == FMessageDialog::Open(EAppMsgType::YesNo, NSLOCTEXT("UnrealEd", "ResetUndoBufferForObjectDeletionPrompt", "The only reference to this object is the undo history.  In order to delete this object, you must clear all undo history - would you like to clear undo history?")) )
+						{
+							GEditor->Trans->Reset(NSLOCTEXT("UnrealEd", "DeleteSelectedItem", "Delete Selected Item"));
+						}
+						else
+						{
+							bIsReferenced = true;
+						}
 					}
 				}
 			}
@@ -1581,7 +1586,8 @@ namespace ObjectTools
 				PotentialPackagesToDelete.AddUnique( ObjectsDeletedSuccessfully[ObjIdx]->GetOutermost() );
 			}
 
-			CleanupAfterSuccessfulDelete( PotentialPackagesToDelete );
+			bool bPerformReferenceCheck = false;
+			CleanupAfterSuccessfulDelete( PotentialPackagesToDelete, bPerformReferenceCheck );
 			ObjectsDeletedSuccessfully.Empty();
 		}
 
diff --git a/Engine/Source/Editor/UnrealEd/Public/AssetDeleteModel.h b/Engine/Source/Editor/UnrealEd/Public/AssetDeleteModel.h
--- a/Engine/Source/Editor/UnrealEd/Public/AssetDeleteModel.h
+++ b/Engine/Source/Editor/UnrealEd/Public/AssetDeleteModel.h
@@ -5,7 +5,7 @@
 /**
  * The pending deleted object.
  */
-struct FPendingDelete
+struct FPendingDelete : TSharedFromThis<FPendingDelete>
 {
 public:
 	FPendingDelete( UObject* InObject );
@@ -34,6 +34,9 @@ public:
 	/** Sets if the pending delete is internal, and need not be shown to the user. */
 	void IsInternal(bool Value) { bIsInternal = Value; }
 
+	/** Support comparing for unique insertions */
+	bool operator == ( const FPendingDelete& Other ) const;
+
 	/** The on disk references to this object */
 	TArray<FName> DiskReferences;
 
@@ -168,6 +171,8 @@ public:
 	}
 
 private:
+	void PrepareToDelete(UObject* InObject);
+
 	/** Sets the current state of the model. */
 	void SetState( EState NewState );
 
diff --git a/Engine/Source/Editor/UnrealEd/Public/ObjectTools.h b/Engine/Source/Editor/UnrealEd/Public/ObjectTools.h
--- a/Engine/Source/Editor/UnrealEd/Public/ObjectTools.h
+++ b/Engine/Source/Editor/UnrealEd/Public/ObjectTools.h
@@ -258,7 +258,7 @@ namespace ObjectTools
 	 *
 	 * @param	ObjectsDeletedSuccessfully		The list of objects that were recently deleted but not yet cleaned up
 	 */
-	UNREALED_API void CleanupAfterSuccessfulDelete ( const TArray<UPackage*>& ObjectsDeletedSuccessfully );
+	UNREALED_API void CleanupAfterSuccessfulDelete ( const TArray<UPackage*>& ObjectsDeletedSuccessfully, bool bPerformReferenceCheck = true );
 
 	/**
 	 * Deletes the list of objects
