Undoing destructible mesh material collapsing - broke importing of APB with materials

diff --git a/Engine/Source/Editor/UnrealEd/Private/Factories/ApexDestructibleAssetImport.cpp b/Engine/Source/Editor/UnrealEd/Private/Factories/ApexDestructibleAssetImport.cpp
--- a/Engine/Source/Editor/UnrealEd/Private/Factories/ApexDestructibleAssetImport.cpp
+++ b/Engine/Source/Editor/UnrealEd/Private/Factories/ApexDestructibleAssetImport.cpp
@@ -476,28 +476,6 @@ static bool FillSkelMeshImporterFromApexDestructibleAsset(FSkeletalMeshImportDat
 	bHaveAllNormals = true;
 	bHaveAllTangents = true;
 
-	TArray<VMaterial> UniqueMaterials;
-
-	for (int32 CurMatIdx=0; CurMatIdx < ImportData.Materials.Num(); ++CurMatIdx)
-	{
-		bool bHasMaterial = false;
-
-		for (int32 CurUniqueMatIdx=0; CurUniqueMatIdx < UniqueMaterials.Num(); ++CurUniqueMatIdx)
-		{
-			if (ImportData.Materials[CurMatIdx].MaterialImportName == UniqueMaterials[CurUniqueMatIdx].MaterialImportName)
-			{
-				bHasMaterial = true;
-				break;
-			}
-		}
-
-		if (!bHasMaterial)
-		{
-			UniqueMaterials.Add(ImportData.Materials[CurMatIdx]);
-		}
-		
-	}
-
 	// APEX render meshes are organized by submesh (render elements)
 	// Looping through submeshes first, can be done either way
 	for (uint32 SubmeshIndex = 0; SubmeshIndex < ApexRenderMesh->getSubmeshCount(); ++SubmeshIndex)
@@ -614,25 +592,7 @@ static bool FillSkelMeshImporterFromApexDestructibleAsset(FSkeletalMeshImportDat
 				Triangle.SmoothingGroups = 255; 
 
 				// Material index
-				int32 MatIdx = 0;
-
-				check(ImportData.Materials.Num() > 0);
-				
-				if (SubmeshIndex < (uint32)ImportData.Materials.Num())
-				{
-					const FString& MaterialName = ImportData.Materials[SubmeshIndex].MaterialImportName;
-
-					for (int32 UniqueMaterialIdx = 0; UniqueMaterialIdx < UniqueMaterials.Num(); ++UniqueMaterialIdx)
-					{
-						if (UniqueMaterials[UniqueMaterialIdx].MaterialImportName == MaterialName)
-						{
-							MatIdx = UniqueMaterialIdx;
-							break;
-						}
-					}
-				}
-				
-				Triangle.MatIndex = MatIdx;
+				Triangle.MatIndex = SubmeshIndex;
 				Triangle.AuxMatIndex = 0;
 
 				// Per-vertex
@@ -683,9 +643,6 @@ static bool FillSkelMeshImporterFromApexDestructibleAsset(FSkeletalMeshImportDat
 		VertexIndexBase += SubmeshVertexCount;
 	}
 
-	// Switch material arrays so we have a list with unique materials
-	ImportData.Materials = UniqueMaterials;
-
 	// Create mapping from import to raw- @TODO trivial at the moment, do we need this info for destructibles?
 	ImportData.PointToRawMap.AddUninitialized(ImportData.Points.Num());
 	for(int32 PointIdx=0; PointIdx<ImportData.PointToRawMap.Num(); PointIdx++)
@@ -769,24 +726,7 @@ bool SetApexDestructibleAsset(UDestructibleMesh& DestructibleMesh, NxDestructibl
 {
 	DestructibleMesh.PreEditChange(NULL);
 
-#if WITH_EDITORONLY_DATA
-	// Handle the StaticMesh import case, where we have the materials only stored in the fracture settings
-	if (DestructibleMesh.Materials.Num() == 0 && DestructibleMesh.FractureSettings != NULL && 
-		DestructibleMesh.FractureSettings->Materials.Num() > 0)
-	{
-		for (int32 i=0; i < DestructibleMesh.FractureSettings->Materials.Num(); ++i)
-		{
-			FSkeletalMaterial Mat;
-			Mat.MaterialInterface = DestructibleMesh.FractureSettings->Materials[i];
-			Mat.bEnableShadowCasting = true;
-
-			DestructibleMesh.Materials.Add(Mat);
-		}
-	}
-#endif
-
-	ExistingDestMeshData* ExistDestMeshDataPtr = NULL;
-	ExistDestMeshDataPtr = SaveExistingDestMeshData(&DestructibleMesh);
+	ExistingDestMeshData * ExistDestMeshDataPtr = SaveExistingDestMeshData(&DestructibleMesh);
 	
 	// The asset is going away, which will destroy any actors created from it.  We must destroy the physics state of any destructible mesh components before we release the asset.
 	for(TObjectIterator<UDestructibleComponent> It; It; ++It)
@@ -835,29 +775,13 @@ bool SetApexDestructibleAsset(UDestructibleMesh& DestructibleMesh, NxDestructibl
 	{
 		SkelMeshImportDataPtr = OutData;
 	}
+	
+	// Get all material names here
+	ImportMaterialsForSkelMesh(*SkelMeshImportDataPtr, ApexDestructibleAsset);
 
 	// Import animation hierarchy, although this is trivial for an Apex Destructible Asset
 	CreateBones(*SkelMeshImportDataPtr, ApexDestructibleAsset);
 
-	if (!(Options & EImportOptions::PreserveSettings))
-	{
-		// Get all material names here
-		ImportMaterialsForSkelMesh(*SkelMeshImportDataPtr, ApexDestructibleAsset);
-	}
-	else
-	{
-		SkelMeshImportDataPtr->Materials.Empty(DestructibleMesh.Materials.Num());
-		for (int32 i=0; i < DestructibleMesh.Materials.Num(); ++i)
-		{
-			UMaterialInterface* MI = DestructibleMesh.Materials[i].MaterialInterface;
-			VMaterial NewMat;
-			
-			NewMat.Material = MI;
-			NewMat.MaterialImportName = MI != NULL ? MI->GetName() : TEXT("");
-			SkelMeshImportDataPtr->Materials.Add(NewMat);
-		}
-	}
-
 	// Import graphics data
 	bool bHaveNormals, bHaveTangents;
 	if (!FillSkelMeshImporterFromApexDestructibleAsset(*SkelMeshImportDataPtr, ApexDestructibleAsset, bHaveNormals, bHaveTangents))
@@ -880,11 +804,8 @@ bool SetApexDestructibleAsset(UDestructibleMesh& DestructibleMesh, NxDestructibl
 	}
 #endif
 
-	if ( !(Options&EImportOptions::PreserveSettings) )
-	{
-		// process materials from import data
-		ProcessImportMeshMaterials( DestructibleMesh.Materials,*SkelMeshImportDataPtr );
-	}
+	// process materials from import data
+	ProcessImportMeshMaterials( DestructibleMesh.Materials,*SkelMeshImportDataPtr );
 	
 	// process reference skeleton from import data
 	int32 SkeletalDepth=0;
@@ -1020,6 +941,13 @@ UNREALED_API bool BuildDestructibleMeshFromFractureSettings(UDestructibleMesh& D
 #if WITH_EDITORONLY_DATA
 	if (DestructibleMesh.FractureSettings != NULL)
 	{
+		DestructibleMesh.Materials.Init(DestructibleMesh.FractureSettings->Materials.Num());
+
+		for (int32 MaterialIndex = 0; MaterialIndex < DestructibleMesh.Materials.Num(); ++MaterialIndex)
+		{
+			DestructibleMesh.Materials[MaterialIndex].MaterialInterface = DestructibleMesh.FractureSettings->Materials[MaterialIndex];
+		}
+
 		NxDestructibleAssetCookingDesc DestructibleAssetCookingDesc;
 		DestructibleMesh.FractureSettings->BuildDestructibleAssetCookingDesc(DestructibleAssetCookingDesc);
 		NewApexDestructibleAsset = DestructibleMesh.FractureSettings->CreateApexDestructibleAsset(DestructibleAssetCookingDesc);
