Enums: L10N: Improper use of when to use the localized versus non-localized value for an enum name was causing potential asset corruption.

diff --git a/Engine/Source/Editor/AnimGraph/Private/AnimGraphNode_BlendListByEnum.cpp b/Engine/Source/Editor/AnimGraph/Private/AnimGraphNode_BlendListByEnum.cpp
--- a/Engine/Source/Editor/AnimGraph/Private/AnimGraphNode_BlendListByEnum.cpp
+++ b/Engine/Source/Editor/AnimGraph/Private/AnimGraphNode_BlendListByEnum.cpp
@@ -219,7 +219,7 @@ void UAnimGraphNode_BlendListByEnum::CustomizePinData(UEdGraphPin* Pin, FName So
 				const int32 EnumIndex = BoundEnum->FindEnumIndex(EnumElementName);
 				if (EnumIndex != INDEX_NONE)
 				{
-					Pin->PinFriendlyName = BoundEnum->GetEnumString(EnumIndex);
+					Pin->PinFriendlyName = BoundEnum->GetEnumText(EnumIndex).ToString();
 				}
 				else
 				{
diff --git a/Engine/Source/Editor/BehaviorTreeEditor/Private/DetailCustomizations/BlackboardDecoratorDetails.cpp b/Engine/Source/Editor/BehaviorTreeEditor/Private/DetailCustomizations/BlackboardDecoratorDetails.cpp
--- a/Engine/Source/Editor/BehaviorTreeEditor/Private/DetailCustomizations/BlackboardDecoratorDetails.cpp
+++ b/Engine/Source/Editor/BehaviorTreeEditor/Private/DetailCustomizations/BlackboardDecoratorDetails.cpp
@@ -144,7 +144,7 @@ void FBlackboardDecoratorDetails::OnKeyIDChanged()
 		{
 			for (int32 i = 0; i < CachedCustomObjectType->NumEnums() - 1; i++)
 			{
-				FString DisplayedName = CachedCustomObjectType->GetEnumString(i);
+				FString DisplayedName = CachedCustomObjectType->GetEnumText(i).ToString();
 				EnumPropValues.Add(DisplayedName);
 			}
 		}
diff --git a/Engine/Source/Editor/EnvironmentQueryEditor/Private/DetailCustomizations/EnvQueryTestDetails.cpp b/Engine/Source/Editor/EnvironmentQueryEditor/Private/DetailCustomizations/EnvQueryTestDetails.cpp
--- a/Engine/Source/Editor/EnvironmentQueryEditor/Private/DetailCustomizations/EnvQueryTestDetails.cpp
+++ b/Engine/Source/Editor/EnvironmentQueryEditor/Private/DetailCustomizations/EnvQueryTestDetails.cpp
@@ -96,19 +96,19 @@ void FEnvQueryTestDetails::BuildConditionValues()
 	check(TestConditionEnum);
 
 	ConditionValues.Reset();
-	ConditionValues.Add(FStringIntPair(TestConditionEnum->GetEnumString(EEnvTestCondition::NoCondition), EEnvTestCondition::NoCondition));
+	ConditionValues.Add(FStringIntPair(TestConditionEnum->GetEnumText(EEnvTestCondition::NoCondition).ToString(), EEnvTestCondition::NoCondition));
 
 	UEnvQueryTest* EditedTest = (UEnvQueryTest*)(MyTest.Get());
 	if (EditedTest)
 	{
 		if (EditedTest->bWorkOnFloatValues)
 		{
-			ConditionValues.Add(FStringIntPair(TestConditionEnum->GetEnumString(EEnvTestCondition::AtLeast), EEnvTestCondition::AtLeast));
-			ConditionValues.Add(FStringIntPair(TestConditionEnum->GetEnumString(EEnvTestCondition::UpTo), EEnvTestCondition::UpTo));
+			ConditionValues.Add(FStringIntPair(TestConditionEnum->GetEnumText(EEnvTestCondition::AtLeast).ToString(), EEnvTestCondition::AtLeast));
+			ConditionValues.Add(FStringIntPair(TestConditionEnum->GetEnumText(EEnvTestCondition::UpTo).ToString(), EEnvTestCondition::UpTo));
 		}
 		else
 		{
-			ConditionValues.Add(FStringIntPair(TestConditionEnum->GetEnumString(EEnvTestCondition::Match), EEnvTestCondition::Match));
+			ConditionValues.Add(FStringIntPair(TestConditionEnum->GetEnumText(EEnvTestCondition::Match).ToString(), EEnvTestCondition::Match));
 		}
 	}
 }
diff --git a/Engine/Source/Editor/EnvironmentQueryEditor/Private/DetailCustomizations/EnvTraceDataCustomization.cpp b/Engine/Source/Editor/EnvironmentQueryEditor/Private/DetailCustomizations/EnvTraceDataCustomization.cpp
--- a/Engine/Source/Editor/EnvironmentQueryEditor/Private/DetailCustomizations/EnvTraceDataCustomization.cpp
+++ b/Engine/Source/Editor/EnvironmentQueryEditor/Private/DetailCustomizations/EnvTraceDataCustomization.cpp
@@ -121,15 +121,15 @@ void FEnvTraceDataCustomization::CacheTraceModes(TSharedRef<class IPropertyHandl
 	TraceModes.Reset();
 	if (bCanDisable)
 	{
-		TraceModes.Add(FStringIntPair(TraceModeEnum->GetEnumString(EEnvQueryTrace::None), EEnvQueryTrace::None));
+		TraceModes.Add(FStringIntPair(TraceModeEnum->GetEnumText(EEnvQueryTrace::None).ToString(), EEnvQueryTrace::None));
 	}
 	if (bCanNavMesh)
 	{
-		TraceModes.Add(FStringIntPair(TraceModeEnum->GetEnumString(EEnvQueryTrace::Navigation), EEnvQueryTrace::Navigation));
+		TraceModes.Add(FStringIntPair(TraceModeEnum->GetEnumText(EEnvQueryTrace::Navigation).ToString(), EEnvQueryTrace::Navigation));
 	}
 	if (bCanGeometry)
 	{
-		TraceModes.Add(FStringIntPair(TraceModeEnum->GetEnumString(EEnvQueryTrace::Geometry), EEnvQueryTrace::Geometry));
+		TraceModes.Add(FStringIntPair(TraceModeEnum->GetEnumText(EEnvQueryTrace::Geometry).ToString(), EEnvQueryTrace::Geometry));
 	}
 
 	ActiveMode = EEnvQueryTrace::None;
diff --git a/Engine/Source/Editor/MainFrame/Private/Frame/MainFrameActions.cpp b/Engine/Source/Editor/MainFrame/Private/Frame/MainFrameActions.cpp
--- a/Engine/Source/Editor/MainFrame/Private/Frame/MainFrameActions.cpp
+++ b/Engine/Source/Editor/MainFrame/Private/Frame/MainFrameActions.cpp
@@ -604,7 +604,7 @@ void FMainFrameActionCallbacks::PackageProject( const FString InPlatformName, co
 		*ProjectPath,
 		*PackagingSettings->StagingDirectory.Path,
 		*PlatformName,
-		*FindObject<UEnum>(ANY_PACKAGE, TEXT("EProjectPackagingBuildConfigurations"))->GetEnumString(PackagingSettings->BuildConfiguration),
+		*FindObject<UEnum>(ANY_PACKAGE, TEXT("EProjectPackagingBuildConfigurations"))->GetEnumName(PackagingSettings->BuildConfiguration),
 		*ExecutableName,
 		*OptionalParams
 	);
diff --git a/Engine/Source/Editor/Matinee/Private/MatineeRecordMovie.cpp b/Engine/Source/Editor/Matinee/Private/MatineeRecordMovie.cpp
--- a/Engine/Source/Editor/Matinee/Private/MatineeRecordMovie.cpp
+++ b/Engine/Source/Editor/Matinee/Private/MatineeRecordMovie.cpp
@@ -265,7 +265,7 @@ void SMatineeRecordMovie::InitializeOptions()
 	UEnum* FormatEnum = FindObject<UEnum>(ANY_PACKAGE, TEXT("EMatineeCaptureType"), true);
 	for(int32 EnumMember = 0; EnumMember < FormatEnum->NumEnums() - 1; EnumMember++)
 	{
-		CaptureFormats.Add(MakeShareable(new FString(FormatEnum->GetEnumString(EnumMember))));
+		CaptureFormats.Add(MakeShareable(new FString(FormatEnum->GetEnumText(EnumMember).ToString())));
 	}	
 
 	Options.CaptureTypeIndex = FMath::Clamp<int32>(Options.CaptureTypeIndex, 0, FormatEnum->NumEnums() - 1);
@@ -599,7 +599,7 @@ void SMatineeRecordMovie::OnCaptureTypeSettingChanged(  TSharedPtr<FString> Chos
 	UEnum* FormatEnum = FindObject<UEnum>(ANY_PACKAGE, TEXT("EMatineeCaptureType"), true);
 	for(int32 EnumMember = 0; EnumMember < FormatEnum->NumEnums() - 1; EnumMember++)
 	{
-		if(*ChosenString.Get() == FormatEnum->GetEnumString(EnumMember))
+		if(*ChosenString.Get() == FormatEnum->GetEnumText(EnumMember).ToString())
 		{
 			CaptureType = (EMatineeCaptureType::Type)EnumMember;
 			break;
diff --git a/Engine/Source/Editor/PropertyEditor/Private/PropertyHandleImpl.cpp b/Engine/Source/Editor/PropertyEditor/Private/PropertyHandleImpl.cpp
--- a/Engine/Source/Editor/PropertyEditor/Private/PropertyHandleImpl.cpp
+++ b/Engine/Source/Editor/PropertyEditor/Private/PropertyHandleImpl.cpp
@@ -67,7 +67,7 @@ FPropertyAccess::Result FPropertyValueImpl::GetPropertyValueString( FString& Out
 				if ( ByteProperty != NULL && ByteProperty->Enum != NULL )
 				{
 					uint8 EnumValue = ByteProperty->GetPropertyValue(ValueAddress);
-					OutString = ByteProperty->Enum->GetEnumString(EnumValue);
+					OutString = ByteProperty->Enum->GetEnumName(EnumValue);
 				}
 			}
 			else
@@ -114,7 +114,7 @@ FPropertyAccess::Result FPropertyValueImpl::GetPropertyValueText( FText& OutText
 				if ( ByteProperty != NULL && ByteProperty->Enum != NULL )
 				{
 					uint8 EnumValueIndex = ByteProperty->GetPropertyValue(ValueAddress);
-					OutString = ByteProperty->Enum->GetEnumString(EnumValueIndex);
+					OutString = ByteProperty->Enum->GetEnumName(EnumValueIndex);
 				}
 
 				OutText = FText::FromString(OutString);
diff --git a/Engine/Source/Editor/UnrealEd/Private/Kismet2/EnumEditorUtils.cpp b/Engine/Source/Editor/UnrealEd/Private/Kismet2/EnumEditorUtils.cpp
--- a/Engine/Source/Editor/UnrealEd/Private/Kismet2/EnumEditorUtils.cpp
+++ b/Engine/Source/Editor/UnrealEd/Private/Kismet2/EnumEditorUtils.cpp
@@ -341,7 +341,7 @@ FString FEnumEditorUtils::GetEnumeratorDisplayName(const UUserDefinedEnum* Enum,
 {
 	if (Enum && Enum->DisplayNames.IsValidIndex(EnumeratorIndex))
 	{
-		return Enum->DisplayNames[EnumeratorIndex];
+		return Enum->DisplayNames[EnumeratorIndex].ToString();
 	}
 	return FString();
 }
@@ -389,16 +389,16 @@ void FEnumEditorUtils::EnsureAllDisplayNamesExist(UUserDefinedEnum* Enum)
 		Enum->DisplayNames.Empty(EnumeratorsToEnsure);
 		for	(int32 Index = 0; Index < EnumeratorsToEnsure; Index++)
 		{
-			FString DisplayNameMetaData = Enum->GetMetaData(FEnumEditorUtilsHelper::DisplayName(), Index);
+			FText DisplayNameMetaData = Enum->GetEnumText(Index);
 			if (DisplayNameMetaData.IsEmpty())
 			{
 				const FString EnumName = Enum->GetEnumName(Index);
-				DisplayNameMetaData = EnumName;
-				for(int32 AddIndex = 0; !IsEnumeratorDisplayNameValid(Enum, DisplayNameMetaData); ++AddIndex)
+				DisplayNameMetaData = FText::FromString(EnumName);
+				for(int32 AddIndex = 0; !IsEnumeratorDisplayNameValid(Enum, DisplayNameMetaData.ToString()); ++AddIndex)
 				{
-					DisplayNameMetaData = FString::Printf(TEXT("%s%d"), *EnumName, AddIndex);
+					DisplayNameMetaData = FText::FromString(FString::Printf(TEXT("%s%d"), *EnumName, AddIndex));
 				}
-				Enum->SetMetaData(FEnumEditorUtilsHelper::DisplayName(), *DisplayNameMetaData, Index);
+				Enum->SetMetaData(FEnumEditorUtilsHelper::DisplayName(), *DisplayNameMetaData.ToString(), Index);
 			}
 			Enum->DisplayNames.Add(DisplayNameMetaData);
 		}
diff --git a/Engine/Source/Runtime/CoreUObject/Public/UObject/Class.h b/Engine/Source/Runtime/CoreUObject/Public/UObject/Class.h
--- a/Engine/Source/Runtime/CoreUObject/Public/UObject/Class.h
+++ b/Engine/Source/Runtime/CoreUObject/Public/UObject/Class.h
@@ -1331,27 +1331,10 @@ public:
 		return FName(NAME_None).ToString();
 	}
 
-
-	/**
-	 * @return	The enum string at the specified index.
-	 */
-	COREUOBJECT_API virtual FString GetEnumString(int32 InIndex) const
-	{
-#if WITH_EDITOR
-		FString LocalizedDisplayName = GetDisplayNameText(InIndex).ToString();
-		if(!LocalizedDisplayName.IsEmpty())
-		{
-			return LocalizedDisplayName;
-		}
-#endif
-
-		return GetEnumName(InIndex);
-	}
-
 	/**
 	 * @return	The enum string at the specified index.
 	 */
-	FText GetEnumText(int32 InIndex) const
+	COREUOBJECT_API virtual FText GetEnumText(int32 InIndex) const
 	{
 #if WITH_EDITOR
 		//@todo These values should be properly localized [9/24/2013 justin.sargent]
diff --git a/Engine/Source/Runtime/Engine/Classes/Engine/UserDefinedEnum.h b/Engine/Source/Runtime/Engine/Classes/Engine/UserDefinedEnum.h
--- a/Engine/Source/Runtime/Engine/Classes/Engine/UserDefinedEnum.h
+++ b/Engine/Source/Runtime/Engine/Classes/Engine/UserDefinedEnum.h
@@ -31,7 +31,7 @@ class ENGINE_API UUserDefinedEnum : public UEnum
 	 * To sync DisplayNames with meta-data use FEnumEditorUtils::EnsureAllDisplayNamesExist.
 	 */
 	UPROPERTY()
-	TArray<FString> DisplayNames;
+	TArray<FText> DisplayNames;
 
 public:
 	/**
@@ -54,7 +54,7 @@ public:
 	/**
 	 * @return	The enum string at the specified index.
 	 */
-	virtual FString GetEnumString(int32 InIndex) const OVERRIDE;
+	virtual FText GetEnumText(int32 InIndex) const OVERRIDE;
 
 #if WITH_EDITOR
 	// Begin UObject interface
diff --git a/Engine/Source/Runtime/Engine/Private/AI/EnvironmentQuery/EnvQueryTypes.cpp b/Engine/Source/Runtime/Engine/Private/AI/EnvironmentQuery/EnvQueryTypes.cpp
--- a/Engine/Source/Runtime/Engine/Private/AI/EnvironmentQuery/EnvQueryTypes.cpp
+++ b/Engine/Source/Runtime/Engine/Private/AI/EnvironmentQuery/EnvQueryTypes.cpp
@@ -106,7 +106,7 @@ FString FEnvTraceData::ToString(FEnvTraceData::EDescriptionMode DescMode) const
 		if (DescMode == FEnvTraceData::Brief)
 		{
 			static UEnum* ChannelEnum = FindObject<UEnum>(ANY_PACKAGE, TEXT("ETraceTypeQuery"), true);
-			Desc += FString::Printf(TEXT(" %s on %s"), bCanProjectDown ? TEXT("projection") : TEXT("trace"), *ChannelEnum->GetEnumString(TraceChannel));
+			Desc += FString::Printf(TEXT(" %s on %s"), bCanProjectDown ? TEXT("projection") : TEXT("trace"), *ChannelEnum->GetEnumText(TraceChannel).ToString());
 		}
 		else
 		{
diff --git a/Engine/Source/Runtime/Engine/Private/AI/EnvironmentQuery/Tests/EnvQueryTest_Trace.cpp b/Engine/Source/Runtime/Engine/Private/AI/EnvironmentQuery/Tests/EnvQueryTest_Trace.cpp
--- a/Engine/Source/Runtime/Engine/Private/AI/EnvironmentQuery/Tests/EnvQueryTest_Trace.cpp
+++ b/Engine/Source/Runtime/Engine/Private/AI/EnvironmentQuery/Tests/EnvQueryTest_Trace.cpp
@@ -95,7 +95,7 @@ void UEnvQueryTest_Trace::RunTest(struct FEnvQueryInstance& QueryInstance)
 FString UEnvQueryTest_Trace::GetDescriptionTitle() const
 {
 	UEnum* ChannelEnum = FindObject<UEnum>(ANY_PACKAGE, TEXT("ETraceTypeQuery"), true);
-	FString ChannelDesc = ChannelEnum->GetEnumString(TraceData.TraceChannel);
+	FString ChannelDesc = ChannelEnum->GetEnumText(TraceData.TraceChannel).ToString();
 
 	FString DirectionDesc = TraceToItem.IsNamedParam() ?
 		FString::Printf(TEXT("%s, direction: %s"), *UEnvQueryTypes::DescribeContext(Context), *UEnvQueryTypes::DescribeBoolParam(TraceToItem)) :
diff --git a/Engine/Source/Runtime/Engine/Private/AI/Navigation/NavigationSystem.cpp b/Engine/Source/Runtime/Engine/Private/AI/Navigation/NavigationSystem.cpp
--- a/Engine/Source/Runtime/Engine/Private/AI/Navigation/NavigationSystem.cpp
+++ b/Engine/Source/Runtime/Engine/Private/AI/Navigation/NavigationSystem.cpp
@@ -1508,7 +1508,7 @@ void UNavigationSystem::DescribeFilterFlags(class UEnum* FlagsEnum) const
 	const int32 NumEnums = FMath::Min(16, FlagsEnum->NumEnums() - 1);	// skip _MAX
 	for (int32 i = 0; i < NumEnums; i++)
 	{
-		FlagDesc[i] = FlagsEnum->GetEnumString(i);
+		FlagDesc[i] = FlagsEnum->GetEnumText(i).ToString();
 	}
 
 	DescribeFilterFlags(FlagDesc);
diff --git a/Engine/Source/Runtime/Engine/Private/CheatManager.cpp b/Engine/Source/Runtime/Engine/Private/CheatManager.cpp
--- a/Engine/Source/Runtime/Engine/Private/CheatManager.cpp
+++ b/Engine/Source/Runtime/Engine/Private/CheatManager.cpp
@@ -1019,7 +1019,7 @@ void UCheatManager::ToggleGameplayDebugView(const FString& InViewName)
 		ViewNames.AddZeroed(EAIDebugDrawDataView::MAX);
 		for (int32 Index = 0; Index < EAIDebugDrawDataView::MAX; ++Index)
 		{
-			ViewNames[Index] = ViewlEnum->GetEnumString(Index);
+			ViewNames[Index] = ViewlEnum->GetEnumName(Index);
 		}
 	}
 
diff --git a/Engine/Source/Runtime/Engine/Private/DialogueVoice.cpp b/Engine/Source/Runtime/Engine/Private/DialogueVoice.cpp
--- a/Engine/Source/Runtime/Engine/Private/DialogueVoice.cpp
+++ b/Engine/Source/Runtime/Engine/Private/DialogueVoice.cpp
@@ -29,14 +29,14 @@ FString UDialogueVoice::GetDesc()
 	FString SummaryString;
 	{
 		UByteProperty* GenderProperty = CastChecked<UByteProperty>( FindFieldChecked<UProperty>( GetClass(), "Gender" ) );
-		SummaryString += GenderProperty->Enum->GetEnumString(Gender);
+		SummaryString += GenderProperty->Enum->GetEnumText(Gender).ToString();
 
 		if( Plurality != EGrammaticalNumber::Singular )
 		{
 			UByteProperty* PluralityProperty = CastChecked<UByteProperty>( FindFieldChecked<UProperty>( GetClass(), "Plurality" ) );
 
 			SummaryString += ", ";
-			SummaryString += PluralityProperty->Enum->GetEnumString(Plurality);
+			SummaryString += PluralityProperty->Enum->GetEnumText(Plurality).ToString();
 		}
 	}
 
diff --git a/Engine/Source/Runtime/Engine/Private/KismetNodeHelperLibrary.cpp b/Engine/Source/Runtime/Engine/Private/KismetNodeHelperLibrary.cpp
--- a/Engine/Source/Runtime/Engine/Private/KismetNodeHelperLibrary.cpp
+++ b/Engine/Source/Runtime/Engine/Private/KismetNodeHelperLibrary.cpp
@@ -151,7 +151,7 @@ FString UKismetNodeHelperLibrary::GetEnumeratorUserFriendlyName(const UEnum* Enu
 {
 	if (NULL != Enum)
 	{
-		return Enum->GetEnumString(EnumeratorIndex);
+		return Enum->GetEnumText(EnumeratorIndex).ToString();
 	}
 
 	return FName().ToString();
diff --git a/Engine/Source/Runtime/Engine/Private/Materials/MaterialExpressions.cpp b/Engine/Source/Runtime/Engine/Private/Materials/MaterialExpressions.cpp
--- a/Engine/Source/Runtime/Engine/Private/Materials/MaterialExpressions.cpp
+++ b/Engine/Source/Runtime/Engine/Private/Materials/MaterialExpressions.cpp
@@ -1100,8 +1100,8 @@ static bool VerifySamplerType(
 			UEnum* SamplerTypeEnum = FindObject<UEnum>( NULL, TEXT("/Script/Engine.EngineTypes.EMaterialSamplerType") );
 			check( SamplerTypeEnum );
 
-			FString SamplerTypeDisplayName = SamplerTypeEnum->GetEnumString(SamplerType);
-			FString TextureTypeDisplayName = SamplerTypeEnum->GetEnumString(CorrectSamplerType);
+			FString SamplerTypeDisplayName = SamplerTypeEnum->GetEnumText(SamplerType).ToString();
+			FString TextureTypeDisplayName = SamplerTypeEnum->GetEnumText(CorrectSamplerType).ToString();
 
 			Compiler->Errorf( TEXT("%s> Sampler type is %s, should be %s for %s"),
 				ExpressionDesc,
@@ -1115,7 +1115,7 @@ static bool VerifySamplerType(
 			UEnum* SamplerTypeEnum = FindObject<UEnum>( NULL, TEXT("/Script/Engine.EngineTypes.EMaterialSamplerType") );
 			check( SamplerTypeEnum );
 
-			FString SamplerTypeDisplayName = SamplerTypeEnum->GetEnumString(SamplerType);
+			FString SamplerTypeDisplayName = SamplerTypeEnum->GetEnumText(SamplerType).ToString();
 
 			Compiler->Errorf( TEXT("%s> To use '%s' as sampler type, SRGB must be disabled for %s"),
 				ExpressionDesc,
diff --git a/Engine/Source/Runtime/Engine/Private/SoundGroups.cpp b/Engine/Source/Runtime/Engine/Private/SoundGroups.cpp
--- a/Engine/Source/Runtime/Engine/Private/SoundGroups.cpp
+++ b/Engine/Source/Runtime/Engine/Private/SoundGroups.cpp
@@ -37,7 +37,7 @@ void USoundGroups::Initialize() const
 		{
 			if (SoundGroupEnum->HasMetaData(TEXT("Hidden"), It.Key))
 			{
-				UE_LOG(LogAudio, Warning, TEXT("Custom Game SoundGroup profile for %s defined but no display name supplied."), *SoundGroupEnum->GetEnumString(It.Key));
+				UE_LOG(LogAudio, Warning, TEXT("Custom Game SoundGroup profile for %s defined but no display name supplied."), *SoundGroupEnum->GetEnumText(It.Key).ToString());
 			}
 			SoundGroupEnum->RemoveMetaData(TEXT("Hidden"), It.Key);
 		}
@@ -58,7 +58,7 @@ const FSoundGroup& USoundGroups::GetSoundGroup(const ESoundGroup SoundGroup) con
 	{
 		UEnum* SoundGroupEnum = FindObjectChecked<UEnum>(NULL, TEXT("/Script/Engine.SoundGroups.ESoundGroup"));
 		
-		UE_LOG(LogAudio, Warning, TEXT("Requested SoundGroup %s does not have defined profile.  Using SOUNDGROUP_Default."), *SoundGroupEnum->GetEnumString(SoundGroup));
+		UE_LOG(LogAudio, Warning, TEXT("Requested SoundGroup %s does not have defined profile.  Using SOUNDGROUP_Default."), *SoundGroupEnum->GetEnumText(SoundGroup).ToString());
 		return SoundGroupMap.FindChecked(SOUNDGROUP_Default);
 	}
 
diff --git a/Engine/Source/Runtime/Engine/Private/UserDefinedEnum.cpp b/Engine/Source/Runtime/Engine/Private/UserDefinedEnum.cpp
--- a/Engine/Source/Runtime/Engine/Private/UserDefinedEnum.cpp
+++ b/Engine/Source/Runtime/Engine/Private/UserDefinedEnum.cpp
@@ -87,11 +87,12 @@ int32 UUserDefinedEnum::ResolveEnumerator(FArchive& Ar, int32 EnumeratorIndex) c
 #endif // WITH_EDITOR
 }
 
-FString UUserDefinedEnum::GetEnumString(int32 InIndex) const
+FText UUserDefinedEnum::GetEnumText(int32 InIndex) const
 {
 	if (DisplayNames.IsValidIndex(InIndex))
 	{
 		return DisplayNames[InIndex];
 	}
-	return Super::GetEnumString(InIndex);
+
+	return Super::GetEnumText(InIndex);
 }
\ No newline at end of file
