#TTP 323090 - LANDSCAPE SPLINES: Add support for exporting spline meshes when exporting landscape

#proj UE4

#summary Added support for landscape splines to FBX export

diff --git a/Engine/Source/Editor/LandscapeEditor/Private/LandscapeEdModeSplineTools.cpp b/Engine/Source/Editor/LandscapeEditor/Private/LandscapeEdModeSplineTools.cpp
--- a/Engine/Source/Editor/LandscapeEditor/Private/LandscapeEdModeSplineTools.cpp
+++ b/Engine/Source/Editor/LandscapeEditor/Private/LandscapeEdModeSplineTools.cpp
@@ -591,9 +591,7 @@ public:
 		SplinesComponent->Modify();
 		ControlPoint->Modify();
 
-		const ALandscapeProxy* Landscape = CastChecked<ALandscapeProxy>(SplinesComponent->GetOuter());
-		const FTransform LocalToWorld = Landscape->GetTransform();
-
+		const FTransform LocalToWorld = SplinesComponent->GetComponentToWorld();
 		const FVector Start = LocalToWorld.TransformPosition(ControlPoint->Location);
 		const FVector End = Start + FVector(0, 0, -HALF_WORLD_MAX);
 
diff --git a/Engine/Source/Editor/UnrealEd/Private/Fbx/FbxMainExport.cpp b/Engine/Source/Editor/UnrealEd/Private/Fbx/FbxMainExport.cpp
--- a/Engine/Source/Editor/UnrealEd/Private/Fbx/FbxMainExport.cpp
+++ b/Engine/Source/Editor/UnrealEd/Private/Fbx/FbxMainExport.cpp
@@ -827,7 +827,7 @@ void FFbxExporter::ExportLandscape(ALandscapeProxy* Actor, bool bSelectedOnly)
 
 	FString FbxNodeName = GetActorNodeName(Actor, InMatineeActor);
 
-	FbxNode* FbxActor = ExportActor(Actor, InMatineeActor);
+	FbxNode* FbxActor = ExportActor(Actor, InMatineeActor, true);
 	ExportLandscapeToFbx(Actor, *FbxNodeName, FbxActor, bSelectedOnly);
 }
 
@@ -1091,9 +1091,10 @@ FbxNode* FFbxExporter::ExportActor(AActor* Actor, AMatineeActor* InMatineeActor,
 					if( SceneComp != Actor->GetRootComponent() )
 					{
 						// Transform is relative to the root component
-						CompNode->LclTranslation.Set( Converter.ConvertToFbxPos( SceneComp->RelativeLocation ) );
-						CompNode->LclRotation.Set( Converter.ConvertToFbxRot( SceneComp->RelativeRotation.Euler() ) );
-						CompNode->LclScaling.Set( Converter.ConvertToFbxScale( SceneComp->RelativeScale3D ) );
+						const FTransform RelativeTransform = SceneComp->GetComponentToWorld().GetRelativeTransform(Actor->GetTransform());
+						CompNode->LclTranslation.Set(Converter.ConvertToFbxPos(RelativeTransform.GetTranslation()));
+						CompNode->LclRotation.Set(Converter.ConvertToFbxRot(RelativeTransform.GetRotation().Euler()));
+						CompNode->LclScaling.Set(Converter.ConvertToFbxScale(RelativeTransform.GetScale3D()));
 					}
 
 					ActorNode->AddChild(CompNode);
@@ -1104,17 +1105,24 @@ FbxNode* FFbxExporter::ExportActor(AActor* Actor, AMatineeActor* InMatineeActor,
 				UStaticMeshComponent* StaticMeshComp = Cast<UStaticMeshComponent>( Component );
 				USkeletalMeshComponent* SkelMeshComp = Cast<USkeletalMeshComponent>( Component );
 
-				if( StaticMeshComp && StaticMeshComp->StaticMesh )
+				if (StaticMeshComp && StaticMeshComp->StaticMesh)
 				{
 					int32 LODIndex = StaticMeshComp->ForcedLodModel;
 					FStaticMeshLODResources& RenderMesh = StaticMeshComp->StaticMesh->GetLODForExport(LODIndex);
 
-					ExportStaticMeshToFbx( StaticMeshComp->StaticMesh, RenderMesh, *StaticMeshComp->GetName(), ExportNode );
+					if (USplineMeshComponent* SplineMeshComp = Cast<USplineMeshComponent>(StaticMeshComp))
+					{
+						ExportSplineMeshToFbx(SplineMeshComp, RenderMesh, *SplineMeshComp->GetName(), ExportNode);
+					}
+					else
+					{
+						ExportStaticMeshToFbx(StaticMeshComp->StaticMesh, RenderMesh, *StaticMeshComp->GetName(), ExportNode);
+					}
 				}
-				else if( SkelMeshComp && SkelMeshComp->SkeletalMesh )
+				else if (SkelMeshComp && SkelMeshComp->SkeletalMesh)
 				{
 					ExportSkeletalMeshToFbx( *SkelMeshComp->SkeletalMesh, *SkelMeshComp->GetName(), ExportNode );
-				}			
+				}
 			}
 		}
 		
@@ -1708,12 +1716,12 @@ FbxNode* FFbxExporter::ExportStaticMeshToFbx(UStaticMesh* StaticMesh, FStaticMes
 	// Verify the integrity of the static mesh.
 	if (RenderMesh.VertexBuffer.GetNumVertices() == 0)
 	{
-			return NULL;
+		return NULL;
 	}
 
 	if (RenderMesh.Sections.Num() == 0)
 	{
-			return NULL;
+		return NULL;
 	}
 
 	// Remaps an Unreal vert to final reduced vertex list
@@ -1762,17 +1770,17 @@ FbxNode* FFbxExporter::ExportStaticMeshToFbx(UStaticMesh* StaticMesh, FStaticMes
 	}
 
 	// Build list of Indices re-used multiple times to lookup Normals, UVs, other per face vertex information
-	TArray<int32> Indices;
+	TArray<uint32> Indices;
 	for (int32 PolygonsIndex = 0; PolygonsIndex < PolygonsCount; ++PolygonsIndex)
 	{
 		FIndexArrayView RawIndices = RenderMesh.IndexBuffer.GetArrayView();
 		FStaticMeshSection& Polygons = RenderMesh.Sections[PolygonsIndex];
-		const int32 TriangleCount = Polygons.NumTriangles;
-		for (int32 TriangleIndex = 0; TriangleIndex < TriangleCount; ++TriangleIndex)
+		const uint32 TriangleCount = Polygons.NumTriangles;
+		for (uint32 TriangleIndex = 0; TriangleIndex < TriangleCount; ++TriangleIndex)
 		{
-			for (int32 PointIndex = 0; PointIndex < 3; PointIndex++)
+			for (uint32 PointIndex = 0; PointIndex < 3; PointIndex++)
 			{
-				int32 UnrealVertIndex = RawIndices[Polygons.FirstIndex + ((TriangleIndex * 3) + PointIndex)];
+				uint32 UnrealVertIndex = RawIndices[Polygons.FirstIndex + ((TriangleIndex * 3) + PointIndex)];
 				Indices.Add(UnrealVertIndex);
 			}
 		}
@@ -1801,7 +1809,7 @@ FbxNode* FFbxExporter::ExportStaticMeshToFbx(UStaticMesh* StaticMesh, FStaticMes
 	// Add one normal per each face index (3 per triangle)
 	for (int32 i=0; i < Indices.Num(); i++)
 	{
-		int32 UnrealVertIndex = Indices[i];
+		uint32 UnrealVertIndex = Indices[i];
 		LayerElementNormal->GetDirectArray().Add( FbxNormals[UnrealVertIndex] );
 	}
 	Layer->SetNormals(LayerElementNormal);
@@ -1825,8 +1833,11 @@ FbxNode* FFbxExporter::ExportStaticMeshToFbx(UStaticMesh* StaticMesh, FStaticMes
 		{
 			FCString::Sprintf(UVChannelName, TEXT("LightMapUV"));
 		}
-		
-		
+		else
+		{
+			FCString::Sprintf(UVChannelName, TEXT(""));
+		}
+
 		FbxLayerElementUV* UVDiffuseLayer = FbxLayerElementUV::Create(Mesh, TCHAR_TO_ANSI(UVChannelName));
 
 		// Note: when eINDEX_TO_DIRECT is used, IndexArray must be 3xTriangle count, DirectArray can be smaller
@@ -1859,7 +1870,7 @@ FbxNode* FFbxExporter::ExportStaticMeshToFbx(UStaticMesh* StaticMesh, FStaticMes
 		UVDiffuseLayer->GetIndexArray().SetCount( Indices.Num() );
 		for (int32 i=0; i < Indices.Num(); i++)
 		{
-			int32 UnrealVertIndex = Indices[i];
+			uint32 UnrealVertIndex = Indices[i];
 			int32 NewVertIndex = UvsRemap[UnrealVertIndex];
 			UVDiffuseLayer->GetIndexArray().SetAt(i,NewVertIndex);
 		}
@@ -1873,7 +1884,7 @@ FbxNode* FFbxExporter::ExportStaticMeshToFbx(UStaticMesh* StaticMesh, FStaticMes
 	Layer->SetMaterials(MatLayer);
 	
 	// Keep track of the number of tri's we export
-	int32 AccountedTriangles = 0;
+	uint32 AccountedTriangles = 0;
 	for (int32 PolygonsIndex = 0; PolygonsIndex < PolygonsCount; ++PolygonsIndex)
 	{
 		FStaticMeshSection& Polygons = RenderMesh.Sections[PolygonsIndex];
@@ -1896,15 +1907,15 @@ FbxNode* FFbxExporter::ExportStaticMeshToFbx(UStaticMesh* StaticMesh, FStaticMes
 		}
 		// Static meshes contain one triangle list per element.
 		// [GLAFORTE] Could it occasionally contain triangle strips? How do I know?
-		int32 TriangleCount = Polygons.NumTriangles;
+		uint32 TriangleCount = Polygons.NumTriangles;
 		
 		// Copy over the index buffer into the FBX polygons set.
-		for (int32 TriangleIndex = 0; TriangleIndex < TriangleCount; ++TriangleIndex)
+		for (uint32 TriangleIndex = 0; TriangleIndex < TriangleCount; ++TriangleIndex)
 		{
 			Mesh->BeginPolygon(ActualIndex);
-			for (int32 PointIndex = 0; PointIndex < 3; PointIndex++)
+			for (uint32 PointIndex = 0; PointIndex < 3; PointIndex++)
 			{
-				int32 OriginalUnrealVertIndex = RawIndices[Polygons.FirstIndex + ((TriangleIndex * 3) + PointIndex)];
+				uint32 OriginalUnrealVertIndex = RawIndices[Polygons.FirstIndex + ((TriangleIndex * 3) + PointIndex)];
 				int32 RemappedVertIndex = VertRemap[OriginalUnrealVertIndex];
 				Mesh->AddPolygon( RemappedVertIndex );
 			}
@@ -1945,7 +1956,7 @@ FbxNode* FFbxExporter::ExportStaticMeshToFbx(UStaticMesh* StaticMesh, FStaticMes
 
 	// Create and fill in the vertex color data source.
 	FColorVertexBuffer* ColorBufferToUse = ColorBuffer? ColorBuffer: &RenderMesh.ColorVertexBuffer;
-	int32 ColorVertexCount = ColorBufferToUse->GetNumVertices();
+	uint32 ColorVertexCount = ColorBufferToUse->GetNumVertices();
 	
 	// Only export vertex colors if they exist
 	if (ColorVertexCount > 0)
@@ -1959,7 +1970,7 @@ FbxNode* FFbxExporter::ExportStaticMeshToFbx(UStaticMesh* StaticMesh, FStaticMes
 		for (int32 i=0; i < Indices.Num(); i++)
 		{
 			FLinearColor VertColor(1.0f, 1.0f, 1.0f);
-			int32 UnrealVertIndex = Indices[i];
+			uint32 UnrealVertIndex = Indices[i];
 			if (UnrealVertIndex < ColorVertexCount)
 			{
 				VertColor = ColorBufferToUse->VertexColor(UnrealVertIndex).ReinterpretAsLinear();
@@ -1980,6 +1991,285 @@ FbxNode* FFbxExporter::ExportStaticMeshToFbx(UStaticMesh* StaticMesh, FStaticMes
 	return FbxActor;
 }
 
+static float& GetAxisValue(FVector& InVector, ESplineMeshAxis::Type InAxis)
+{
+	switch (InAxis)
+	{
+	case ESplineMeshAxis::X:
+		return InVector.X;
+	case ESplineMeshAxis::Y:
+		return InVector.Y;
+	case ESplineMeshAxis::Z:
+		return InVector.Z;
+	default:
+		check(0);
+		return InVector.Z;
+	}
+}
+
+FbxNode* FFbxExporter::ExportSplineMeshToFbx(USplineMeshComponent* SplineMeshComp, FStaticMeshLODResources& RenderMesh, const TCHAR* MeshName, FbxNode* FbxActor)
+{
+	const UStaticMesh* StaticMesh = SplineMeshComp->StaticMesh;
+	check(StaticMesh);
+
+	// Verify the integrity of the static mesh.
+	if (RenderMesh.VertexBuffer.GetNumVertices() == 0)
+	{
+		return NULL;
+	}
+
+	if (RenderMesh.Sections.Num() == 0)
+	{
+		return NULL;
+	}
+
+	// Remaps an Unreal vert to final reduced vertex list
+	TArray<int32> VertRemap;
+	TArray<int32> UniqueVerts;
+
+	if (bStaticMeshExportUnWeldedVerts == false)
+	{
+		// Weld verts
+		DetermineVertsToWeld(VertRemap, UniqueVerts, RenderMesh);
+	}
+	else
+	{
+		// Do not weld verts
+		VertRemap.Add(RenderMesh.VertexBuffer.GetNumVertices());
+		for (int32 i = 0; i < VertRemap.Num(); i++)
+		{
+			VertRemap[i] = i;
+		}
+		UniqueVerts = VertRemap;
+	}
+
+	FbxMesh* Mesh = FbxMesh::Create(Scene, TCHAR_TO_ANSI(MeshName));
+
+	// Create and fill in the vertex position data source.
+	// The position vertices are duplicated, for some reason, retrieve only the first half vertices.
+	const int32 VertexCount = VertRemap.Num();
+	const int32 PolygonsCount = RenderMesh.Sections.Num();
+
+	Mesh->InitControlPoints(UniqueVerts.Num());
+
+	FbxVector4* ControlPoints = Mesh->GetControlPoints();
+	for (int32 PosIndex = 0; PosIndex < UniqueVerts.Num(); ++PosIndex)
+	{
+		int32 UnrealPosIndex = UniqueVerts[PosIndex];
+		FVector Position = RenderMesh.PositionVertexBuffer.VertexPosition(UnrealPosIndex);
+
+		const FTransform SliceTransform = SplineMeshComp->CalcSliceTransform(GetAxisValue(Position, SplineMeshComp->ForwardAxis));
+		GetAxisValue(Position, SplineMeshComp->ForwardAxis) = 0;
+		Position = SliceTransform.TransformPosition(Position);
+
+		ControlPoints[PosIndex] = FbxVector4(Position.X, -Position.Y, Position.Z);
+	}
+
+	// Set the normals on Layer 0.
+	FbxLayer* Layer = Mesh->GetLayer(0);
+	if (Layer == NULL)
+	{
+		Mesh->CreateLayer();
+		Layer = Mesh->GetLayer(0);
+	}
+
+	// Build list of Indices re-used multiple times to lookup Normals, UVs, other per face vertex information
+	TArray<uint32> Indices;
+	for (int32 PolygonsIndex = 0; PolygonsIndex < PolygonsCount; ++PolygonsIndex)
+	{
+		FIndexArrayView RawIndices = RenderMesh.IndexBuffer.GetArrayView();
+		FStaticMeshSection& Polygons = RenderMesh.Sections[PolygonsIndex];
+		const uint32 TriangleCount = Polygons.NumTriangles;
+		for (uint32 TriangleIndex = 0; TriangleIndex < TriangleCount; ++TriangleIndex)
+		{
+			for (uint32 PointIndex = 0; PointIndex < 3; PointIndex++)
+			{
+				uint32 UnrealVertIndex = RawIndices[Polygons.FirstIndex + ((TriangleIndex * 3) + PointIndex)];
+				Indices.Add(UnrealVertIndex);
+			}
+		}
+	}
+
+	// Create and fill in the per-face-vertex normal data source.
+	// We extract the Z-tangent and drop the X/Y-tangents which are also stored in the render mesh.
+	FbxLayerElementNormal* LayerElementNormal = FbxLayerElementNormal::Create(Mesh, "");
+
+	// Set 3 normals per triangle instead of storing normals on positional control points
+	LayerElementNormal->SetMappingMode(FbxLayerElement::eByPolygonVertex);
+
+	// Set the normal values for every polygon vertex.
+	LayerElementNormal->SetReferenceMode(FbxLayerElement::eDirect);
+
+	TArray<FbxVector4> FbxNormals;
+	FbxNormals.AddUninitialized(VertexCount);
+	for (int32 VertIndex = 0; VertIndex < VertexCount; ++VertIndex)
+	{
+		FVector Position = RenderMesh.PositionVertexBuffer.VertexPosition(VertIndex);
+		const FTransform SliceTransform = SplineMeshComp->CalcSliceTransform(GetAxisValue(Position, SplineMeshComp->ForwardAxis));
+		FVector Normal = FVector(RenderMesh.VertexBuffer.VertexTangentZ(VertIndex));
+		Normal = SliceTransform.TransformVector(Normal);
+		FbxVector4& FbxNormal = FbxNormals[VertIndex];
+		FbxNormal = FbxVector4(Normal.X, -Normal.Y, Normal.Z);
+		FbxNormal.Normalize();
+	}
+
+	// Add one normal per each face index (3 per triangle)
+	for (uint32 UnrealVertIndex : Indices)
+	{
+		LayerElementNormal->GetDirectArray().Add(FbxNormals[UnrealVertIndex]);
+	}
+	Layer->SetNormals(LayerElementNormal);
+	FbxNormals.Empty();
+
+	// Create and fill in the per-face-vertex texture coordinate data source(s).
+	// Create UV for Diffuse channel.
+	int32 TexCoordSourceCount = RenderMesh.VertexBuffer.GetNumTexCoords();
+	TCHAR UVChannelName[32] = { 0 };
+	for (int32 TexCoordSourceIndex = 0; TexCoordSourceIndex < TexCoordSourceCount; ++TexCoordSourceIndex)
+	{
+		FbxLayer* UVsLayer = Mesh->GetLayer(TexCoordSourceIndex);
+		if (UVsLayer == NULL)
+		{
+			Mesh->CreateLayer();
+			UVsLayer =  Mesh->GetLayer(TexCoordSourceIndex);
+		}
+
+		if (TexCoordSourceIndex == 1)
+		{
+			FCString::Sprintf(UVChannelName, TEXT("LightMapUV"));
+		}
+		else
+		{
+			FCString::Sprintf(UVChannelName, TEXT(""));
+		}
+
+		FbxLayerElementUV* UVDiffuseLayer = FbxLayerElementUV::Create(Mesh, TCHAR_TO_ANSI(UVChannelName));
+
+		// Note: when eINDEX_TO_DIRECT is used, IndexArray must be 3xTriangle count, DirectArray can be smaller
+		UVDiffuseLayer->SetMappingMode(FbxLayerElement::eByPolygonVertex);
+		UVDiffuseLayer->SetReferenceMode(FbxLayerElement::eIndexToDirect);
+
+		TArray<int32> UvsRemap;
+		TArray<int32> UniqueUVs;
+		if (bStaticMeshExportUnWeldedVerts == false)
+		{
+			// Weld UVs
+			DetermineUVsToWeld(UvsRemap, UniqueUVs, RenderMesh.VertexBuffer, TexCoordSourceIndex);
+		}
+		else
+		{
+			// Do not weld UVs
+			UvsRemap = VertRemap;
+			UniqueUVs = UvsRemap;
+		}
+
+		// Create the texture coordinate data source.
+		for (int32 UnrealVertIndex : UniqueUVs)
+		{
+			const FVector2D& TexCoord = RenderMesh.VertexBuffer.GetVertexUV(UnrealVertIndex, TexCoordSourceIndex);
+			UVDiffuseLayer->GetDirectArray().Add(FbxVector2(TexCoord.X, -TexCoord.Y + 1.0));
+		}
+
+		// For each face index, point to a texture uv
+		UVDiffuseLayer->GetIndexArray().SetCount(Indices.Num());
+		for (int32 i = 0; i < Indices.Num(); i++)
+		{
+			uint32 UnrealVertIndex = Indices[i];
+			int32 NewVertIndex = UvsRemap[UnrealVertIndex];
+			UVDiffuseLayer->GetIndexArray().SetAt(i, NewVertIndex);
+		}
+
+		UVsLayer->SetUVs(UVDiffuseLayer, FbxLayerElement::eTextureDiffuse);
+	}
+
+	FbxLayerElementMaterial* MatLayer = FbxLayerElementMaterial::Create(Mesh, "");
+	MatLayer->SetMappingMode(FbxLayerElement::eByPolygon);
+	MatLayer->SetReferenceMode(FbxLayerElement::eIndexToDirect);
+	Layer->SetMaterials(MatLayer);
+
+	for (int32 PolygonsIndex = 0; PolygonsIndex < PolygonsCount; ++PolygonsIndex)
+	{
+		FStaticMeshSection& Polygons = RenderMesh.Sections[PolygonsIndex];
+		FIndexArrayView RawIndices = RenderMesh.IndexBuffer.GetArrayView();
+		UMaterialInterface* Material = StaticMesh->GetMaterial(Polygons.MaterialIndex);
+
+		FbxSurfaceMaterial* FbxMaterial = Material ? ExportMaterial(Material->GetMaterial()) : NULL;
+		if (!FbxMaterial)
+		{
+			FbxMaterial = CreateDefaultMaterial();
+		}
+		int32 MatIndex = FbxActor->AddMaterial(FbxMaterial);
+
+		// Static meshes contain one triangle list per element.
+		uint32 TriangleCount = Polygons.NumTriangles;
+
+		// Copy over the index buffer into the FBX polygons set.
+		for (uint32 TriangleIndex = 0; TriangleIndex < TriangleCount; ++TriangleIndex)
+		{
+			Mesh->BeginPolygon(MatIndex);
+			for (uint32 PointIndex = 0; PointIndex < 3; PointIndex++)
+			{
+				uint32 OriginalUnrealVertIndex = RawIndices[Polygons.FirstIndex + ((TriangleIndex * 3) + PointIndex)];
+				int32 RemappedVertIndex = VertRemap[OriginalUnrealVertIndex];
+				Mesh->AddPolygon(RemappedVertIndex);
+			}
+			Mesh->EndPolygon();
+		}
+	}
+
+#if TODO_FBX
+	// This is broken. We are exporting the render mesh but providing smoothing
+	// information from the source mesh. The render triangles are not in the
+	// same order. Therefore we should export the raw mesh or not export
+	// smoothing group information!
+	int32 TriangleCount = RenderMesh.RawTriangles.GetElementCount();
+	FStaticMeshTriangle* RawTriangleData = (FStaticMeshTriangle*)RenderMesh.RawTriangles.Lock(LOCK_READ_ONLY);
+	for (int32 TriangleIndex = 0; TriangleIndex < TriangleCount; TriangleIndex++)
+	{
+		FStaticMeshTriangle* Triangle = (RawTriangleData++);
+
+		SmoothingArray.Add(Triangle->SmoothingMask);
+	}
+	RenderMesh.RawTriangles.Unlock();
+#endif // #if TODO_FBX
+
+	// Create and fill in the vertex color data source.
+	FColorVertexBuffer* ColorBufferToUse = &RenderMesh.ColorVertexBuffer;
+	uint32 ColorVertexCount = ColorBufferToUse->GetNumVertices();
+
+	// Only export vertex colors if they exist
+	if (ColorVertexCount > 0)
+	{
+		FbxLayerElementVertexColor* VertexColor = FbxLayerElementVertexColor::Create(Mesh, "");
+		VertexColor->SetMappingMode(FbxLayerElement::eByPolygonVertex);
+		VertexColor->SetReferenceMode(FbxLayerElement::eIndexToDirect);
+		FbxLayerElementArrayTemplate<FbxColor>& VertexColorArray = VertexColor->GetDirectArray();
+		Layer->SetVertexColors(VertexColor);
+
+		for (int32 i = 0; i < Indices.Num(); i++)
+		{
+			FLinearColor VertColor(1.0f, 1.0f, 1.0f);
+			uint32 UnrealVertIndex = Indices[i];
+			if (UnrealVertIndex < ColorVertexCount)
+			{
+				VertColor = ColorBufferToUse->VertexColor(UnrealVertIndex).ReinterpretAsLinear();
+			}
+
+			VertexColorArray.Add(FbxColor(VertColor.R, VertColor.G, VertColor.B, VertColor.A));
+		}
+
+		VertexColor->GetIndexArray().SetCount(Indices.Num());
+		for (int32 i = 0; i < Indices.Num(); i++)
+		{
+			VertexColor->GetIndexArray().SetAt(i, i);
+		}
+	}
+
+	FbxActor->SetNodeAttribute(Mesh);
+
+	return FbxActor;
+}
+
 /**
  * Exports a Landscape
  */
@@ -2167,4 +2457,4 @@ FbxNode* FFbxExporter::ExportLandscapeToFbx(ALandscapeProxy* Landscape, const TC
 }
 
 
-} // namespace UnFbx
+} // namespace UnFbx
\ No newline at end of file
diff --git a/Engine/Source/Editor/UnrealEd/Private/FbxExporter.h b/Engine/Source/Editor/UnrealEd/Private/FbxExporter.h
--- a/Engine/Source/Editor/UnrealEd/Private/FbxExporter.h
+++ b/Engine/Source/Editor/UnrealEd/Private/FbxExporter.h
@@ -183,6 +183,7 @@ private:
 	
 	/**
 	 * Exports a static mesh
+	 * @param StaticMesh	The static mesh to export
 	 * @param RenderMesh	The static mesh render data to export
 	 * @param MeshName		The name of the mesh for the FBX file
 	 * @param FbxActor		The fbx node representing the mesh
@@ -192,6 +193,21 @@ private:
 	 */
 	FbxNode* ExportStaticMeshToFbx(UStaticMesh* StaticMesh, FStaticMeshLODResources& RenderMesh, const TCHAR* MeshName, FbxNode* FbxActor, int32 LightmapUVChannel = -1, FColorVertexBuffer* ColorBuffer = NULL, const TArray<UMaterialInterface*>* MaterialOrderOverride = NULL );
 
+	/**
+	* Exports a spline mesh
+	* @param SplineMeshComp	The spline mesh component to export
+	* @param RenderMesh		The mesh render data to export
+	* @param MeshName		The name of the mesh for the FBX file
+	* @param FbxActor		The fbx node representing the mesh
+	*/
+	FbxNode* ExportSplineMeshToFbx(USplineMeshComponent* SplineMeshComp, FStaticMeshLODResources& RenderMesh, const TCHAR* MeshName, FbxNode* FbxActor);
+
+	/**
+	* Exports a landscape
+	* @param Landscape		The landscape to export
+	* @param MeshName		The name of the mesh for the FBX file
+	* @param FbxActor		The fbx node representing the mesh
+	*/
 	FbxNode* ExportLandscapeToFbx(ALandscapeProxy* Landscape, const TCHAR* MeshName, FbxNode* FbxActor, bool bSelectedOnly);
 
 	/**
diff --git a/Engine/Source/Runtime/Engine/Classes/Components/SplineMeshComponent.h b/Engine/Source/Runtime/Engine/Classes/Components/SplineMeshComponent.h
--- a/Engine/Source/Runtime/Engine/Classes/Components/SplineMeshComponent.h
+++ b/Engine/Source/Runtime/Engine/Classes/Components/SplineMeshComponent.h
@@ -158,7 +158,7 @@ class USplineMeshComponent : public UStaticMeshComponent, public IInterface_Coll
 	 * Calculates the spline transform, including roll, scale, and offset along the spline at a specified distance
 	 * @Note:  This is mirrored to Lightmass::CalcSliceTransform() and LocalVertexShader.usf.  If you update one of these, please update them all! 
 	 */
-	FTransform CalcSliceTransform(const float DistanceAlong) const;
+	ENGINE_API FTransform CalcSliceTransform(const float DistanceAlong) const;
 };
 
 
