#TTP 299847 - ROCKET: BIG: K2: USABILITY: Ability to go to C++ code for native functions/variables in BP editor
#Branch UE4
#Proj  BlueprintGraph, GraphEditor, Kismet, UnrealEd

#Summary Added code to navigate the users IDE to native functions and variables. Added options in the graph and myblueprint node context menus to navigate to the native code they represent.

#Add Added function FSourceCodeNavigation::NavigateToProperty to locate and open the header file where the variable is defined.
#Add Added function FSourceCodeNavigation::NavigateToFunction to locate the source code from the symbolmap and open the IDE at the definition location.
#Change Moved the code used to determine the correct modulename from FSourceCodeNavigation::GatherFunctionsForActors to FSourceCodeNavigation::FindClassModuleName so other functions didn't have to replicate the code.
#Change added Goto Code Definition for native variable nodes and native call_function nodes in the blueprint graph node context menu.
#Change added Goto Code Definition for native variables in the myblueprint node context menu.
#Add Added GotoNativeFunctionDefinition, IsSelectionNativeFunction, GotoNativeVariableDefinition and IsSelectionNativeVariable to FBlueprintEditor so the context menu actions can call these functions to goto source code.

diff --git a/Engine/Source/Editor/BlueprintGraph/Private/EdGraphSchema_K2.cpp b/Engine/Source/Editor/BlueprintGraph/Private/EdGraphSchema_K2.cpp
--- a/Engine/Source/Editor/BlueprintGraph/Private/EdGraphSchema_K2.cpp
+++ b/Engine/Source/Editor/BlueprintGraph/Private/EdGraphSchema_K2.cpp
@@ -766,10 +766,17 @@ void UEdGraphSchema_K2::GetContextMenuActions(const UEdGraph* CurrentGraph, cons
 						MenuBuilder->AddMenuEntry(FGraphEditorCommands::Get().AssignReferencedActor);
 					}
 
-					// show search for references for variable nodes
+					// Add the goto source code action for native functions
+					if (InGraphNode->IsA(UK2Node_CallFunction::StaticClass()))
+					{
+						MenuBuilder->AddMenuEntry(FGraphEditorCommands::Get().GotoNativeFunctionDefinition);
+					}
+
+					// show search for references for variable nodes and goto source code action
 					if (InGraphNode->IsA(UK2Node_Variable::StaticClass()))
 					{
 						MenuBuilder->AddMenuEntry(FGraphEditorCommands::Get().FindVariableReferences);
+						MenuBuilder->AddMenuEntry(FGraphEditorCommands::Get().GotoNativeVariableDefinition);
 					}
 
 					MenuBuilder->AddMenuEntry(FGenericCommands::Get().Rename, NAME_None, LOCTEXT("Rename", "Rename"), LOCTEXT("Rename_Tooltip", "Renames selected function or variable in blueprint.") );
diff --git a/Engine/Source/Editor/GraphEditor/Private/GraphEditorActions.cpp b/Engine/Source/Editor/GraphEditor/Private/GraphEditorActions.cpp
--- a/Engine/Source/Editor/GraphEditor/Private/GraphEditorActions.cpp
+++ b/Engine/Source/Editor/GraphEditor/Private/GraphEditorActions.cpp
@@ -40,6 +40,9 @@ void FGraphEditorCommandsImpl::RegisterCommands()
 	UI_COMMAND( AssignReferencedActor, "Assign selected Actor", "Assign the selected actor to be this node's referenced object", EUserInterfaceActionType::Button, FInputGesture() )
 	UI_COMMAND( FindVariableReferences, "Find Variable References", "Find references of this variable", EUserInterfaceActionType::Button, FInputGesture() )
 	
+	UI_COMMAND( GotoNativeFunctionDefinition, "Goto Code Definition", "Goto the native code definition of this function", EUserInterfaceActionType::Button, FInputGesture() )
+	UI_COMMAND( GotoNativeVariableDefinition, "Goto Code Definition", "Goto the native code definition of this variable", EUserInterfaceActionType::Button, FInputGesture() )
+
 	UI_COMMAND( BreakPinLinks, "Break Link(s)", "Breaks pin links", EUserInterfaceActionType::Button, FInputGesture() )
 	UI_COMMAND( PromoteToVariable, "Promote to Variable", "Promotes something to a variable", EUserInterfaceActionType::Button, FInputGesture() )
 	UI_COMMAND( StartWatchingPin, "Watch this value", "Adds this pin or variable to the watch list", EUserInterfaceActionType::Button, FInputGesture() )
diff --git a/Engine/Source/Editor/GraphEditor/Public/GraphEditorActions.h b/Engine/Source/Editor/GraphEditor/Public/GraphEditorActions.h
--- a/Engine/Source/Editor/GraphEditor/Public/GraphEditorActions.h
+++ b/Engine/Source/Editor/GraphEditor/Public/GraphEditorActions.h
@@ -61,6 +61,10 @@ public:
 	// Find variable references
 	TSharedPtr< FUICommandInfo > FindVariableReferences;
 
+	// Goto native code actions
+	TSharedPtr< FUICommandInfo > GotoNativeFunctionDefinition;
+	TSharedPtr< FUICommandInfo > GotoNativeVariableDefinition;
+
 	// Pin-specific actions
 	TSharedPtr< FUICommandInfo > BreakPinLinks;
 	TSharedPtr< FUICommandInfo > PromoteToVariable;
diff --git a/Engine/Source/Editor/Kismet/Private/BlueprintEditor.cpp b/Engine/Source/Editor/Kismet/Private/BlueprintEditor.cpp
--- a/Engine/Source/Editor/Kismet/Private/BlueprintEditor.cpp
+++ b/Engine/Source/Editor/Kismet/Private/BlueprintEditor.cpp
@@ -627,6 +627,16 @@ TSharedRef<SGraphEditor> FBlueprintEditor::CreateGraphEditorWidget(TSharedRef<FT
 				FExecuteAction::CreateSP( this, &FBlueprintEditor::OnFindVariableReferences ),
 				FCanExecuteAction::CreateSP( this, &FBlueprintEditor::CanFindVariableReferences )
 				);
+
+			GraphEditorCommands->MapAction( FGraphEditorCommands::Get().GotoNativeFunctionDefinition,
+				FExecuteAction::CreateSP( this, &FBlueprintEditor::GotoNativeFunctionDefinition ),
+				FCanExecuteAction::CreateSP(this, &FBlueprintEditor::IsSelectionNativeFunction)
+				);
+
+			GraphEditorCommands->MapAction( FGraphEditorCommands::Get().GotoNativeVariableDefinition,
+				FExecuteAction::CreateSP( this, &FBlueprintEditor::GotoNativeVariableDefinition ),
+				FCanExecuteAction::CreateSP(this, &FBlueprintEditor::IsSelectionNativeVariable)
+				);
 		}
 	}
 
@@ -5864,6 +5874,93 @@ bool FBlueprintEditor::CanFindVariableReferences()
 	return false;
 }
 
+void FBlueprintEditor::GotoNativeFunctionDefinition()
+{
+	auto GraphEditor = FocusedGraphEdPtr.Pin();
+	if( GraphEditor.IsValid() && IsSelectionNativeFunction() )
+	{
+		const FGraphPanelSelectionSet SelectedNodes = GraphEditor->GetSelectedNodes();
+		FGraphPanelSelectionSet::TConstIterator NodeIter( SelectedNodes );
+
+		const UK2Node_CallFunction* FunctionNode = Cast<UK2Node_CallFunction>( *NodeIter );
+
+		if( FunctionNode )
+		{
+			UFunction* TargetFunction = FunctionNode->GetTargetFunction();
+
+			if( TargetFunction )
+			{
+				FSourceCodeNavigation::NavigateToFunctionAsync( TargetFunction );
+			}
+		}
+	}
+}
+
+bool FBlueprintEditor::IsSelectionNativeFunction()
+{
+	auto GraphEditor = FocusedGraphEdPtr.Pin();
+	if( GraphEditor.IsValid() && FSourceCodeNavigation::IsCompilerAvailable() )
+	{
+		FGraphPanelSelectionSet SelectedNodes = GraphEditor->GetSelectedNodes();
+		FGraphPanelSelectionSet::TIterator NodeIter( SelectedNodes );
+		const UK2Node_CallFunction* FunctionNode = Cast<UK2Node_CallFunction>( *NodeIter );
+
+		if( FunctionNode && SelectedNodes.Num() == 1 )
+		{
+			UClass* OwningClass = FunctionNode->FunctionReference.GetMemberParentClass( FunctionNode );
+
+			if( OwningClass && OwningClass->HasAllClassFlags( CLASS_Native ))
+			{
+				return true;
+			}
+		}
+	}
+	return false;
+}
+
+void FBlueprintEditor::GotoNativeVariableDefinition()
+{
+	auto GraphEditor = FocusedGraphEdPtr.Pin();
+	if( GraphEditor.IsValid() && IsSelectionNativeVariable() )
+	{
+		const FGraphPanelSelectionSet SelectedNodes = GraphEditor->GetSelectedNodes();
+		FGraphPanelSelectionSet::TConstIterator NodeIter( SelectedNodes );
+		const UK2Node_Variable* VarNode = Cast<UK2Node_Variable>( *NodeIter );
+
+		if( VarNode )
+		{
+			UProperty* VariableProperty = VarNode->VariableReference.ResolveMember<UProperty>( VarNode );
+
+			if( VariableProperty )
+			{
+				FSourceCodeNavigation::NavigateToProperty( VariableProperty );
+			}
+		}
+	}
+}
+
+bool FBlueprintEditor::IsSelectionNativeVariable()
+{
+	auto GraphEditor = FocusedGraphEdPtr.Pin();
+	if( GraphEditor.IsValid() )
+	{
+		FGraphPanelSelectionSet SelectedNodes = GraphEditor->GetSelectedNodes();
+		FGraphPanelSelectionSet::TIterator NodeIter( SelectedNodes );
+		const UK2Node_Variable* VarNode = Cast<UK2Node_Variable>( *NodeIter );
+
+		if( VarNode && SelectedNodes.Num() == 1 )
+		{
+			UProperty* VariableProperty = VarNode->VariableReference.ResolveMember<UProperty>( VarNode );
+
+			if( VariableProperty && VariableProperty->HasAllFlags( RF_Native ))
+			{
+				return true;
+			}
+		}
+	}
+	return false;
+}
+
 void FBlueprintEditor::OnFindInstancesCustomEvent()
 {
 	auto GraphEditor = FocusedGraphEdPtr.Pin();
diff --git a/Engine/Source/Editor/Kismet/Private/SMyBlueprint.cpp b/Engine/Source/Editor/Kismet/Private/SMyBlueprint.cpp
--- a/Engine/Source/Editor/Kismet/Private/SMyBlueprint.cpp
+++ b/Engine/Source/Editor/Kismet/Private/SMyBlueprint.cpp
@@ -34,6 +34,7 @@
 #include "AssetToolsModule.h"
 #include "Editor/GraphEditor/Private/GraphActionNode.h"
 #include "IDocumentation.h"
+#include "Editor/UnrealEd/Public/SourceCodeNavigation.h"
 
 #define LOCTEXT_NAMESPACE "MyBlueprint"
 
@@ -50,6 +51,7 @@ void FMyBlueprintCommands::RegisterCommands()
 	UI_COMMAND(DeleteEntry, "Delete", "Deletes this function or variable from this blueprint.", EUserInterfaceActionType::Button, FInputGesture(EKeys::Platform_Delete));
 	UI_COMMAND( FindUserDefinedEnumInContentBrowser, "Find in Content Browser...", "Find user defined enum in content browser...", EUserInterfaceActionType::Button, FInputGesture() );
 	UI_COMMAND( AddNewUserDefinedEnum, "Create Enum Asset", "Create new user defined enum asset", EUserInterfaceActionType::Button, FInputGesture() );
+	UI_COMMAND( GotoNativeVarDefinition, "Goto Code Definition", "Goto the native code definition of this variable", EUserInterfaceActionType::Button, FInputGesture() );
 }
 
 //////////////////////////////////////////////////////////////////////////
@@ -172,7 +174,13 @@ void SMyBlueprint::Construct(const FArguments& InArgs, TWeakPtr<FBlueprintEditor
 		FCanExecuteAction::CreateSP(this, &SMyBlueprint::CanDuplicateAction),
 		FIsActionChecked(),
 		FIsActionButtonVisible::CreateSP(this, &SMyBlueprint::IsDuplicateActionVisible) );
-	
+
+	ToolKitCommandList->MapAction( FMyBlueprintCommands::Get().GotoNativeVarDefinition,
+		FExecuteAction::CreateSP(this, &SMyBlueprint::GotoNativeCodeVarDefinition),
+		FCanExecuteAction(),
+		FIsActionChecked(),
+		FIsActionButtonVisible::CreateSP(this, &SMyBlueprint::IsNativeVariable) );
+
 	TSharedPtr<FBlueprintEditorToolbar> Toolbar = MakeShareable(new FBlueprintEditorToolbar(InBlueprintEditor.Pin()));
 	TSharedPtr<FExtender> Extender = MakeShareable(new FExtender);
 	Toolbar->AddNewToolbar(Extender);
@@ -1248,6 +1256,7 @@ TSharedPtr<SWidget> SMyBlueprint::OnContextMenuOpening()
 			MenuBuilder.AddMenuEntry(FGenericCommands::Get().Rename, NAME_None, LOCTEXT("Rename", "Rename"), LOCTEXT("Rename_Tooltip", "Renames this function or variable from blueprint.") );
 			MenuBuilder.AddMenuEntry(FMyBlueprintCommands::Get().ImplementFunction);
 			MenuBuilder.AddMenuEntry(FMyBlueprintCommands::Get().FindEntry);
+			MenuBuilder.AddMenuEntry(FMyBlueprintCommands::Get().GotoNativeVarDefinition);
 			MenuBuilder.AddMenuEntry(FGenericCommands::Get().Duplicate);
 			MenuBuilder.AddMenuEntry(FMyBlueprintCommands::Get().DeleteEntry);
 		}
@@ -1732,6 +1741,31 @@ void SMyBlueprint::OnDuplicateAction()
 	}
 }
 
+void SMyBlueprint::GotoNativeCodeVarDefinition()
+{
+	if( FEdGraphSchemaAction_K2Var* VarAction = SelectionAsVar() )
+	{
+		if( UProperty* VarProperty = VarAction->GetProperty() )
+		{
+			FSourceCodeNavigation::NavigateToProperty( VarProperty );
+		}
+	}
+}
+
+bool SMyBlueprint::IsNativeVariable() const
+{
+	if( FEdGraphSchemaAction_K2Var* VarAction = SelectionAsVar() )
+	{
+		UProperty* VarProperty = VarAction->GetProperty();
+
+		if( VarProperty && VarProperty->HasAllFlags( RF_Native ))
+		{
+			return true;
+		}
+	}
+	return false;
+}
+
 void SMyBlueprint::OnResetItemFilter()
 {
 	FilterBox->SetText(FText::GetEmpty());
diff --git a/Engine/Source/Editor/Kismet/Private/SMyBlueprint.h b/Engine/Source/Editor/Kismet/Private/SMyBlueprint.h
--- a/Engine/Source/Editor/Kismet/Private/SMyBlueprint.h
+++ b/Engine/Source/Editor/Kismet/Private/SMyBlueprint.h
@@ -39,6 +39,7 @@ public:
 	TSharedPtr<FUICommandInfo> DeleteEntry;
 	TSharedPtr<FUICommandInfo> FindUserDefinedEnumInContentBrowser;
 	TSharedPtr<FUICommandInfo> AddNewUserDefinedEnum;
+	TSharedPtr<FUICommandInfo> GotoNativeVarDefinition;
 
 	// Add New Item
 	/** Initialize commands */
@@ -174,6 +175,8 @@ private:
 	bool IsDuplicateActionVisible() const;
 	bool CanDuplicateAction() const;
 	void OnDuplicateAction();
+	void GotoNativeCodeVarDefinition();
+	bool IsNativeVariable() const;
 
 	/** Callback when the filter is changed, forces the action tree(s) to filter */
 	void OnFilterTextChanged( const FText& InFilterText );
diff --git a/Engine/Source/Editor/Kismet/Public/BlueprintEditor.h b/Engine/Source/Editor/Kismet/Public/BlueprintEditor.h
--- a/Engine/Source/Editor/Kismet/Public/BlueprintEditor.h
+++ b/Engine/Source/Editor/Kismet/Public/BlueprintEditor.h
@@ -251,6 +251,18 @@ public:
 	/** Handles opening the header file of native parent class */
 	void OnEditParentClassNativeCodeClicked();
 
+	/** Called to open native function definition of the current node selection in an IDE */
+	void GotoNativeFunctionDefinition();
+
+	/** Called to check if the current selection is a native function */
+	bool IsSelectionNativeFunction();
+
+	/** Called to open native variable declaration of the current node selection in an IDE */
+	void GotoNativeVariableDefinition();
+
+	/** Called to check if the current selection is a native variable */
+	bool IsSelectionNativeVariable();
+
 	/** Returns: "(<NativeParentClass>.h)" */
 	FText GetTextForNativeParentClassHeaderLink() const;
 
diff --git a/Engine/Source/Editor/UnrealEd/Private/SourceCodeNavigation.cpp b/Engine/Source/Editor/UnrealEd/Private/SourceCodeNavigation.cpp
--- a/Engine/Source/Editor/UnrealEd/Private/SourceCodeNavigation.cpp
+++ b/Engine/Source/Editor/UnrealEd/Private/SourceCodeNavigation.cpp
@@ -1124,31 +1124,7 @@ void FSourceCodeNavigation::GatherFunctionsForActors( TArray< AActor* >& Actors,
 			// Only bother getting the correct module if we're gathering functions, too, since it can slow down the process a bit.
 			if( GatherMode == EGatherMode::ClassesAndFunctions )
 			{
-				UPackage* ClassPackage = CurClass->GetOuterUPackage();
-				if( ClassPackage != NULL )
-				{
-					//@Package name transition
-					FName ShortClassPackageName = FPackageName::GetShortFName(ClassPackage->GetFName());
-
-					// Is this module loaded?  In many cases, we may not have a loaded module for this class' package,
-					// as it might be statically linked into the executable, etc.
-					if( FModuleManager::Get().IsModuleLoaded( ShortClassPackageName ) )
-					{
-						// Because the module loaded into memory may have a slightly mutated file name (for
-						// hot reload, etc), we ask the module manager for the actual file name being used.  This
-						// is important as we need to be sure to get the correct symbols.
-						FModuleManager::FModuleStatus ModuleStatus;
-						if( ensure( FModuleManager::Get().QueryModule( ShortClassPackageName, ModuleStatus ) ) )
-						{
-							// Use the base file name (no path, no extension) as the module name for symbol look up!
-							ModuleName = FPaths::GetBaseFilename(ModuleStatus.FilePath);
-						}
-						else
-						{
-							// This module should always be known.  Should never happen.
-						}
-					}
-				}
+				FindClassModuleName( CurClass, ModuleName );
 			}
 
 			{
@@ -1241,6 +1217,84 @@ void FSourceCodeNavigation::GatherFunctionsForActors( TArray< AActor* >& Actors,
 	}
 }
 
+bool FSourceCodeNavigation::NavigateToFunctionAsync( UFunction* InFunction )
+{
+	bool bResult = false;
+
+	if( InFunction )
+	{
+		UClass* OwningClass = InFunction->GetOwnerClass();
+
+		if(  OwningClass->HasAllClassFlags( CLASS_Native ))
+		{
+			FString ModuleName;
+			// Find module name for class
+			if( FindClassModuleName( OwningClass, ModuleName ))
+			{
+				const FString SymbolName = FString::Printf( TEXT( "%s%s::%s" ), OwningClass->GetPrefixCPP(), *OwningClass->GetName(), *InFunction->GetName() );
+				NavigateToFunctionSourceAsync( SymbolName, ModuleName, false );
+				bResult = true;
+			}
+		}
+	}
+	return bResult;
+}
+
+bool FSourceCodeNavigation::NavigateToProperty( UProperty* InProperty )
+{
+	bool bResult = false;
+
+	if( InProperty && InProperty->HasAllFlags( RF_Native ))
+	{
+		FString SourceFilePath;
+		const bool bFileLocated =	FindClassHeaderPath( InProperty, SourceFilePath ) && 
+									IFileManager::Get().FileSize( *SourceFilePath ) != INDEX_NONE;
+
+		if( bFileLocated )
+		{
+			const FString AbsoluteSourcePath = IFileManager::Get().ConvertToAbsolutePathForExternalAppForRead( *SourceFilePath );
+			bResult = OpenSourceFile( AbsoluteSourcePath );
+		}
+	}
+	return bResult;
+}
+
+bool FSourceCodeNavigation::FindClassModuleName( UClass* InClass, FString& ModuleName )
+{
+	bool bResult = false;
+	// Find module name from class
+	if( InClass )
+	{
+		UPackage* ClassPackage = InClass->GetOuterUPackage();
+
+		if( ClassPackage )
+		{
+			//@Package name transition
+			FName ShortClassPackageName = FPackageName::GetShortFName(ClassPackage->GetFName());
+
+			// Is this module loaded?  In many cases, we may not have a loaded module for this class' package,
+			// as it might be statically linked into the executable, etc.
+			if( FModuleManager::Get().IsModuleLoaded( ShortClassPackageName ) )
+			{
+				// Because the module loaded into memory may have a slightly mutated file name (for
+				// hot reload, etc), we ask the module manager for the actual file name being used.  This
+				// is important as we need to be sure to get the correct symbols.
+				FModuleManager::FModuleStatus ModuleStatus;
+				if( ensure( FModuleManager::Get().QueryModule( ShortClassPackageName, ModuleStatus ) ) )
+				{
+					// Use the base file name (no path, no extension) as the module name for symbol look up!
+					ModuleName = FPaths::GetBaseFilename(ModuleStatus.FilePath);
+					bResult = true;
+				}
+				else
+				{
+					// This module should always be known.  Should never happen.
+				}
+			}
+		}
+	}
+	return bResult;
+}
 
 /** Call this to access the multi-cast delegate that you can register a callback with */
 FSourceCodeNavigation::FOnSymbolQueryFinished& FSourceCodeNavigation::AccessOnSymbolQueryFinished()
diff --git a/Engine/Source/Editor/UnrealEd/Public/SourceCodeNavigation.h b/Engine/Source/Editor/UnrealEd/Public/SourceCodeNavigation.h
--- a/Engine/Source/Editor/UnrealEd/Public/SourceCodeNavigation.h
+++ b/Engine/Source/Editor/UnrealEd/Public/SourceCodeNavigation.h
@@ -82,6 +82,22 @@ public:
 	 */
 	UNREALED_API static void GatherFunctionsForActors( TArray< AActor* >& Actors, const EGatherMode::Type GatherMode, TArray< FEditCodeMenuClass >& Classes );
 
+	/**
+	 * Navigates asynchronously to the UFunction definition in IDE
+	 *
+	 * @param	InFunction		UFunction to navigate to in source code
+	 * @return					whether the navigation was successfull or not
+	 */
+	UNREALED_API static bool NavigateToFunctionAsync( UFunction* InFunction );
+
+	/**
+	 * Navigates to UProperty source header in the IDE
+	 *
+	 * @param	InProperty		UProperty to navigate to in source code
+	 * @return					whether the navigation was successfull or not
+	 */
+	UNREALED_API static bool NavigateToProperty( UProperty* InProperty );
+
 	/** Delegate that's triggered when any symbol query has completed */
 	DECLARE_MULTICAST_DELEGATE( FOnSymbolQueryFinished );
 
@@ -111,6 +127,9 @@ public:
 
 	UNREALED_API static bool OpenModuleSolution();
 
+	/** Attempt to locate fully qualified class module name */
+	UNREALED_API static bool FindClassModuleName( UClass* InClass, FString& ModuleName );
+
 	/** Delegate that's triggered when any symbol query has completed */
 	DECLARE_MULTICAST_DELEGATE( FOnCompilerNotFound );
 
