#TTP 332018 - Regression: EDITOR: LANDSCAPE: CRASH: Choosing a landscape material before creation and then closing the editor will result in a crash

#proj UE4

#summary Fixed crash when closing new landscape dialog with a material selected
(fix CIS)

diff --git a/Engine/Source/Editor/UnrealEd/Private/Settings/SettingsClasses.cpp b/Engine/Source/Editor/UnrealEd/Private/Settings/SettingsClasses.cpp
--- a/Engine/Source/Editor/UnrealEd/Private/Settings/SettingsClasses.cpp
+++ b/Engine/Source/Editor/UnrealEd/Private/Settings/SettingsClasses.cpp
@@ -11,6 +11,8 @@
 /* UContentBrowserSettings interface
  *****************************************************************************/
 
+UContentBrowserSettings::FSettingChangedEvent UContentBrowserSettings::SettingChangedEvent;
+
 UContentBrowserSettings::UContentBrowserSettings( const class FPostConstructInitializeProperties& PCIP )
 	: Super(PCIP)
 { }
