Added initial support for cook on the side.

diff --git a/Engine/Source/Editor/UnrealEd/Classes/Commandlets/CookCommandlet.h b/Engine/Source/Editor/UnrealEd/Classes/Commandlets/CookCommandlet.h
--- a/Engine/Source/Editor/UnrealEd/Classes/Commandlets/CookCommandlet.h
+++ b/Engine/Source/Editor/UnrealEd/Classes/Commandlets/CookCommandlet.h
@@ -99,15 +99,6 @@ public:
 	
 	// End UCommandlet Interface
 
-private:
-
-	// Callback for handling a network file request.
-	void HandleNetworkFileServerFileRequest( const FString& Filename, const FString &Platformname, TArray<FString>& UnsolicitedFiles );
-
-	// Callback for recompiling shaders
-	void HandleNetworkFileServerRecompileShaders(const struct FShaderRecompileData& RecompileData);
-
-
 private:
 
 	/** Holds the instance identifier. */
diff --git a/Engine/Source/Editor/UnrealEd/Classes/CookOnTheSide/CookOnTheFlyServer.h b/Engine/Source/Editor/UnrealEd/Classes/CookOnTheSide/CookOnTheFlyServer.h
new file mode 100644
--- /dev/null
+++ b/Engine/Source/Editor/UnrealEd/Classes/CookOnTheSide/CookOnTheFlyServer.h
@@ -0,0 +1,327 @@
+// Copyright 1998-2014 Epic Games, Inc. All Rights Reserved.
+
+/*=============================================================================
+	CookOnTheFlyServer.h : handles polite cook requests via network ;)
+=============================================================================*/
+
+#pragma once
+
+#include "CookOnTheFlyServer.generated.h"
+
+
+
+UCLASS()
+class UCookOnTheFlyServer : public UObject, public FTickableEditorObject
+{
+	GENERATED_UCLASS_BODY()
+
+private:
+	/** Cook ont he fly is almost always true, just leaving this option here if we want to switch the cook commandlet over to using this structure to cook */
+	bool bCookOnTheFly;
+	/** If true, iterative cooking is being done */
+	bool bIterativeCooking;
+	/** If true, packages are cooked compressed */
+	bool bCompressed;
+	/** Skip saving any packages in Engine/COntent/Editor* UNLESS TARGET HAS EDITORONLY DATA (in which case it will save those anyway) */
+	bool bSkipEditorContent;
+	/** Should we use Async loading and saving (good for the editor) */
+	bool bAsyncSaveLoad;
+	/** Should we auto tick (used in the editor) */
+	bool bIsAutoTick;
+
+	/** Holds the instance identifier. */
+	FGuid InstanceId;
+
+	/** Holds the sandbox file wrapper to handle sandbox path conversion. */
+	TAutoPtr<class FSandboxPlatformFile> SandboxFile;
+	/** Cook on the fly server uses the NetworkFileServer */
+	class INetworkFileServer* NetworkFileServer;
+
+	/** We hook this up to a delegate to avoid reloading textures and whatnot */
+	TSet<FString> PackagesToNotReload;
+
+
+
+	/** Array which has been made thread safe :) */
+	template<typename Type>
+	struct FThreadSafeArray
+	{
+	private:
+		mutable FCriticalSection	SynchronizationObject; // made this mutable so this class can have const functions and still be thread safe
+		TArray<Type>		Items;
+	public:
+		void Enqueue(const Type& Item)
+		{
+			FScopeLock ScopeLock(&SynchronizationObject);
+			Items.Add(Item);
+		}
+		bool Dequeue(Type* Result)
+		{
+			FScopeLock ScopeLock(&SynchronizationObject);
+			if (Items.Num())
+			{
+				*Result = Items[0];
+				Items.RemoveAt(0);
+				return true;
+			}
+			return false;
+		}
+		void DequeueAll(TArray<Type>& Results)
+		{
+			FScopeLock ScopeLock(&SynchronizationObject);
+			Results += Items;
+			Items.Empty();
+		}
+
+		bool HasItems() const
+		{
+			FScopeLock ScopeLock( &SynchronizationObject );
+			return Items.Num() > 0;
+		}
+	};
+
+public:
+	/** cooked file requests which includes platform which file is requested for */
+	struct FFilePlatformRequest
+	{
+	public:
+		FFilePlatformRequest() { }
+		FFilePlatformRequest( const FString &InFilename, const FString &InPlatformname ) : Filename( FPaths::ConvertRelativePathToFull( InFilename) ), Platformname( InPlatformname ) 
+		{
+		}
+
+		FString Filename;
+		FString Platformname;
+
+		bool IsValid()  const
+		{
+			return Filename.Len() > 0;
+		}
+
+		void Clear()
+		{
+
+			Filename = TEXT("");
+			Platformname = TEXT("");
+		}
+
+		bool operator ==( const FFilePlatformRequest &InFileRequest ) const
+		{
+			if ( InFileRequest.Filename == Filename )
+			{
+				if ( InFileRequest.Platformname == Platformname )
+					return true;
+			}
+			return false;
+		}
+
+		FORCEINLINE FString ToString() const
+		{
+			return FString::Printf( TEXT("%s %s"), *Filename, *Platformname );
+		}
+
+	};
+
+private:
+
+	FThreadSafeArray<FFilePlatformRequest> ThreadSafeFilenameArray; // list of requested files
+	FThreadSafeArray<FFilePlatformRequest> UnsolicitedCookedPackages; // list of files which haven't been requested but we think should cook based on previous requests
+
+
+
+	/** Helper list of all files which have been cooked */
+	struct FThreadSafeFilenameSet
+	{
+		FCriticalSection	SynchronizationObject;
+		TSet<FFilePlatformRequest> FilesProcessed;
+
+		void Add(const FFilePlatformRequest &Filename)
+		{
+			FScopeLock ScopeLock(&SynchronizationObject);
+			check(Filename.IsValid());
+			FilesProcessed.Add(Filename);
+		}
+		bool Exists(const FFilePlatformRequest &Filename)
+		{
+			FScopeLock ScopeLock(&SynchronizationObject);
+			return FilesProcessed.Contains(Filename);
+		}
+	};
+	FThreadSafeFilenameSet ThreadSafeFilenameSet; // set of files which have been cooked when needing to recook a file the entry will need to be removed from here
+
+
+	FThreadSafeArray<struct FRecompileRequest*> RecompileRequests;
+
+
+public:
+
+
+	
+
+	enum ECookOnTheSideResult
+	{
+		COSR_CookedMap			= 0x00000001,
+		COSR_CookedPackage		= 0x00000002,
+		COSR_ErrorLoadingPackage= 0x00000004,
+	};
+
+
+
+	/**
+	 * FTickableEditorObject interface used by cook on the side
+	 */
+	TStatId GetStatId() const OVERRIDE;
+	 void Tick(float DeltaTime) OVERRIDE;
+	 bool IsTickable() const OVERRIDE;
+
+
+
+
+
+	/**
+	 * Initialize the CookServer so that either CookOnTheFly can be called or Cook on the side can be started and ticked
+	 */
+	void Initialize( bool inCompressed, bool inIterativeCooking, bool inSkipEditorContent, bool inIsAutoTick = false, bool inAsyncLoadSave = false, const FString &OutputDirectoryOverride = FString(TEXT("")) );
+
+	/**
+	 * Cook on the side, cooks while also running the editor...
+	 *
+	 * @param  BindAnyPort					Whether to bind on any port or the default port.
+	 *
+	 * @return true on success, false otherwise.
+	 */
+	bool StartNetworkFileServer( bool BindAnyPort );
+	
+	/**
+	 * Handles cook package requests until there are no more requests, then returns
+	 *
+	 * @param  CookFlags output of the things that might have happened in the cook on the side
+	 * 
+	 * @return returns ECookOnTheSideResult
+	 */
+	uint32 TickCookOnTheSide( const float TimeSlice, uint32 &CookedPackagesCount );
+	
+	/** 
+	 * Process any shader recompile requests
+	 *
+	 */
+	void TickRecompileShaderRequests();
+
+	bool HasCookRequests() const { return ThreadSafeFilenameArray.HasItems(); }
+
+	bool HasRecompileShaderRequests() const { return RecompileRequests.HasItems(); }
+
+	/** 
+	 * Stop the network file server
+	 *
+	 */
+	void EndNetworkFileServer();
+
+
+	uint32 NumConnections() const;
+
+	
+	/**
+	 * BLOCKING Run the network file server and handle cook requests 
+	 *	don't return until finished cooking or we haven't received a cook requests for a length of time
+	 *  shouldn't be used in conjunction with StartCookOnTheSide, TickCookOnTheSide, TickRecompileShaderRequests, FinishCookOnTheSide
+	 *
+	 * @param  BindAnyPort					Whether to bind on any port or the default port.
+	 * @param  Timeout						Length of time to wait for connections before attempting to close
+	 * @param  bForceClose					Whether or not the server should always shutdown after a timeout or after a user disconnects
+	 *
+	 * @return true on success, false otherwise.
+	 */
+	bool CookOnTheFly( bool BindAnyPort, int32 Timeout = 180, bool bForceClose = false );
+
+
+	virtual void BeginDestroy() OVERRIDE;
+
+
+
+
+private:
+	
+
+
+
+	/**
+	 * Returns cooker output directory.
+	 *
+	 * @param PlatformName Target platform name.
+	 * @return Cooker output directory for the specified platform.
+	 */
+	FString GetOutputDirectory( const FString& PlatformName ) const;
+
+	/**
+	 *	Get the given packages 'cooked' timestamp (i.e. account for dependencies)
+	 *
+	 *	@param	InFilename			The filename of the package
+	 *	@param	OutDateTime			The timestamp the cooked file should have
+	 *
+	 *	@return	bool				true if the package timestamp was found, false if not
+	 */
+	bool GetPackageTimestamp( const FString& InFilename, FDateTime& OutDateTime );
+
+	/**
+	 *	Cook (save) the given package
+	 *
+	 *	@param	Package				The package to cook/save
+	 *	@param	SaveFlags			The flags to pass to the SavePackage function
+	 *	@param	bOutWasUpToDate		Upon return, if true then the cooked package was cached (up to date)
+	 *
+	 *	@return	bool			true if packages was cooked
+	 */
+	bool SaveCookedPackage( UPackage* Package, uint32 SaveFlags, bool& bOutWasUpToDate );
+	/**
+	 *	Cook (save) the given package
+	 *
+	 *	@param	Package				The package to cook/save
+	 *	@param	SaveFlags			The flags to pass to the SavePackage function
+	 *	@param	bOutWasUpToDate		Upon return, if true then the cooked package was cached (up to date)
+	 *	@param  TargetPlatformNames Only cook for target platforms which are included in this array (if empty cook for all target platforms specified on commandline options)
+	 *									TargetPlatformNames is in and out value returns the platforms which the SaveCookedPackage function saved for
+	 *
+	 *	@return	bool			true if packages was cooked
+	 */
+	bool SaveCookedPackage( UPackage* Package, uint32 SaveFlags, bool& bOutWasUpToDate, TArray<FString> &TargetPlatformNames );
+
+	// bool ShouldCook(const FString& InFilename, const FString &InPlatformname = TEXT(""));
+
+	// Callback for handling a network file request.
+	void HandleNetworkFileServerFileRequest( const FString& Filename, const FString &Platformname, TArray<FString>& UnsolicitedFiles );
+
+	// Callback for recompiling shaders
+	void HandleNetworkFileServerRecompileShaders(const struct FShaderRecompileData& RecompileData);
+
+
+	void MaybeMarkPackageAsAlreadyLoaded(UPackage *Package);
+
+	/** Gets the output directory respecting any command line overrides */
+	FString GetOutputDirectoryOverride( const FString &OutputDirectoryOverride ) const;
+
+	/** Cleans sandbox folders for all target platforms */
+	void CleanSandbox(const TArray<ITargetPlatform*>& Platforms);
+
+	/** Generates asset registry */
+	void GenerateAssetRegistry(const TArray<ITargetPlatform*>& Platforms);
+
+	/** Saves global shader map files */
+	// void SaveGlobalShaderMapFiles(const TArray<ITargetPlatform*>& Platforms);
+
+	/** Collects all files to be cooked. This includes all commandline specified maps */
+	//void CollectFilesToCook(TArray<FString>& FilesInPath);
+
+	/** Generates long package names for all files to be cooked */
+	void GenerateLongPackageNames(TArray<FString>& FilesInPath);
+
+	/** Cooks all files */
+	//bool Cook(const TArray<ITargetPlatform*>& Platforms, TArray<FString>& FilesInPath);
+
+};
+
+FORCEINLINE uint32 GetTypeHash(const UCookOnTheFlyServer::FFilePlatformRequest &Key)
+{
+	return GetTypeHash( Key.Filename ) ^ GetTypeHash( Key.Platformname );
+}
+
diff --git a/Engine/Source/Editor/UnrealEd/Classes/Editor/UnrealEdEngine.h b/Engine/Source/Editor/UnrealEd/Classes/Editor/UnrealEdEngine.h
--- a/Engine/Source/Editor/UnrealEd/Classes/Editor/UnrealEdEngine.h
+++ b/Engine/Source/Editor/UnrealEd/Classes/Editor/UnrealEdEngine.h
@@ -140,6 +140,10 @@ class UNREALED_API UUnrealEdEngine : public UEditorEngine, public FNotifyHook
 	UPROPERTY(config)
 	TArray<FTemplateMapInfo> TemplateMapInfos;
 
+	/** Cooker server incase we want to cook ont he side while editing... */
+	UPROPERTY()
+	class UCookOnTheFlyServer* CookServer;
+
 	/** A mapping of packages to their checkout notify state.  This map only contains dirty packages.  Once packages become clean again, they are removed from the map.*/
 	TMap<TWeakObjectPtr<UPackage>, uint8>	PackageToNotifyState;
 
diff --git a/Engine/Source/Editor/UnrealEd/Private/Commandlets/CookCommandlet.cpp b/Engine/Source/Editor/UnrealEd/Private/Commandlets/CookCommandlet.cpp
--- a/Engine/Source/Editor/UnrealEd/Private/Commandlets/CookCommandlet.cpp
+++ b/Engine/Source/Editor/UnrealEd/Private/Commandlets/CookCommandlet.cpp
@@ -31,42 +31,6 @@ struct FRecompileRequest
 	bool bComplete;
 };
 
-////////////////////////////////////////////////////////////////
-/// Cook on the fly prototype
-///////////////////////////////////////////////////////////////
-
-template<typename Type>
-struct FThreadSafeArray
-{
-	FCriticalSection	SynchronizationObject;
-	TArray<Type>		Items;
-
-	void Enqueue(const Type& Item)
-	{
-		FScopeLock ScopeLock(&SynchronizationObject);
-		Items.Add(Item);
-	}
-	bool Dequeue(Type* Result)
-	{
-		FScopeLock ScopeLock(&SynchronizationObject);
-		if (Items.Num())
-		{
-			*Result = Items[0];
-			Items.RemoveAt(0);
-			return true;
-		}
-		return false;
-	}
-	void DequeueAll(TArray<Type>& Results)
-	{
-		FScopeLock ScopeLock(&SynchronizationObject);
-		Results += Items;
-		Items.Empty();
-	}
-};
-
-FThreadSafeArray<FRecompileRequest*> RecompileRequests;
-
 /* Static functions
  *****************************************************************************/
 
@@ -95,114 +59,32 @@ UCookCommandlet::UCookCommandlet( const class FPostConstructInitializeProperties
 /* UCookCommandlet interface
  *****************************************************************************/
 
-struct FFilePlatformRequest
+bool UCookCommandlet::CookOnTheFly( bool BindAnyPort, int32 Timeout, bool bForceClose )
 {
-public:
-	FFilePlatformRequest() { }
-	FFilePlatformRequest( const FString &InFilename, const FString &InPlatformname ) : Filename( FPaths::ConvertRelativePathToFull( InFilename) ), Platformname( InPlatformname ) 
-	{
-		// FPaths::ConvertRelativePathToFull(InFilename);
-	}
+	UCookOnTheFlyServer *CookOnTheFlyServer = ConstructObject<UCookOnTheFlyServer>( UCookOnTheFlyServer::StaticClass() );
 
-	FString Filename;
-	FString Platformname;
-
-	bool IsValid()  const
-	{
-		return Filename.Len() > 0;
-	}
-
-	void Clear()
+	struct FScopeRootObject
 	{
-
-		Filename = TEXT("");
-		Platformname = TEXT("");
-	}
-
-	bool operator ==( const FFilePlatformRequest &InFileRequest ) const
-	{
-		if ( InFileRequest.Filename == Filename )
+		UObject *Object;
+		FScopeRootObject( UObject *InObject ) : Object( InObject )
 		{
-			if ( InFileRequest.Platformname == Platformname )
-				return true;
+			Object->AddToRoot();
 		}
-		return false;
-	}
-
-	
-};
-uint32 GetTypeHash(const FFilePlatformRequest &Key)
-{
-	return GetTypeHash( Key.Filename ) ^ GetTypeHash( Key.Platformname );
-}
-
-FThreadSafeArray<FFilePlatformRequest> ThreadSafeFilenameArray;
-FThreadSafeArray<FFilePlatformRequest> UnsolicitedCookedPackages;
-
-
-
-
-struct FThreadSafeFilenameSet
-{
-	FCriticalSection	SynchronizationObject;
-	TSet<FFilePlatformRequest> FilesProcessed;
-
-	void Add(const FFilePlatformRequest &Filename)
-	{
-		FScopeLock ScopeLock(&SynchronizationObject);
-		check(Filename.IsValid());
-		FilesProcessed.Add(Filename);
-	}
-	bool Exists(const FFilePlatformRequest &Filename)
-	{
-		FScopeLock ScopeLock(&SynchronizationObject);
-		return FilesProcessed.Contains(Filename);
-	}
-};
-
-FThreadSafeFilenameSet ThreadSafeFilenameSet;
 
-
-bool UCookCommandlet::CookOnTheFly( bool BindAnyPort, int32 Timeout, bool bForceClose )
-{
-	//GetDerivedDataCacheRef().WaitForQuiescence(false);
-
-	// start the listening thread
-	FFileRequestDelegate FileRequestDelegate(FFileRequestDelegate::CreateUObject(this, &UCookCommandlet::HandleNetworkFileServerFileRequest));
-	FRecompileShadersDelegate RecompileShadersDelegate(FRecompileShadersDelegate::CreateUObject(this, &UCookCommandlet::HandleNetworkFileServerRecompileShaders));
-
-	INetworkFileServer* NetworkFileServer = FModuleManager::LoadModuleChecked<INetworkFileSystemModule>("NetworkFileSystem")
-		.CreateNetworkFileServer(BindAnyPort ? 0 : -1, &FileRequestDelegate, &RecompileShadersDelegate);
-
-	TArray<TSharedPtr<FInternetAddr> > AddressList;
-
-	if ((NetworkFileServer == NULL) || !NetworkFileServer->GetAddressList(AddressList))
-	{
-		UE_LOG(LogCookCommandlet, Error, TEXT("Failed to create network file server"));
-
-		return false;
-	}
-
-	// broadcast our presence
-	if (InstanceId.IsValid())
-	{
-		TArray<FString> AddressStringList;
-
-		for (int32 AddressIndex = 0; AddressIndex < AddressList.Num(); ++AddressIndex)
+		~FScopeRootObject()
 		{
-			AddressStringList.Add(AddressList[AddressIndex]->ToString(true));
+			Object->RemoveFromRoot();
 		}
+	};
 
-		FMessageEndpointPtr MessageEndpoint = FMessageEndpoint::Builder("UCookCommandlet").Build();
+	// make sure that the cookonthefly server doesn't get cleaned up while we are garbage collecting below :)
+	FScopeRootObject S(CookOnTheFlyServer);
+
+	CookOnTheFlyServer->Initialize( bCompressed, bIterativeCooking, bSkipEditorContent );
 
-		if (MessageEndpoint.IsValid())
-		{
-			MessageEndpoint->Publish(new FFileServerReady(AddressStringList, InstanceId), EMessageScope::Network);
-		}		
-	}
 
-	// loop while waiting for requests
-	GIsRequestingExit = false;
+	if ( CookOnTheFlyServer->StartNetworkFileServer(BindAnyPort) == false )
+		return false;
 
 
 	// Garbage collection should happen when either
@@ -212,11 +94,11 @@ bool UCookCommandlet::CookOnTheFly( bool BindAnyPort, int32 Timeout, bool bForce
 	//		b. we have been idle for 20 seconds.
 	bool bShouldGC = true;
 
-	bool bCookedAMapSinceLastGC = false;
-
+	// megamoth
+	uint32 NonMapPackageCountSinceLastGC = 0;
+	
 	const int32 PackagesPerGC = 50;
-	int32 NonMapPackageCountSinceLastGC = 0;
-
+	
 	const double IdleTimeToGC = 20.0;
 	double LastCookActionTime = FPlatformTime::Seconds();
 
@@ -225,150 +107,23 @@ bool UCookCommandlet::CookOnTheFly( bool BindAnyPort, int32 Timeout, bool bForce
 
 	while (!GIsRequestingExit)
 	{
-		FFilePlatformRequest ToBuild;
 
 		while (!GIsRequestingExit)
 		{
-			// This is all the target platforms which we needed to process requests for this iteration
-			// we use this in the unsolicited packages processing below
-			TArray<FString> AllTargetPlatformNames;
-
+			uint32 TickResults = 0;
+			
+			CookOnTheFlyServer->TickCookOnTheSide(TickResults, NonMapPackageCountSinceLastGC);
 
-			while (ToBuild.IsValid() && !GIsRequestingExit)
+			bool bCookedAMapSinceLastGC = TickResults & UCookOnTheFlyServer::COSR_CookedMap;
+			if ( TickResults & (UCookOnTheFlyServer::COSR_CookedMap | UCookOnTheFlyServer::COSR_CookedPackage))
 			{
-				bool bWasUpToDate = false;
-				TArray<FString> TargetPlatformNames;
-				if( ToBuild.Platformname.Len() )
-				{
-					TargetPlatformNames.Add( ToBuild.Platformname );
-				}
-
-				FString Name = FPackageName::FilenameToLongPackageName(*ToBuild.Filename);
-				if (!ThreadSafeFilenameSet.Exists(ToBuild))
-				{
-					UPackage* Package = LoadPackage( NULL, *ToBuild.Filename, LOAD_None );
-
-					if( Package == NULL )
-					{
-						UE_LOG(LogCookCommandlet, Error, TEXT("Error loading %s!"), *ToBuild.Filename );
-					}
-					else
-					{
-						if( SaveCookedPackage(Package, SAVE_KeepGUID | SAVE_Async, bWasUpToDate, TargetPlatformNames) )
-						{
-							// Update flags used to determine garbage collection.
-							if (Package->ContainsMap())
-							{
-								bCookedAMapSinceLastGC = true;
-							}
-							else
-							{
-								NonMapPackageCountSinceLastGC++;
-								LastCookActionTime = FPlatformTime::Seconds();
-							}
-						}
-
-						//@todo ResetLoaders outside of this (ie when Package is NULL) causes problems w/ default materials
-						if (Package->HasAnyFlags(RF_RootSet) == false)
-						{
-							ResetLoaders(Package);
-						}
-					}
-				}
-
-				for ( int Index = 0; Index < TargetPlatformNames.Num(); ++Index )
-				{
-					AllTargetPlatformNames.AddUnique(TargetPlatformNames[Index]);
-
-					FFilePlatformRequest FileRequest( ToBuild.Filename, TargetPlatformNames[Index]);
-					
-					ThreadSafeFilenameSet.Add( FileRequest );
-				}
-
-
-				ToBuild.Clear();
-
-				ThreadSafeFilenameArray.Dequeue(&ToBuild);
+				LastCookActionTime = FPlatformTime::Seconds();
 			}
 
-			TArray<UObject *> ObjectsInOuter;
-			GetObjectsWithOuter(NULL, ObjectsInOuter, false);
-			for( int32 Index = 0; Index < ObjectsInOuter.Num() && !GIsRequestingExit; Index++ )
-			{
-				UPackage* Package = Cast<UPackage>(ObjectsInOuter[Index]);
-
-				if (Package)
-				{
-					FString Name = Package->GetPathName();
-					FString PackageFilename(GetPackageFilename(Package));
-					FPaths::MakeStandardFilename(PackageFilename);
-
-					const FString FullPath = FPaths::ConvertRelativePathToFull(PackageFilename);
-
-					bool AlreadyCooked = true;
-					for ( int Index = 0; Index < AllTargetPlatformNames.Num(); ++Index )
-					{
-						if ( !ThreadSafeFilenameSet.Exists(FFilePlatformRequest(FullPath, AllTargetPlatformNames[Index])) )
-						{
-							AlreadyCooked = false;
-							break;
-						}
-					}
-					if ( !AlreadyCooked )
-					{
-						bool bUnsolicitedWasUpdateToDate = false;
-
-						if (SaveCookedPackage(Package, SAVE_KeepGUID | SAVE_Async, bUnsolicitedWasUpdateToDate, AllTargetPlatformNames))
-						{
-							
-
-							if ((FPaths::FileExists(PackageFilename) == true) &&
-								(PackageFilename.Len() > 0) && 
-								(bUnsolicitedWasUpdateToDate == false))
-							{
-								UE_LOG(LogCookCommandlet, Display, TEXT("UnsolicitedCookedPackages: %s"), *PackageFilename);
-
-								for ( int Index = 0; Index < AllTargetPlatformNames.Num(); ++Index )
-								{
-									const FFilePlatformRequest Request( FullPath, AllTargetPlatformNames[Index] );
-									UnsolicitedCookedPackages.Enqueue( Request );
-									ThreadSafeFilenameSet.Add( Request );
-
-									UE_LOG( LogCookCommandlet, Display, TEXT("Unsolicited saved package %s.%s"), *Request.Filename, *Request.Platformname );
-								}
-							}
-
-							// Update flags used to determine garbage collection.
-							if (Package->ContainsMap())
-							{
-								bCookedAMapSinceLastGC = true;
-							}
-							else
-							{
-								NonMapPackageCountSinceLastGC++;
-								LastCookActionTime = FPlatformTime::Seconds();
-							}
-						}
-					}
-
-					ToBuild.Clear();
-
-					ThreadSafeFilenameArray.Dequeue(&ToBuild);
-
-					if (ToBuild.IsValid())
-					{
-						break;
-					}
-				}
-			}
 
-			while (!ToBuild.IsValid() && !GIsRequestingExit)
+			while ( (CookOnTheFlyServer->HasCookRequests() == false) && !GIsRequestingExit)
 			{
-				ToBuild.Clear();
-
-				ThreadSafeFilenameArray.Dequeue(&ToBuild);
-
-				if (!ToBuild.IsValid())
+				
 				{
 					if (NonMapPackageCountSinceLastGC > 0)
 					{
@@ -390,19 +145,7 @@ bool UCookCommandlet::CookOnTheFly( bool BindAnyPort, int32 Timeout, bool bForce
 					}
 					else
 					{
-						// try to pull off a request
-						FRecompileRequest* Request = NULL;
-
-						RecompileRequests.Dequeue(&Request);
-
-						// process it
-						if (Request)
-						{
-							HandleNetworkFileServerRecompileShaders(Request->RecompileData);
-						
-							// all done! other thread can unblock now
-							Request->bComplete = true;
-						}
+						CookOnTheFlyServer->TickRecompileShaderRequests();
 
 						FPlatformProcess::Sleep(0.0f);
 					}
@@ -422,7 +165,7 @@ bool UCookCommandlet::CookOnTheFly( bool BindAnyPort, int32 Timeout, bool bForce
 				// handle server timeout
 				if (InstanceId.IsValid() || bForceClose)
 				{
-					if (NetworkFileServer->NumConnections() > 0)
+					if (CookOnTheFlyServer->NumConnections() > 0)
 					{
 						bHadConnection = true;
 						LastConnectionTime = FDateTime::UtcNow();
@@ -441,7 +184,7 @@ bool UCookCommandlet::CookOnTheFly( bool BindAnyPort, int32 Timeout, bool bForce
 							GIsRequestingExit = true;
 						}
 					}
-					else if (bHadConnection && NetworkFileServer->NumConnections() == 0 && bForceClose) // immediately shut down if we previously had a connection and now do not
+					else if (bHadConnection && (CookOnTheFlyServer->NumConnections() == 0) && bForceClose) // immediately shut down if we previously had a connection and now do not
 					{
 						GIsRequestingExit = true;
 					}
@@ -450,10 +193,7 @@ bool UCookCommandlet::CookOnTheFly( bool BindAnyPort, int32 Timeout, bool bForce
 		}
 	}
 
-	// shutdown the server
-	NetworkFileServer->Shutdown();
-	delete NetworkFileServer;
-
+	CookOnTheFlyServer->EndNetworkFileServer();
 	return true;
 }
 
@@ -747,35 +487,8 @@ int32 UCookCommandlet::Main(const FString& CmdLineParams)
 		}
 	}
 
-	ITargetPlatformManagerModule& TPM = GetTargetPlatformManagerRef();
-	const TArray<ITargetPlatform*>& Platforms = TPM.GetActiveTargetPlatforms();
-	if (TPM.RestrictFormatsToRuntimeOnly())
-	{
-		UE_LOG(LogCookCommandlet, Fatal, TEXT("You must provide a TargetPlatform argument to cook."));
-	}
 
-	// Local sandbox file wrapper. This will be used to handle path conversions,
-	// but will not be used to actually write/read files so we can safely
-	// use [Platform] token in the sandbox directory name and then replace it
-	// with the actual platform name.
-	SandboxFile = new FSandboxPlatformFile(false);
-	
-	// Output directory override.	
-	FString OutputDirectory = GetOutputDirectoryOverride();
-
-	// Use SandboxFile to do path conversion to properly handle sandbox paths (outside of standard paths in particular).
-	SandboxFile->Initialize(&FPlatformFileManager::Get().GetPlatformFile(), *FString::Printf(TEXT("-sandbox=%s"), *OutputDirectory));
-
-	CleanSandbox(Platforms);
-
-	// allow the game to fill out the asset registry, as well as get a list of objects to always cook
-	TArray<FString> FilesInPath;
-	FGameDelegates::Get().GetCookModificationDelegate().ExecuteIfBound(FilesInPath);
-
-	// always generate the asset registry before starting to cook, for either method
-	GenerateAssetRegistry(Platforms);
-
-	if (bCookOnTheFly)
+	if ( bCookOnTheFly )
 	{
 		// parse instance identifier
 		FString InstanceIdString;
@@ -789,21 +502,38 @@ int32 UCookCommandlet::Main(const FString& CmdLineParams)
 			}
 		}
 
-		int32 Timeout = 180;
-		if (!FParse::Value(*Params, TEXT("timeout="), Timeout))
-		{
-			Timeout = 180;
-		}
-		if (!CookOnTheFly(InstanceId.IsValid(), Timeout, bForceClose))
-		{
-			return -1;
-		}
+		CookOnTheFly( InstanceId.IsValid() );
 	}
 	else
 	{
+		ITargetPlatformManagerModule& TPM = GetTargetPlatformManagerRef();
+		const TArray<ITargetPlatform*>& Platforms = TPM.GetTargetPlatforms();
+	
+		// Local sandbox file wrapper. This will be used to handle path conversions,
+		// but will not be used to actually write/read files so we can safely
+		// use [Platform] token in the sandbox directory name and then replace it
+		// with the actual platform name.
+		SandboxFile = new FSandboxPlatformFile(false);
+	
+		// Output directory override.	
+		FString OutputDirectory = GetOutputDirectoryOverride();
+
+		// Use SandboxFile to do path conversion to properly handle sandbox paths (outside of standard paths in particular).
+		SandboxFile->Initialize(&FPlatformFileManager::Get().GetPlatformFile(), *FString::Printf(TEXT("-sandbox=%s"), *OutputDirectory));
+
+		CleanSandbox(Platforms);
+
+		// allow the game to fill out the asset registry, as well as get a list of objects to always cook
+		TArray<FString> FilesInPath;
+		FGameDelegates::Get().GetCookModificationDelegate().ExecuteIfBound(FilesInPath);
+
+		// always generate the asset registry before starting to cook, for either method
+		GenerateAssetRegistry(Platforms);
+
 		Cook(Platforms, FilesInPath);
-	}
 
+	}
+	
 	return 0;
 }
 
@@ -954,7 +684,19 @@ void UCookCommandlet::SaveGlobalShaderMapFiles(const TArray<ITargetPlatform*>& P
 		RecompileData.ShaderPlatform = -1;
 		RecompileData.ModifiedFiles = &Files;
 		RecompileData.MeshMaterialMaps = NULL;
-		HandleNetworkFileServerRecompileShaders(RecompileData);
+
+		check( IsInGameThread() );
+
+		FString OutputDir = GetOutputDirectory(RecompileData.PlatformName);
+
+		RecompileShadersForRemote
+			(RecompileData.PlatformName, 
+			RecompileData.ShaderPlatform == -1 ? SP_NumPlatforms : (EShaderPlatform)RecompileData.ShaderPlatform,
+			OutputDir, 
+			RecompileData.MaterialsToLoad, 
+			RecompileData.SerializedShaderResources, 
+			RecompileData.MeshMaterialMaps, 
+			RecompileData.ModifiedFiles);
 	}
 }
 
@@ -1385,80 +1127,3 @@ bool UCookCommandlet::Cook(const TArray<ITargetPlatform*>& Platforms, TArray<FSt
 }
 
 
-/* UCookCommandlet callbacks
- *****************************************************************************/
-
-void UCookCommandlet::HandleNetworkFileServerFileRequest( const FString& Filename, const FString &Platformname, TArray<FString>& UnsolicitedFiles )
-{
-	bool bIsCookable = FPackageName::IsPackageExtension(*FPaths::GetExtension(Filename, true));
-
-	if (!bIsCookable)
-	{
-		TArray<FFilePlatformRequest> FileRequests;
-		UnsolicitedCookedPackages.DequeueAll(FileRequests);
-		for ( int Index=  0; Index < FileRequests.Num(); ++Index)
-		{
-			UnsolicitedFiles.Add( FileRequests[Index].Filename);
-		}
-		UPackage::WaitForAsyncFileWrites();
-		return;
-	}
-
-	FFilePlatformRequest FileRequest( Filename, Platformname );
-	ThreadSafeFilenameArray.Enqueue(FileRequest);
-
-	do
-	{
-		FPlatformProcess::Sleep(0.0f);
-	}
-	while (!ThreadSafeFilenameSet.Exists(FileRequest));
-
-	TArray<FFilePlatformRequest> FileRequests;
-	UnsolicitedCookedPackages.DequeueAll(FileRequests);
-	for ( int Index=  0; Index < FileRequests.Num(); ++Index)
-	{
-		UnsolicitedFiles.Add( FileRequests[Index].Filename);
-	}
-	
-	UPackage::WaitForAsyncFileWrites();
-}
-
-
-void UCookCommandlet::HandleNetworkFileServerRecompileShaders(const FShaderRecompileData& RecompileData)
-{
-	// if we aren't in the game thread, we need to push this over to the game thread and wait for it to finish
-	if (!IsInGameThread())
-	{
-		UE_LOG(LogCookCommandlet, Display, TEXT("Got a recompile request on non-game thread"));
-
-		// make a new request
-		FRecompileRequest* Request = new FRecompileRequest;
-		Request->RecompileData = RecompileData;
-		Request->bComplete = false;
-
-		// push the request for the game thread to process
-		RecompileRequests.Enqueue(Request);
-
-		// wait for it to complete (the game thread will pull it out of the TArray, but I will delete it)
-		while (!Request->bComplete)
-		{
-			FPlatformProcess::Sleep(0);
-		}
-		delete Request;
-		UE_LOG(LogCookCommandlet, Display, TEXT("Completed recompile..."));
-
-		// at this point, we are done on the game thread, and ModifiedFiles will have been filled out
-		return;
-	}
-
-	FString OutputDir = GetOutputDirectory(RecompileData.PlatformName);
-
-	RecompileShadersForRemote
-		(RecompileData.PlatformName, 
-		RecompileData.ShaderPlatform == -1 ? SP_NumPlatforms : (EShaderPlatform)RecompileData.ShaderPlatform,
-		OutputDir, 
-		RecompileData.MaterialsToLoad, 
-		RecompileData.SerializedShaderResources, 
-		RecompileData.MeshMaterialMaps, 
-		RecompileData.ModifiedFiles);
-}
diff --git a/Engine/Source/Editor/UnrealEd/Private/CookOnTheFlyServer.cpp b/Engine/Source/Editor/UnrealEd/Private/CookOnTheFlyServer.cpp
new file mode 100644
--- /dev/null
+++ b/Engine/Source/Editor/UnrealEd/Private/CookOnTheFlyServer.cpp
@@ -0,0 +1,947 @@
+// Copyright 1998-2014 Epic Games, Inc. All Rights Reserved.
+
+/*=============================================================================
+	CookOnTheFlyServer.cpp: handles polite cook requests via network ;)
+=============================================================================*/
+
+#include "UnrealEd.h"
+
+
+#include "PackageHelperFunctions.h"
+#include "DerivedDataCacheInterface.h"
+#include "ISourceControlModule.h"
+#include "GlobalShader.h"
+#include "TargetPlatform.h"
+#include "IConsoleManager.h"
+#include "Developer/PackageDependencyInfo/Public/PackageDependencyInfo.h"
+#include "IPlatformFileSandboxWrapper.h"
+#include "Messaging.h"
+#include "NetworkFileSystem.h"
+#include "AssetRegistryModule.h"
+#include "UnrealEdMessages.h"
+#include "GameDelegates.h"
+
+
+
+
+DEFINE_LOG_CATEGORY_STATIC(LogCookCommandlet, Log, All);
+
+#define DEBUG_COOKONTHEFLY 0
+
+
+////////////////////////////////////////////////////////////////
+/// Cook on the fly server
+///////////////////////////////////////////////////////////////
+
+
+/* Static functions
+ *****************************************************************************/
+
+
+
+static FString GetPackageFilename( UPackage* Package )
+{
+	FString Filename;
+	if (FPackageName::DoesPackageExist(Package->GetName(), NULL, &Filename))
+	{
+		Filename = FPaths::ConvertRelativePathToFull(Filename);
+	}
+	return Filename;
+}
+
+
+/* helper structs
+ *****************************************************************************/
+
+/** Helper to pass a recompile request to game thread */
+struct FRecompileRequest
+{
+	struct FShaderRecompileData RecompileData;
+	bool bComplete;
+};
+
+
+/* UCookOnTheFlyServer structors
+ *****************************************************************************/
+
+UCookOnTheFlyServer::UCookOnTheFlyServer(const class FPostConstructInitializeProperties& PCIP)
+	: Super(PCIP)
+{
+	bCookOnTheFly = true;
+}
+
+
+
+void UCookOnTheFlyServer::Tick(float DeltaTime)
+{
+	uint32 CookedPackagesCount = 0;
+	TickCookOnTheSide( 0.01f, CookedPackagesCount);
+	TickRecompileShaderRequests();
+}
+
+bool UCookOnTheFlyServer::IsTickable() const 
+{ 
+	return bIsAutoTick; 
+}
+
+TStatId UCookOnTheFlyServer::GetStatId() const
+{
+	RETURN_QUICK_DECLARE_CYCLE_STAT(UCookServer, STATGROUP_Tickables);
+}
+
+
+
+bool UCookOnTheFlyServer::StartNetworkFileServer( bool BindAnyPort )
+{
+	//GetDerivedDataCacheRef().WaitForQuiescence(false);
+
+	// start the listening thread
+	FFileRequestDelegate FileRequestDelegate(FFileRequestDelegate::CreateUObject(this, &UCookOnTheFlyServer::HandleNetworkFileServerFileRequest));
+	FRecompileShadersDelegate RecompileShadersDelegate(FRecompileShadersDelegate::CreateUObject(this, &UCookOnTheFlyServer::HandleNetworkFileServerRecompileShaders));
+
+	NetworkFileServer = FModuleManager::LoadModuleChecked<INetworkFileSystemModule>("NetworkFileSystem")
+		.CreateNetworkFileServer(BindAnyPort ? 0 : -1, &FileRequestDelegate, &RecompileShadersDelegate);
+
+	TArray<TSharedPtr<FInternetAddr> > AddressList;
+
+	if ((NetworkFileServer == NULL) || !NetworkFileServer->GetAddressList(AddressList))
+	{
+		UE_LOG(LogCookCommandlet, Error, TEXT("Failed to create network file server"));
+
+		return false;
+	}
+
+	// broadcast our presence
+	if (InstanceId.IsValid())
+	{
+		TArray<FString> AddressStringList;
+
+		for (int32 AddressIndex = 0; AddressIndex < AddressList.Num(); ++AddressIndex)
+		{
+			AddressStringList.Add(AddressList[AddressIndex]->ToString(true));
+		}
+
+		FMessageEndpointPtr MessageEndpoint = FMessageEndpoint::Builder("UCookOnTheFlyServer").Build();
+
+		if (MessageEndpoint.IsValid())
+		{
+			MessageEndpoint->Publish(new FFileServerReady(AddressStringList, InstanceId), EMessageScope::Network);
+		}		
+	}
+
+	// loop while waiting for requests
+	GIsRequestingExit = false;
+	return true;
+}
+
+uint32 UCookOnTheFlyServer::TickCookOnTheSide( const float TimeSlice, uint32 &CookedPackageCount )
+{
+
+	check( NetworkFileServer != NULL );
+
+
+	double StartTime = FPlatformTime::Seconds();
+
+	uint32 Result = 0;
+
+	// This is all the target platforms which we needed to process requests for this iteration
+	// we use this in the unsolicited packages processing below
+	TArray<FString> AllTargetPlatformNames;
+	bool bCookedPackage = false;
+	while (ThreadSafeFilenameArray.HasItems() && !GIsRequestingExit)
+	{
+		FFilePlatformRequest ToBuild;
+		ThreadSafeFilenameArray.Dequeue( &ToBuild );
+
+		check( ToBuild.IsValid() );
+		bool bWasUpToDate = false;
+		TArray<FString> TargetPlatformNames;
+		if( ToBuild.Platformname.Len() )
+		{
+			TargetPlatformNames.Add( ToBuild.Platformname );
+		}
+
+		FString Name = FPackageName::FilenameToLongPackageName(*ToBuild.Filename);
+		if (!ThreadSafeFilenameSet.Exists(ToBuild))
+		{
+			UPackage* Package = LoadPackage( NULL, *ToBuild.Filename, LOAD_None );
+
+			if( Package == NULL )
+			{
+				UE_LOG(LogCookCommandlet, Error, TEXT("Error loading %s!"), *ToBuild.Filename );
+				Result |= COSR_ErrorLoadingPackage;
+			}
+			else
+			{
+				if( SaveCookedPackage(Package, SAVE_KeepGUID | SAVE_Async, bWasUpToDate, TargetPlatformNames) )
+				{
+					// Update flags used to determine garbage collection.
+					if (Package->ContainsMap())
+					{
+						Result |= COSR_CookedMap;
+						//bCookedAMapSinceLastGC = true;
+					}
+					else
+					{
+						++CookedPackageCount;
+						Result |= COSR_CookedPackage;
+						//NonMapPackageCountSinceLastGC++;
+						// LastCookActionTime = FPlatformTime::Seconds();
+					}
+					bCookedPackage = true;
+				}
+
+				//@todo ResetLoaders outside of this (ie when Package is NULL) causes problems w/ default materials
+				if (Package->HasAnyFlags(RF_RootSet) == false)
+				{
+					ResetLoaders(Package);
+				}
+			}
+		}
+
+		for ( int Index = 0; Index < TargetPlatformNames.Num(); ++Index )
+		{
+			AllTargetPlatformNames.AddUnique(TargetPlatformNames[Index]);
+
+			FFilePlatformRequest FileRequest( ToBuild.Filename, TargetPlatformNames[Index]);
+
+			ThreadSafeFilenameSet.Add( FileRequest );
+		}
+
+		if ( (FPlatformTime::Seconds() - StartTime) > TimeSlice)
+		{
+			return Result;
+		}
+	}
+
+	if ( bCookedPackage )
+	{
+		TArray<UObject *> ObjectsInOuter;
+		GetObjectsWithOuter(NULL, ObjectsInOuter, false);
+		for( int32 Index = 0; Index < ObjectsInOuter.Num() && !GIsRequestingExit; Index++ )
+		{
+			if ( (FPlatformTime::Seconds() - StartTime) > TimeSlice)
+			{
+				return Result;
+			}
+
+			UPackage* Package = Cast<UPackage>(ObjectsInOuter[Index]);
+
+			if (Package)
+			{
+				FString Name = Package->GetPathName();
+				FString PackageFilename(GetPackageFilename(Package));
+				FPaths::MakeStandardFilename(PackageFilename);
+
+				const FString FullPath = FPaths::ConvertRelativePathToFull(PackageFilename);
+
+				bool AlreadyCooked = true;
+				for ( int Index = 0; Index < AllTargetPlatformNames.Num(); ++Index )
+				{
+					if ( !ThreadSafeFilenameSet.Exists(FFilePlatformRequest(FullPath, AllTargetPlatformNames[Index])) )
+					{
+						AlreadyCooked = false;
+						break;
+					}
+				}
+				if ( !AlreadyCooked )
+				{
+					bool bUnsolicitedWasUpdateToDate = false;
+
+					if (SaveCookedPackage(Package, SAVE_KeepGUID | SAVE_Async, bUnsolicitedWasUpdateToDate, AllTargetPlatformNames))
+					{
+
+
+						if ((FPaths::FileExists(PackageFilename) == true) &&
+							(PackageFilename.Len() > 0) && 
+							(bUnsolicitedWasUpdateToDate == false))
+						{
+							UE_LOG(LogCookCommandlet, Display, TEXT("UnsolicitedCookedPackages: %s"), *PackageFilename);
+
+							for ( int Index = 0; Index < AllTargetPlatformNames.Num(); ++Index )
+							{
+								const FFilePlatformRequest Request( FullPath, AllTargetPlatformNames[Index] );
+								UnsolicitedCookedPackages.Enqueue( Request );
+								ThreadSafeFilenameSet.Add( Request );
+
+								UE_LOG( LogCookCommandlet, Display, TEXT("Unsolicited saved package %s.%s"), *Request.Filename, *Request.Platformname );
+							}
+						}
+
+						// Update flags used to determine garbage collection.
+						if (Package->ContainsMap())
+						{
+							Result |= COSR_CookedMap;
+						}
+						else
+						{
+							++CookedPackageCount;
+							Result |= COSR_CookedPackage;
+						}
+					}
+				}
+
+				if ( ThreadSafeFilenameArray.HasItems())
+				{
+					// break out of saving unsolicited packages if we get a real request
+					break;
+				}
+			}
+		}
+	}
+	return Result;
+}
+
+
+void UCookOnTheFlyServer::EndNetworkFileServer()
+{
+	if ( NetworkFileServer )
+	{
+		// shutdown the server
+		NetworkFileServer->Shutdown();
+		delete NetworkFileServer;
+		NetworkFileServer = NULL;
+	}
+}
+
+void UCookOnTheFlyServer::BeginDestroy()
+{
+	EndNetworkFileServer();
+
+	Super::BeginDestroy();
+
+}
+
+void UCookOnTheFlyServer::TickRecompileShaderRequests()
+{
+	// try to pull off a request
+	FRecompileRequest* Request = NULL;
+
+	RecompileRequests.Dequeue(&Request);
+
+	// process it
+	if (Request)
+	{
+		HandleNetworkFileServerRecompileShaders(Request->RecompileData);
+
+		// all done! other thread can unblock now
+		Request->bComplete = true;
+	}
+}
+
+bool UCookOnTheFlyServer::CookOnTheFly( bool BindAnyPort, int32 Timeout, bool bForceClose )
+{
+	if ( StartNetworkFileServer(BindAnyPort) == false )
+	{
+		return false;
+	}
+
+	// Garbage collection should happen when either
+	//	1. We have cooked a map
+	//	2. We have cooked non-map packages and...
+	//		a. we have accumulated 50 of these since the last GC.
+	//		b. we have been idle for 20 seconds.
+	bool bShouldGC = true;
+
+	// megamoth
+	uint32 NonMapPackageCountSinceLastGC = 0;
+	
+	const int32 PackagesPerGC = 50;
+	
+	const double IdleTimeToGC = 20.0;
+	double LastCookActionTime = FPlatformTime::Seconds();
+
+	FDateTime LastConnectionTime = FDateTime::UtcNow();
+	bool bHadConnection = false;
+
+	while (!GIsRequestingExit)
+	{
+
+		while (!GIsRequestingExit)
+		{
+			uint32 TickResults = 0;
+			
+			TickCookOnTheSide(TickResults, NonMapPackageCountSinceLastGC);
+
+			bool bCookedAMapSinceLastGC = TickResults & COSR_CookedMap;
+			if ( TickResults & (COSR_CookedMap | COSR_CookedPackage))
+			{
+				LastCookActionTime = FPlatformTime::Seconds();
+			}
+
+
+			while ( (ThreadSafeFilenameArray.HasItems() == false) && !GIsRequestingExit)
+			{
+				
+				{
+					if (NonMapPackageCountSinceLastGC > 0)
+					{
+						// We should GC if we have packages to collect and we've been idle for some time.
+						bShouldGC = (NonMapPackageCountSinceLastGC > PackagesPerGC) || 
+							((FPlatformTime::Seconds() - LastCookActionTime) >= IdleTimeToGC);
+					}
+					bShouldGC |= bCookedAMapSinceLastGC;
+
+					if (bShouldGC)
+					{
+						bShouldGC = false;
+						bCookedAMapSinceLastGC = false;
+						NonMapPackageCountSinceLastGC = 0;
+
+						UE_LOG(LogCookCommandlet, Display, TEXT("GC..."));
+
+						CollectGarbage( RF_Native );
+					}
+					else
+					{
+						TickRecompileShaderRequests();
+
+						FPlatformProcess::Sleep(0.0f);
+					}
+				}
+
+				// update task graph
+				FTaskGraphInterface::Get().ProcessThreadUntilIdle(ENamedThreads::GameThread);
+
+				// execute deferred commands
+				for (int32 DeferredCommandsIndex = 0; DeferredCommandsIndex<GEngine->DeferredCommands.Num(); ++DeferredCommandsIndex)
+				{
+					GEngine->Exec( GWorld, *GEngine->DeferredCommands[DeferredCommandsIndex], *GLog);
+				}
+
+				GEngine->DeferredCommands.Empty();
+
+				// handle server timeout
+				if (InstanceId.IsValid() || bForceClose)
+				{
+					if (NetworkFileServer && NetworkFileServer->NumConnections() > 0)
+					{
+						bHadConnection = true;
+						LastConnectionTime = FDateTime::UtcNow();
+					}
+
+					if ((FDateTime::UtcNow() - LastConnectionTime) > FTimespan::FromSeconds(Timeout))
+					{
+						uint32 Result = FMessageDialog::Open(EAppMsgType::YesNo, NSLOCTEXT("UnrealEd", "FileServerIdle", "The file server did not receive any connections in the past 3 minutes. Would you like to shut it down?"));
+
+						if (Result == EAppReturnType::No && !bForceClose)
+						{
+							LastConnectionTime = FDateTime::UtcNow();
+						}
+						else
+						{
+							GIsRequestingExit = true;
+						}
+					}
+					else if (bHadConnection && (NetworkFileServer && NetworkFileServer->NumConnections() == 0) && bForceClose) // immediately shut down if we previously had a connection and now do not
+					{
+						GIsRequestingExit = true;
+					}
+				}
+			}
+		}
+	}
+
+	EndNetworkFileServer();
+	return true;
+}
+
+
+FString UCookOnTheFlyServer::GetOutputDirectory( const FString& PlatformName ) const
+{
+	// Use SandboxFile to get the correct sandbox directory.
+	FString OutputDirectory = SandboxFile->GetSandboxDirectory();
+
+	return OutputDirectory.Replace(TEXT("[Platform]"), *PlatformName);
+}
+
+
+bool UCookOnTheFlyServer::GetPackageTimestamp( const FString& InFilename, FDateTime& OutDateTime )
+{
+	FPackageDependencyInfoModule& PDInfoModule = FModuleManager::LoadModuleChecked<FPackageDependencyInfoModule>("PackageDependencyInfo");
+	FDateTime DependentTime;
+
+	if (PDInfoModule.DeterminePackageDependentTimeStamp(*InFilename, DependentTime) == true)
+	{
+		OutDateTime = DependentTime;
+
+		return true;
+	}
+
+	return false;
+}
+
+
+bool UCookOnTheFlyServer::SaveCookedPackage( UPackage* Package, uint32 SaveFlags, bool& bOutWasUpToDate ) 
+{
+	TArray<FString> TargetPlatformNames; 
+	return SaveCookedPackage( Package, SaveFlags, bOutWasUpToDate, TargetPlatformNames );
+}
+
+bool UCookOnTheFlyServer::SaveCookedPackage( UPackage* Package, uint32 SaveFlags, bool& bOutWasUpToDate, TArray<FString> &TargetPlatformNames )
+{
+
+
+	bool bSavedCorrectly = true;
+
+	FString Filename(GetPackageFilename(Package));
+
+	if (Filename.Len())
+	{
+		FString PkgFilename;
+		FDateTime DependentTimeStamp = FDateTime::MinValue();
+
+		// We always want to use the dependent time stamp when saving a cooked package...
+		// Iterative or not!
+		FString PkgFile;
+		FString Name = Package->GetPathName();
+
+		if (bIterativeCooking && FPackageName::DoesPackageExist(Name, NULL, &PkgFile))
+		{
+			PkgFilename = PkgFile;
+
+			if (GetPackageTimestamp(FPaths::GetBaseFilename(PkgFilename, false), DependentTimeStamp) == false)
+			{
+				UE_LOG(LogCookCommandlet, Display, TEXT("Failed to find depedency timestamp for: %s"), *PkgFilename);
+			}
+		}
+
+		// Use SandboxFile to do path conversion to properly handle sandbox paths (outside of standard paths in particular).
+		Filename = SandboxFile->ConvertToAbsolutePathForExternalAppForWrite(*Filename);
+
+		uint32 OriginalPackageFlags = Package->PackageFlags;
+ 		UWorld* World = NULL;
+ 		EObjectFlags Flags = RF_NoFlags;
+		bool bPackageFullyLoaded = false;
+
+		if (bCompressed)
+		{
+			Package->PackageFlags |= PKG_StoreCompressed;
+		}
+
+		ITargetPlatformManagerModule& TPM = GetTargetPlatformManagerRef();
+
+		static TArray<ITargetPlatform*> ActiveStartupPlatforms = TPM.GetActiveTargetPlatforms();
+
+		TArray<ITargetPlatform*> Platforms;
+
+		if ( TargetPlatformNames.Num() )
+		{
+			for ( int Index = 0; Index < TargetPlatformNames.Num(); ++Index )
+			{
+				const FString &TargetPlatformName = TargetPlatformNames[Index];
+
+
+				const TArray<ITargetPlatform*>& TargetPlatforms = TPM.GetTargetPlatforms();	
+
+
+				for ( int Index = 0; Index < TargetPlatforms.Num(); ++Index )
+				{
+					ITargetPlatform *TargetPlatform = TargetPlatforms[ Index ];
+					if ( TargetPlatform->PlatformName() == TargetPlatformName )
+					{
+						Platforms.Add( TargetPlatform );
+					}
+				}
+			}
+		}
+		else
+		{
+			Platforms = ActiveStartupPlatforms;
+
+			for ( int Index = 0; Index < Platforms.Num(); ++Index )
+			{
+				TargetPlatformNames.Add( Platforms[Index]->PlatformName() );
+			}
+		}
+		
+
+		for (int32 Index = 0; Index < Platforms.Num(); Index++)
+		{
+			ITargetPlatform* Target = Platforms[Index];
+			FString PlatFilename = Filename.Replace(TEXT("[Platform]"), *Target->PlatformName());
+
+			// If we are not iterative cooking, then cook the package
+			bool bCookPackage = (bIterativeCooking == false);
+
+			if (bCookPackage == false)
+			{
+				// If the cooked package doesn't exist, or if the cooked is older than the dependent, re-cook it
+				FDateTime CookedTimeStamp = IFileManager::Get().GetTimeStamp(*PlatFilename);
+				int32 CookedTimespanSeconds = (CookedTimeStamp - DependentTimeStamp).GetTotalSeconds();
+				bCookPackage = (CookedTimeStamp == FDateTime::MinValue()) || (CookedTimespanSeconds < 0);
+			}
+
+			// don't save Editor resources from the Engine if the target doesn't have editoronly data
+			if (bSkipEditorContent && Name.StartsWith(TEXT("/Engine/Editor")) && !Target->HasEditorOnlyData())
+			{
+				bCookPackage = false;
+			}
+
+			if (bCookPackage == true)
+			{
+				if (bPackageFullyLoaded == false)
+				{
+					Package->FullyLoad();
+					if (!Package->IsFullyLoaded())
+					{
+						UE_LOG(LogCookCommandlet, Warning, TEXT("Package %s supposed to be fully loaded but isn't. RF_WasLoaded is %s"), 
+							*Package->GetName(), Package->HasAnyFlags(RF_WasLoaded) ? TEXT("set") : TEXT("not set"));
+					}
+					bPackageFullyLoaded = true;
+
+					// If fully loading has caused a blueprint to be regenerated, make sure we eliminate all meta data outside the package
+					UMetaData* MetaData = Package->GetMetaData();
+					MetaData->RemoveMetaDataOutsidePackage();
+
+					// look for a world object in the package (if there is one, there's a map)
+					World = UWorld::FindWorldInPackage(Package);
+					Flags = World ? RF_NoFlags : RF_Standalone;
+				}
+
+				UE_LOG(LogCookCommandlet, Display, TEXT("Cooking %s -> %s"), *Package->GetName(), *PlatFilename);
+
+				bool bSwap = (!Target->IsLittleEndian()) ^ (!PLATFORM_LITTLE_ENDIAN);
+
+
+				if (!Target->HasEditorOnlyData())
+				{
+					Package->PackageFlags |= PKG_FilterEditorOnly;
+				}
+				else
+				{
+					Package->PackageFlags &= ~PKG_FilterEditorOnly;
+				}
+
+				if (World)
+				{
+					World->PersistentLevel->OwningWorld = World;
+				}
+
+				if( PlatFilename.Len() >= PLATFORM_MAX_FILEPATH_LENGTH )
+				{
+					UE_LOG( LogCookCommandlet, Error, TEXT( "Couldn't save package, filename is too long :%s" ), *PlatFilename );
+					bSavedCorrectly = false;
+				}
+				else
+				{
+					bSavedCorrectly &= GEditor->SavePackage( Package, World, Flags, *PlatFilename, GError, NULL, bSwap, false, SaveFlags, Target, FDateTime::MinValue() );
+				}
+
+				
+				bOutWasUpToDate = false;
+			}
+			else
+			{
+				UE_LOG(LogCookCommandlet, Display, TEXT("Up to date: %s"), *PlatFilename);
+
+				bOutWasUpToDate = true;
+			}
+		}
+
+		Package->PackageFlags = OriginalPackageFlags;
+	}
+
+	// return success
+	return bSavedCorrectly;
+}
+
+
+void UCookOnTheFlyServer::MaybeMarkPackageAsAlreadyLoaded(UPackage *Package)
+{
+	FString Name = Package->GetName();
+	if (PackagesToNotReload.Contains(Name))
+	{
+		UE_LOG(LogCookCommandlet, Verbose, TEXT("Marking %s already loaded."), *Name);
+		Package->PackageFlags |= PKG_ReloadingForCooker;
+	}
+}
+
+
+void UCookOnTheFlyServer::Initialize( bool inCompressed, bool inIterativeCooking, bool inSkipEditorContent, bool inIsAutoTick, bool inAsyncSaveLoad, const FString &OutputDirectoryOverride )
+{
+	bCompressed = inCompressed;
+	bIterativeCooking = inIterativeCooking;
+	bSkipEditorContent = inSkipEditorContent;
+	bAsyncSaveLoad = inAsyncSaveLoad;
+	bIsAutoTick = inIsAutoTick;
+
+	
+	ITargetPlatformManagerModule& TPM = GetTargetPlatformManagerRef();
+	const TArray<ITargetPlatform*>& Platforms = TPM.GetTargetPlatforms();
+
+	// Local sandbox file wrapper. This will be used to handle path conversions,
+	// but will not be used to actually write/read files so we can safely
+	// use [Platform] token in the sandbox directory name and then replace it
+	// with the actual platform name.
+	SandboxFile = new FSandboxPlatformFile(false);
+
+	// Output directory override.	
+	FString OutputDirectory = GetOutputDirectoryOverride(OutputDirectoryOverride);
+
+	// Use SandboxFile to do path conversion to properly handle sandbox paths (outside of standard paths in particular).
+	SandboxFile->Initialize(&FPlatformFileManager::Get().GetPlatformFile(), *FString::Printf(TEXT("-sandbox=%s"), *OutputDirectory));
+
+	CleanSandbox(Platforms);
+
+	// allow the game to fill out the asset registry, as well as get a list of objects to always cook
+	TArray<FString> FilesInPath;
+	FGameDelegates::Get().GetCookModificationDelegate().ExecuteIfBound(FilesInPath);
+
+	// always generate the asset registry before starting to cook, for either method
+	GenerateAssetRegistry(Platforms);
+
+}
+
+uint32 UCookOnTheFlyServer::NumConnections() const
+{
+	if ( NetworkFileServer )
+	{
+		return NetworkFileServer->NumConnections();
+	}
+	return 0;
+}
+
+FString UCookOnTheFlyServer::GetOutputDirectoryOverride( const FString &OutputDirectoryOverride ) const
+{
+	FString OutputDirectory = OutputDirectoryOverride;
+	// Output directory override.	
+	if ( OutputDirectory.Len() <= 0 )
+	{
+		OutputDirectory = TEXT("Cooked-[Platform]");
+	}
+	else if (!OutputDirectory.Contains(TEXT("[Platform]"), ESearchCase::IgnoreCase, ESearchDir::FromEnd) )
+	{
+		// Output directory needs to contain [Platform] token to be able to cook for multiple targets.
+		OutputDirectory += TEXT("/Cooked-[Platform]");
+	}
+	FPaths::NormalizeDirectoryName(OutputDirectory);
+
+	return OutputDirectory;
+}
+
+void UCookOnTheFlyServer::CleanSandbox(const TArray<ITargetPlatform*>& Platforms)
+{
+	double SandboxCleanTime = 0.0;
+	{
+		SCOPE_SECONDS_COUNTER(SandboxCleanTime);
+
+		if (bIterativeCooking == false)
+		{
+			// for now we are going to wipe the cooked directory
+			for (int32 Index = 0; Index < Platforms.Num(); Index++)
+			{
+				ITargetPlatform* Target = Platforms[Index];
+				FString SandboxDirectory = GetOutputDirectory(Target->PlatformName());
+				IFileManager::Get().DeleteDirectory(*SandboxDirectory, false, true);
+			}
+		}
+		else
+		{
+			FPackageDependencyInfoModule& PDInfoModule = FModuleManager::LoadModuleChecked<FPackageDependencyInfoModule>("PackageDependencyInfo");
+			
+			// list of directories to skip
+			TArray<FString> DirectoriesToSkip;
+			TArray<FString> DirectoriesToNotRecurse;
+
+			// See what files are out of date in the sandbox folder
+			for (int32 Index = 0; Index < Platforms.Num(); Index++)
+			{
+				ITargetPlatform* Target = Platforms[Index];
+				FString SandboxDirectory = GetOutputDirectory(Target->PlatformName());
+
+				// use the timestamp grabbing visitor
+				IPlatformFile& PlatformFile = FPlatformFileManager::Get().GetPlatformFile();
+				FLocalTimestampDirectoryVisitor Visitor(PlatformFile, DirectoriesToSkip, DirectoriesToNotRecurse, false);
+			
+				PlatformFile.IterateDirectory(*SandboxDirectory, Visitor);
+
+				for (TMap<FString, FDateTime>::TIterator TimestampIt(Visitor.FileTimes); TimestampIt; ++TimestampIt)
+				{
+					FString CookedFilename = TimestampIt.Key();
+					FDateTime CookedTimestamp = TimestampIt.Value();
+					FString StandardCookedFilename = CookedFilename.Replace(*SandboxDirectory, *(FPaths::GetRelativePathToRoot()));
+					FDateTime DependentTimestamp;
+
+					if (PDInfoModule.DeterminePackageDependentTimeStamp(*(FPaths::GetBaseFilename(StandardCookedFilename, false)), DependentTimestamp) == true)
+					{
+						double Diff = (CookedTimestamp - DependentTimestamp).GetTotalSeconds();
+
+						if (Diff < 0.0)
+						{
+							UE_LOG(LogCookCommandlet, Display, TEXT("Deleting out of date cooked file: %s"), *CookedFilename);
+
+							IFileManager::Get().Delete(*CookedFilename);
+						}
+					}
+				}
+			}
+
+			// Collect garbage to ensure we don't have any packages hanging around from dependent time stamp determination
+			CollectGarbage(RF_Native);
+		}
+	}
+
+	UE_LOG(LogCookCommandlet, Display, TEXT("Sandbox cleanup took %5.3f seconds"), SandboxCleanTime);
+}
+
+void UCookOnTheFlyServer::GenerateAssetRegistry(const TArray<ITargetPlatform*>& Platforms)
+{
+	// load the interface
+	FAssetRegistryModule& AssetRegistryModule = FModuleManager::LoadModuleChecked<FAssetRegistryModule>(TEXT("AssetRegistry"));
+	IAssetRegistry& AssetRegistry = AssetRegistryModule.Get();
+
+	double GenerateAssetRegistryTime = 0.0;
+	{
+		SCOPE_SECONDS_COUNTER(GenerateAssetRegistryTime);
+		UE_LOG(LogCookCommandlet, Display, TEXT("Creating asset registry [is editor: %d]"), GIsEditor);
+
+		// Perform a synchronous search of any .ini based asset paths (note that the per-game delegate may
+		// have already scanned paths on its own)
+		// We want the registry to be fully initialized when generating streaming manifests too.
+		TArray<FString> ScanPaths;
+		if (GConfig->GetArray(TEXT("AssetRegistry"), TEXT("PathsToScanForCook"), ScanPaths, GEngineIni) > 0)
+		{
+			AssetRegistry.ScanPathsSynchronous(ScanPaths);
+		}
+		else
+		{
+			AssetRegistry.SearchAllAssets(true);
+		}
+		
+		// When not cooking on the fly the registry will be saved after the cooker has finished
+		if (bCookOnTheFly)
+		{
+			// write it out to a memory archive
+			FArrayWriter SerializedAssetRegistry;
+			AssetRegistry.Serialize(SerializedAssetRegistry);
+			UE_LOG(LogCookCommandlet, Display, TEXT("Generated asset registry size is %5.2fkb"), (float)SerializedAssetRegistry.Num() / 1024.f);
+
+			// now save it in each cooked directory
+			FString RegistryFilename = FPaths::GameDir() / TEXT("AssetRegistry.bin");
+			// Use SandboxFile to do path conversion to properly handle sandbox paths (outside of standard paths in particular).
+			FString SandboxFilename = SandboxFile->ConvertToAbsolutePathForExternalAppForWrite(*RegistryFilename);
+
+			for (int32 Index = 0; Index < Platforms.Num(); Index++)
+			{
+				FString PlatFilename = SandboxFilename.Replace(TEXT("[Platform]"), *Platforms[Index]->PlatformName());
+				FFileHelper::SaveArrayToFile(SerializedAssetRegistry, *PlatFilename);
+			}
+		}
+	}
+	UE_LOG(LogCookCommandlet, Display, TEXT("Done creating registry. It took %5.2fs."), GenerateAssetRegistryTime);
+	
+}
+
+void UCookOnTheFlyServer::GenerateLongPackageNames(TArray<FString>& FilesInPath)
+{
+	TArray<FString> FilesInPathReverse;
+	FilesInPathReverse.Reserve(FilesInPath.Num());
+	for( int32 FileIndex = 0; FileIndex < FilesInPath.Num(); FileIndex++ )
+	{
+		const FString& FileInPath = FilesInPath[FilesInPath.Num() - FileIndex - 1];
+		if (FPackageName::IsValidLongPackageName(FileInPath))
+		{
+			FilesInPathReverse.AddUnique(FileInPath);
+		}
+		else
+		{
+			FString LongPackageName;
+			if (FPackageName::TryConvertFilenameToLongPackageName(FileInPath, LongPackageName))
+			{
+				FilesInPathReverse.AddUnique(LongPackageName);
+			}
+			else
+			{
+				UE_LOG(LogCookCommandlet, Warning, TEXT("Unable to generate long package name for %s"), *FileInPath);
+			}
+		}
+	}
+	Exchange(FilesInPathReverse, FilesInPath);
+}
+
+
+/* UCookOnTheFlyServer callbacks
+ *****************************************************************************/
+
+
+void UCookOnTheFlyServer::HandleNetworkFileServerFileRequest( const FString& Filename, const FString &Platformname, TArray<FString>& UnsolicitedFiles )
+{
+	bool bIsCookable = FPackageName::IsPackageExtension(*FPaths::GetExtension(Filename, true));
+
+	if (!bIsCookable)
+	{
+		TArray<FFilePlatformRequest> FileRequests;
+		UnsolicitedCookedPackages.DequeueAll(FileRequests);
+		for ( int Index=  0; Index < FileRequests.Num(); ++Index)
+		{
+			UnsolicitedFiles.Add( FileRequests[Index].Filename);
+		}
+		UPackage::WaitForAsyncFileWrites();
+		return;
+	}
+
+	FFilePlatformRequest FileRequest( Filename, Platformname );
+	ThreadSafeFilenameArray.Enqueue(FileRequest);
+
+	do
+	{
+		FPlatformProcess::Sleep(0.0f);
+	}
+	while (!ThreadSafeFilenameSet.Exists(FileRequest));
+
+	TArray<FFilePlatformRequest> FileRequests;
+	UnsolicitedCookedPackages.DequeueAll(FileRequests);
+	for ( int Index=  0; Index < FileRequests.Num(); ++Index)
+	{
+		UnsolicitedFiles.Add( FileRequests[Index].Filename);
+	}
+	
+	UPackage::WaitForAsyncFileWrites();
+
+
+#if DEBUG_COOKONTHEFLY
+	UE_LOG( LogCookCommandlet, Display, TEXT("Processed file request %s"), *Filename );
+#endif
+
+}
+
+
+void UCookOnTheFlyServer::HandleNetworkFileServerRecompileShaders(const FShaderRecompileData& RecompileData)
+{
+	// if we aren't in the game thread, we need to push this over to the game thread and wait for it to finish
+	if (!IsInGameThread())
+	{
+		UE_LOG(LogCookCommandlet, Display, TEXT("Got a recompile request on non-game thread"));
+
+		// make a new request
+		FRecompileRequest* Request = new FRecompileRequest;
+		Request->RecompileData = RecompileData;
+		Request->bComplete = false;
+
+		// push the request for the game thread to process
+		RecompileRequests.Enqueue(Request);
+
+		// wait for it to complete (the game thread will pull it out of the TArray, but I will delete it)
+		while (!Request->bComplete)
+		{
+			FPlatformProcess::Sleep(0);
+		}
+		delete Request;
+		UE_LOG(LogCookCommandlet, Display, TEXT("Completed recompile..."));
+
+		// at this point, we are done on the game thread, and ModifiedFiles will have been filled out
+		return;
+	}
+
+	FString OutputDir = GetOutputDirectory(RecompileData.PlatformName);
+
+	RecompileShadersForRemote
+		(RecompileData.PlatformName, 
+		RecompileData.ShaderPlatform == -1 ? SP_NumPlatforms : (EShaderPlatform)RecompileData.ShaderPlatform,
+		OutputDir, 
+		RecompileData.MaterialsToLoad, 
+		RecompileData.SerializedShaderResources, 
+		RecompileData.MeshMaterialMaps, 
+		RecompileData.ModifiedFiles);
+}
diff --git a/Engine/Source/Editor/UnrealEd/Private/UnrealEdEngine.cpp b/Engine/Source/Editor/UnrealEd/Private/UnrealEdEngine.cpp
--- a/Engine/Source/Editor/UnrealEd/Private/UnrealEdEngine.cpp
+++ b/Engine/Source/Editor/UnrealEd/Private/UnrealEdEngine.cpp
@@ -92,6 +92,14 @@ void UUnrealEdEngine::Init(IEngineLoop* InEngineLoop)
 		PropertyModule.RegisterCustomPropertyLayout("GameMapsSettings", FOnGetDetailCustomizationInstance::CreateStatic(&FGameMapsSettingsCustomization::MakeInstance));
 		PropertyModule.RegisterCustomPropertyLayout("LevelEditorPlaySettings", FOnGetDetailCustomizationInstance::CreateStatic(&FLevelEditorPlaySettingsCustomization::MakeInstance));
 	}
+
+	bool bCookOnTheSide = FParse::Param( FCommandLine::Get(),TEXT("COOKONTHESIDE"));
+	if ( bCookOnTheSide )
+	{
+		CookServer = ConstructObject<UCookOnTheFlyServer>( UCookOnTheFlyServer::StaticClass() );
+		CookServer->Initialize( false, false, false, true, true );
+		CookServer->StartNetworkFileServer( false );
+	}
 }
 
 void UUnrealEdEngine::MakeSortedSpriteInfo(TArray<FSpriteCategoryInfo>& OutSortedSpriteInfo) const
