#TTP 329377 - RocketGDC: POSTLAUNCH: BP: CRASH:  While deleting a blueprint, clicking 'no' on the popup message regarding clearing undo history will result in a crash
#Branch UE4
#Proj  UnrealEd

#Summary Fixed the intial delete assets check for references in the undo buffer so the user is warned about the history being cleared and added an automatic clear once the user clicks ok in the delete asset dialog.

#Change Fixed what appears to be a logical error in checking if the undo buffer contains references to the object in question in FPendingDelete::CheckForReferences. Whereas before it would only check if the undo buffer was the only reference before setting bIsReferencedInMemoryByUndo now it checks to see if there are any references.
#Change Added some code in SDeleteAssetsDialog::ForceDelete to delete the undo buffer before continuing. At this point the user has had a notification in the DeleteAsset dialog and has clicked force delete to accept this - see information as to why this was neccesary.

#info The issue was caused by the blueprint comprising of 3 components being deleted and as each object was deleted, it nullified its references in other objects. When the process came across an undo buffer reference it would prompt the user to accept the undo history clear or to skip deleting that object. If the user replied to skip the history clear, the current object references would be gathered and logged and of course this meant hitting references that had been nullified and this caused the assert. Admittedly I'm not an expert on this piece of code but it does seem to me to be wrong to allow the user to get that far and then back out of the delete causing an inderminate and dangerous state with a blueprint. My solution is to check up front if we have undo history references and alert the user as quickly as possibly before anything is modified so if the user does decide to back out of the delete there is no change in state. This appears to be the intended functionality anyway.

diff --git a/Engine/Source/Editor/UnrealEd/Private/AssetDeleteModel.cpp b/Engine/Source/Editor/UnrealEd/Private/AssetDeleteModel.cpp
--- a/Engine/Source/Editor/UnrealEd/Private/AssetDeleteModel.cpp
+++ b/Engine/Source/Editor/UnrealEd/Private/AssetDeleteModel.cpp
@@ -510,14 +510,14 @@ void FPendingDelete::CheckForReferences()
 	bIsReferencedInMemoryByUndo = false;
 	if ( bIsReferencedInMemory )
 	{
-		// determine whether the transaction buffer is the only thing holding a reference to the object
+		FReferencerInformationList ReferencesExcludingUndo;
+		// determine whether the transaction buffer is holding a reference to the object
 		// and if so, offer the user the option to reset the transaction buffer.
 		GEditor->Trans->DisableObjectSerialization();
-		bIsReferencedInMemory = IsReferenced(Object, GARBAGE_COLLECTION_KEEPFLAGS, true, &MemoryReferences);
+		IsReferenced(Object, GARBAGE_COLLECTION_KEEPFLAGS, true, &ReferencesExcludingUndo);
 		GEditor->Trans->EnableObjectSerialization();
-
 		// only ref to this object is the transaction buffer - set a flag so we know we need to clear the undo stack
-		if ( !bIsReferencedInMemory )
+		if ( MemoryReferences.InternalReferences.Num() > ReferencesExcludingUndo.InternalReferences.Num() )
 		{
 			bIsReferencedInMemoryByUndo = true;
 		}
diff --git a/Engine/Source/Editor/UnrealEd/Private/Dialogs/SDeleteAssetsDialog.cpp b/Engine/Source/Editor/UnrealEd/Private/Dialogs/SDeleteAssetsDialog.cpp
--- a/Engine/Source/Editor/UnrealEd/Private/Dialogs/SDeleteAssetsDialog.cpp
+++ b/Engine/Source/Editor/UnrealEd/Private/Dialogs/SDeleteAssetsDialog.cpp
@@ -582,6 +582,11 @@ FReply SDeleteAssetsDialog::ForceDelete()
 {
 	ParentWindow.Get()->RequestDestroyWindow();
 
+	if( DeleteModel->IsAnythingReferencedInMemoryByUndo() )
+	{
+		GEditor->Trans->Reset( LOCTEXT("DeleteSelectedItem", "Delete Selected Item") );
+	}
+
 	DeleteModel->DoForceDelete();
 
 	return FReply::Handled();
