#summary Compositing Blueprints with actors sharing the same label will no longer cause crash or invalid Blueprints

#change When compositing Blueprints from actors, it will ensure that each component is given a unique and valid name in the Blueprint

#remove Removed util function for giving the components unique names.

#change Fixed an issue with FBlueprintEditorUtils::FindUniqueKismetName when names are too long it will go into a infinite loop. Now it will canibilize the end of the name to keep it from being too long.

#ttp 330066 - EDITOR: BP: CRASH: Compositing Blueprint with actors sharing the same display name

diff --git a/Engine/Source/Editor/UnrealEd/Private/Kismet2/BlueprintEditorUtils.cpp b/Engine/Source/Editor/UnrealEd/Private/Kismet2/BlueprintEditorUtils.cpp
--- a/Engine/Source/Editor/UnrealEd/Private/Kismet2/BlueprintEditorUtils.cpp
+++ b/Engine/Source/Editor/UnrealEd/Private/Kismet2/BlueprintEditorUtils.cpp
@@ -4341,11 +4341,20 @@ FName FBlueprintEditorUtils::FindUniqueKismetName(const UBlueprint* InBlueprint,
 {
 	int32 Count = 0;
 	FString KismetName;
+	FString BaseName = InBaseName;
 
-	TSharedPtr<INameValidatorInterface> NameValidator = MakeShareable(new FKismetNameValidator(InBlueprint));
+	TSharedPtr<FKismetNameValidator> NameValidator = MakeShareable(new FKismetNameValidator(InBlueprint));
 	while(NameValidator->IsValid(KismetName) != EValidatorResult::Ok)
 	{
-		KismetName = FString::Printf(TEXT("%s_%d"), *InBaseName, Count);
+		// Calculate the number of digits in the number, adding 2 (1 extra to correctly count digits, another to account for the '_' that will be added to the name
+		int32 CountLength = Count > 0? (int32)log((double)Count) + 2 : 2;
+
+		// If the length of the final string will be too long, cut off the end so we can fit the number
+		if(CountLength + BaseName.Len() > NameValidator->GetMaximumNameLength())
+		{
+			BaseName = BaseName.Left(NameValidator->GetMaximumNameLength() - CountLength);
+		}
+		KismetName = FString::Printf(TEXT("%s_%d"), *BaseName, Count);
 		Count++;
 	}
 
diff --git a/Engine/Source/Editor/UnrealEd/Private/Kismet2/Kismet2.cpp b/Engine/Source/Editor/UnrealEd/Private/Kismet2/Kismet2.cpp
--- a/Engine/Source/Editor/UnrealEd/Private/Kismet2/Kismet2.cpp
+++ b/Engine/Source/Editor/UnrealEd/Private/Kismet2/Kismet2.cpp
@@ -19,7 +19,6 @@
 #include "IContentBrowserSingleton.h"
 #include "ContentBrowserModule.h"
 #include "AssetRegistryModule.h"
-#include "Kismet2NameValidators.h"
 #include "Layers/Layers.h"
 #include "ScopedTransaction.h"
 #include "AssetToolsModule.h"
@@ -598,22 +597,6 @@ public:
 	{
 	}
 
-	/** Util to get a unique name within the BP from the supplied Actor */
-	FName GetUniqueNameFromActor(FKismetNameValidator& InValidator, const AActor* InActor)
-	{
-		FName NewName;
-		NewName = FName(*InActor->GetActorLabel());
-
-		//find valid name--
-		for (int32 i = 1; EValidatorResult::Ok != InValidator.IsValid(NewName); ++i)
-		{
-			NewName = FName(*InActor->GetActorLabel(), i);
-			check(i < 10000);
-		}
-
-		return NewName;
-	}
-
 	UBlueprint* Execute(FString Path, TArray<AActor*> SelectedActors, bool bReplaceInWorld)
 	{
 		if (SelectedActors.Num() > 0)
@@ -669,13 +652,11 @@ public:
 			check(Blueprint->SimpleConstructionScript != NULL);
 			SCS = Blueprint->SimpleConstructionScript;
 
-			FKismetNameValidator Validator(Blueprint);
-
 			// If we have a single component selected, make a BP with a single component
 			if(SingleSceneComp != NULL)
 			{
 				USCS_Node* Node = CreateUSCSNode(SingleSceneComp);
-				Node->VariableName = GetUniqueNameFromActor(Validator, SingleSceneComp->GetOwner());
+				Node->VariableName = FBlueprintEditorUtils::FindUniqueKismetName(Blueprint, SingleSceneComp->GetOwner()->GetActorLabel());
 				SCS->AddNode(Node);
 
 				NewActorTransform = SingleSceneComp->ComponentToWorld;
@@ -715,12 +696,18 @@ public:
 					if(ActorRootComponent != NULL && ActorRootComponent->GetClass()->HasMetaData(FBlueprintMetadata::MD_BlueprintSpawnableComponent))
 					{
 						USCS_Node* Node = CreateUSCSNode(ActorRootComponent);
-						Node->VariableName = GetUniqueNameFromActor(Validator, Actor);
+						Node->VariableName = FBlueprintEditorUtils::FindUniqueKismetName(Blueprint, Actor->GetActorLabel());
+
+						// Add the node as a child to the root node, this will be rearranged later for correct parent-child relations but needs to be there now for name validation purposes
+						RootNode->AddChildNode(Node);
 
 						ActorToUSCSNodeMap.Add(Actor, Node);
 					}
 				}
 
+				// Clear the root node of all children we added, they will be re-added with correct parent-child relationships
+				RootNode->ChildNodes.Empty();
+
 				// Attach all immediate children to their parent, or the root if their parent is not being added
 				for (auto ActorIt = SelectedActors.CreateConstIterator(); ActorIt; ++ActorIt)
 				{
