#summary Warnings that there are compile errors will only appear once when attempting to play in PIE, unless the error BP is compiled unsuccessfully again.

#ttp 300044 - ROCKET: TASK: Improve 'one or more BP has an unresolve compiler error' warning

diff --git a/Engine/Source/Editor/UnrealEd/Private/Kismet2/Kismet2.cpp b/Engine/Source/Editor/UnrealEd/Private/Kismet2/Kismet2.cpp
--- a/Engine/Source/Editor/UnrealEd/Private/Kismet2/Kismet2.cpp
+++ b/Engine/Source/Editor/UnrealEd/Private/Kismet2/Kismet2.cpp
@@ -296,6 +296,9 @@ UBlueprint* FKismetEditorUtilities::CreateBlueprint(UClass* ParentClass, UObject
 
 void FKismetEditorUtilities::CompileBlueprint(UBlueprint* BlueprintObj, bool bIsRegeneratingOnLoad, bool bSkipGarbageCollection, bool bSaveIntermediateProducts, FCompilerResultsLog* pResults)
 {
+	// Reset the flag, so if the user tries to use PIE it will warn them if the BP did not compile
+	BlueprintObj->bDisplayCompilePIEWarning = true;
+
 	UPackage* const BlueprintPackage = Cast<UPackage>(BlueprintObj->GetOutermost());
 	// compiling the blueprint will inherently dirty the package, but if there 
 	// weren't any changes to save before, there shouldn't be after
diff --git a/Engine/Source/Editor/UnrealEd/Private/PlayLevel.cpp b/Engine/Source/Editor/UnrealEd/Private/PlayLevel.cpp
--- a/Engine/Source/Editor/UnrealEd/Private/PlayLevel.cpp
+++ b/Engine/Source/Editor/UnrealEd/Private/PlayLevel.cpp
@@ -1499,7 +1499,7 @@ void UEditorEngine::PlayInEditor( UWorld* InWorld, bool bInSimulateInEditor )
 	check(EditorPlayInSettings);
 
 	// Prompt the user to compile any dirty Blueprints before PIE can occur.
-	bool bAnyBlueprintErrors = false;
+	TArray< UBlueprint* > ErrorBlueprintList;
 	bool bAnyBlueprintsDirty = false;
 	{
 		FString DirtyBlueprints;
@@ -1523,9 +1523,9 @@ void UEditorEngine::PlayInEditor( UWorld* InWorld, bool bInSimulateInEditor )
 				{
 					BlueprintsToRecompile.Add(Blueprint);
 				}
-				else if(BS_Error == Blueprint->Status)
+				else if(BS_Error == Blueprint->Status && Blueprint->bDisplayCompilePIEWarning)
 				{
-					bAnyBlueprintErrors = true;
+					ErrorBlueprintList.Add(Blueprint);
 					ErrorBlueprints += FString::Printf(TEXT("\n   %s"), *Blueprint->GetName());
 				}
 			}
@@ -1538,8 +1538,11 @@ void UEditorEngine::PlayInEditor( UWorld* InWorld, bool bInSimulateInEditor )
 					// Treating unknown as up to date for right now
 					break;
 				case BS_Error:
-					bAnyBlueprintErrors = true;
-					ErrorBlueprints += FString::Printf(TEXT("\n   %s"), *Blueprint->GetName());
+					if( Blueprint->bDisplayCompilePIEWarning)
+					{
+						ErrorBlueprintList.Add(Blueprint);
+						ErrorBlueprints += FString::Printf(TEXT("\n   %s"), *Blueprint->GetName());
+					}
 					break;
 				case BS_UpToDate:
 				case BS_UpToDateWithWarnings:
@@ -1606,9 +1609,13 @@ void UEditorEngine::PlayInEditor( UWorld* InWorld, bool bInSimulateInEditor )
 					UE_LOG(LogPlayLevel, Log, TEXT("[PIE] Compiling %s before PIE..."), *Blueprint->GetName());
 					FKismetEditorUtilities::CompileBlueprint(Blueprint, false, true);
 					const bool bHadError = (!Blueprint->IsUpToDate() && Blueprint->Status != BS_Unknown);
-					bAnyBlueprintErrors |= bHadError;
-					if (bHadError)
+
+					// Check if the Blueprint has already been added to the error list to prevent it from being added again
+					if (bHadError && ErrorBlueprintList.Find(Blueprint) == INDEX_NONE)
 					{
+						ErrorBlueprintList.Add(Blueprint);
+						ErrorBlueprints += FString::Printf(TEXT("\n   %s"), *Blueprint->GetName());
+
 						FFormatNamedArguments Arguments;
 						Arguments.Add(TEXT("Name"), FText::FromString(Blueprint->GetName()));
 
@@ -1657,9 +1664,9 @@ void UEditorEngine::PlayInEditor( UWorld* InWorld, bool bInSimulateInEditor )
 							const bool bIsPackageDirty = Package ? Package->IsDirty() : false;
 
 							FKismetEditorUtilities::CompileBlueprint(Blueprint);
-							if (Blueprint->Status == BS_Error)
+							if (Blueprint->Status == BS_Error && Blueprint->bDisplayCompilePIEWarning)
 							{
-								bAnyBlueprintErrors = true;
+								ErrorBlueprintList.Add(Blueprint);
 							}
 
 							// Restore the dirty flag
@@ -1673,7 +1680,7 @@ void UEditorEngine::PlayInEditor( UWorld* InWorld, bool bInSimulateInEditor )
 			}
 		}
 
-		if ( bAnyBlueprintErrors && !GIsDemoMode )
+		if ( ErrorBlueprintList.Num() && !GIsDemoMode )
 		{
 			FFormatNamedArguments Args;
 			Args.Add( TEXT("ErrorBlueprints"), FText::FromString( ErrorBlueprints ) );
@@ -1690,6 +1697,14 @@ void UEditorEngine::PlayInEditor( UWorld* InWorld, bool bInSimulateInEditor )
 
 				return;
 			}
+			else
+			{
+				// The user wants to ignore the compiler errors, mark the Blueprints and do not warn them again unless the Blueprint attempts to compile
+				for( auto Blueprint : ErrorBlueprintList )
+				{
+					Blueprint->bDisplayCompilePIEWarning = false;
+				}
+			}
 		}
 	}
 
@@ -1751,6 +1766,8 @@ void UEditorEngine::PlayInEditor( UWorld* InWorld, bool bInSimulateInEditor )
 	// Can't allow realtime viewports whilst in PIE so disable it for ALL viewports here.
 	DisableRealtimeViewports();
 
+	bool bAnyBlueprintErrors = ErrorBlueprintList.Num()? true : false;
+
 	if (bInSimulateInEditor || (PlayInSettings->PlayNetMode == EPlayNetMode::PIE_Standalone) || !PlayInSettings->RunUnderOneProcess)
 	{
 		// Only spawning 1 PIE instance under this process
diff --git a/Engine/Source/Runtime/Engine/Classes/Engine/Blueprint.h b/Engine/Source/Runtime/Engine/Classes/Engine/Blueprint.h
--- a/Engine/Source/Runtime/Engine/Classes/Engine/Blueprint.h
+++ b/Engine/Source/Runtime/Engine/Classes/Engine/Blueprint.h
@@ -317,6 +317,10 @@ class ENGINE_API UBlueprint : public UBlueprintCore
 	UPROPERTY(EditAnywhere, Category=BlueprintOption)
 	FString BlueprintDescription;
 
+	/** TRUE to show a warning when attempting to start in PIE and there is a compiler error on this Blueprint */
+	UPROPERTY(transient)
+	bool bDisplayCompilePIEWarning;
+
 #endif //WITH_EDITORONLY_DATA
 
 	/** 'Simple' construction script - graph of components to instance */
