#summary Adding localization import support for Portable Object(.po) files
#change Added a import option to the localization export commandlet for Epic Friday work.  This was added in the export commandlet and should probably be moved to a standalone commandlet in the future.
#change Modified the config files to include import/export settings which are now required
#ttp 321739 - L10N: Create a commandlet to convert between our loc formats and the po format for use with translation houses

diff --git a/Engine/Config/Localization/PortableObjectExport.ini b/Engine/Config/Localization/PortableObjectExport.ini
--- a/Engine/Config/Localization/PortableObjectExport.ini
+++ b/Engine/Config/Localization/PortableObjectExport.ini
@@ -5,6 +5,7 @@ CulturesToGenerate=zh
 
 [GatherTextStep0]
 CommandletClass=InternationalizationExport
+bExportLoc=true
 SourcePath=../../../Engine/Content/Localization/Engine
 DestinationPath=../../../Engine/Content/Localization/Engine
 ManifestName=Engine.manifest
@@ -12,6 +13,7 @@ ArchiveName=Engine.archive
 
 [GatherTextStep1]
 CommandletClass=InternationalizationExport
+bExportLoc=true
 SourcePath=../../../Engine/Content/Localization/Editor
 DestinationPath=../../../Engine/Content/Localization/Editor
 ManifestName=Editor.manifest
@@ -19,6 +21,7 @@ ArchiveName=Editor.archive
 
 [GatherTextStep2]
 CommandletClass=InternationalizationExport
+bExportLoc=true
 SourcePath=../../../Engine/Content/Localization/Editor
 DestinationPath=../../../Engine/Content/Localization/Editor
 ManifestName=ToolTips.manifest
@@ -26,6 +29,7 @@ ArchiveName=ToolTips.archive
 
 [GatherTextStep3]
 CommandletClass=InternationalizationExport
+bExportLoc=true
 SourcePath=../../../Engine/Content/Localization/Editor
 DestinationPath=../../../Engine/Content/Localization/Editor
 ManifestName=PropertyNames.manifest
diff --git a/Engine/Config/Localization/PortableObjectImport.ini b/Engine/Config/Localization/PortableObjectImport.ini
new file mode 100644
--- /dev/null
+++ b/Engine/Config/Localization/PortableObjectImport.ini
@@ -0,0 +1,20 @@
+[CommonSettings]
+CulturesToGenerate=ja
+CulturesToGenerate=ko
+CulturesToGenerate=zh
+
+[GatherTextStep0]
+CommandletClass=InternationalizationExport
+bImportLoc=true
+SourcePath=../../../Engine/Content/Localization/Engine
+DestinationPath=../../../Engine/Content/Localization/Engine
+ManifestName=Engine.manifest
+ArchiveName=Engine.archive
+
+[GatherTextStep1]
+CommandletClass=InternationalizationExport
+bImportLoc=true
+SourcePath=../../../Engine/Content/Localization/Editor
+DestinationPath=../../../Engine/Content/Localization/Editor
+ManifestName=Editor.manifest
+ArchiveName=Editor.archive
\ No newline at end of file
diff --git a/Engine/Source/Editor/UnrealEd/Classes/Commandlets/GenerateGatherArchiveCommandlet.h b/Engine/Source/Editor/UnrealEd/Classes/Commandlets/GenerateGatherArchiveCommandlet.h
--- a/Engine/Source/Editor/UnrealEd/Classes/Commandlets/GenerateGatherArchiveCommandlet.h
+++ b/Engine/Source/Editor/UnrealEd/Classes/Commandlets/GenerateGatherArchiveCommandlet.h
@@ -20,7 +20,7 @@ public:
 	virtual int32 Main( const FString& Params ) OVERRIDE;
 	// End UCommandlet Interface
 	
-	bool WriteArchiveToFile( TSharedRef< FJsonObject> ManifestObject, const FString& OutputDirectoryPath, const TCHAR* Culture, const FString& FileName );
+	bool WriteArchiveToFile( TSharedRef< FJsonObject> ArchiveJSONObject, const FString& OutputDirectoryPath, const TCHAR* Culture, const FString& FileName );
 	void BuildArchiveFromManifest( TSharedRef< const FInternationalizationManifest > InManifest, TSharedRef< FInternationalizationArchive > Archive, const FString& SourceCulture, const FString& TargetCulture  );
 	void AppendArchiveData( TSharedRef< const FInternationalizationArchive > InArchiveToAppend, TSharedRef< FInternationalizationArchive > ArchiveCombined );
 	void ConditionTranslation( FLocItem& LocItem );
diff --git a/Engine/Source/Editor/UnrealEd/Classes/Commandlets/InternationalizationExportCommandlet.h b/Engine/Source/Editor/UnrealEd/Classes/Commandlets/InternationalizationExportCommandlet.h
--- a/Engine/Source/Editor/UnrealEd/Classes/Commandlets/InternationalizationExportCommandlet.h
+++ b/Engine/Source/Editor/UnrealEd/Classes/Commandlets/InternationalizationExportCommandlet.h
@@ -341,6 +341,11 @@ public:
 	/* Sets the project name from the Project-Id-Version header entry */
 	FString GetProjectName() { return ProjectName; }
 
+	TArray< TSharedPtr< FPortableObjectEntry > >::TConstIterator GetEntriesIterator()
+	{
+		return Entries.CreateConstIterator();
+	}
+
 private:
 	/** Copying is not supported */
 	FPortableObjectFormatDOM( const FPortableObjectFormatDOM& ); 
@@ -366,6 +371,12 @@ public:
 	// Begin UCommandlet Interface
 	virtual int32 Main(const FString& Params) OVERRIDE;
 	// End UCommandlet Interface
+private:
+	bool DoExport( const FString& SourcePath, const FString& DestinationPath );
+	bool DoImport( const FString& SourcePath, const FString& DestinationPath );
 
+	TArray<FString> CulturesToGenerate;
+	FString ConfigPath;
+	FString SectionName;
 
 };
\ No newline at end of file
diff --git a/Engine/Source/Editor/UnrealEd/Private/Commandlets/GenerateGatherArchiveCommandlet.cpp b/Engine/Source/Editor/UnrealEd/Private/Commandlets/GenerateGatherArchiveCommandlet.cpp
--- a/Engine/Source/Editor/UnrealEd/Private/Commandlets/GenerateGatherArchiveCommandlet.cpp
+++ b/Engine/Source/Editor/UnrealEd/Private/Commandlets/GenerateGatherArchiveCommandlet.cpp
@@ -262,11 +262,11 @@ int32 UGenerateGatherArchiveCommandlet::Main( const FString& Params )
 	return 0;
 }
 
-bool UGenerateGatherArchiveCommandlet::WriteArchiveToFile(TSharedRef< FJsonObject > ManifestObject, const FString& OutputDirectoryPath, const TCHAR* Culture, const FString& FileName)
+bool UGenerateGatherArchiveCommandlet::WriteArchiveToFile(TSharedRef< FJsonObject > ArchiveJSONObject, const FString& OutputDirectoryPath, const TCHAR* Culture, const FString& FileName)
 {
 	FString FullOutPath = FPaths::ConvertRelativePathToFull(OutputDirectoryPath) / Culture / FileName;
 
-	if( !WriteJSONToTextFile(ManifestObject, FullOutPath, SourceControlInfo ) )
+	if( !WriteJSONToTextFile(ArchiveJSONObject, FullOutPath, SourceControlInfo ) )
 	{
 		return false;
 	}
diff --git a/Engine/Source/Editor/UnrealEd/Private/Commandlets/InternationalizationConditioningCommandlet.cpp b/Engine/Source/Editor/UnrealEd/Private/Commandlets/InternationalizationConditioningCommandlet.cpp
--- a/Engine/Source/Editor/UnrealEd/Private/Commandlets/InternationalizationConditioningCommandlet.cpp
+++ b/Engine/Source/Editor/UnrealEd/Private/Commandlets/InternationalizationConditioningCommandlet.cpp
@@ -504,7 +504,11 @@ bool UInternationalizationConditioningCommandlet::ProcessArchive( const FString&
 		TSharedRef<FJsonObject> FinalArchiveJsonObj = MakeShareable( new FJsonObject );
 		ArchiveSerializer.SerializeArchive( InternationalizationArchive, FinalArchiveJsonObj );
 
-		WriteJSONToTextFile( FinalArchiveJsonObj, DestinationArchiveFileName, SourceControlInfo );
+		if( !WriteJSONToTextFile( FinalArchiveJsonObj, DestinationArchiveFileName, SourceControlInfo ) )
+		{
+			UE_LOG(LogInternationalizationConditioningCommandlet, Error, TEXT("Could not save file %s"), *DestinationArchiveFileName);
+			return false;
+		}
 
 		LocPairs.Empty();
 	}
diff --git a/Engine/Source/Editor/UnrealEd/Private/Commandlets/InternationalizationExportCommandlet.cpp b/Engine/Source/Editor/UnrealEd/Private/Commandlets/InternationalizationExportCommandlet.cpp
--- a/Engine/Source/Editor/UnrealEd/Private/Commandlets/InternationalizationExportCommandlet.cpp
+++ b/Engine/Source/Editor/UnrealEd/Private/Commandlets/InternationalizationExportCommandlet.cpp
@@ -234,8 +234,9 @@
 
 FString ConvertSrcLocationToPORef( const FString& InSrcLocation )
 {
-	// Source locaiton format: /Path1/Path2/file.cpp - line 123
+	// Source location format: /Path1/Path2/file.cpp - line 123
 	// PO Reference format: /Path1/Path2/file.cpp:123
+	// @TODO: Note, we assume the source location format here but it could be arbitrary.
 	return InSrcLocation.Replace( TEXT(" - line "), TEXT(":") );
 }
 
@@ -544,6 +545,7 @@
 					{
 						return false;
 					}
+					ProjectName = Header.GetEntryValue(TEXT("Project-Id-Version"));
 				}
 				else
 				{
@@ -958,64 +960,14 @@
 }
 
 
-int32 UInternationalizationExportCommandlet::Main( const FString& Params )
+bool UInternationalizationExportCommandlet::DoExport( const FString& SourcePath, const FString& DestinationPath  )
 {
-	TArray<FString> Tokens;
-	TArray<FString> Switches;
-	TMap<FString, FString> ParamVals;
-	FString ConfigPath;
-	FString SectionName;
-
-
-	UCommandlet::ParseCommandLine(*Params, Tokens, Switches, ParamVals);
-	
-	const FString* ParamVal = ParamVals.Find(FString(TEXT("Config")));
-
-	if ( ParamVal )
-	{
-		ConfigPath = *ParamVal;
-	}
-	else
-	{
-		UE_LOG(LogInternationalizationExportCommandlet, Error, TEXT("No config specified."));
-		return -1;
-	}
-
-	//Set config section
-	ParamVal = ParamVals.Find(FString(TEXT("Section")));
-	
-
-	if ( ParamVal )
-	{
-		SectionName = *ParamVal;
-	}
-	else
-	{
-		UE_LOG(LogInternationalizationExportCommandlet, Error, TEXT("No config section specified."));
-		return -1;
-	}
-
-
-	FString SourcePath; // Source path to the root folder that manifest/archive files live in.
-	if( !GetConfigString( *SectionName, TEXT("SourcePath"), SourcePath, ConfigPath ) )
-	{
-		UE_LOG( LogInternationalizationExportCommandlet, Error, TEXT("No source path specified.") );
-		return -1;
-	}
-
-	FString DestinationPath; // Destination path that we will write files to.
-	if( !GetConfigString( *SectionName, TEXT("DestinationPath"), DestinationPath, ConfigPath ) )
-	{
-		UE_LOG( LogInternationalizationExportCommandlet, Error, TEXT("No destination path specified.") );
-		return -1;
-	}
-
 	// Get manifest name.
 	FString ManifestName;
 	if( !GetConfigString( *SectionName, TEXT("ManifestName"), ManifestName, ConfigPath ) )
 	{
 		UE_LOG( LogInternationalizationExportCommandlet, Error, TEXT("No manifest name specified.") );
-		return -1;
+		return false;
 	}
 
 	// Get archive name.
@@ -1023,17 +975,10 @@
 	if( !( GetConfigString(* SectionName, TEXT("ArchiveName"), ArchiveName, ConfigPath ) ) )
 	{
 		UE_LOG(LogInternationalizationExportCommandlet, Error, TEXT("No archive name specified."));
-		return -1;
+		return false;
 	}
 
-	// Get cultures to generate.
-	TArray<FString> CulturesToGenerate;
-	if( GetConfigArray(*SectionName, TEXT("CulturesToGenerate"), CulturesToGenerate, ConfigPath) == 0 )
-	{
-		UE_LOG(LogInternationalizationExportCommandlet, Error, TEXT("No cultures specified for generation."));
-		return -1;
-	}
-	
+
 	TSharedRef< FInternationalizationManifest > InternationalizationManifest = MakeShareable( new FInternationalizationManifest );
 	// Load the manifest info
 	{
@@ -1041,7 +986,7 @@
 		if( !FPaths::FileExists(ManifestFilePath) )
 		{
 			UE_LOG(LogInternationalizationExportCommandlet, Error, TEXT("Could not find manifest file %s."), *ManifestFilePath);
-			return -1;
+			return false;
 		}
 
 		TSharedPtr<FJsonObject> ManifestJsonObject = ReadJSONTextFile( ManifestFilePath );
@@ -1049,13 +994,13 @@
 		if( !ManifestJsonObject.IsValid() )
 		{
 			UE_LOG(LogInternationalizationExportCommandlet, Error, TEXT("Could not read manifest file %s."), *ManifestFilePath);
-			return -1;
+			return false;
 		}
 
 		FInternationalizationManifestJsonSerializer ManifestSerializer;
 		ManifestSerializer.DeserializeManifest( ManifestJsonObject.ToSharedRef(), InternationalizationManifest );
 	}
-	
 
+
 	// Process the desired cultures
 	for(int32 Culture = 0; Culture < CulturesToGenerate.Num(); Culture++)
@@ -1065,7 +1010,7 @@
 		const FString CulturePath = SourcePath / CultureName;
 		FString ArchiveFileName = CulturePath / ArchiveName;
 		TSharedPtr< FJsonObject > ArchiveJsonObject = NULL;
-		
+
 		if( FPaths::FileExists(ArchiveFileName) )
 		{
 			ArchiveJsonObject = ReadJSONTextFile( ArchiveFileName );
@@ -1124,11 +1069,222 @@
 					if( !FFileHelper::SaveStringToFile(OutputString, *OutputFileName, FFileHelper::EEncodingOptions::ForceUTF8) )
 					{
 						UE_LOG( LogInternationalizationExportCommandlet, Error, TEXT("Could not write file %s"), *OutputFileName );
+						return false;
 					}
 				}
 			}
 		}
 	}
+	return true;
+}
 
+bool UInternationalizationExportCommandlet::DoImport( const FString& SourcePath, const FString& DestinationPath  )
+{
+	// Get manifest name.
+	FString ManifestName;
+	if( !GetConfigString( *SectionName, TEXT("ManifestName"), ManifestName, ConfigPath ) )
+	{
+		UE_LOG( LogInternationalizationExportCommandlet, Error, TEXT("No manifest name specified.") );
+		return false;
+	}
+
+	// Get archive name.
+	FString ArchiveName;
+	if( !( GetConfigString(* SectionName, TEXT("ArchiveName"), ArchiveName, ConfigPath ) ) )
+	{
+		UE_LOG(LogInternationalizationExportCommandlet, Error, TEXT("No archive name specified."));
+		return false;
+	}
+
+	// Process the desired cultures
+	for(int32 Culture = 0; Culture < CulturesToGenerate.Num(); Culture++)
+	{
+		// Load the Portable Object file if found
+		const FString CultureName = CulturesToGenerate[Culture];
+		const FString POFilePath = SourcePath / CultureName + TEXT(".po");
+
+		if( !FPaths::FileExists(POFilePath) )
+		{
+			UE_LOG( LogInternationalizationExportCommandlet, Error, TEXT("Could not find file %s"), *POFilePath );
+			continue;
+		}
+
+		FString POFileContents;
+		if ( !FFileHelper::LoadFileToString( POFileContents, *POFilePath ) )
+		{
+			UE_LOG( LogInternationalizationExportCommandlet, Error, TEXT("Failed to load file %s."), *POFilePath);
+			continue;
+		}
+
+		FPortableObjectFormatDOM PortableObject;
+		if( !PortableObject.FromString( POFileContents ) )
+		{
+			UE_LOG( LogInternationalizationExportCommandlet, Error, TEXT("Failed to parse Portable Object file %s."), *POFilePath);
+			continue;
+		}
+
+		if( PortableObject.GetProjectName() != ManifestName.Replace(TEXT(".manifest"), TEXT("")) )
+		{
+			UE_LOG( LogInternationalizationExportCommandlet, Warning, TEXT("The project name (%s) in the file (%s) did not match the target manifest project (%s)."), *POFilePath, *PortableObject.GetProjectName(), *ManifestName.Replace(TEXT(".manifest"), TEXT("")));
+		}
+
+
+		const FString DestinationCulturePath = DestinationPath / CultureName;
+		FString ArchiveFileName = DestinationCulturePath / ArchiveName;
+		
+		if( !FPaths::FileExists(ArchiveFileName) )
+		{
+			UE_LOG( LogInternationalizationExportCommandlet, Error, TEXT("Failed to find destination archive %s."), *ArchiveFileName);
+			continue;
+		}
+
+		TSharedPtr< FJsonObject > ArchiveJsonObject = NULL;
+		ArchiveJsonObject = ReadJSONTextFile( ArchiveFileName );
+
+		FInternationalizationArchiveJsonSerializer ArchiveSerializer;
+		TSharedRef< FInternationalizationArchive > InternationalizationArchive = MakeShareable( new FInternationalizationArchive );
+		ArchiveSerializer.DeserializeArchive( ArchiveJsonObject.ToSharedRef(), InternationalizationArchive );
+
+		bool bModifiedArchive = false;
+		{
+			for( auto EntryIter = PortableObject.GetEntriesIterator(); EntryIter; ++EntryIter )
+			{
+				auto POEntry = *EntryIter;
+				if( POEntry->MsgId.IsEmpty() || POEntry->MsgStr.Num() == 0 || POEntry->MsgStr[0].IsEmpty() )
+				{
+					// We ignore the header entry or entries with no translation.
+					continue;
+				}
+
+				// Some warning messages for data we don't process at the moment
+				if( !POEntry->MsgIdPlural.IsEmpty() || POEntry->MsgStr.Num() > 1 )
+				{
+					UE_LOG( LogInternationalizationExportCommandlet, Error, TEXT("Portable Object entry has plural form we did not process.  File: %s  MsgCtxt: %s  MsgId: %s"), *POFilePath, *POEntry->MsgCtxt, *POEntry->MsgId );
+				}
+				
+				const FString& Namespace = POEntry->MsgCtxt;
+				const FString& SourceText = POEntry->MsgId;
+				const FString& Translation = POEntry->MsgStr[0];
+
+				//@TODO: Take into account optional entries and entries that differ by keymetadata.  Ex. Each optional entry needs a unique msgCtxt
+				TSharedPtr< FArchiveEntry > FoundEntry = InternationalizationArchive->FindEntryBySource( Namespace, SourceText, NULL );
+				if( !FoundEntry.IsValid() )
+				{
+					UE_LOG(LogInternationalizationExportCommandlet, Error, TEXT("Could not find corresponding archive entry for PO entry.  File: %s  MsgCtxt: %s  MsgId: %s"), *POFilePath, *POEntry->MsgCtxt, *POEntry->MsgId );
+					continue;
+				}
+				
+				if( FoundEntry->Translation != Translation )
+				{
+					FoundEntry->Translation = Translation;
+					bModifiedArchive = true;
+				}
+			}
+		}
+
+		if( bModifiedArchive )
+		{
+			TSharedRef<FJsonObject> FinalArchiveJsonObj = MakeShareable( new FJsonObject );
+			ArchiveSerializer.SerializeArchive( InternationalizationArchive, FinalArchiveJsonObj );
+
+			if( !WriteJSONToTextFile(FinalArchiveJsonObj, ArchiveFileName, SourceControlInfo ) )
+			{
+				UE_LOG( LogInternationalizationExportCommandlet, Error, TEXT("Failed to write archive to %s."), *ArchiveFileName );				
+				return false;
+			}
+		}
+	}
+
+	return true;
+}
+
+int32 UInternationalizationExportCommandlet::Main( const FString& Params )
+{
+	TArray<FString> Tokens;
+	TArray<FString> Switches;
+	TMap<FString, FString> ParamVals;
+
+
+	UCommandlet::ParseCommandLine(*Params, Tokens, Switches, ParamVals);
+	
+	const FString* ParamVal = ParamVals.Find(FString(TEXT("Config")));
+
+	if ( ParamVal )
+	{
+		ConfigPath = *ParamVal;
+	}
+	else
+	{
+		UE_LOG(LogInternationalizationExportCommandlet, Error, TEXT("No config specified."));
+		return -1;
+	}
+
+	//Set config section
+	ParamVal = ParamVals.Find(FString(TEXT("Section")));
+	
+
+	if ( ParamVal )
+	{
+		SectionName = *ParamVal;
+	}
+	else
+	{
+		UE_LOG(LogInternationalizationExportCommandlet, Error, TEXT("No config section specified."));
+		return -1;
+	}
+
+
+	FString SourcePath; // Source path to the root folder that manifest/archive files live in.
+	if( !GetConfigString( *SectionName, TEXT("SourcePath"), SourcePath, ConfigPath ) )
+	{
+		UE_LOG( LogInternationalizationExportCommandlet, Error, TEXT("No source path specified.") );
+		return -1;
+	}
+
+	FString DestinationPath; // Destination path that we will write files to.
+	if( !GetConfigString( *SectionName, TEXT("DestinationPath"), DestinationPath, ConfigPath ) )
+	{
+		UE_LOG( LogInternationalizationExportCommandlet, Error, TEXT("No destination path specified.") );
+		return -1;
+	}
+
+	// Get cultures to generate.
+	if( GetConfigArray(*SectionName, TEXT("CulturesToGenerate"), CulturesToGenerate, ConfigPath) == 0 )
+	{
+		UE_LOG(LogInternationalizationExportCommandlet, Error, TEXT("No cultures specified for generation."));
+		return -1;
+	}
+
+
+	bool bDoExport = false;
+	bool bDoImport = false;
+
+	GetConfigBool( *SectionName, TEXT("bExportLoc"), bDoExport, ConfigPath );
+	GetConfigBool( *SectionName, TEXT("bImportLoc"), bDoImport, ConfigPath );
+	
+	if( !bDoImport && !bDoExport )
+	{
+		UE_LOG(LogInternationalizationExportCommandlet, Error, TEXT("Import/Export operation not detected.  Use bExportLoc or bImportLoc in config section."));
+		return -1;
+	}
+
+	if( bDoExport )
+	{
+		if( !DoExport( SourcePath, DestinationPath ) )
+		{
+			UE_LOG(LogInternationalizationExportCommandlet, Error, TEXT("Failed to export localization files."));
+			return -1;
+		}
+	}
+
+	if( bDoImport )
+	{
+		if( !DoImport( SourcePath, DestinationPath ) )
+		{
+			UE_LOG(LogInternationalizationExportCommandlet, Error, TEXT("Failed to import localization files."));
+			return -1;
+		}
+	}
+
 	return 0;
 }
