- ProcessImportMeshInfluences optimization

- Made to one loop to calculate in-game bone influence

diff --git a/Engine/Source/Editor/UnrealEd/Private/Factories/SkeletalMeshImport.cpp b/Engine/Source/Editor/UnrealEd/Private/Factories/SkeletalMeshImport.cpp
--- a/Engine/Source/Editor/UnrealEd/Private/Factories/SkeletalMeshImport.cpp
+++ b/Engine/Source/Editor/UnrealEd/Private/Factories/SkeletalMeshImport.cpp
@@ -267,69 +267,63 @@ void ProcessImportMeshInfluences(FSkeletalMeshImportData& ImportData)
 	};
 	Influences.Sort( FCompareVertexIndex() );
 
-	// Remove more than allowed number of weights by removing least important influences (setting them to 0). 
-	// Relies on influences sorted by vertex index and weight and the code actually removing the influences below.
-	int32 LastVertexIndex		= INDEX_NONE;
-	int32 InfluenceCount		= 0;
-	for(  int32 i=0; i<Influences.Num(); i++ )
-	{		
-		if( ( LastVertexIndex != Influences[i].VertexIndex ) )
-		{
-			InfluenceCount	= 0;
-			LastVertexIndex	= Influences[i].VertexIndex;
-		}
+	TArray <VRawBoneInfluence> NewInfluences;
+	int32	LastNewInfluenceIndex=0;
+	int32	LastVertexIndex		= INDEX_NONE;
+	int32	InfluenceCount			= 0;
 
-		InfluenceCount++;
-		if( InfluenceCount > MAX_TOTAL_INFLUENCES || LastVertexIndex >= Points.Num() )
-		{
-			Influences[i].Weight = 0.f;
-		}
-	}
+	float TotalWeight		= 0.f;
+	const float MINWEIGHT   = 0.01f;
 
-	// Remove influences below a certain threshold.
-	int32 RemovedInfluences	= 0;
-	const float MINWEIGHT	= 0.01f; // 1%
-	for( int32 i=Influences.Num()-1; i>=0; i-- )
-	{
-		if( Influences[i].Weight < MINWEIGHT )
-		{
-			Influences.RemoveAt(i);
-			RemovedInfluences++;
-		}
-	}
 
-	// Renormalize influence weights.
-	int32	LastInfluenceCount	= 0;
-	InfluenceCount			= 0;
-	LastVertexIndex			= INDEX_NONE;
-	float TotalWeight		= 0.f;
 	for( int32 i=0; i<Influences.Num(); i++ )
 	{
-		if( LastVertexIndex != Influences[i].VertexIndex )
+		// we found next verts, normalize it now
+		if (LastVertexIndex != Influences[i].VertexIndex )
 		{
-			LastInfluenceCount	= InfluenceCount;
-			InfluenceCount		= 0;
-
 			// Normalize the last set of influences.
-			if( LastInfluenceCount && (TotalWeight != 1.0f) )
-			{				
+			if (InfluenceCount && (TotalWeight != 1.0f))
+			{
 				float OneOverTotalWeight = 1.f / TotalWeight;
-				for( int r=0; r<LastInfluenceCount; r++)
+				for (int r = 0; r < InfluenceCount; r++)
+				{
+					NewInfluences[LastNewInfluenceIndex - r].Weight *= OneOverTotalWeight;
+				}
+			}
+
+			// now we insert missing verts
+			if (LastVertexIndex != INDEX_NONE)
+			{
+				int32 CurrentVertexIndex = Influences[i].VertexIndex;
+				for(int32 j=LastVertexIndex+1; j<CurrentVertexIndex; j++)
 				{
-					Influences[i-r-1].Weight *= OneOverTotalWeight;
+					// Add a 0-bone weight if none other present (known to happen with certain MAX skeletal setups).
+					LastNewInfluenceIndex = NewInfluences.AddUninitialized();
+					NewInfluences[LastNewInfluenceIndex].VertexIndex	= j;
+					NewInfluences[LastNewInfluenceIndex].BoneIndex		= 0;
+					NewInfluences[LastNewInfluenceIndex].Weight		= 1.f;
 				}
 			}
-			TotalWeight		= 0.f;				
-			LastVertexIndex = Influences[i].VertexIndex;							
+
+			// clear to count next one
+			InfluenceCount = 0;
+			TotalWeight = 0.f;
+			LastVertexIndex = Influences[i].VertexIndex;
+		}
+		
+		// if less than min weight, or it's more than 8, then we clear it to use weight
+		if (Influences[i].Weight > MINWEIGHT && InfluenceCount < MAX_TOTAL_INFLUENCES)
+		{
+			LastNewInfluenceIndex = NewInfluences.Add(Influences[i]);
+			InfluenceCount++;
+			TotalWeight	+= Influences[i].Weight;
 		}
-		InfluenceCount++;
-		TotalWeight	+= Influences[i].Weight;			
 	}
 
+	Influences = NewInfluences;
+
 	// Ensure that each vertex has at least one influence as e.g. CreateSkinningStream relies on it.
 	// The below code relies on influences being sorted by vertex index.
-	LastVertexIndex = -1;
-	InfluenceCount	= 0;
 	if( Influences.Num() == 0 )
 	{
 		UnFbx::FFbxImporter* FFbxImporter = UnFbx::FFbxImporter::GetInstance();
@@ -343,22 +337,22 @@ void ProcessImportMeshInfluences(FSkeletalMeshImportData& ImportData)
 			Influences[WedgeIdx].BoneIndex = 0;
 			Influences[WedgeIdx].Weight = 1.0f;
 		}		
-	}
-	for( int32 i=0; i<Influences.Num(); i++ )
-	{
-		int32 CurrentVertexIndex = Influences[i].VertexIndex;
-
-		if( LastVertexIndex != CurrentVertexIndex )
+		for(int32 i=0; i<Influences.Num(); i++)
 		{
-			for( int32 j=LastVertexIndex+1; j<CurrentVertexIndex; j++ )
+			int32 CurrentVertexIndex = Influences[i].VertexIndex;
+
+			if(LastVertexIndex != CurrentVertexIndex)
 			{
-				// Add a 0-bone weight if none other present (known to happen with certain MAX skeletal setups).
-				Influences.InsertUninitialized(i,1);
-				Influences[i].VertexIndex	= j;
-				Influences[i].BoneIndex		= 0;
-				Influences[i].Weight		= 1.f;
+				for(int32 j=LastVertexIndex+1; j<CurrentVertexIndex; j++)
+				{
+					// Add a 0-bone weight if none other present (known to happen with certain MAX skeletal setups).
+					Influences.InsertUninitialized(i, 1);
+					Influences[i].VertexIndex	= j;
+					Influences[i].BoneIndex		= 0;
+					Influences[i].Weight		= 1.f;
+				}
+				LastVertexIndex = CurrentVertexIndex;
 			}
-			LastVertexIndex = CurrentVertexIndex;
 		}
 	}
 }
