Fix for global shader map for all platforms being process when using cookbythebook.

diff --git a/Engine/Source/Editor/UnrealEd/Private/Commandlets/CookCommandlet.cpp b/Engine/Source/Editor/UnrealEd/Private/Commandlets/CookCommandlet.cpp
--- a/Engine/Source/Editor/UnrealEd/Private/Commandlets/CookCommandlet.cpp
+++ b/Engine/Source/Editor/UnrealEd/Private/Commandlets/CookCommandlet.cpp
@@ -507,7 +507,7 @@ int32 UCookCommandlet::Main(const FString& CmdLineParams)
 	else
 	{
 		ITargetPlatformManagerModule& TPM = GetTargetPlatformManagerRef();
-		const TArray<ITargetPlatform*>& Platforms = TPM.GetTargetPlatforms();
+		const TArray<ITargetPlatform*>& Platforms = TPM.GetActiveTargetPlatforms();
 	
 		// Local sandbox file wrapper. This will be used to handle path conversions,
 		// but will not be used to actually write/read files so we can safely
