#TTP 331119 - Regression: EDITOR: In the world settings, clicking the find in content browser button for game mode blueprints does not highlight and bring it into view

#branch UE4

#proj Editor.UnrealEd

#summary The Game info customizer now browses to the correct blueprint in the content browser.

#change Syncing to content requires the actual object rather than its class be passed. The code now passes the actual blueprints.

diff --git a/Engine/Source/Editor/UnrealEd/Public/GameModeInfoCustomizer.h b/Engine/Source/Editor/UnrealEd/Public/GameModeInfoCustomizer.h
--- a/Engine/Source/Editor/UnrealEd/Public/GameModeInfoCustomizer.h
+++ b/Engine/Source/Editor/UnrealEd/Public/GameModeInfoCustomizer.h
@@ -221,33 +221,40 @@ public:
 
 	bool CanBrowseDefaultClass(FName ClassPropertyName) const
 	{
-		return (OnGetDefaultClass(ClassPropertyName) != NULL);
+		return CanSyncToClass(OnGetDefaultClass(ClassPropertyName));
 	}
 
 	void OnBrowseDefaultClassClicked(FName ClassPropertyName)
 	{
-		UClass* Class = const_cast<UClass*>( OnGetDefaultClass(ClassPropertyName) );
-		if(Class != NULL)
-		{
-			TArray<UObject*> SyncObjects;
-			SyncObjects.Add(Class);
-			GEditor->SyncBrowserToObjects(SyncObjects);
-		}
+		SyncBrowserToClass(OnGetDefaultClass(ClassPropertyName));
 	}
 
 	bool CanBrowseGameMode() const
 	{
-		return (GetCurrentGameModeClass() != NULL);
+		return CanSyncToClass(GetCurrentGameModeClass());
 	}
 
 	void OnBrowseGameModeClicked()
 	{
-		UClass* Class = const_cast<UClass*>( GetCurrentGameModeClass() );
-		if (Class != NULL)
+		SyncBrowserToClass(GetCurrentGameModeClass());
+	}
+
+	bool CanSyncToClass(const UClass* Class) const
+	{
+		return (Class != NULL && Class->ClassGeneratedBy != NULL);
+	}
+
+	void SyncBrowserToClass(const UClass* Class)
+	{
+		if (CanSyncToClass(Class))
 		{
-			TArray<UObject*> SyncObjects;
-			SyncObjects.Add(Class);
-			GEditor->SyncBrowserToObjects(SyncObjects);
+			UBlueprint* Blueprint = Cast<UBlueprint>(Class->ClassGeneratedBy);
+			if (ensure(Blueprint != NULL))
+			{
+				TArray<UObject*> SyncObjects;
+				SyncObjects.Add(Blueprint);
+				GEditor->SyncBrowserToObjects(SyncObjects);
+			}
 		}
 	}
 
