#TTP 331023 - LIVE: EDITOR: Undoing a selection when multiple items are created only undoes one at a time.

#branch UE4

#proj Editor.UnrealEd

#summary Creating multiple actors using drag and drop from the content browser now undoes as a single transaction

#add Added a scoped transaction surrounding the for loop used to create the individual actors. This cancels itself if there is only a single actor and the transaction is handled by the transaction in AssetSelection.cpp - PrivateAddActor

diff --git a/Engine/Source/Editor/UnrealEd/Private/LevelEditorViewport.cpp b/Engine/Source/Editor/UnrealEd/Private/LevelEditorViewport.cpp
--- a/Engine/Source/Editor/UnrealEd/Private/LevelEditorViewport.cpp
+++ b/Engine/Source/Editor/UnrealEd/Private/LevelEditorViewport.cpp
@@ -861,94 +861,109 @@ bool FLevelEditorViewportClient::DropObjectsOnActor(FViewportCursorLocation& Cur
 
 	if ( DroppedUponActor != NULL )
 	{
-		bResult = DroppedObjects.Num() > 0;
-
 		const FVector TargetLocation = DroppedLocation ? *DroppedLocation : DroppedUponActor->GetActorLocation();
-		for ( int32 DroppedObjectsIdx = 0; DroppedObjectsIdx < DroppedObjects.Num(); ++DroppedObjectsIdx )
+		if(DroppedObjects.Num() > 1)
 		{
-			UObject* AssetObj = DroppedObjects[DroppedObjectsIdx];
-			ensure( AssetObj );
+			bResult = true;
 
-			// Attempt to apply the dropped asset as a material to the actor
-			const bool bAppliedToActor = ( FactoryToUse == NULL ) ? AttemptApplyObjToActor( AssetObj, DroppedUponActor, DroppedUponSlot ) : false;
-			if ( !bAppliedToActor )
+			// Create a transaction if we have more than 1 item. This causes them all to get "un-created" together
+			FScopedTransaction Transaction( NSLOCTEXT("UnrealEd", "CreateActors", "Create Actors") );
+			for ( auto DroppedObject : DroppedObjects )
 			{
-				// Actor
-				GEditor->ClickLocation = TargetLocation;
-				GEditor->ClickPlane = FPlane(TargetLocation, FVector(0.0f, 0.0f, 1.0f));
-				TArray<FHitResult> Hits;
-				FCollisionQueryParams Param(TEXT("DragDropTrace"), true);
-				// grab the bounds of the actor we're being dropped upon
-				FVector DUAOrigin, DUAExtent;
-				DroppedUponActor->GetActorBounds(true, DUAOrigin, DUAExtent);
-				// now calculate how far back we need to be to collide a ray with it - extend it slightly
-				float WorldDistanceMultiplier = 0.0f;
-				switch( Cursor.GetViewportClient()->GetViewportType() )
+				if (!DropSingleObjectOnActor(Cursor, DroppedObject, DroppedUponActor, DroppedUponSlot, TargetLocation, ObjectFlags, OutNewActors, bUsedHitProxy, bSelectActors, FactoryToUse))
 				{
-				case LVT_OrthoXY:	WorldDistanceMultiplier = DUAExtent.Z * 1.1; break; // Top
-				case LVT_OrthoXZ:	WorldDistanceMultiplier = DUAExtent.Y * 1.1; break; // Front
-				case LVT_OrthoYZ:	WorldDistanceMultiplier = DUAExtent.X * 1.1; break; // Side
+					bResult = false;
 				}
-				
-				if ( GetWorld()->LineTraceMulti(Hits, Cursor.GetOrigin() - Cursor.GetDirection() * WorldDistanceMultiplier, Cursor.GetOrigin() + Cursor.GetDirection() * HALF_WORLD_MAX, ECC_Visibility, Param) )
+			}
+		}
+		else if (DroppedObjects.Num() == 1)
+		{
+			bResult = DropSingleObjectOnActor(Cursor, DroppedObjects[0], DroppedUponActor, DroppedUponSlot, TargetLocation, ObjectFlags, OutNewActors, bUsedHitProxy, bSelectActors, FactoryToUse);
+		}
+	}
+
+	return bResult;
+}
+
+bool FLevelEditorViewportClient::DropSingleObjectOnActor(FViewportCursorLocation& Cursor, UObject* DroppedObject, AActor* DroppedUponActor, int32 DroppedUponSlot, const FVector& DroppedLocation, EObjectFlags ObjectFlags, TArray<AActor*>& OutNewActors, bool bUsedHitProxy /*= true*/, bool bSelectActors /*= true*/, class UActorFactory* FactoryToUse /*= NULL*/ )
+{
+	if (ensure( DroppedObject ))
+	{
+		// Attempt to apply the dropped asset as a material to the actor
+		const bool bAppliedToActor = ( FactoryToUse == NULL ) ? AttemptApplyObjToActor( DroppedObject, DroppedUponActor, DroppedUponSlot ) : false;
+		if ( !bAppliedToActor )
+		{
+			// Actor
+			GEditor->ClickLocation = DroppedLocation;
+			GEditor->ClickPlane = FPlane(DroppedLocation, FVector(0.0f, 0.0f, 1.0f));
+			TArray<FHitResult> Hits;
+			FCollisionQueryParams Param(TEXT("DragDropTrace"), true);
+			// grab the bounds of the actor we're being dropped upon
+			FVector DUAOrigin, DUAExtent;
+			DroppedUponActor->GetActorBounds(true, DUAOrigin, DUAExtent);
+			// now calculate how far back we need to be to collide a ray with it - extend it slightly
+			float WorldDistanceMultiplier = 0.0f;
+			switch( Cursor.GetViewportClient()->GetViewportType() )
+			{
+			case LVT_OrthoXY:	WorldDistanceMultiplier = DUAExtent.Z * 1.1; break; // Top
+			case LVT_OrthoXZ:	WorldDistanceMultiplier = DUAExtent.Y * 1.1; break; // Front
+			case LVT_OrthoYZ:	WorldDistanceMultiplier = DUAExtent.X * 1.1; break; // Side
+			}
+
+			if ( GetWorld()->LineTraceMulti(Hits, Cursor.GetOrigin() - Cursor.GetDirection() * WorldDistanceMultiplier, Cursor.GetOrigin() + Cursor.GetDirection() * HALF_WORLD_MAX, ECC_Visibility, Param) )
+			{
+				bool FoundMatch = false;
+
+				// We only care about the collision with the hit proxy actor as line trace will report hits on hidden actors too.
+				for( int32 HitIdx = 0; HitIdx < Hits.Num(); HitIdx++ )
 				{
-					bool FoundMatch = false;
+					const FHitResult& Hit = Hits[HitIdx];
+					if ( Hit.GetActor() == DroppedUponActor )
+					{
+						GEditor->ClickLocation = Hit.Location;
+						GEditor->ClickPlane = FPlane( Hit.Location,Hit.Normal );
+						FoundMatch = true;
+						break;
+					}
+				}
 
-					// We only care about the collision with the hit proxy actor as line trace will report hits on hidden actors too.
+				// If unsuccessful, use the first visible actor/component* instead
+				// *Limit this just to models/brush for the time being as its only needed for TTP#307379
+				if ( !FoundMatch )
+				{
 					for( int32 HitIdx = 0; HitIdx < Hits.Num(); HitIdx++ )
 					{
 						const FHitResult& Hit = Hits[HitIdx];
-						if ( Hit.GetActor() == DroppedUponActor )
+						if((Hit.GetActor() && !Hit.GetActor()->IsHiddenEd())
+							|| (Hit.Component.IsValid() && Hit.Component->IsVisibleInEditor() && Hit.Component->IsA(UModelComponent::StaticClass()) && DroppedUponActor->IsA(ABrush::StaticClass())))
 						{
 							GEditor->ClickLocation = Hit.Location;
-							GEditor->ClickPlane = FPlane( Hit.Location,Hit.Normal );
-							FoundMatch = true;
+							GEditor->ClickPlane = FPlane( Hit.Location, Hit.Normal );
 							break;
 						}
 					}
-
-					// If unsuccessful, use the first visible actor/component* instead
-					// *Limit this just to models/brush for the time being as its only needed for TTP#307379
-					if ( !FoundMatch )
-					{
-						for( int32 HitIdx = 0; HitIdx < Hits.Num(); HitIdx++ )
-						{
-							const FHitResult& Hit = Hits[HitIdx];
-							if((Hit.GetActor() && !Hit.GetActor()->IsHiddenEd())
-							|| (Hit.Component.IsValid() && Hit.Component->IsVisibleInEditor() && Hit.Component->IsA(UModelComponent::StaticClass()) && DroppedUponActor->IsA(ABrush::StaticClass())))
-							{
-								GEditor->ClickLocation = Hit.Location;
-								GEditor->ClickPlane = FPlane( Hit.Location, Hit.Normal );
-								break;
-							}
-						}
-					}
 				}
-				else
+			}
+			else
+			{
+				// if the line check fails clear the hit proxy flag so that we place it on the background
+				if ( Cursor.GetViewportType() == LVT_Perspective )
 				{
-					// if the line check fails clear the hit proxy flag so that we place it on the background
-					if ( Cursor.GetViewportType() == LVT_Perspective )
-					{
-						bUsedHitProxy = false;
-					}
+					bUsedHitProxy = false;
 				}
+			}
 
-				// Attempt to create actors from the dropped object
-				TArray<AActor*> NewActors = AttemptDropObjAsActors(GetWorld()->GetCurrentLevel(), AssetObj, Cursor, bUsedHitProxy, bSelectActors, ObjectFlags, FactoryToUse);
+			// Attempt to create actors from the dropped object
+			TArray<AActor*> NewActors = AttemptDropObjAsActors(GetWorld()->GetCurrentLevel(), DroppedObject, Cursor, bUsedHitProxy, bSelectActors, ObjectFlags, FactoryToUse);
 
-				if ( NewActors.Num() > 0 )
-				{
-					OutNewActors.Append( NewActors );
-				}
-				else
-				{
-					bResult = false;
-				}
+			if ( NewActors.Num() > 0 )
+			{
+				OutNewActors.Append( NewActors );
+				return true;
 			}
 		}
 	}
-
-	return bResult;
+	return false;
 }
 
 bool FLevelEditorViewportClient::DropObjectsOnBSPSurface( FSceneView* View, FViewportCursorLocation& Cursor, const TArray<UObject*>& DroppedObjects, HModel* TargetProxy, EObjectFlags ObjectFlags, TArray<AActor*>& OutNewActors, bool bSelectActors, UActorFactory* FactoryToUse )
diff --git a/Engine/Source/Editor/UnrealEd/Public/LevelEditorViewport.h b/Engine/Source/Editor/UnrealEd/Public/LevelEditorViewport.h
--- a/Engine/Source/Editor/UnrealEd/Public/LevelEditorViewport.h
+++ b/Engine/Source/Editor/UnrealEd/Public/LevelEditorViewport.h
@@ -609,6 +609,24 @@ private:
 	*/
 	bool DropObjectsOnActor(struct FViewportCursorLocation& Cursor, const TArray<UObject*>& DroppedObjects, AActor* DroppedUponActor, int32 DroppedUponSlot, FVector* DroppedLocation, EObjectFlags ObjectFlags, TArray<AActor*>& OutNewActors, bool bUsedHitProxy = true, bool bSelectActors = true, class UActorFactory* FactoryToUse = NULL);
 
+	/**
+	* Called when an asset is dropped upon an existing actor.
+	*
+	* @param	Cursor				Mouse cursor location
+	* @param	DroppedObject		Object dropped into the viewport
+	* @param	DroppedUponActor	The actor that we are dropping upon
+	* @param    DroppedUponSlot		The material slot/submesh that was identified as the drop location.  If unknown use -1.
+	* @param	DroppedLocation		The location that we're dropping the objects
+	* @param	ObjectFlags			The object flags to place on the actors that this function spawns.
+	* @param	OutNewActors		The list of actors created while dropping
+	* @param	bUsedHitProxy		Whether or not a hit proxy was used for spawning
+	* @param	bSelectActors		If true, select the newly dropped actors (defaults: true)
+	* @param	FactoryToUse		The preferred actor factory to use (optional)
+	*
+	* @return	true if the drop operation was successfully handled; false otherwise
+	*/
+	bool DropSingleObjectOnActor(struct FViewportCursorLocation& Cursor, UObject* DroppedObject, AActor* DroppedUponActor, int32 DroppedUponSlot, const FVector& DroppedLocation, EObjectFlags ObjectFlags, TArray<AActor*>& OutNewActors, bool bUsedHitProxy = true, bool bSelectActors = true, class UActorFactory* FactoryToUse = NULL);
+
 	/**
 	 * Called when an asset is dropped upon a BSP surface.
 	 *
