#UE4 - Removed message entry for old FBX files. There are no known issues with old FBX files.  We still log this information for debugging.

diff --git a/Engine/Source/Editor/UnrealEd/Private/Fbx/FbxMainImport.cpp b/Engine/Source/Editor/UnrealEd/Private/Fbx/FbxMainImport.cpp
--- a/Engine/Source/Editor/UnrealEd/Private/Fbx/FbxMainImport.cpp
+++ b/Engine/Source/Editor/UnrealEd/Private/Fbx/FbxMainImport.cpp
@@ -594,8 +594,7 @@ bool FFbxImporter::OpenFile(FString Filename, bool bParseStatistics, bool bForSc
 			const FText WarningText = FText::Format(
 				NSLOCTEXT("UnrealEd", "Warning_OutOfDateFBX", "An out of date FBX has been detected.\nImporting different versions of FBX files than the SDK version can cause undesirable results.\n\nFile Version: {0}\nSDK Version: {1}" ),
 				FText::FromString(FileVerStr), FText::FromString(SDKVerStr) );
-			
-			AddTokenizedErrorMessage(FTokenizedMessage::Create(EMessageSeverity::Warning, WarningText));
+
 		}
 	}
 
