- Added warning when morphtarget LOd is added without finding correct morphtarget first.
- Fixed trace in BP construction script not working if you cook.

- Accidently checked in both of them.

diff --git a/Engine/Source/Editor/UnrealEd/Private/Commandlets/CookCommandlet.cpp b/Engine/Source/Editor/UnrealEd/Private/Commandlets/CookCommandlet.cpp
--- a/Engine/Source/Editor/UnrealEd/Private/Commandlets/CookCommandlet.cpp
+++ b/Engine/Source/Editor/UnrealEd/Private/Commandlets/CookCommandlet.cpp
@@ -422,6 +422,7 @@ bool UCookCommandlet::SaveCookedPackage( UPackage* Package, uint32 SaveFlags, bo
 				if (World)
 				{
 					World->PersistentLevel->OwningWorld = World;
+					World->InitWorld(UWorld::InitializationValues().RequiresHitProxies(false).ShouldSimulatePhysics(false).EnableTraceCollision(false).CreateNavigation(false).AllowAudioPlayback(false).CreatePhysicsScene(true).CreateWorldComposition(false));
 				}
 
 				if( PlatFilename.Len() >= PLATFORM_MAX_FILEPATH_LENGTH )
diff --git a/Engine/Source/Editor/UnrealEd/Private/Fbx/FbxSkeletalMeshImport.cpp b/Engine/Source/Editor/UnrealEd/Private/Fbx/FbxSkeletalMeshImport.cpp
--- a/Engine/Source/Editor/UnrealEd/Private/Fbx/FbxSkeletalMeshImport.cpp
+++ b/Engine/Source/Editor/UnrealEd/Private/Fbx/FbxSkeletalMeshImport.cpp
@@ -2519,25 +2519,36 @@ void UnFbx::FFbxImporter::ImportMorphTargetsInternal( TArray<FbxNode*>& SkelMesh
 
 							FSkeletalMeshImportData ImportData;
 
-							// now we get a shape for whole mesh, import to unreal as a morph target
-							// @todo AssetImportData do we need import data for this temp mesh?
-							UFbxSkeletalMeshImportData* TmpMeshImportData = NULL;
-							USkeletalMesh* TmpSkeletalMesh = (USkeletalMesh*)ImportSkeletalMesh( GetTransientPackage(), SkelMeshNodeArray, NAME_None, (EObjectFlags)0, TmpMeshImportData, FPaths::GetBaseFilename(InFilename), &FbxShapeArray, &ImportData, false );
-							TempMeshes.Add( TmpSkeletalMesh );
-
 							// See if this morph target already exists.
 							UMorphTarget * Result = FindObject<UMorphTarget>(BaseSkelMesh,*ShapeName);
+							// we only create new one for LOD0, otherwise don't create new one
 							if(!Result)
 							{
-								Result = NewNamedObject<UMorphTarget>(BaseSkelMesh, FName(*ShapeName));
+								if (LODIndex == 0)
+								{
+									Result = NewNamedObject<UMorphTarget>(BaseSkelMesh, FName(*ShapeName));
+								}
+								else
+								{
+									AddTokenizedErrorMessage(FTokenizedMessage::Create(EMessageSeverity::Error, FText::Format(FText::FromString("Could not find the {0} morphtarget for LOD {1}. \
+																													Make sure the name for morphtarget matches with LOD 0"), FText::FromString(ShapeName), FText::FromString(FString::FromInt(LODIndex)))));
+								}
 							}
 
-							MorphTargets.Add(Result);
-					
-							// Process the skeletal mesh on a separate thread
-							FAsyncTask<FAsyncImportMorphTargetWork>* NewWork = new (PendingWork) FAsyncTask<FAsyncImportMorphTargetWork>( TmpSkeletalMesh, LODIndex, ImportData, ImportOptions->bKeepOverlappingVertices );
-							NewWork->StartBackgroundTask();
-							++CurrentShapeIndex;
+							if (Result)
+							{
+								// now we get a shape for whole mesh, import to unreal as a morph target
+								// @todo AssetImportData do we need import data for this temp mesh?
+								UFbxSkeletalMeshImportData* TmpMeshImportData = NULL;
+								USkeletalMesh* TmpSkeletalMesh = (USkeletalMesh*)ImportSkeletalMesh(GetTransientPackage(), SkelMeshNodeArray, NAME_None, (EObjectFlags)0, TmpMeshImportData, FPaths::GetBaseFilename(InFilename), &FbxShapeArray, &ImportData, false);
+								TempMeshes.Add(TmpSkeletalMesh);
+								MorphTargets.Add(Result);
+
+								// Process the skeletal mesh on a separate thread
+								FAsyncTask<FAsyncImportMorphTargetWork>* NewWork = new (PendingWork)FAsyncTask<FAsyncImportMorphTargetWork>(TmpSkeletalMesh, LODIndex, ImportData, ImportOptions->bKeepOverlappingVertices);
+								NewWork->StartBackgroundTask();
+								++CurrentShapeIndex;
+							}
 						}
 					}
 				}
