Added comment

diff --git a/Engine/Source/Editor/UnrealEd/Private/Commandlets/CookCommandlet.cpp b/Engine/Source/Editor/UnrealEd/Private/Commandlets/CookCommandlet.cpp
--- a/Engine/Source/Editor/UnrealEd/Private/Commandlets/CookCommandlet.cpp
+++ b/Engine/Source/Editor/UnrealEd/Private/Commandlets/CookCommandlet.cpp
@@ -422,6 +422,7 @@ bool UCookCommandlet::SaveCookedPackage( UPackage* Package, uint32 SaveFlags, bo
 				if (World)
 				{
 					World->PersistentLevel->OwningWorld = World;
+					// we need to initialize the world - at least need physics scene since BP construction script runs during cooking, otherwise trace won't work
 					World->InitWorld(UWorld::InitializationValues().RequiresHitProxies(false).ShouldSimulatePhysics(false).EnableTraceCollision(false).CreateNavigation(false).AllowAudioPlayback(false).CreatePhysicsScene(true).CreateWorldComposition(false));
 				}
 
