InternationalizationExportCommandlet:
- Export was using SourcePath as DestinationPath
- Also added ability to destination filename as well as path, so we can use the "Save File" dialog

diff --git a/Engine/Source/Editor/UnrealEd/Private/Commandlets/InternationalizationExportCommandlet.cpp b/Engine/Source/Editor/UnrealEd/Private/Commandlets/InternationalizationExportCommandlet.cpp
--- a/Engine/Source/Editor/UnrealEd/Private/Commandlets/InternationalizationExportCommandlet.cpp
+++ b/Engine/Source/Editor/UnrealEd/Private/Commandlets/InternationalizationExportCommandlet.cpp
@@ -1063,7 +1063,15 @@ bool UInternationalizationExportCommandlet::DoExport( const FString& SourcePath,
 				// Write out the Portable Object to .po file.
 				{
 					FString OutputString = PortableObj.ToString();
-					FString OutputFileName = SourcePath / CultureName + TEXT(".po");
+					FString OutputFileName;
+					if (FPaths::GetExtension(DestinationPath) != "po")
+					{
+						OutputFileName = DestinationPath / CultureName + TEXT(".po");
+					}
+					else
+					{
+						OutputFileName = DestinationPath;
+					}
 
 					//@TODO We force UTF8 at the moment but we want this to be based on the format found in the header info.
 					if( !FFileHelper::SaveStringToFile(OutputString, *OutputFileName, FFileHelper::EEncodingOptions::ForceUTF8) )
