UE4: When saving world assets, the correct extension is assumed and a save prompt is no longer required

diff --git a/Engine/Source/Editor/UnrealEd/Private/FileHelpers.cpp b/Engine/Source/Editor/UnrealEd/Private/FileHelpers.cpp
--- a/Engine/Source/Editor/UnrealEd/Private/FileHelpers.cpp
+++ b/Engine/Source/Editor/UnrealEd/Private/FileHelpers.cpp
@@ -1909,14 +1909,18 @@ static int32 InternalSavePackage( UPackage* PackageToSave, bool& bOutPackageLoca
 	// True if we should attempt saving
 	bool bAttemptSave = true;
 
+	// If we are treating worlds as assets, there is never a need to go down the "Save As" codepath
+	static const bool bUsingWorldAssets = FParse::Param(FCommandLine::Get(), TEXT("WorldAssets"));
+
 	FString ExistingFilename;
 	const bool bPackageAlreadyExists = FPackageName::DoesPackageExist( PackageToSave->GetName(), NULL, &ExistingFilename );
-	if( !bIsMapPackage || bPackageAlreadyExists )
+	if( !bIsMapPackage || bPackageAlreadyExists || bUsingWorldAssets )
 	{
 		if (!bPackageAlreadyExists)
 		{
 			// Construct a filename from long package name.
-			ExistingFilename = FPackageName::LongPackageNameToFilename(PackageToSave->GetName(), FPackageName::GetAssetPackageExtension());
+			const FString& FileExtension = bIsMapPackage ? FPackageName::GetMapPackageExtension() : FPackageName::GetAssetPackageExtension();
+			ExistingFilename = FPackageName::LongPackageNameToFilename(PackageToSave->GetName(), FileExtension);
 
 			// Check if we can use this filename.
 			FText ErrorText;
@@ -2024,7 +2028,7 @@ static int32 InternalSavePackage( UPackage* PackageToSave, bool& bOutPackageLoca
 	while( bAttemptSave )
 	{
 		bool bWasSuccessful = false;
-		if ( bIsMapPackage )
+		if ( bIsMapPackage && !bUsingWorldAssets )
 		{
 			// have a Helper attempt to save the map
 			SaveOutput.Log("LogFileHelpers", ELogVerbosity::Log, FString::Printf(TEXT("Saving Map: %s"), *PackageName));
