#ttp 331823 UE4: Adding geometry primitives (BSP) at their default size is problematic.
#proj Editor
#branch UE4
#summary Give all of the brushes a more uniform default size that also matches their description Cube is now a Cube etc.

diff --git a/Engine/Source/Editor/UnrealEd/Private/EditorBrushBuilder.cpp b/Engine/Source/Editor/UnrealEd/Private/EditorBrushBuilder.cpp
--- a/Engine/Source/Editor/UnrealEd/Private/EditorBrushBuilder.cpp
+++ b/Engine/Source/Editor/UnrealEd/Private/EditorBrushBuilder.cpp
@@ -247,10 +247,10 @@ UConeBuilder::UConeBuilder(const class FPostConstructInitializeProperties& PCIP)
 	};
 	static FConstructorStatics ConstructorStatics;
 
-	Z = 500.0f;
-	CapZ = 490.0f;
-	OuterRadius = 1000.0f;
-	InnerRadius = 990.0f;
+	Z = 300.0f;
+	CapZ = 290.0f;
+	OuterRadius = 200.0f;
+	InnerRadius = 190.0f;
 	Sides = 8;
 	GroupName = ConstructorStatics.NAME_Cone;
 	AlignToSide = true;
@@ -325,8 +325,8 @@ UCubeBuilder::UCubeBuilder(const class FPostConstructInitializeProperties& PCIP)
 	};
 	static FConstructorStatics ConstructorStatics;
 
-	X = 2000.0f;
-	Y = 2000.0f;
+	X = 200.0f;
+	Y = 200.0f;
 	Z = 200.0;
 	WallThickness = 10.0f;
 	GroupName = ConstructorStatics.NAME_Cube;
@@ -521,9 +521,9 @@ UCylinderBuilder::UCylinderBuilder(const class FPostConstructInitializePropertie
 	};
 	static FConstructorStatics ConstructorStatics;
 
-	Z = 500.0f;
-	OuterRadius = 1000.0f;
-	InnerRadius = 990.0f;
+	Z = 200.0f;
+	OuterRadius = 200.0f;
+	InnerRadius = 190.0f;
 	Sides = 8;
 	GroupName = ConstructorStatics.NAME_Cylinder;
 	AlignToSide = true;
