#ttp 295072 - UE4: Editor: [ and ] keys to change the size of the transform gizmos
#proj UE4
#branch UE4
#summary You can now use Alt + [] to adjust the size of the transform gizmo
#extra This is stored and applied as an adjustment on the base gizmo size (as the different gizmos have different starting sizes).
#extra I had to use Alt+[] rather than just [] as [] is already used for grid size adjustment.

diff --git a/Engine/Source/Editor/UnrealEd/Classes/Settings/LevelEditorViewportSettings.h b/Engine/Source/Editor/UnrealEd/Classes/Settings/LevelEditorViewportSettings.h
--- a/Engine/Source/Editor/UnrealEd/Classes/Settings/LevelEditorViewportSettings.h
+++ b/Engine/Source/Editor/UnrealEd/Classes/Settings/LevelEditorViewportSettings.h
@@ -356,6 +356,12 @@ public:
 	UPROPERTY(EditAnywhere, config, Category=LookAndFeel, meta=(ClampMin = "0.01", UIMin = "0.01", UIMax = "5"))
 	float BillboardScale;
 
+	/**
+	 * The size adjustment to apply to the translate/rotate/scale widgets (in Unreal units).
+	 */
+	UPROPERTY(EditAnywhere, config, Category=LookAndFeel, meta=(ClampMin="-10",ClampMax="150") )
+	int32 TransformWidgetSizeAdjustment;
+
 public:
 
 	/**
diff --git a/Engine/Source/Editor/UnrealEd/Private/EditorViewportClient.cpp b/Engine/Source/Editor/UnrealEd/Private/EditorViewportClient.cpp
--- a/Engine/Source/Editor/UnrealEd/Private/EditorViewportClient.cpp
+++ b/Engine/Source/Editor/UnrealEd/Private/EditorViewportClient.cpp
@@ -3257,6 +3257,13 @@ void FEditorViewportClient::SetCurrentWidgetAxis( EAxisList::Type InAxis )
 	Widget->SetCurrentAxis( InAxis );
 }
 
+void FEditorViewportClient::AdjustTransformWidgetSize(const int32 SizeDelta)
+{
+	 ULevelEditorViewportSettings &ViewportSettings = *GetMutableDefault<ULevelEditorViewportSettings>();
+	 ViewportSettings.TransformWidgetSizeAdjustment = FMath::Clamp(ViewportSettings.TransformWidgetSizeAdjustment + SizeDelta, -10, 150);
+	 ViewportSettings.PostEditChange();
+}
+
 float FEditorViewportClient::GetNearClipPlane() const
 {
 	return (NearPlane < 0.0f) ? GNearClippingPlane : NearPlane;
diff --git a/Engine/Source/Editor/UnrealEd/Private/EditorViewportCommands.cpp b/Engine/Source/Editor/UnrealEd/Private/EditorViewportCommands.cpp
--- a/Engine/Source/Editor/UnrealEd/Private/EditorViewportCommands.cpp
+++ b/Engine/Source/Editor/UnrealEd/Private/EditorViewportCommands.cpp
@@ -45,6 +45,9 @@ void FEditorViewportCommands::RegisterCommands()
 	UI_COMMAND( ScaleMode, "Scale Mode", "Select and scale objects", EUserInterfaceActionType::ToggleButton, FInputGesture(EKeys::R) );
 	UI_COMMAND( TranslateRotateMode, "Combined Translate and Rotate Mode", "Select and translate or rotate objects", EUserInterfaceActionType::ToggleButton, FInputGesture() );
 
+	UI_COMMAND( ShrinkTransformWidget, "Shrink Transform Widget", "Shrink the level editor transform widget", EUserInterfaceActionType::Button, FInputGesture(EModifierKey::Alt, EKeys::LeftBracket) );
+	UI_COMMAND( ExpandTransformWidget, "Expand Transform Widget", "Expand the level editor transform widget", EUserInterfaceActionType::Button, FInputGesture(EModifierKey::Alt, EKeys::RightBracket) );
+
 	UI_COMMAND( RelativeCoordinateSystem_World, "World-relative Transform", "Move and rotate objects relative to the cardinal world axes", EUserInterfaceActionType::RadioButton, FInputGesture() );
 	UI_COMMAND( RelativeCoordinateSystem_Local, "Local-relative Transform", "Move and rotate objects relative to the object's local axes", EUserInterfaceActionType::RadioButton, FInputGesture() );
 
diff --git a/Engine/Source/Editor/UnrealEd/Private/SEditorViewport.cpp b/Engine/Source/Editor/UnrealEd/Private/SEditorViewport.cpp
--- a/Engine/Source/Editor/UnrealEd/Private/SEditorViewport.cpp
+++ b/Engine/Source/Editor/UnrealEd/Private/SEditorViewport.cpp
@@ -217,6 +217,16 @@ void SEditorViewport::BindCommands()
 		FIsActionButtonVisible::CreateSP( this, &SEditorViewport::IsTranslateRotateModeVisible )
 		);
 
+	CommandListRef.MapAction( 
+		Commands.ShrinkTransformWidget,
+		FExecuteAction::CreateSP( ClientRef, &FEditorViewportClient::AdjustTransformWidgetSize, -1 )
+		);
+
+	CommandListRef.MapAction( 
+		Commands.ExpandTransformWidget,
+		FExecuteAction::CreateSP( ClientRef, &FEditorViewportClient::AdjustTransformWidgetSize, 1 )
+		);
+
 	CommandListRef.MapAction( 
 		Commands.RelativeCoordinateSystem_World,
 		FExecuteAction::CreateSP( ClientRef, &FEditorViewportClient::SetWidgetCoordSystemSpace, COORD_World ),
diff --git a/Engine/Source/Editor/UnrealEd/Private/Settings/SettingsClasses.cpp b/Engine/Source/Editor/UnrealEd/Private/Settings/SettingsClasses.cpp
--- a/Engine/Source/Editor/UnrealEd/Private/Settings/SettingsClasses.cpp
+++ b/Engine/Source/Editor/UnrealEd/Private/Settings/SettingsClasses.cpp
@@ -180,6 +180,8 @@ ULevelEditorViewportSettings::ULevelEditorViewportSettings( const class FPostCon
 	bLevelStreamingVolumePrevis = false;
 	BillboardScale = 1.0f;
 
+	TransformWidgetSizeAdjustment = 0.0f;
+
 	// Set a default preview mesh
 	PreviewMeshes.Add(FStringAssetReference("/Engine/EditorMeshes/ColorCalibrator/SM_ColorCalibrator.SM_ColorCalibrator"));
 }
@@ -260,6 +262,8 @@ void ULevelEditorViewportSettings::PostEditChangeProperty( struct FPropertyChang
 		SaveConfig();
 	}
 
+	GEditor->RedrawAllViewports();
+
 	SettingChangedEvent.Broadcast(Name);
 }
 
diff --git a/Engine/Source/Editor/UnrealEd/Private/UnrealWidget.cpp b/Engine/Source/Editor/UnrealEd/Private/UnrealWidget.cpp
--- a/Engine/Source/Editor/UnrealEd/Private/UnrealWidget.cpp
+++ b/Engine/Source/Editor/UnrealEd/Private/UnrealWidget.cpp
@@ -5,20 +5,12 @@
 
 IMPLEMENT_HIT_PROXY(HWidgetAxis,HHitProxy);
 
-static const float AXIS_ARROW_RADIUS = 3.5f;
-static const float AXIS_CIRCLE_RADIUS	= 48.0f;
-static const float TRANSLATE_ROTATE_AXIS_CIRCLE_RADIUS	= 20.0f;
+static const float AXIS_LENGTH = 35.0f;
+static const float TRANSLATE_ROTATE_AXIS_CIRCLE_RADIUS = 20.0f;
 static const float INNER_AXIS_CIRCLE_RADIUS = 48.0f;
 static const float OUTER_AXIS_CIRCLE_RADIUS = 56.0f;
 static const float ROTATION_TEXT_RADIUS = 75.0f;
-static const int32 AXIS_CIRCLE_SIDES		= 24;
-
-namespace WidgetDefs
-{
-	static const int32 GScreenAlignedSphereSides = 16;
-	static const float GScreenAlignedSphereRadius = 3.0f;
-}
-
+static const int32 AXIS_CIRCLE_SIDES = 24;
 
 FWidget::FWidget()
 {
@@ -227,7 +219,8 @@ void FWidget::Render_Axis( const FSceneView* View, FPrimitiveDrawInterface* PDI,
 		const bool bDisabled = EditorModeTools ? (EditorModeTools->IsModeActive(FBuiltinEditorModes::EM_Default) && GEditor->HasLockedActors() ) : false;
 		PDI->SetHitProxy( new HWidgetAxis( InAxis, bDisabled) );
 
-		const float HalfHeight = 35/2.0f;
+		const float AxisLength = AXIS_LENGTH + GetDefault<ULevelEditorViewportSettings>()->TransformWidgetSizeAdjustment;
+		const float HalfHeight = AxisLength/2.0f;
 		const float CylinderRadius = 1.2f;
 		const FVector Offset( 0,0,HalfHeight );
 
@@ -252,13 +245,15 @@ void FWidget::Render_Axis( const FSceneView* View, FPrimitiveDrawInterface* PDI,
 
 		if ( bCubeHead )
 		{
-			FVector RootPos(38, 0, 0);
+			const float CubeHeadOffset = 3.0f;
+			FVector RootPos(AxisLength + CubeHeadOffset, 0, 0);
 
 			Render_Cube(PDI, (FTranslationMatrix(RootPos) * ArrowToWorld) * FScaleMatrix(FlattenScale), InMaterial, FVector(4.0f));
 		}
 		else
 		{
-			FVector RootPos(47, 0, 0);
+			const float ConeHeadOffset = 12.0f;
+			FVector RootPos(AxisLength + ConeHeadOffset, 0, 0);
 
 			float Angle = FMath::DegreesToRadians( PI * 5 );
 			DrawCone(PDI, ( FScaleMatrix(-13) * FTranslationMatrix(RootPos) * ArrowToWorld ) * FScaleMatrix(FlattenScale), Angle, Angle, 32, false, FColor::White, InMaterial->GetRenderProxy(false), SDPG_Foreground);
@@ -734,12 +729,13 @@ void FWidget::Render_TranslateRotateZ( const FSceneView* View, FPrimitiveDrawInt
 
 		const bool bDisabled = IsWidgetDisabled();
 
+		const float ScaledRadius = (TRANSLATE_ROTATE_AXIS_CIRCLE_RADIUS * UniformScale) + GetDefault<ULevelEditorViewportSettings>()->TransformWidgetSizeAdjustment;
+
 		//ZRotation
 		if( DrawAxis&EAxisList::ZRotation && (bIsPerspective || bIsLocalSpace || View->ViewMatrices.ViewMatrix.M[0][2] != -1.f) )
 		{
 			PDI->SetHitProxy( new HWidgetAxis(EAxisList::ZRotation, bDisabled) );
 			{
-				float ScaledRadius = TRANSLATE_ROTATE_AXIS_CIRCLE_RADIUS * UniformScale;
 				FVector XAxis = CustomCoordSystem.TransformPosition( FVector(1,0,0).RotateAngleAxis( (EditorModeTools ? EditorModeTools->TranslateRotateXAxisAngle : 0 ), FVector(0,0,1)) );
 				FVector YAxis = CustomCoordSystem.TransformPosition( FVector(0,1,0).RotateAngleAxis( (EditorModeTools ? EditorModeTools->TranslateRotateXAxisAngle : 0 ), FVector(0,0,1)) );
 				FVector BaseArrowPoint = InLocation + XAxis * ScaledRadius;
@@ -753,11 +749,17 @@ void FWidget::Render_TranslateRotateZ( const FSceneView* View, FPrimitiveDrawInt
 		{
 			if( (DrawAxis & EAxisList::XY) == EAxisList::XY ) 
 			{
+				// Add more sides to the circle if we've been scaled up to keep the circle looking circular
+				// An extra side for every 5 extra unreal units seems to produce a nice result
+				const int32 CircleSides = (GetDefault<ULevelEditorViewportSettings>()->TransformWidgetSizeAdjustment > 0) 
+					? AXIS_CIRCLE_SIDES + (GetDefault<ULevelEditorViewportSettings>()->TransformWidgetSizeAdjustment / 5)
+					: AXIS_CIRCLE_SIDES;
+
 				PDI->SetHitProxy( new HWidgetAxis(EAxisList::XY, bDisabled) );
 				{
-					DrawCircle( PDI, InLocation, CustomCoordSystem.TransformPosition( FVector(1,0,0) ), CustomCoordSystem.TransformPosition( FVector(0,1,0) ), XYPlaneColor, TRANSLATE_ROTATE_AXIS_CIRCLE_RADIUS * UniformScale, AXIS_CIRCLE_SIDES, SDPG_Foreground );
+					DrawCircle( PDI, InLocation, CustomCoordSystem.TransformPosition( FVector(1,0,0) ), CustomCoordSystem.TransformPosition( FVector(0,1,0) ), XYPlaneColor, ScaledRadius, CircleSides, SDPG_Foreground );
 					XYPlaneColor.A = ((CurrentAxis&EAxisList::XY) == EAxisList::XY) ? 0x3f : 0x0f;	//make the disc transparent
-					DrawDisc  ( PDI, InLocation, CustomCoordSystem.TransformPosition( FVector(1,0,0) ), CustomCoordSystem.TransformPosition( FVector(0,1,0) ), XYPlaneColor, TRANSLATE_ROTATE_AXIS_CIRCLE_RADIUS * UniformScale, AXIS_CIRCLE_SIDES, TransparentPlaneMaterialXY->GetRenderProxy(false), SDPG_Foreground );
+					DrawDisc  ( PDI, InLocation, CustomCoordSystem.TransformPosition( FVector(1,0,0) ), CustomCoordSystem.TransformPosition( FVector(0,1,0) ), XYPlaneColor, ScaledRadius, CircleSides, TransparentPlaneMaterialXY->GetRenderProxy(false), SDPG_Foreground );
 				}
 				PDI->SetHitProxy( NULL );
 			}
@@ -1420,10 +1422,13 @@ void FWidget::DrawRotationArc(const FSceneView* View, FPrimitiveDrawInterface* P
  */
 void FWidget::DrawPartialRotationArc(const FSceneView* View, FPrimitiveDrawInterface* PDI, EAxisList::Type InAxis, const FVector& InLocation, const FVector& Axis0, const FVector& Axis1, const float InStartAngle, const float InEndAngle, const FColor& InColor, const float InScale, const FVector& InDirectionToWidget )
 {
+	const float InnerRadius = (INNER_AXIS_CIRCLE_RADIUS * InScale) + GetDefault<ULevelEditorViewportSettings>()->TransformWidgetSizeAdjustment;
+	const float OuterRadius = (OUTER_AXIS_CIRCLE_RADIUS * InScale) + GetDefault<ULevelEditorViewportSettings>()->TransformWidgetSizeAdjustment;
+
 	bool bIsPerspective = ( View->ViewMatrices.ProjMatrix.M[3][3] < 1.0f );
 	PDI->SetHitProxy( new HWidgetAxis( InAxis ) );
 	{
-		FThickArcParams OuterArcParams(PDI, InLocation, TransparentPlaneMaterialXY, INNER_AXIS_CIRCLE_RADIUS*InScale, OUTER_AXIS_CIRCLE_RADIUS*InScale);
+		FThickArcParams OuterArcParams(PDI, InLocation, TransparentPlaneMaterialXY, InnerRadius, OuterRadius);
 		FColor OuterColor = ( CurrentAxis&InAxis ? CurrentColor : InColor );
 		//Pass through alpha
 		OuterColor.A = InColor.A;
@@ -1433,7 +1438,7 @@ void FWidget::DrawPartialRotationArc(const FSceneView* View, FPrimitiveDrawInter
 
 	if (bIsPerspective)
 	{
-		FThickArcParams InnerArcParams(PDI, InLocation, GridMaterial, 0.0f, INNER_AXIS_CIRCLE_RADIUS*InScale);
+		FThickArcParams InnerArcParams(PDI, InLocation, GridMaterial, 0.0f, InnerRadius);
 		FColor InnerColor = InColor;
 		//if something is selected and it's not this
 		InnerColor.A = ((CurrentAxis & InAxis) && !bDragging) ? LargeInnerAlpha : SmallInnerAlpha;
@@ -1454,7 +1459,13 @@ void FWidget::DrawThickArc (const FThickArcParams& InParams, const FVector& Axis
 	{
 		return;
 	}
-	const int32 NumPoints = FMath::Trunc(AXIS_CIRCLE_SIDES * (InEndAngle-InStartAngle)/(PI/2)) + 1;
+
+	// Add more sides to the circle if we've been scaled up to keep the circle looking circular
+	// An extra side for every 5 extra unreal units seems to produce a nice result
+	const int32 CircleSides = (GetDefault<ULevelEditorViewportSettings>()->TransformWidgetSizeAdjustment > 0) 
+		? AXIS_CIRCLE_SIDES + (GetDefault<ULevelEditorViewportSettings>()->TransformWidgetSizeAdjustment / 5)
+		: AXIS_CIRCLE_SIDES;
+	const int32 NumPoints = FMath::Trunc(CircleSides * (InEndAngle-InStartAngle)/(PI/2)) + 1;
 
 	FColor TriangleColor = InColor;
 	FColor RingColor = InColor;
@@ -1536,8 +1547,8 @@ void FWidget::DrawThickArc (const FThickArcParams& InParams, const FVector& Axis
  */
 void FWidget::DrawSnapMarker(FPrimitiveDrawInterface* PDI, const FVector& InLocation, const FVector& Axis0, const FVector& Axis1, const FColor& InColor, const float InScale, const float InWidthPercent, const float InPercentSize)
 {
-	const float InnerDistance = (INNER_AXIS_CIRCLE_RADIUS*InScale);
-	const float OuterDistance = (OUTER_AXIS_CIRCLE_RADIUS*InScale);
+	const float InnerDistance = (INNER_AXIS_CIRCLE_RADIUS * InScale) + GetDefault<ULevelEditorViewportSettings>()->TransformWidgetSizeAdjustment;
+	const float OuterDistance = (OUTER_AXIS_CIRCLE_RADIUS * InScale) + GetDefault<ULevelEditorViewportSettings>()->TransformWidgetSizeAdjustment;
 	const float MaxMarkerHeight = OuterDistance - InnerDistance;
 	const float MarkerWidth = MaxMarkerHeight*InWidthPercent;
 	const float MarkerHeight = MaxMarkerHeight*InPercentSize;
@@ -1594,8 +1605,8 @@ void FWidget::DrawSnapMarker(FPrimitiveDrawInterface* PDI, const FVector& InLoca
 void FWidget::DrawStartStopMarker(FPrimitiveDrawInterface* PDI, const FVector& InLocation, const FVector& Axis0, const FVector& Axis1, const float InAngle, const FColor& InColor, const float InScale)
 {
 	const float ArrowHeightPercent = .8f;
-	const float InnerDistance = (INNER_AXIS_CIRCLE_RADIUS*InScale);
-	const float OuterDistance = (OUTER_AXIS_CIRCLE_RADIUS*InScale);
+	const float InnerDistance = (INNER_AXIS_CIRCLE_RADIUS * InScale) + GetDefault<ULevelEditorViewportSettings>()->TransformWidgetSizeAdjustment;
+	const float OuterDistance = (OUTER_AXIS_CIRCLE_RADIUS * InScale) + GetDefault<ULevelEditorViewportSettings>()->TransformWidgetSizeAdjustment;
 	const float RingHeight = OuterDistance - InnerDistance;
 	const float ArrowHeight = RingHeight*ArrowHeightPercent;
 	const float ThirtyDegrees = PI / 6.0f;
@@ -1649,7 +1660,7 @@ void FWidget::DrawStartStopMarker(FPrimitiveDrawInterface* PDI, const FVector& I
  */
 void FWidget::CacheRotationHUDText(const FSceneView* View, FPrimitiveDrawInterface* PDI, const FVector& InLocation, const FVector& Axis0, const FVector& Axis1, const float AngleOfChange, const float InScale)
 {
-	const float TextDistance = (ROTATION_TEXT_RADIUS*InScale);
+	const float TextDistance = (ROTATION_TEXT_RADIUS * InScale) + GetDefault<ULevelEditorViewportSettings>()->TransformWidgetSizeAdjustment;
 
 	FVector AxisVectors[4] = { Axis0, Axis1, -Axis0, -Axis1};
 
diff --git a/Engine/Source/Editor/UnrealEd/Public/EditorViewportClient.h b/Engine/Source/Editor/UnrealEd/Public/EditorViewportClient.h
--- a/Engine/Source/Editor/UnrealEd/Public/EditorViewportClient.h
+++ b/Engine/Source/Editor/UnrealEd/Public/EditorViewportClient.h
@@ -499,6 +499,11 @@ public:
 	 */
 	virtual void SetCurrentWidgetAxis( EAxisList::Type InAxis );
 
+	/**
+	 * Adjusts the current transform widget size by the provided delta value
+	 */
+	void AdjustTransformWidgetSize(const int32 SizeDelta);
+
 	/**
 	 * Called to do any additional set up of the view for rendering
 	 * 
diff --git a/Engine/Source/Editor/UnrealEd/Public/EditorViewportCommands.h b/Engine/Source/Editor/UnrealEd/Public/EditorViewportCommands.h
--- a/Engine/Source/Editor/UnrealEd/Public/EditorViewportCommands.h
+++ b/Engine/Source/Editor/UnrealEd/Public/EditorViewportCommands.h
@@ -108,6 +108,12 @@ public:
 	/** TranslateRotate Mode */
 	TSharedPtr< FUICommandInfo > TranslateRotateMode;
 
+	/** Shrink the level editor transform widget */
+	TSharedPtr< FUICommandInfo > ShrinkTransformWidget;
+
+	/** Expand the level editor transform widget */
+	TSharedPtr< FUICommandInfo > ExpandTransformWidget;
+
 	/** World relative coordinate system */
 	TSharedPtr< FUICommandInfo > RelativeCoordinateSystem_World;
 
