#ue4 Editoronly BP: Local variable conversion.

diff --git a/Engine/Source/Editor/UnrealEd/Private/Kismet2/BlueprintEditorUtils.cpp b/Engine/Source/Editor/UnrealEd/Private/Kismet2/BlueprintEditorUtils.cpp
--- a/Engine/Source/Editor/UnrealEd/Private/Kismet2/BlueprintEditorUtils.cpp
+++ b/Engine/Source/Editor/UnrealEd/Private/Kismet2/BlueprintEditorUtils.cpp
@@ -817,6 +817,28 @@ struct FEditoronlyBlueprintHelper
 		}
 	}
 
+	static void HandleTemporaryVariableNode(UK2Node_TemporaryVariable* Node)
+	{
+		check(Node);
+
+		auto Pin = Node->GetVariablePin();
+		if (Pin && IsUnwantedType(Pin->PinType))
+		{
+			ChangePinType(Pin->PinType);
+			
+			if (Pin->DefaultObject)
+			{
+				UBlueprint* DefaultBlueprint = Cast<UBlueprint>(Pin->DefaultObject);
+				Pin->DefaultObject = DefaultBlueprint ? *DefaultBlueprint->GeneratedClass : NULL;
+			}
+		}
+
+		if (IsUnwantedType(Node->VariableType))
+		{
+			ChangePinType(Node->VariableType);
+		}
+	}
+
 	static void HandleDefaultObjects(UK2Node* Node)
 	{
 		if (Node)
@@ -839,7 +861,7 @@ struct FEditoronlyBlueprintHelper
 	{
 		if (ShouldBeFixed(Blueprint, false))
 		{
-			UE_LOG(LogBlueprint, Log, TEXT("Bluepirnt references will b e removed from '%s'"), *Blueprint->GetName());
+			UE_LOG(LogBlueprint, Log, TEXT("Bluepirnt references will be removed from '%s'"), *Blueprint->GetName());
 
 			const UEdGraphSchema_K2* Schema = GetDefault<UEdGraphSchema_K2>();
 
@@ -854,11 +876,22 @@ struct FEditoronlyBlueprintHelper
 			//any blueprint pin (function/event/macro parameter
 			for (const auto Graph : Graphs)
 			{
-				TArray<UK2Node_EditablePinBase*> EditablePinNodes;
-				Graph->GetNodesOfClass(EditablePinNodes);
-				for (const auto Node : EditablePinNodes)
 				{
-					HandleEditablePinNode(Node);
+					TArray<UK2Node_EditablePinBase*> EditablePinNodes;
+					Graph->GetNodesOfClass(EditablePinNodes);
+					for (const auto Node : EditablePinNodes)
+					{
+						HandleEditablePinNode(Node);
+					}
+				}
+
+				{
+					TArray<UK2Node_TemporaryVariable*> TempVariablesNodes;
+					Graph->GetNodesOfClass(TempVariablesNodes);
+					for (const auto Node : TempVariablesNodes)
+					{
+						HandleTemporaryVariableNode(Node);
+					}
 				}
 			}
 
