Provide the entire reference stack for ungarbage collected worlds

diff --git a/Engine/Source/Editor/UnrealEd/Private/Editor.cpp b/Engine/Source/Editor/UnrealEd/Private/Editor.cpp
--- a/Engine/Source/Editor/UnrealEd/Private/Editor.cpp
+++ b/Engine/Source/Editor/UnrealEd/Private/Editor.cpp
@@ -6551,7 +6551,7 @@ void UEditorEngine::VerifyLoadMapWorldCleanup()
 			{
 				// Print some debug information...
 				UE_LOG(LogLoad, Log, TEXT("%s not cleaned up by garbage collection! "), *World->GetFullName());
-				StaticExec(World, *FString::Printf(TEXT("OBJ REFS CLASS=WORLD NAME=%s direct"), *World->GetPathName()));
+				StaticExec(World, *FString::Printf(TEXT("OBJ REFS CLASS=WORLD NAME=%s"), *World->GetPathName()));
 				TMap<UObject*,UProperty*>	Route		= FArchiveTraceRoute::FindShortestRootPath( World, true, GARBAGE_COLLECTION_KEEPFLAGS );
 				FString						ErrorString	= FArchiveTraceRoute::PrintRootPath( Route, World );
 				UE_LOG(LogLoad, Log, TEXT("%s"),*ErrorString);
