InternationalizationExportCommandlet:
- added ability to destination filename as well as path, so we can use the "Open File" dialog

diff --git a/Engine/Source/Editor/UnrealEd/Private/Commandlets/InternationalizationExportCommandlet.cpp b/Engine/Source/Editor/UnrealEd/Private/Commandlets/InternationalizationExportCommandlet.cpp
--- a/Engine/Source/Editor/UnrealEd/Private/Commandlets/InternationalizationExportCommandlet.cpp
+++ b/Engine/Source/Editor/UnrealEd/Private/Commandlets/InternationalizationExportCommandlet.cpp
@@ -1109,7 +1109,15 @@ bool UInternationalizationExportCommandlet::DoImport( const FString& SourcePath,
 	{
 		// Load the Portable Object file if found
 		const FString CultureName = CulturesToGenerate[Culture];
-		const FString POFilePath = SourcePath / CultureName + TEXT(".po");
+		FString POFilePath;
+		if (FPaths::GetExtension(SourcePath) != "po")
+		{
+			POFilePath = SourcePath / CultureName + TEXT(".po");
+		}
+		else
+		{
+			POFilePath = SourcePath;
+		}
 
 		if( !FPaths::FileExists(POFilePath) )
 		{
