TTP #331746/#332309/#332445: Better cleaning up after blueprint automation tests (unloading all dependencies along with tested blueprints, & clearing the undo transaction buffer to remove references to those blueprints).

diff --git a/Engine/Source/Editor/UnrealEd/Private/Kismet2/BlueprintEditorUtils.cpp b/Engine/Source/Editor/UnrealEd/Private/Kismet2/BlueprintEditorUtils.cpp
--- a/Engine/Source/Editor/UnrealEd/Private/Kismet2/BlueprintEditorUtils.cpp
+++ b/Engine/Source/Editor/UnrealEd/Private/Kismet2/BlueprintEditorUtils.cpp
@@ -2074,7 +2074,7 @@ void FBlueprintEditorUtils::GetDependentBlueprints(UBlueprint* Blueprint, TArray
 	{
 		// we know the class is correct so a fast cast is ok here
 		UBlueprint* const TestBP = (UBlueprint*) *ObjIt;
-		if ( FBlueprintEditorUtils::IsBlueprintDependentOn(TestBP, Blueprint))
+		if ( FBlueprintEditorUtils::IsBlueprintDependentOn(TestBP, Blueprint) && !TestBP->HasAnyFlags(RF_PendingKill) )
 		{
 			DependentBlueprints.Add(TestBP);
 		}
diff --git a/Engine/Source/Editor/UnrealEd/Private/Tests/BlueprintAutomationTests.cpp b/Engine/Source/Editor/UnrealEd/Private/Tests/BlueprintAutomationTests.cpp
--- a/Engine/Source/Editor/UnrealEd/Private/Tests/BlueprintAutomationTests.cpp
+++ b/Engine/Source/Editor/UnrealEd/Private/Tests/BlueprintAutomationTests.cpp
@@ -350,6 +350,11 @@ public:
 		// up garbage so an attempted load doesn't stick us with an invalid asset
 		if (bForceFlush)
 		{
+#if WITH_EDITOR
+			// clear undo history to ensure that the transaction buffer isn't 
+			// holding onto any references to the blueprints we want unloaded
+			GEditor->Trans->Reset(NSLOCTEXT("BpAutomation", "BpAutomationTest", "Blueprint Automation Test"));
+#endif // #if WITH_EDITOR
 			CollectGarbage(RF_Native);
 		}
 	}
@@ -410,6 +415,23 @@ public:
 		return bIsLoaded;
 	}
 
+	/** */
+	static bool GetExternalReferences(UObject* Obj, TArray<FReferencerInformation>& ExternalRefsOut)
+	{
+		CollectGarbage(GARBAGE_COLLECTION_KEEPFLAGS);
+
+		bool bHasReferences = false;
+
+		FReferencerInformationList Refs;
+		if (IsReferenced(Obj, RF_Native | RF_Public, true, &Refs))
+		{
+			ExternalRefsOut = Refs.ExternalReferences;
+			bHasReferences = true;
+		}
+
+		return bHasReferences;
+	}
+
 	/**
 	 * Helper method for determining if the specified asset has pending changes.
 	 * 
@@ -682,6 +704,85 @@ public:
 uint32 FBlueprintAutomationTestUtilities::QueuedTempId = 0u;
 TArray<FName> FBlueprintAutomationTestUtilities::DontSavePackagesList;
 
+/************************************************************************/
+/* FScopedBlueprintUnloader                                             */
+/************************************************************************/
+
+class FScopedBlueprintUnloader
+{
+public:
+	FScopedBlueprintUnloader(bool bAutoOpenScope, bool bRunGCOnCloseIn)
+		: bIsOpen(false)
+		, bRunGCOnClose(bRunGCOnCloseIn)
+	{
+		if (bAutoOpenScope)
+		{
+			OpenScope();
+		}
+	}
+
+	~FScopedBlueprintUnloader()
+	{
+		CloseScope();
+	}
+
+	/** Tracks currently loaded blueprints at the time of this object's creation */
+	void OpenScope()
+	{
+		PreLoadedBlueprints.Empty();
+
+		// keep a list of blueprints that were loaded at the start (so we can unload new ones after)
+		for (TObjectIterator<UBlueprint> BpIt; BpIt; ++BpIt)
+		{
+			UBlueprint* Blueprint = *BpIt;
+			PreLoadedBlueprints.Add(Blueprint);
+		}
+		bIsOpen = true;
+	}
+
+	/** Unloads any blueprints that weren't loaded when this object was created */
+	void CloseScope()
+	{
+		if (bIsOpen)
+		{
+			// clean up any dependencies that we're loading in the scope of this object's lifetime
+			for (TObjectIterator<UBlueprint> BpIt; BpIt; ++BpIt)
+			{
+				UBlueprint* Blueprint = *BpIt;
+				if (PreLoadedBlueprints.Find(Blueprint) == nullptr)
+				{
+					FBlueprintAutomationTestUtilities::UnloadBlueprint(Blueprint);
+				}
+			}
+
+			bIsOpen = false;
+		}
+
+		// run, even if it was not open (some tests may be relying on this, and 
+		// not running it themselves)
+		if (bRunGCOnClose)
+		{
+#if WITH_EDITOR
+			// clear undo history to ensure that the transaction buffer isn't 
+			// holding onto any references to the blueprints we want unloaded
+			GEditor->Trans->Reset(NSLOCTEXT("BpAutomation", "BpAutomationTest", "Blueprint Automation Test"));
+#endif // #if WITH_EDITOR
+			CollectGarbage(RF_Native);
+		}
+	}
+
+	void ClearScope()
+	{
+		PreLoadedBlueprints.Empty();
+		bIsOpen = false;
+	}
+
+private:
+	bool bIsOpen;
+	TSet<UBlueprint*> PreLoadedBlueprints;
+	bool bRunGCOnClose;
+};
+
 /************************************************************************/
 /* FBlueprintCompileOnLoadTest                                          */
 /************************************************************************/
@@ -720,6 +821,10 @@ bool FBlueprintCompileOnLoadTest::RunTest(const FString& BlueprintAssetPath)
 			FBlueprintAutomationTestUtilities::UnloadBlueprint(ExistingBP);
 		}		
 	}
+	
+	// tracks blueprints that were already loaded (and cleans up any that were 
+	// loaded in its lifetime, once it is destroyed)
+	FScopedBlueprintUnloader NewBlueprintUnloader(/*bAutoOpenScope =*/true, /*bRunGCOnCloseIn =*/true);
 
 	// We load the blueprint twice and compare the two for discrepancies. This is 
 	// to bring dependency load issues to light (among other things). If a blueprint's
@@ -833,12 +938,13 @@ bool FBlueprintCompileOnLoadTest::RunTest(const FString& BlueprintAssetPath)
 		}
 	}
 
-	FBlueprintAutomationTestUtilities::UnloadBlueprint(ReloadedBlueprint);
-	FBlueprintAutomationTestUtilities::UnloadBlueprint(InitialBlueprint);
-	// run garbage collection, in hopes that the imports for this blueprint get destroyed
-	// so that they don't invalidate other tests that share the same dependencies
-	CollectGarbage(RF_Native);
-
+	// At the close of this function, the FScopedBlueprintUnloader should prep 
+	// for following tests by unloading any blueprint dependencies that were 
+	// loaded for this one (should catch InitialBlueprint and ReloadedBlueprint) 
+	// 
+	// The FScopedBlueprintUnloader should also run garbage-collection after,
+	// in hopes that the imports for this blueprint get destroyed so that they 
+	// don't invalidate other tests that share the same dependencies
 	return !bDiffsFound;
 }
 
@@ -1049,6 +1155,11 @@ bool FBlueprintReparentTest::RunTest(const FString& BlueprintAssetPath)
 			FBlueprintAutomationTestUtilities::UnloadBlueprint(BlueprintObj);
 		}
 
+#if WITH_EDITOR
+		// clear undo history to ensure that the transaction buffer isn't 
+		// holding onto any references to the blueprints we want unloaded
+		GEditor->Trans->Reset(NSLOCTEXT("BpAutomation", "ReparentTest", "Reparent Blueprint Test"));
+#endif // #if WITH_EDITOR
 		// make sure the unloaded blueprints are properly flushed (for future tests)
 		CollectGarbage(RF_Native);
 	}
@@ -1109,6 +1220,10 @@ bool FBlueprintRenameAndCloneTest::RunTest(const FString& BlueprintAssetPath)
 		AddWarning(FString::Printf(TEXT("'%s' is already loaded, and possibly referenced by external objects (unable to perform rename tests... please run again in an empty map)."), *BlueprintAssetPath));
 	}
 
+	// track the loaded blueprint (and any other BP dependencies) so we can 
+	// unload them if we end up renaming it.
+	FScopedBlueprintUnloader NewBlueprintUnloader(/*bAutoOpenScope =*/true, /*bRunGCOnCloseIn =*/false);
+
 	UBlueprint* const OriginalBlueprint = Cast<UBlueprint>(StaticLoadObject(UBlueprint::StaticClass(), NULL, *BlueprintAssetPath));
 	if (OriginalBlueprint == NULL)
 	{
@@ -1152,11 +1267,25 @@ bool FBlueprintRenameAndCloneTest::RunTest(const FString& BlueprintAssetPath)
 			// we don't save it in this state and so we can reload the blueprint later
 			FBlueprintAutomationTestUtilities::InvalidatePackage(OriginalPackage);
 
+			// need to unload the renamed blueprint (and any other blueprints 
+			// that were relying on it), so that the renamed blueprint doesn't get used by the user
 			FBlueprintAutomationTestUtilities::UnloadBlueprint(OriginalBlueprint);
+			NewBlueprintUnloader.CloseScope();
+		}
+		else 
+		{
+			// no need to unload the blueprint or any of its dependencies (since 
+			// we didn't muck with it by renaming it) 
+			NewBlueprintUnloader.ClearScope();
 		}
 
+#if WITH_EDITOR
+		// clear undo history to ensure that the transaction buffer isn't 
+		// holding onto any references to the blueprints we want unloaded
+		GEditor->Trans->Reset(NSLOCTEXT("BpAutomation", "RenameCloneTest", "Rename and Clone Test"));
+#endif // #if WITH_EDITOR
 		// make sure the unloaded blueprints are properly flushed (for future tests)
-		CollectGarbage(RF_Native);
+		CollectGarbage(RF_Native);		
 	}
 
 	return !bTestFailed;
