#ttp 332393 - Regression: EDITOR: BSP: CRASH: If a vert/face of a brush is selected in geometry mode and you switch to placement mode, clicking the brush causes a crash
#summary Switching editor modes now invalidates viewport hit proxy data
#proj Editor
#branch UE4

diff --git a/Engine/Source/Editor/UnrealEd/Private/EditorViewportClient.cpp b/Engine/Source/Editor/UnrealEd/Private/EditorViewportClient.cpp
--- a/Engine/Source/Editor/UnrealEd/Private/EditorViewportClient.cpp
+++ b/Engine/Source/Editor/UnrealEd/Private/EditorViewportClient.cpp
@@ -248,10 +248,14 @@ FEditorViewportClient::FEditorViewportClient(FPreviewScene* InPreviewScene)
 	EngineShowFlags.SetSnap(1);
 
 	SetViewMode(VMI_Lit);
+
+	GEditorModeTools().OnEditorModeChanged().AddRaw(this, &FEditorViewportClient::OnEditorModeChanged);
 }
 
 FEditorViewportClient::~FEditorViewportClient()
 {
+	GEditorModeTools().OnEditorModeChanged().RemoveAll(this);
+
 	delete Widget;
 	delete MouseDeltaTracker;
 
@@ -286,6 +290,11 @@ void FEditorViewportClient::RequestInvalidateHitProxy(FViewport* InViewport)
 	bNeedsInvalidateHitProxy = true;
 }
 
+void FEditorViewportClient::OnEditorModeChanged(FEdMode* EditorMode, bool bIsEntering)
+{
+	RequestInvalidateHitProxy(Viewport);
+}
+
 float FEditorViewportClient::GetOrthoUnitsPerPixel(const FViewport* InViewport) const
 {
 	const float SizeX = InViewport->GetSizeXY().X;
diff --git a/Engine/Source/Editor/UnrealEd/Public/EditorViewportClient.h b/Engine/Source/Editor/UnrealEd/Public/EditorViewportClient.h
--- a/Engine/Source/Editor/UnrealEd/Public/EditorViewportClient.h
+++ b/Engine/Source/Editor/UnrealEd/Public/EditorViewportClient.h
@@ -394,6 +394,9 @@ public:
 
 	void TakeHighResScreenShot();
 
+	/** Called when an editor mode has been (de)activated */
+	void OnEditorModeChanged(FEdMode* EditorMode, bool bIsEntering);
+
 	/** FViewElementDrawer interface */
 	virtual void Draw(const FSceneView* View,FPrimitiveDrawInterface* PDI) OVERRIDE;
 	virtual void Draw(FViewport* Viewport,FCanvas* Canvas) OVERRIDE;
