#UE4 - added support for OSS integration with PIE
- need to add a PIELogins list to DefaultEngine.ini

diff --git a/Engine/Config/BaseEngine.ini b/Engine/Config/BaseEngine.ini
--- a/Engine/Config/BaseEngine.ini
+++ b/Engine/Config/BaseEngine.ini
@@ -1261,6 +1261,7 @@ MinSmoothedFrameRate=5
 MaxSmoothedFrameRate=120
 UseOldStyleMICEditorGroups=true
 InEditorGameURLOptions=
+bOnlinePIEEnabled=false
 
 [/Script/UnrealEd.UnrealEdEngine]
 AutoSaveIndex=0
diff --git a/Engine/Source/Editor/UnrealEd/Classes/Editor/EditorEngine.h b/Engine/Source/Editor/UnrealEd/Classes/Editor/EditorEngine.h
--- a/Engine/Source/Editor/UnrealEd/Classes/Editor/EditorEngine.h
+++ b/Engine/Source/Editor/UnrealEd/Classes/Editor/EditorEngine.h
@@ -52,6 +52,62 @@ struct FSlatePlayInEditorInfo
 	{}
 };
 
+/**
+ * Data structure for storing PIE login credentials
+ */
+USTRUCT()
+struct FPIELoginInfo
+{
+public:
+
+	GENERATED_USTRUCT_BODY()
+
+	/** Type of account. Needed to identity the auth method to use (epic, internal, facebook, etc) */
+	UPROPERTY()
+	FString Type;
+	/** Id of the user logging in (email, display name, facebook id, etc) */
+	UPROPERTY()
+	FString Id;
+	/** Credentials of the user logging in (password or auth token) */
+	UPROPERTY()
+	FString Token;
+};
+
+/**
+ * Holds various data to pass to the post login delegate for PIE logins
+ */
+struct FPieLoginStruct
+{
+	/** World context handle for this login */
+	int32 WorldContextHandle;
+	/** Setting index for window positioning */
+	int32 SettingsIndex;
+	/** X location for window positioning */
+	int32 NextX;
+	/** Y location for window positioning */
+	int32 NextY;
+	/** Will this instance run as a server */
+	bool bIsServer;
+	/** Passthrough condition of blueprint compilation*/
+	bool bAnyBlueprintErrors;
+	/** Passthrough condition of spectator mode */
+	bool bStartInSpectatorMode;
+	/** Passthrough start time of PIE */
+	double PIEStartTime;
+
+	FPieLoginStruct() :
+		WorldContextHandle(INDEX_NONE),
+		SettingsIndex(0),
+		NextX(0),
+		NextY(0),
+		bIsServer(false),
+		bAnyBlueprintErrors(false),
+		bStartInSpectatorMode(false),
+		PIEStartTime(0)
+	{
+	}
+};
+
 USTRUCT()
 struct FCopySelectedInfo
 {
@@ -326,6 +382,14 @@ class UNREALED_API UEditorEngine : public UEngine
 	UPROPERTY(config)
 	int32 BuildPlayDevice;
 
+	/** Enabled online PIE */
+	UPROPERTY(config)
+	bool bOnlinePIEEnabled;
+
+	/** Logins available for use when running PIE instances */
+	UPROPERTY(config)
+	TArray<FPIELoginInfo> PIELogins;
+
 	/** Maps world contexts to their slate data */
 	TMap<int32, FSlatePlayInEditorInfo>	SlatePlayInEditorMap;
 
@@ -2221,8 +2285,48 @@ private:
 
 	UWorld* CreatePIEWorldFromEntry(FWorldContext &WorldContext, UWorld* InWorld, FString &PlayWorldMapName);
 
+	/** @return true if the editor is able to launch PIE with online platform support */
+	virtual bool SupportsOnlinePIE() const;
+
+	/**
+	 * Login PIE instances with the online platform before actually creating any PIE worlds
+	 *
+	 * @param bAnyBlueprintErrors passthrough value of blueprint errors encountered during PIE creation
+	 * @param bStartInSpectatorMode passthrough value if it is expected that PIE will start in spectator mode
+	 * @param PIEStartTime passthrough value of the time that PIE was initiated by the user
+	 */
+	virtual void LoginPIEInstances(bool bAnyBlueprintErrors, bool bStartInSpectatorMode, double PIEStartTime);
+
+	/**
+	 * Delegate called as each PIE instance login is complete, continues creating the PIE world for a single instance
+	 * 
+	 * @param LocalUserNum local user id, for PIE is going to be 0 (there is no splitscreen)
+	 * @param bWasSuccessful was the login successful
+	 * @param UserId userid of the logged in account
+	 * @param ErrorString descriptive error when applicable
+	 * @param DataStruct data required to continue PIE creation, set at login time
+	 */
+	virtual void OnLoginPIEComplete(int32 LocalUserNum, bool bWasSuccessful, const FUniqueNetId& UserId, const FString& ErrorString, FPieLoginStruct DataStruct);
+
+	/**
+	 * Continue the creation of a single PIE world after a login was successful
+	 *
+	 * @param PieWorldContext world context for this PIE instance
+	 * @param PlayNetMode mode to create this PIE world in (as server, client, etc)
+	 * @param DataStruct data required to continue PIE creation, set at login time
+	 */
+	void CreatePIEWorldFromLogin(FWorldContext& PieWorldContext, enum EPlayNetMode PlayNetMode, FPieLoginStruct& DataStruct);
+
+	/**
+	 * Non Online PIE creation flow, creates all instances of PIE at once when online isn't requested/required
+	 *
+	 * @param bAnyBlueprintErrors passthrough value of blueprint errors encountered during PIE creation
+	 * @param bStartInSpectatorMode passthrough value if it is expected that PIE will start in spectator mode
+	 */
+	void SpawnIntraProcessPIEWorlds(bool bAnyBlueprintErrors, bool bStartInSpectatorMode);
+
 	/** Common init shared by CreatePIEWorldByDuplication and CreatePIEWorldBySavingToTemp */
-	void	PostCreatePIEWorld(UWorld *InWorld);
+	void PostCreatePIEWorld(UWorld *InWorld);
 
 	/**
 	 * Toggles PIE to SIE or vice-versa
diff --git a/Engine/Source/Editor/UnrealEd/Private/Editor.cpp b/Engine/Source/Editor/UnrealEd/Private/Editor.cpp
--- a/Engine/Source/Editor/UnrealEd/Private/Editor.cpp
+++ b/Engine/Source/Editor/UnrealEd/Private/Editor.cpp
@@ -799,8 +799,8 @@ void UEditorEngine::Tick( float DeltaSeconds, bool bIdleMode )
 	{
 		for (auto It=WorldList.CreateIterator(); It; ++It)
 		{
-			// For now, kill PIE session is any of the viewports are closed
-			if (It->WorldType == EWorldType::PIE && It->GameViewport == NULL && !It->RunAsDedicated)
+			// For now, kill PIE session if any of the viewports are closed
+			if (It->WorldType == EWorldType::PIE && It->GameViewport == NULL && !It->RunAsDedicated && !It->bWaitingOnOnlineSubsystem)
 			{
 				EndPlayMap();
 				break;
@@ -1088,7 +1088,7 @@ void UEditorEngine::Tick( float DeltaSeconds, bool bIdleMode )
 		for (auto ContextIt = WorldList.CreateIterator(); ContextIt; ++ContextIt)
 		{
 			FWorldContext &PieContext = *ContextIt;
-			if (PieContext.WorldType != EWorldType::PIE)
+			if (PieContext.WorldType != EWorldType::PIE || PieContext.World() == NULL)
 			{
 				continue;
 			}
@@ -1214,7 +1214,11 @@ void UEditorEngine::Tick( float DeltaSeconds, bool bIdleMode )
 	// If all viewports closed, close the current play level.
 	if( GameViewport == NULL && PlayWorld && !bIsSimulatingInEditor )
 	{
-		EndPlayMap();
+		FWorldContext& PieWorldContext = GetWorldContextFromWorldChecked(PlayWorld);
+		if (!PieWorldContext.RunAsDedicated && !PieWorldContext.bWaitingOnOnlineSubsystem)
+		{
+			EndPlayMap();
+		}
 	}
 
 	// Update viewports.
@@ -1338,7 +1342,7 @@ void UEditorEngine::Tick( float DeltaSeconds, bool bIdleMode )
 			GameViewport = PieContext.GameViewport;
 
 			// Render playworld. This needs to happen after the other viewports for screenshots to work correctly in PIE.
-			if(PlayWorld && GameViewport && !bIsSimulatingInEditor )
+			if(PlayWorld && GameViewport && !bIsSimulatingInEditor)
 			{
 				// Use the PlayWorld as the GWorld, because who knows what will happen in the Tick.
 				UWorld* OldGWorld = SetPlayInEditorWorld( PlayWorld );
diff --git a/Engine/Source/Editor/UnrealEd/Private/PlayLevel.cpp b/Engine/Source/Editor/UnrealEd/Private/PlayLevel.cpp
--- a/Engine/Source/Editor/UnrealEd/Private/PlayLevel.cpp
+++ b/Engine/Source/Editor/UnrealEd/Private/PlayLevel.cpp
@@ -27,6 +27,8 @@
 #include "GameProjectGenerationModule.h"
 #include "SourceCodeNavigation.h"
 
+#include "Online.h"
+
 DEFINE_LOG_CATEGORY_STATIC(LogPlayLevel, Log, All);
 
 #define LOCTEXT_NAMESPACE "PlayLevel"
@@ -246,7 +248,6 @@ void UEditorEngine::TeardownPlaySession(FWorldContext &PieWorldContext)
 	check(PieWorldContext.WorldType == EWorldType::PIE);
 	PlayWorld = PieWorldContext.World();
 
-	
 	if (!PieWorldContext.RunAsDedicated)
 	{
 		// Slate data for this pie world
@@ -454,6 +455,12 @@ void UEditorEngine::TeardownPlaySession(FWorldContext &PieWorldContext)
 		// Remove the slate info from the map (note that the UWorld* is long gone at this point, but the WorldContext still exists. It will be removed outside of this function)
 		SlatePlayInEditorMap.Remove(PieWorldContext.ContextHandle);
 	}
+
+	if (SupportsOnlinePIE())
+	{
+		FName OnlineIdentifier = FName(*FString::Printf(TEXT(":%d"), PieWorldContext.ContextHandle));
+		IOnlineSubsystem::Destroy(OnlineIdentifier);
+	}
 }
 
 void UEditorEngine::PlayMap( const FVector* StartLocation, const FRotator* StartRotation, int32 Destination, int32 InPlayInViewportIndex, bool bUseMobilePreview, bool bMovieCapture )
@@ -580,19 +587,16 @@ void UEditorEngine::PlaySessionPaused()
 	FEditorDelegates::PausePIE.Broadcast(bIsSimulatingInEditor);
 }
 
-
 void UEditorEngine::PlaySessionResumed()
 {
 	FEditorDelegates::ResumePIE.Broadcast(bIsSimulatingInEditor);
 }
 
-
 void UEditorEngine::PlaySessionSingleStepped()
 {
 	FEditorDelegates::SingleStepPIE.Broadcast(bIsSimulatingInEditor);
 }
 
-
 void AdvanceWindowPositionsForNextPIEWindow(int32 &WinX, int32 &WinY, int32 Width, int32 Height)
 {
 	const FIntPoint WindowPadding(0, 0);//(15, 32);
@@ -616,7 +620,6 @@ void AdvanceWindowPositionsForNextPIEWindow(int32 &WinX, int32 &WinY, int32 Widt
 	}
 }
 
-
 void GetMultipleInstancePositions(int32 index, int32 &LastX, int32 &LastY)
 {
 	ULevelEditorPlaySettings* PlayInSettings = Cast<ULevelEditorPlaySettings>(ULevelEditorPlaySettings::StaticClass()->GetDefaultObject());
@@ -638,7 +641,6 @@ void GetMultipleInstancePositions(int32 index, int32 &LastX, int32 &LastY)
 	AdvanceWindowPositionsForNextPIEWindow(LastX, LastY, PlayInSettings->NewWindowWidth, PlayInSettings->NewWindowHeight);
 }
 
-
 FString GenerateCmdLineForNextPieInstance(int32 &WinX, int32 &WinY, int32 &InstanceNum, bool IsServer)
 {
 	// Get GameSettings INI override
@@ -1462,15 +1464,6 @@ void UEditorEngine::PlayInEditor( UWorld* InWorld, bool bInSimulateInEditor )
 {
 	double PIEStartTime = FPlatformTime::Seconds();
 
-	bool bStartInSpectatorMode = false;
-	
-	FModifierKeysState KeysState = FSlateApplication::Get().GetModifierKeys();
-	if (bInSimulateInEditor || KeysState.IsControlDown())
-	{
-		// if control is pressed, start in spectator mode
-		bStartInSpectatorMode = true;
-	}
-
 	// Prompt the user that Matinee must be closed before PIE can occur.
 	if( GEditorModeTools().IsModeActive(FBuiltinEditorModes::EM_InterpEdit) )
 	{
@@ -1733,7 +1726,6 @@ void UEditorEngine::PlayInEditor( UWorld* InWorld, bool bInSimulateInEditor )
 	EditorWorld->bAllowAudioPlayback = false;
 
 	ULevelEditorPlaySettings* PlayInSettings = Cast<ULevelEditorPlaySettings>(ULevelEditorPlaySettings::StaticClass()->GetDefaultObject());
-	int32 PIEInstance = 0;
 
 	if (!PlayInSettings->EnableSound && GetAudioDevice())
 	{
@@ -1766,76 +1758,293 @@ void UEditorEngine::PlayInEditor( UWorld* InWorld, bool bInSimulateInEditor )
 	DisableRealtimeViewports();
 
 	bool bAnyBlueprintErrors = ErrorBlueprintList.Num()? true : false;
+	bool bStartInSpectatorMode = false;
+
+	FModifierKeysState KeysState = FSlateApplication::Get().GetModifierKeys();
+	if (bInSimulateInEditor || KeysState.IsControlDown())
+	{
+		// if control is pressed, start in spectator mode
+		bStartInSpectatorMode = true;
+	}
 
 	if (bInSimulateInEditor || (PlayInSettings->PlayNetMode == EPlayNetMode::PIE_Standalone) || !PlayInSettings->RunUnderOneProcess)
 	{
 		// Only spawning 1 PIE instance under this process
 		FWorldContext &PieWorldContext = CreateNewWorldContext(EWorldType::PIE);
-		PieWorldContext.PIEInstance = PIEInstance++;
+		PieWorldContext.PIEInstance = 0;
 
 		CreatePlayInEditorWorld(PieWorldContext, bInSimulateInEditor, bAnyBlueprintErrors, bStartInSpectatorMode, PIEStartTime);
+
+		if (bInSimulateInEditor)
+		{
+			ToggleBetweenPIEandSIE();
+		}
 	}
 	else
 	{
-		// Spawning multiple PIE instances
-		if (PlayInSettings->MultipleInstancePositions.Num() == 0)
+		if (SupportsOnlinePIE())
 		{
-			PlayInSettings->MultipleInstancePositions.SetNum(1);
+			// Make sure all instances of PIE are logged in before creating/launching worlds
+			LoginPIEInstances(bAnyBlueprintErrors, bStartInSpectatorMode, PIEStartTime);
 		}
+		else
+		{
+			// Normal, non-online creation/launching of worlds
+			SpawnIntraProcessPIEWorlds(bAnyBlueprintErrors, bStartInSpectatorMode);
+		}
+	}
 
-		PlayInSettings->MultipleInstancePositions[0] = PlayInSettings->NewWindowPosition; 
+	PlayInSettings->MultipleInstanceLastHeight = PlayInSettings->NewWindowHeight;
+	PlayInSettings->MultipleInstanceLastWidth = PlayInSettings->NewWindowWidth;
+	PlayInSettings->PlayNetMode = OrigNetMode;
+}
+
+void UEditorEngine::SpawnIntraProcessPIEWorlds(bool bAnyBlueprintErrors, bool bStartInSpectatorMode)
+{
+	double PIEStartTime = FPlatformTime::Seconds();
 
-		int32 NextX = 0;
-		int32 NextY = 0;
-		int32 SettingsIndex = 1;
-		int32 ClientNum = 0;
+	// Has to be false or this function wouldn't be called
+	bool bInSimulateInEditor = false;
+	ULevelEditorPlaySettings* PlayInSettings = Cast<ULevelEditorPlaySettings>(ULevelEditorPlaySettings::StaticClass()->GetDefaultObject());
+	int32 PIEInstance = 0;
+
+	// Spawning multiple PIE instances
+	if (PlayInSettings->MultipleInstancePositions.Num() == 0)
+	{
+		PlayInSettings->MultipleInstancePositions.SetNum(1);
+	}
 
-		PIEInstance = 1;
+	PlayInSettings->MultipleInstancePositions[0] = PlayInSettings->NewWindowPosition;
 
-		// Server
-		FString ServerPrefix;
+	int32 NextX = 0;
+	int32 NextY = 0;
+	int32 SettingsIndex = 1;
+	int32 ClientNum = 0;
+
+	PIEInstance = 1;
+
+	// Server
+	FString ServerPrefix;
+	{
+		PlayInSettings->PlayNetMode = EPlayNetMode::PIE_ListenServer;
+		FWorldContext &PieWorldContext = CreateNewWorldContext(EWorldType::PIE);
+
+		PieWorldContext.PIEInstance = PIEInstance++;
+		PieWorldContext.RunAsDedicated = PlayInSettings->PlayNetDedicated;
+
+		if (!PlayInSettings->PlayNetDedicated)
 		{
-			PlayInSettings->PlayNetMode = EPlayNetMode::PIE_ListenServer;
-			FWorldContext &PieWorldContext = CreateNewWorldContext(EWorldType::PIE);
+			ClientNum++;
+			GetMultipleInstancePositions(SettingsIndex++, NextX, NextY);
+		}
 
-			PieWorldContext.PIEInstance = PIEInstance++;
-			PieWorldContext.RunAsDedicated = PlayInSettings->PlayNetDedicated;
+		CreatePlayInEditorWorld(PieWorldContext, bInSimulateInEditor, bAnyBlueprintErrors, bStartInSpectatorMode, PIEStartTime);
+		ServerPrefix = PieWorldContext.PIEPrefix;
+	}
 
-			if (!PlayInSettings->PlayNetDedicated)
+	// Clients
+	for (; ClientNum < PlayInSettings->PlayNumberOfClients; ++ClientNum)
+	{
+		PlayInSettings->PlayNetMode = EPlayNetMode::PIE_Client;
+		FWorldContext &PieWorldContext = CreateNewWorldContext(EWorldType::PIE);
+		PieWorldContext.PIEInstance = PIEInstance++;
+
+		GetMultipleInstancePositions(SettingsIndex++, NextX, NextY);
+		CreatePlayInEditorWorld(PieWorldContext, bInSimulateInEditor, bAnyBlueprintErrors, bStartInSpectatorMode, PIEStartTime);
+
+		PieWorldContext.PIERemapPrefix = ServerPrefix;
+	}
+
+	// Restore window settings
+	GetMultipleInstancePositions(0, NextX, NextY);	// restore cached settings
+}
+
+void UEditorEngine::CreatePIEWorldFromLogin(FWorldContext& PieWorldContext, EPlayNetMode PlayNetMode, FPieLoginStruct& DataStruct)
+{
+	ULevelEditorPlaySettings* PlayInSettings = Cast<ULevelEditorPlaySettings>(ULevelEditorPlaySettings::StaticClass()->GetDefaultObject());
+	PlayInSettings->PlayNetMode = PlayNetMode;
+
+	// Set window position
+	GetMultipleInstancePositions(DataStruct.SettingsIndex, DataStruct.NextX, DataStruct.NextY);
+	
+	CreatePlayInEditorWorld(PieWorldContext, false, DataStruct.bAnyBlueprintErrors, DataStruct.bStartInSpectatorMode, DataStruct.PIEStartTime);
+	
+	// Restore window settings
+	GetMultipleInstancePositions(0, DataStruct.NextX, DataStruct.NextY);	// restore cached settings
+
+	PieWorldContext.bWaitingOnOnlineSubsystem = false;
+
+	if (PlayNetMode == EPlayNetMode::PIE_ListenServer)
+	{
+		// If any clients finished before us, update their PIERemapPrefix
+		for (FWorldContext &WorldContext : WorldList)
+		{
+			if (WorldContext.WorldType == EWorldType::PIE && WorldContext.World() != NULL && WorldContext.ContextHandle != PieWorldContext.ContextHandle)
 			{
-				ClientNum++;
-				GetMultipleInstancePositions(SettingsIndex++, NextX, NextY);
+				WorldContext.PIERemapPrefix = PieWorldContext.PIEPrefix;
 			}
-			
-			CreatePlayInEditorWorld(PieWorldContext, bInSimulateInEditor, bAnyBlueprintErrors, bStartInSpectatorMode, PIEStartTime);		
-			ServerPrefix = PieWorldContext.PIEPrefix;
 		}
-
-		// Clients
-		for ( ; ClientNum < PlayInSettings->PlayNumberOfClients; ++ClientNum)
+	}
+	else
+	{
+		// Grab a valid PIERemapPrefix
+		for (FWorldContext &WorldContext : WorldList)
 		{
-			PlayInSettings->PlayNetMode = EPlayNetMode::PIE_Client;
-			FWorldContext &PieWorldContext = CreateNewWorldContext(EWorldType::PIE);
-			PieWorldContext.PIEInstance = PIEInstance++;
-
-			GetMultipleInstancePositions(SettingsIndex++, NextX, NextY);
-			CreatePlayInEditorWorld(PieWorldContext, bInSimulateInEditor, bAnyBlueprintErrors, bStartInSpectatorMode, PIEStartTime);
-			
-			PieWorldContext.PIERemapPrefix = ServerPrefix;
+			// This relies on the server being the first in the WorldList. Might be risky.
+			if (WorldContext.WorldType == EWorldType::PIE && WorldContext.World() != NULL && WorldContext.ContextHandle != PieWorldContext.ContextHandle)
+			{
+				PieWorldContext.PIERemapPrefix = WorldContext.PIEPrefix;
+				break;
+			}
 		}
+	}
+}
 
-		// Restore window settings
-		GetMultipleInstancePositions(0, NextX, NextY);	// restore cached settings
+bool UEditorEngine::SupportsOnlinePIE() const
+{
+	if (bOnlinePIEEnabled && PIELogins.Num() > 0)
+	{
+		// If we can't get the identity interface then things are either not configured right or disabled
+		IOnlineIdentityPtr IdentityInt = Online::GetIdentityInterface();
+		return IdentityInt.IsValid();
 	}
 
-	PlayInSettings->MultipleInstanceLastHeight = PlayInSettings->NewWindowHeight;
-	PlayInSettings->MultipleInstanceLastWidth = PlayInSettings->NewWindowWidth;
-	PlayInSettings->PlayNetMode = OrigNetMode;
+	return false;
+}
+
+void UEditorEngine::LoginPIEInstances(bool bAnyBlueprintErrors, bool bStartInSpectatorMode, double PIEStartTime)
+{
+	ULevelEditorPlaySettings* PlayInSettings = Cast<ULevelEditorPlaySettings>(ULevelEditorPlaySettings::StaticClass()->GetDefaultObject());
+
+	/** Setup the common data values for each login instance */
+	FPieLoginStruct DataStruct;
+	DataStruct.SettingsIndex = 1;
+	DataStruct.bAnyBlueprintErrors = bAnyBlueprintErrors;
+	DataStruct.bStartInSpectatorMode = bStartInSpectatorMode;
+	DataStruct.PIEStartTime = PIEStartTime;
 
-	if (bInSimulateInEditor)
+	int32 ClientNum = 0;
+	int32 PIEInstance = 1;
+	int32 NextX = 0;
+	int32 NextY = 0;
+
+	bool bHasRequiredLogins = PlayInSettings->PlayNumberOfClients <= PIELogins.Num();
+	if (!bHasRequiredLogins)
 	{
-		ToggleBetweenPIEandSIE();
+		UE_LOG(LogOnline, Verbose, TEXT("Not enough login credentials to launch all PIE instances, modify [/Script/UnrealEd.UnrealEdEngine].PIELogins"));
 	}
+
+	// Server
+	{
+		FWorldContext &PieWorldContext = CreateNewWorldContext(EWorldType::PIE);
+		PieWorldContext.PIEInstance = PIEInstance++;
+		PieWorldContext.RunAsDedicated = PlayInSettings->PlayNetDedicated;
+		PieWorldContext.bWaitingOnOnlineSubsystem = true;
+
+		// Update login struct parameters
+		DataStruct.WorldContextHandle = PieWorldContext.ContextHandle;
+		DataStruct.bIsServer = true;
+
+		// Always get the interface (it will create the subsystem regardless)
+		FName OnlineIdentifier = FName(*FString::Printf(TEXT(":%d"), PieWorldContext.ContextHandle));
+		IOnlineSubsystem* OnlineSub = IOnlineSubsystem::Get(OnlineIdentifier);
+		IOnlineIdentityPtr IdentityInt = OnlineSub->GetIdentityInterface();
+
+		if (!PlayInSettings->PlayNetDedicated)
+		{
+			DataStruct.NextX = NextX;
+			DataStruct.NextY = NextY;
+			GetMultipleInstancePositions(DataStruct.SettingsIndex, NextX, NextY);
+
+			if (bHasRequiredLogins)
+			{
+				// Login to online platform before creating world
+				FOnlineAccountCredentials AccountCreds;
+				AccountCreds.Id = PIELogins[ClientNum].Id;
+				AccountCreds.Token = PIELogins[ClientNum].Token;
+				AccountCreds.Type = PIELogins[ClientNum].Type;
+
+				FOnLoginCompleteDelegate Delegate;
+				Delegate.BindUObject(this, &UEditorEngine::OnLoginPIEComplete, DataStruct);
+
+				// Login first and continue the flow later
+				IdentityInt->AddOnLoginCompleteDelegate(0, Delegate);
+				IdentityInt->Login(0, AccountCreds);
+			}
+			else
+			{
+				// No logins, but continue anyway
+				CreatePIEWorldFromLogin(PieWorldContext, EPlayNetMode::PIE_ListenServer, DataStruct);
+			}
+
+			ClientNum++;
+		}
+		else
+		{
+			// Dedicated servers don't use a login
+			OnlineSub->SetForceDedicated(true);
+			CreatePIEWorldFromLogin(PieWorldContext, EPlayNetMode::PIE_ListenServer, DataStruct);
+		}
+	}
+
+	// Clients
+	for (; ClientNum < PlayInSettings->PlayNumberOfClients; ++ClientNum)
+	{
+		PlayInSettings->PlayNetMode = EPlayNetMode::PIE_Client;
+		FWorldContext &PieWorldContext = CreateNewWorldContext(EWorldType::PIE);
+		PieWorldContext.PIEInstance = PIEInstance++;
+		PieWorldContext.bWaitingOnOnlineSubsystem = true;
+
+		// Update login struct parameters
+		DataStruct.WorldContextHandle = PieWorldContext.ContextHandle;
+		DataStruct.SettingsIndex++;
+		DataStruct.NextX = NextX;
+		DataStruct.NextY = NextY;
+		GetMultipleInstancePositions(DataStruct.SettingsIndex, NextX, NextY);
+		DataStruct.bIsServer = false;
+
+		if (bHasRequiredLogins)
+		{
+			FName OnlineIdentifier = FName(*FString::Printf(TEXT(":%d"), PieWorldContext.ContextHandle));
+			IOnlineIdentityPtr IdentityInt = Online::GetIdentityInterface(OnlineIdentifier);
+
+			FOnlineAccountCredentials AccountCreds;
+			AccountCreds.Id = PIELogins[ClientNum].Id;
+			AccountCreds.Token = PIELogins[ClientNum].Token;
+			AccountCreds.Type = PIELogins[ClientNum].Type;
+
+			FOnLoginCompleteDelegate Delegate;
+			Delegate.BindUObject(this, &UEditorEngine::OnLoginPIEComplete, DataStruct);
+
+			IdentityInt->ClearOnLoginCompleteDelegate(0, Delegate);
+			IdentityInt->AddOnLoginCompleteDelegate(0, Delegate);
+			IdentityInt->Login(0, AccountCreds);
+		}
+		else
+		{
+			// No logins, but continue anyway
+			CreatePIEWorldFromLogin(PieWorldContext, EPlayNetMode::PIE_Client, DataStruct);
+		}
+	}
+
+	// Restore window settings
+	GetMultipleInstancePositions(0, NextX, NextY);	// restore cached settings
+}
+
+void UEditorEngine::OnLoginPIEComplete(int32 LocalUserNum, bool bWasSuccessful, const FUniqueNetId& UserId, const FString& ErrorString, FPieLoginStruct DataStruct)
+{
+	UE_LOG(LogOnline, Verbose, TEXT("OnLoginPIEComplete LocalUserNum: %d bSuccess: %d %s"), LocalUserNum, bWasSuccessful, *ErrorString);
+
+	FWorldContext& PieWorldContext = GetWorldContextFromHandleChecked(DataStruct.WorldContextHandle);
+
+	CreatePIEWorldFromLogin(PieWorldContext, DataStruct.bIsServer ? EPlayNetMode::PIE_ListenServer : EPlayNetMode::PIE_Client, DataStruct);
+
+	FName OnlineIdentifier = FName(*FString::Printf(TEXT(":%d"), PieWorldContext.ContextHandle));
+	IOnlineIdentityPtr IdentityInt = Online::GetIdentityInterface(OnlineIdentifier);
+
+	FOnLoginCompleteDelegate Delegate;
+	Delegate.BindUObject(this, &UEditorEngine::OnLoginPIEComplete, DataStruct);
+	IdentityInt->ClearOnLoginCompleteDelegate(0, Delegate);
 }
 
 UWorld* UEditorEngine::CreatePlayInEditorWorld(FWorldContext &PieWorldContext, bool bInSimulateInEditor, bool bAnyBlueprintErrors, bool bStartInSpectatorMode, float PIEStartTime)
@@ -2801,4 +3010,6 @@ void UEditorEngine::RemapGamepadControllerIdForPIE(class UGameViewportClient* Ga
 	}
 }
 
+
+
 #undef LOCTEXT_NAMESPACE
diff --git a/Engine/Source/Editor/UnrealEd/UnrealEd.Build.cs b/Engine/Source/Editor/UnrealEd/UnrealEd.Build.cs
--- a/Engine/Source/Editor/UnrealEd/UnrealEd.Build.cs
+++ b/Engine/Source/Editor/UnrealEd/UnrealEd.Build.cs
@@ -70,6 +70,7 @@ public class UnrealEd : ModuleRules
 			{ 
 				"AnimGraph",
 				"BlueprintGraph",
+				"OnlineSubsystem",
 				"OnlineBlueprintSupport",
 				"DirectoryWatcher",
 				"EditorSettingsViewer",
diff --git a/Engine/Source/Runtime/Engine/Classes/Engine/Engine.h b/Engine/Source/Runtime/Engine/Classes/Engine/Engine.h
--- a/Engine/Source/Runtime/Engine/Classes/Engine/Engine.h
+++ b/Engine/Source/Runtime/Engine/Classes/Engine/Engine.h
@@ -274,6 +274,10 @@ struct FWorldContext
 	UPROPERTY()
 	bool	RunAsDedicated;
 
+	/** Is this world context waiting for an online login to complete (for PIE) */
+	UPROPERTY()
+	bool	bWaitingOnOnlineSubsystem;
+
 	/**************************************************************/
 
 	/** Outside pointers to CurrentWorld that should be kept in sync if current world changes  */
@@ -322,6 +326,7 @@ struct FWorldContext
 		, GameViewport(NULL)
 		, PIEInstance(INDEX_NONE)
 		, RunAsDedicated(false)
+		, bWaitingOnOnlineSubsystem(false)
 		, ThisCurrentWorld(NULL)
 	{
 	}
