Fixed unsolicited cooked packages in the cook on the fly server.
Working towards merging some more cookonthefly and cookontheside code paths.

diff --git a/Engine/Source/Editor/UnrealEd/Classes/CookOnTheSide/CookOnTheFlyServer.h b/Engine/Source/Editor/UnrealEd/Classes/CookOnTheSide/CookOnTheFlyServer.h
--- a/Engine/Source/Editor/UnrealEd/Classes/CookOnTheSide/CookOnTheFlyServer.h
+++ b/Engine/Source/Editor/UnrealEd/Classes/CookOnTheSide/CookOnTheFlyServer.h
@@ -44,7 +44,7 @@ private:
 
 	/** Array which has been made thread safe :) */
 	template<typename Type>
-	struct FThreadSafeArray
+	struct FThreadSafeQueue
 	{
 	private:
 		mutable FCriticalSection	SynchronizationObject; // made this mutable so this class can have const functions and still be thread safe
@@ -55,6 +55,11 @@ private:
 			FScopeLock ScopeLock(&SynchronizationObject);
 			Items.Add(Item);
 		}
+		void EnqueueUnique( const Type& Item )
+		{
+			FScopeLock ScopeLock(&SynchronizationObject);
+			Items.AddUnique(Item);
+		}
 		bool Dequeue(Type* Result)
 		{
 			FScopeLock ScopeLock(&SynchronizationObject);
@@ -78,6 +83,55 @@ private:
 			FScopeLock ScopeLock( &SynchronizationObject );
 			return Items.Num() > 0;
 		}
+
+		void Remove( const Type& Item ) 
+		{
+			FScopeLock ScopeLock( &SynchronizationObject );
+			Items.Remove( Item );
+		}
+	};
+
+	/** Normal queue to match the interface of FThreadSafeQueue :) */
+	template<typename Type>
+	struct FQueue
+	{
+	private:
+		TArray<Type>		Items;
+	public:
+		void Enqueue(const Type& Item)
+		{
+			Items.Add(Item);
+		}
+		void EnqueueUnique( const Type& Item )
+		{
+			Items.AddUnique(Item);
+		}
+		bool Dequeue(Type* Result)
+		{
+			if (Items.Num())
+			{
+				*Result = Items[0];
+				Items.RemoveAt(0);
+				return true;
+			}
+			return false;
+		}
+		void DequeueAll(TArray<Type>& Results)
+		{
+			Results += Items;
+			Items.Empty();
+		}
+
+		bool HasItems() const
+		{
+			return Items.Num() > 0;
+		}
+
+
+		void Remove( const Type& Item ) 
+		{
+			Items.Remove( Item );
+		}
 	};
 
 public:
@@ -124,8 +178,9 @@ public:
 
 private:
 
-	FThreadSafeArray<FFilePlatformRequest> ThreadSafeFilenameArray; // list of requested files
-	FThreadSafeArray<FFilePlatformRequest> UnsolicitedCookedPackages; // list of files which haven't been requested but we think should cook based on previous requests
+	FThreadSafeQueue<FFilePlatformRequest> FileRequests; // list of requested files
+	FQueue<FFilePlatformRequest> UnsolicitedFileRequests;
+	FThreadSafeQueue<FFilePlatformRequest> UnsolicitedCookedPackages; // list of files which haven't been requested but we think should cook based on previous requests
 
 
 
@@ -150,7 +205,7 @@ private:
 	FThreadSafeFilenameSet ThreadSafeFilenameSet; // set of files which have been cooked when needing to recook a file the entry will need to be removed from here
 
 
-	FThreadSafeArray<struct FRecompileRequest*> RecompileRequests;
+	FThreadSafeQueue<struct FRecompileRequest*> RecompileRequests;
 
 
 public:
@@ -207,7 +262,8 @@ public:
 	 */
 	void TickRecompileShaderRequests();
 
-	bool HasCookRequests() const { return ThreadSafeFilenameArray.HasItems(); }
+	bool HasCookRequests() const { return FileRequests.HasItems(); }
+	bool HasUnsolicitedCookRequests() const { return UnsolicitedFileRequests.HasItems(); }
 
 	bool HasRecompileShaderRequests() const { return RecompileRequests.HasItems(); }
 
@@ -221,18 +277,6 @@ public:
 	uint32 NumConnections() const;
 
 	
-	/**
-	 * BLOCKING Run the network file server and handle cook requests 
-	 *	don't return until finished cooking or we haven't received a cook requests for a length of time
-	 *  shouldn't be used in conjunction with StartCookOnTheSide, TickCookOnTheSide, TickRecompileShaderRequests, FinishCookOnTheSide
-	 *
-	 * @param  BindAnyPort					Whether to bind on any port or the default port.
-	 * @param  Timeout						Length of time to wait for connections before attempting to close
-	 * @param  bForceClose					Whether or not the server should always shutdown after a timeout or after a user disconnects
-	 *
-	 * @return true on success, false otherwise.
-	 */
-	bool CookOnTheFly( bool BindAnyPort, int32 Timeout = 180, bool bForceClose = false );
 
 
 	virtual void BeginDestroy() OVERRIDE;
@@ -243,7 +287,7 @@ public:
 private:
 	
 
-
+	bool UCookOnTheFlyServer::ShouldCook(const FString& InFileName, const FString &InPlatformName);
 
 	/**
 	 * Returns cooker output directory.
diff --git a/Engine/Source/Editor/UnrealEd/Private/Commandlets/CookCommandlet.cpp b/Engine/Source/Editor/UnrealEd/Private/Commandlets/CookCommandlet.cpp
--- a/Engine/Source/Editor/UnrealEd/Private/Commandlets/CookCommandlet.cpp
+++ b/Engine/Source/Editor/UnrealEd/Private/Commandlets/CookCommandlet.cpp
@@ -105,90 +105,88 @@ bool UCookCommandlet::CookOnTheFly( bool BindAnyPort, int32 Timeout, bool bForce
 	FDateTime LastConnectionTime = FDateTime::UtcNow();
 	bool bHadConnection = false;
 
+	
+
 	while (!GIsRequestingExit)
 	{
+		uint32 TickResults = 0;
+		static const float CookOnTheSideTimeSlice = 10.0f;
+		TickResults = CookOnTheFlyServer->TickCookOnTheSide(CookOnTheSideTimeSlice, NonMapPackageCountSinceLastGC);
 
-		while (!GIsRequestingExit)
+		bool bCookedAMapSinceLastGC = TickResults & UCookOnTheFlyServer::COSR_CookedMap;
+		if ( TickResults & (UCookOnTheFlyServer::COSR_CookedMap | UCookOnTheFlyServer::COSR_CookedPackage))
 		{
-			uint32 TickResults = 0;
-			
-			CookOnTheFlyServer->TickCookOnTheSide(TickResults, NonMapPackageCountSinceLastGC);
-
-			bool bCookedAMapSinceLastGC = TickResults & UCookOnTheFlyServer::COSR_CookedMap;
-			if ( TickResults & (UCookOnTheFlyServer::COSR_CookedMap | UCookOnTheFlyServer::COSR_CookedPackage))
-			{
-				LastCookActionTime = FPlatformTime::Seconds();
-			}
+			LastCookActionTime = FPlatformTime::Seconds();
+		}
 
 
-			while ( (CookOnTheFlyServer->HasCookRequests() == false) && !GIsRequestingExit)
-			{
+		while ( (CookOnTheFlyServer->HasCookRequests() == false) && !GIsRequestingExit)
+		{
 				
+			{
+				if (NonMapPackageCountSinceLastGC > 0)
 				{
-					if (NonMapPackageCountSinceLastGC > 0)
-					{
-						// We should GC if we have packages to collect and we've been idle for some time.
-						bShouldGC = (NonMapPackageCountSinceLastGC > PackagesPerGC) || 
-							((FPlatformTime::Seconds() - LastCookActionTime) >= IdleTimeToGC);
-					}
-					bShouldGC |= bCookedAMapSinceLastGC;
+					// We should GC if we have packages to collect and we've been idle for some time.
+					bShouldGC = (NonMapPackageCountSinceLastGC > PackagesPerGC) || 
+						((FPlatformTime::Seconds() - LastCookActionTime) >= IdleTimeToGC);
+				}
+				bShouldGC |= bCookedAMapSinceLastGC;
 
-					if (bShouldGC)
-					{
-						bShouldGC = false;
-						bCookedAMapSinceLastGC = false;
-						NonMapPackageCountSinceLastGC = 0;
+				if (bShouldGC)
+				{
+					bShouldGC = false;
+					bCookedAMapSinceLastGC = false;
+					NonMapPackageCountSinceLastGC = 0;
 
-						UE_LOG(LogCookCommandlet, Display, TEXT("GC..."));
+					UE_LOG(LogCookCommandlet, Display, TEXT("GC..."));
 
-						CollectGarbage( RF_Native );
-					}
-					else
-					{
-						CookOnTheFlyServer->TickRecompileShaderRequests();
+					CollectGarbage( RF_Native );
+				}
+				else
+				{
+					CookOnTheFlyServer->TickRecompileShaderRequests();
 
-						FPlatformProcess::Sleep(0.0f);
-					}
+					FPlatformProcess::Sleep(0.0f);
 				}
+			}
+
+			// update task graph
+			FTaskGraphInterface::Get().ProcessThreadUntilIdle(ENamedThreads::GameThread);
+
+			// execute deferred commands
+			for (int32 DeferredCommandsIndex = 0; DeferredCommandsIndex<GEngine->DeferredCommands.Num(); ++DeferredCommandsIndex)
+			{
+				GEngine->Exec( GWorld, *GEngine->DeferredCommands[DeferredCommandsIndex], *GLog);
+			}
 
-				// update task graph
-				FTaskGraphInterface::Get().ProcessThreadUntilIdle(ENamedThreads::GameThread);
+			GEngine->DeferredCommands.Empty();
 
-				// execute deferred commands
-				for (int32 DeferredCommandsIndex = 0; DeferredCommandsIndex<GEngine->DeferredCommands.Num(); ++DeferredCommandsIndex)
+			// handle server timeout
+			if (InstanceId.IsValid() || bForceClose)
+			{
+				if (CookOnTheFlyServer->NumConnections() > 0)
 				{
-					GEngine->Exec( GWorld, *GEngine->DeferredCommands[DeferredCommandsIndex], *GLog);
+					bHadConnection = true;
+					LastConnectionTime = FDateTime::UtcNow();
 				}
 
-				GEngine->DeferredCommands.Empty();
-
-				// handle server timeout
-				if (InstanceId.IsValid() || bForceClose)
+				if ((FDateTime::UtcNow() - LastConnectionTime) > FTimespan::FromSeconds(Timeout))
 				{
-					if (CookOnTheFlyServer->NumConnections() > 0)
-					{
-						bHadConnection = true;
-						LastConnectionTime = FDateTime::UtcNow();
-					}
+					uint32 Result = FMessageDialog::Open(EAppMsgType::YesNo, NSLOCTEXT("UnrealEd", "FileServerIdle", "The file server did not receive any connections in the past 3 minutes. Would you like to shut it down?"));
 
-					if ((FDateTime::UtcNow() - LastConnectionTime) > FTimespan::FromSeconds(Timeout))
+					if (Result == EAppReturnType::No && !bForceClose)
 					{
-						uint32 Result = FMessageDialog::Open(EAppMsgType::YesNo, NSLOCTEXT("UnrealEd", "FileServerIdle", "The file server did not receive any connections in the past 3 minutes. Would you like to shut it down?"));
-
-						if (Result == EAppReturnType::No && !bForceClose)
-						{
-							LastConnectionTime = FDateTime::UtcNow();
-						}
-						else
-						{
-							GIsRequestingExit = true;
-						}
+						LastConnectionTime = FDateTime::UtcNow();
 					}
-					else if (bHadConnection && (CookOnTheFlyServer->NumConnections() == 0) && bForceClose) // immediately shut down if we previously had a connection and now do not
+					else
 					{
 						GIsRequestingExit = true;
 					}
 				}
+				else if (bHadConnection && (CookOnTheFlyServer->NumConnections() == 0) && bForceClose) // immediately shut down if we previously had a connection and now do not
+				{
+					GIsRequestingExit = true;
+				}
 			}
 		}
 	}
diff --git a/Engine/Source/Editor/UnrealEd/Private/CookOnTheFlyServer.cpp b/Engine/Source/Editor/UnrealEd/Private/CookOnTheFlyServer.cpp
--- a/Engine/Source/Editor/UnrealEd/Private/CookOnTheFlyServer.cpp
+++ b/Engine/Source/Editor/UnrealEd/Private/CookOnTheFlyServer.cpp
@@ -75,7 +75,8 @@ UCookOnTheFlyServer::UCookOnTheFlyServer(const class FPostConstructInitializePro
 void UCookOnTheFlyServer::Tick(float DeltaTime)
 {
 	uint32 CookedPackagesCount = 0;
-	TickCookOnTheSide( 0.01f, CookedPackagesCount);
+	const static float CookOnTheSideTimeSlice = 0.01f;
+	TickCookOnTheSide( CookOnTheSideTimeSlice, CookedPackagesCount);
 	TickRecompileShaderRequests();
 }
 
@@ -148,69 +149,115 @@ uint32 UCookOnTheFlyServer::TickCookOnTheSide( const float TimeSlice, uint32 &Co
 	// we use this in the unsolicited packages processing below
 	TArray<FString> AllTargetPlatformNames;
 	bool bCookedPackage = false;
-	while (ThreadSafeFilenameArray.HasItems() && !GIsRequestingExit)
+	while (!GIsRequestingExit)
 	{
 		FFilePlatformRequest ToBuild;
-		ThreadSafeFilenameArray.Dequeue( &ToBuild );
+
+		bool bIsUnsolicitedRequest = false;
+		if ( FileRequests.HasItems() )
+		{
+			FileRequests.Dequeue( &ToBuild );
+		}
+		else if ( UnsolicitedFileRequests.HasItems() )
+		{
+			bIsUnsolicitedRequest = true;
+			UnsolicitedFileRequests.Dequeue( &ToBuild );
+		}
+		else
+		{
+			// no more to do this tick break out and do some other stuff
+			break;
+		}
 
 		check( ToBuild.IsValid() );
-		bool bWasUpToDate = false;
 		TArray<FString> TargetPlatformNames;
+		
 		if( ToBuild.Platformname.Len() )
 		{
 			TargetPlatformNames.Add( ToBuild.Platformname );
 		}
 
+		bool bWasUpToDate = false;
 		FString Name = FPackageName::FilenameToLongPackageName(*ToBuild.Filename);
 		if (!ThreadSafeFilenameSet.Exists(ToBuild))
 		{
-			UPackage* Package = LoadPackage( NULL, *ToBuild.Filename, LOAD_None );
-
-			if( Package == NULL )
+			// if we have no target platforms then we want to cook because this will cook for all target platforms in that case
+			bool bShouldCook = TargetPlatformNames.Num() > 0 ? false : true;
+			for ( int Index = 0; Index < TargetPlatformNames.Num(); ++Index )
 			{
-				UE_LOG(LogCookCommandlet, Error, TEXT("Error loading %s!"), *ToBuild.Filename );
-				Result |= COSR_ErrorLoadingPackage;
+				bShouldCook |= ShouldCook( ToBuild.Filename, TargetPlatformNames[Index] );
 			}
-			else
+
+			if ( bShouldCook ) // if we should cook the package then cook it otherwise add it to the list of already cooked packages below
 			{
-				if( SaveCookedPackage(Package, SAVE_KeepGUID | SAVE_Async, bWasUpToDate, TargetPlatformNames) )
+				//  if the package is already loaded then don't do any work :)
+				UPackage* Package = LoadPackage( NULL, *ToBuild.Filename, LOAD_None );
+
+				if( Package == NULL )
 				{
-					// Update flags used to determine garbage collection.
-					if (Package->ContainsMap())
+					UE_LOG(LogCookCommandlet, Error, TEXT("Error loading %s!"), *ToBuild.Filename );
+					Result |= COSR_ErrorLoadingPackage;
+				}
+				else
+				{
+					if( SaveCookedPackage(Package, SAVE_KeepGUID | SAVE_Async, bWasUpToDate, TargetPlatformNames) )
 					{
-						Result |= COSR_CookedMap;
-						//bCookedAMapSinceLastGC = true;
+						// Update flags used to determine garbage collection.
+						if (Package->ContainsMap())
+						{
+							Result |= COSR_CookedMap;
+							//bCookedAMapSinceLastGC = true;
+						}
+						else
+						{
+							++CookedPackageCount;
+							Result |= COSR_CookedPackage;
+							//NonMapPackageCountSinceLastGC++;
+							// LastCookActionTime = FPlatformTime::Seconds();
+						}
+						bCookedPackage = true;
 					}
-					else
+
+					//@todo ResetLoaders outside of this (ie when Package is NULL) causes problems w/ default materials
+					if (Package->HasAnyFlags(RF_RootSet) == false)
 					{
-						++CookedPackageCount;
-						Result |= COSR_CookedPackage;
-						//NonMapPackageCountSinceLastGC++;
-						// LastCookActionTime = FPlatformTime::Seconds();
+						ResetLoaders(Package);
 					}
-					bCookedPackage = true;
 				}
+			}
 
-				//@todo ResetLoaders outside of this (ie when Package is NULL) causes problems w/ default materials
-				if (Package->HasAnyFlags(RF_RootSet) == false)
+			for ( int Index = 0; Index < TargetPlatformNames.Num(); ++Index )
+			{
+				AllTargetPlatformNames.AddUnique(TargetPlatformNames[Index]);
+
+				FFilePlatformRequest FileRequest( ToBuild.Filename, TargetPlatformNames[Index]);
+
+				if ( bIsUnsolicitedRequest )
 				{
-					ResetLoaders(Package);
+					if ((FPaths::FileExists(FileRequest.Filename) == true) &&
+						(bWasUpToDate == false))
+					{
+						UnsolicitedCookedPackages.EnqueueUnique( FileRequest );
+						UE_LOG(LogCookCommandlet, Display, TEXT("UnsolicitedCookedPackages: %s"), *FileRequest.Filename);
+					}
 				}
-			}
-		}
 
-		for ( int Index = 0; Index < TargetPlatformNames.Num(); ++Index )
-		{
-			AllTargetPlatformNames.AddUnique(TargetPlatformNames[Index]);
+				UnsolicitedCookedPackages.Remove( FileRequest );
+				
+				ThreadSafeFilenameSet.Add( FileRequest );
 
-			FFilePlatformRequest FileRequest( ToBuild.Filename, TargetPlatformNames[Index]);
 
-			ThreadSafeFilenameSet.Add( FileRequest );
+			}
 		}
+		else
+		{
+			UE_LOG(LogCookCommandlet, Display, TEXT("Package for platform already cooked %s, discarding request"), *ToBuild.Filename);
+		}
+		
 
 		if ( (FPlatformTime::Seconds() - StartTime) > TimeSlice)
 		{
-			return Result;
+			break;
 		}
 	}
 
@@ -220,11 +267,6 @@ uint32 UCookOnTheFlyServer::TickCookOnTheSide( const float TimeSlice, uint32 &Co
 		GetObjectsWithOuter(NULL, ObjectsInOuter, false);
 		for( int32 Index = 0; Index < ObjectsInOuter.Num() && !GIsRequestingExit; Index++ )
 		{
-			if ( (FPlatformTime::Seconds() - StartTime) > TimeSlice)
-			{
-				return Result;
-			}
-
 			UPackage* Package = Cast<UPackage>(ObjectsInOuter[Index]);
 
 			if (Package)
@@ -246,46 +288,12 @@ uint32 UCookOnTheFlyServer::TickCookOnTheSide( const float TimeSlice, uint32 &Co
 				}
 				if ( !AlreadyCooked )
 				{
-					bool bUnsolicitedWasUpdateToDate = false;
-
-					if (SaveCookedPackage(Package, SAVE_KeepGUID | SAVE_Async, bUnsolicitedWasUpdateToDate, AllTargetPlatformNames))
+					for ( int Index = 0; Index < AllTargetPlatformNames.Num(); ++Index )
 					{
-
-
-						if ((FPaths::FileExists(PackageFilename) == true) &&
-							(PackageFilename.Len() > 0) && 
-							(bUnsolicitedWasUpdateToDate == false))
-						{
-							UE_LOG(LogCookCommandlet, Display, TEXT("UnsolicitedCookedPackages: %s"), *PackageFilename);
-
-							for ( int Index = 0; Index < AllTargetPlatformNames.Num(); ++Index )
-							{
-								const FFilePlatformRequest Request( FullPath, AllTargetPlatformNames[Index] );
-								UnsolicitedCookedPackages.Enqueue( Request );
-								ThreadSafeFilenameSet.Add( Request );
-
-								UE_LOG( LogCookCommandlet, Display, TEXT("Unsolicited saved package %s.%s"), *Request.Filename, *Request.Platformname );
-							}
-						}
-
-						// Update flags used to determine garbage collection.
-						if (Package->ContainsMap())
-						{
-							Result |= COSR_CookedMap;
-						}
-						else
-						{
-							++CookedPackageCount;
-							Result |= COSR_CookedPackage;
-						}
+						const FFilePlatformRequest Request( FullPath, AllTargetPlatformNames[Index] );
+						UnsolicitedFileRequests.EnqueueUnique(Request);
 					}
 				}
-
-				if ( ThreadSafeFilenameArray.HasItems())
-				{
-					// break out of saving unsolicited packages if we get a real request
-					break;
-				}
 			}
 		}
 	}
@@ -329,123 +337,6 @@ void UCookOnTheFlyServer::TickRecompileShaderRequests()
 	}
 }
 
-bool UCookOnTheFlyServer::CookOnTheFly( bool BindAnyPort, int32 Timeout, bool bForceClose )
-{
-	if ( StartNetworkFileServer(BindAnyPort) == false )
-	{
-		return false;
-	}
-
-	// Garbage collection should happen when either
-	//	1. We have cooked a map
-	//	2. We have cooked non-map packages and...
-	//		a. we have accumulated 50 of these since the last GC.
-	//		b. we have been idle for 20 seconds.
-	bool bShouldGC = true;
-
-	// megamoth
-	uint32 NonMapPackageCountSinceLastGC = 0;
-	
-	const int32 PackagesPerGC = 50;
-	
-	const double IdleTimeToGC = 20.0;
-	double LastCookActionTime = FPlatformTime::Seconds();
-
-	FDateTime LastConnectionTime = FDateTime::UtcNow();
-	bool bHadConnection = false;
-
-	while (!GIsRequestingExit)
-	{
-
-		while (!GIsRequestingExit)
-		{
-			uint32 TickResults = 0;
-			
-			TickCookOnTheSide(TickResults, NonMapPackageCountSinceLastGC);
-
-			bool bCookedAMapSinceLastGC = TickResults & COSR_CookedMap;
-			if ( TickResults & (COSR_CookedMap | COSR_CookedPackage))
-			{
-				LastCookActionTime = FPlatformTime::Seconds();
-			}
-
-
-			while ( (ThreadSafeFilenameArray.HasItems() == false) && !GIsRequestingExit)
-			{
-				
-				{
-					if (NonMapPackageCountSinceLastGC > 0)
-					{
-						// We should GC if we have packages to collect and we've been idle for some time.
-						bShouldGC = (NonMapPackageCountSinceLastGC > PackagesPerGC) || 
-							((FPlatformTime::Seconds() - LastCookActionTime) >= IdleTimeToGC);
-					}
-					bShouldGC |= bCookedAMapSinceLastGC;
-
-					if (bShouldGC)
-					{
-						bShouldGC = false;
-						bCookedAMapSinceLastGC = false;
-						NonMapPackageCountSinceLastGC = 0;
-
-						UE_LOG(LogCookCommandlet, Display, TEXT("GC..."));
-
-						CollectGarbage( RF_Native );
-					}
-					else
-					{
-						TickRecompileShaderRequests();
-
-						FPlatformProcess::Sleep(0.0f);
-					}
-				}
-
-				// update task graph
-				FTaskGraphInterface::Get().ProcessThreadUntilIdle(ENamedThreads::GameThread);
-
-				// execute deferred commands
-				for (int32 DeferredCommandsIndex = 0; DeferredCommandsIndex<GEngine->DeferredCommands.Num(); ++DeferredCommandsIndex)
-				{
-					GEngine->Exec( GWorld, *GEngine->DeferredCommands[DeferredCommandsIndex], *GLog);
-				}
-
-				GEngine->DeferredCommands.Empty();
-
-				// handle server timeout
-				if (InstanceId.IsValid() || bForceClose)
-				{
-					if (NetworkFileServer && NetworkFileServer->NumConnections() > 0)
-					{
-						bHadConnection = true;
-						LastConnectionTime = FDateTime::UtcNow();
-					}
-
-					if ((FDateTime::UtcNow() - LastConnectionTime) > FTimespan::FromSeconds(Timeout))
-					{
-						uint32 Result = FMessageDialog::Open(EAppMsgType::YesNo, NSLOCTEXT("UnrealEd", "FileServerIdle", "The file server did not receive any connections in the past 3 minutes. Would you like to shut it down?"));
-
-						if (Result == EAppReturnType::No && !bForceClose)
-						{
-							LastConnectionTime = FDateTime::UtcNow();
-						}
-						else
-						{
-							GIsRequestingExit = true;
-						}
-					}
-					else if (bHadConnection && (NetworkFileServer && NetworkFileServer->NumConnections() == 0) && bForceClose) // immediately shut down if we previously had a connection and now do not
-					{
-						GIsRequestingExit = true;
-					}
-				}
-			}
-		}
-	}
-
-	EndNetworkFileServer();
-	return true;
-}
-
 
 FString UCookOnTheFlyServer::GetOutputDirectory( const FString& PlatformName ) const
 {
@@ -478,6 +369,65 @@ bool UCookOnTheFlyServer::SaveCookedPackage( UPackage* Package, uint32 SaveFlags
 	return SaveCookedPackage( Package, SaveFlags, bOutWasUpToDate, TargetPlatformNames );
 }
 
+
+bool UCookOnTheFlyServer::ShouldCook(const FString& InFileName, const FString &InPlatformName)
+{
+	bool bDoCook = false;
+
+
+	FString PkgFile;
+	FString PkgFilename;
+	FDateTime DependentTimeStamp = FDateTime::MinValue();
+
+	if (bIterativeCooking && FPackageName::DoesPackageExist(InFileName, NULL, &PkgFile))
+	{
+		PkgFilename = PkgFile;
+
+		if (GetPackageTimestamp(FPaths::GetBaseFilename(PkgFilename, false), DependentTimeStamp) == false)
+		{
+			UE_LOG(LogCookCommandlet, Display, TEXT("Failed to find dependency timestamp for: %s"), *PkgFilename);
+		}
+	}
+
+	// Use SandboxFile to do path conversion to properly handle sandbox paths (outside of standard paths in particular).
+	PkgFilename = SandboxFile->ConvertToAbsolutePathForExternalAppForWrite(*PkgFilename);
+
+	ITargetPlatformManagerModule& TPM = GetTargetPlatformManagerRef();
+
+	static const TArray<ITargetPlatform*> &ActiveTargetPlatforms = TPM.GetActiveTargetPlatforms();
+
+	TArray<ITargetPlatform*> Platforms;
+
+	if ( InPlatformName.Len() > 0 )
+	{
+		Platforms.Add( TPM.FindTargetPlatform( InPlatformName ) );
+	}
+	else 
+	{
+		Platforms = ActiveTargetPlatforms;
+	}
+
+	for (int32 Index = 0; Index < Platforms.Num() && !bDoCook; Index++)
+	{
+		ITargetPlatform* Target = Platforms[Index];
+		FString PlatFilename = PkgFilename.Replace(TEXT("[Platform]"), *Target->PlatformName());
+
+		// If we are not iterative cooking, then cook the package
+		bool bCookPackage = (bIterativeCooking == false);
+
+		if (bCookPackage == false)
+		{
+			// If the cooked package doesn't exist, or if the cooked is older than the dependent, re-cook it
+			FDateTime CookedTimeStamp = IFileManager::Get().GetTimeStamp(*PlatFilename);
+			int32 CookedTimespanSeconds = (CookedTimeStamp - DependentTimeStamp).GetTotalSeconds();
+			bCookPackage = (CookedTimeStamp == FDateTime::MinValue()) || (CookedTimespanSeconds < 0);
+		}
+		bDoCook |= bCookPackage;
+	}
+
+	return bDoCook;
+}
+
 bool UCookOnTheFlyServer::SaveCookedPackage( UPackage* Package, uint32 SaveFlags, bool& bOutWasUpToDate, TArray<FString> &TargetPlatformNames )
 {
 
@@ -882,7 +832,7 @@ void UCookOnTheFlyServer::HandleNetworkFileServerFileRequest( const FString& Fil
 	}
 
 	FFilePlatformRequest FileRequest( Filename, Platformname );
-	ThreadSafeFilenameArray.Enqueue(FileRequest);
+	FileRequests.Enqueue(FileRequest);
 
 	do
 	{
@@ -894,7 +844,21 @@ void UCookOnTheFlyServer::HandleNetworkFileServerFileRequest( const FString& Fil
 	UnsolicitedCookedPackages.DequeueAll(FileRequests);
 	for ( int Index=  0; Index < FileRequests.Num(); ++Index)
 	{
-		UnsolicitedFiles.Add( FileRequests[Index].Filename);
+		const FFilePlatformRequest &Request = FileRequests[Index];
+		if ( Request.Filename == FileRequest.Filename )
+		{
+			// don't do anything we don't want this guy
+		}
+		else if ( Request.Platformname == Platformname )
+		{
+			FString StandardFilename = FileRequests[Index].Filename;
+			FPaths::MakeStandardFilename( StandardFilename );
+			UnsolicitedFiles.Add( StandardFilename );
+		}
+		else
+		{
+			UnsolicitedCookedPackages.Enqueue( Request );
+		}
 	}
 	
 	UPackage::WaitForAsyncFileWrites();
