#UE4 - hopefully last of the compile errors on various platforms

diff --git a/Engine/Source/Editor/UnrealEd/Classes/Editor/EditorEngine.h b/Engine/Source/Editor/UnrealEd/Classes/Editor/EditorEngine.h
--- a/Engine/Source/Editor/UnrealEd/Classes/Editor/EditorEngine.h
+++ b/Engine/Source/Editor/UnrealEd/Classes/Editor/EditorEngine.h
@@ -4,6 +4,7 @@
 
 #include "EditorUserSettings.h"
 #include "Transactor.h"
+#include "../Settings/LevelEditorPlaySettings.h"
 #include "EditorEngine.generated.h"
 
 //
@@ -2315,7 +2316,7 @@ private:
 	 * @param PlayNetMode mode to create this PIE world in (as server, client, etc)
 	 * @param DataStruct data required to continue PIE creation, set at login time
 	 */
-	void CreatePIEWorldFromLogin(FWorldContext& PieWorldContext, enum EPlayNetMode PlayNetMode, FPieLoginStruct& DataStruct);
+	void CreatePIEWorldFromLogin(FWorldContext& PieWorldContext, EPlayNetMode PlayNetMode, FPieLoginStruct& DataStruct);
 
 	/**
 	 * Non Online PIE creation flow, creates all instances of PIE at once when online isn't requested/required
diff --git a/Engine/Source/Runtime/Online/OnlineSubsystem/Private/OnlineSubsystemImpl.cpp b/Engine/Source/Runtime/Online/OnlineSubsystem/Private/OnlineSubsystemImpl.cpp
--- a/Engine/Source/Runtime/Online/OnlineSubsystem/Private/OnlineSubsystemImpl.cpp
+++ b/Engine/Source/Runtime/Online/OnlineSubsystem/Private/OnlineSubsystemImpl.cpp
@@ -6,7 +6,7 @@
 
 FOnlineSubsystemImpl::FOnlineSubsystemImpl() :
 	InstanceName(NAME_None),
-	bForceDedicated(false)
+	bForceDedicated(false),
 	NamedInterfaces(NULL)
 {
 }
