#ttp 331968 - EDITOR: Updating Socket Transformation doesn't effect in level assets
#branch UE4
#proj Editor.UnrealEd, Runtime.Engine
#summary Changes to socket transforms are now reflected immediately in the level viewport, and any asset attached to that socket will have its transform updated.
#add In SSocketManager::OnSocketPropertyChanged, upon a change of socket transform, all UStaticMeshComponents are iterated through, and where they reference the mesh being edited, and have children attached to the socket which has had its transform changed, those children have their transform updated.
#change Changed UStaticMeshSocket::PostEditChangeProperty to pass FPropertyChangedEvent::MemberProperty, rather than FPropertyChangedEvent::Property, to the delegate, so we can notice if just a component of the vector/rotator (rather than the entire property) has been changed.
#change Fixed all weak pointer pinning code in SSocketManager.cpp.

diff --git a/Engine/Source/Editor/UnrealEd/Private/SSocketManager.cpp b/Engine/Source/Editor/UnrealEd/Private/SSocketManager.cpp
--- a/Engine/Source/Editor/UnrealEd/Private/SSocketManager.cpp
+++ b/Engine/Source/Editor/UnrealEd/Private/SSocketManager.cpp
@@ -65,7 +65,11 @@ public:
 				.IsSelected( this, &STableRow< TSharedPtr<FString> >::IsSelectedExclusively )
 		];
 
-		SocketItem.Pin()->OnRenameRequested.BindSP( InlineWidget.Get(), &SInlineEditableTextBlock::EnterEditingMode );
+		TSharedPtr<SocketListItem> SocketItemPinned = SocketItem.Pin();
+		if (SocketItemPinned.IsValid())
+		{
+			SocketItemPinned->OnRenameRequested.BindSP(InlineWidget.Get(), &SInlineEditableTextBlock::EnterEditingMode);
+		}
 
 		STableRow< TSharedPtr<FString> >::ConstructInternal(
 			STableRow::FArguments()
@@ -77,7 +81,8 @@ private:
 	/** Returns the socket name */
 	FText GetSocketName() const
 	{
-		return FText::FromName(SocketItem.Pin()->Socket->SocketName);
+		TSharedPtr<SocketListItem> SocketItemPinned = SocketItem.Pin();
+		return SocketItemPinned.IsValid() ? FText::FromName(SocketItemPinned->Socket->SocketName) : FText();
 	}
 
 	bool OnVerifySocketNameChanged( const FText& InNewText, FText& OutErrorMessage )
@@ -89,10 +94,17 @@ private:
 			OutErrorMessage = LOCTEXT( "EmptySocketName_Error", "Sockets must have a name!");
 			bVerifyName = false;
 		}
-		else if(SocketItem.Pin()->Socket->SocketName.ToString() != InNewText.ToString() && SocketManagerPtr.IsValid() && SocketManagerPtr.Pin()->CheckForDuplicateSocket(InNewText.ToString()))
+		else
 		{
-			OutErrorMessage = LOCTEXT( "DuplicateSocket_Error", "Socket name in use!");
-			bVerifyName = false;
+			TSharedPtr<SocketListItem> SocketItemPinned = SocketItem.Pin();
+			TSharedPtr<SSocketManager> SocketManagerPinned = SocketManagerPtr.Pin();
+
+			if (SocketItemPinned.IsValid() && SocketItemPinned->Socket != nullptr && SocketItemPinned->Socket->SocketName.ToString() != InNewText.ToString() &&
+				SocketManagerPinned.IsValid() && SocketManagerPinned->CheckForDuplicateSocket(InNewText.ToString()))
+			{
+				OutErrorMessage = LOCTEXT("DuplicateSocket_Error", "Socket name in use!");
+				bVerifyName = false;
+			}
 		}
 
 		return bVerifyName;
@@ -149,10 +161,16 @@ void SSocketManager::Construct(const FArguments& InArgs)
 
 	OnSocketSelectionChanged = InArgs._OnSocketSelectionChanged;
 
+	TSharedPtr<IStaticMeshEditor> StaticMeshEditorPinned = StaticMeshEditorPtr.Pin();
+	if (!StaticMeshEditorPinned.IsValid())
+	{
+		return;
+	}
+
 	// Register a post undo function which keeps the socket manager list view consistent with the static mesh
-	StaticMeshEditorPtr.Pin()->RegisterOnPostUndo( IStaticMeshEditor::FOnPostUndo::CreateSP( this, &SSocketManager::PostUndo ) );
+	StaticMeshEditorPinned->RegisterOnPostUndo(IStaticMeshEditor::FOnPostUndo::CreateSP(this, &SSocketManager::PostUndo));
 
-	StaticMesh = StaticMeshEditorPtr.Pin()->GetStaticMesh();
+	StaticMesh = StaticMeshEditorPinned->GetStaticMesh();
 
 	FDetailsViewArgs Args;
 	Args.bHideSelectionTip = true;
@@ -170,7 +188,7 @@ void SSocketManager::Construct(const FArguments& InArgs)
 		[
 			SNew(SHorizontalBox)
 			+SHorizontalBox::Slot()
-			.FillWidth(1.0f)
+			.MaxWidth(220.0f)
 			.Padding(4.0f, 0.0f, 0.0f, 0.0f)
 			[
 				SNew(SVerticalBox)
@@ -222,7 +240,7 @@ void SSocketManager::Construct(const FArguments& InArgs)
 						+SHorizontalBox::Slot()
 						.HAlign(HAlign_Right)
 						.FillWidth(1.0f)
-						.Padding(9.0f, 0.0f, 16.0f, 0.0f)
+						.Padding(5.0f, 0.0f, 4.0f, 0.0f)
 						[
 							SNew(SButton)
 							.Text(LOCTEXT("CreateSocket", "Create Socket").ToString())
@@ -231,7 +249,7 @@ void SSocketManager::Construct(const FArguments& InArgs)
 						+SHorizontalBox::Slot()
 							.HAlign(HAlign_Left)
 							.FillWidth(1.0f)
-							.Padding(16.0f, 0.0f, 0.0f, 0.0f)
+							.Padding(4.0f, 0.0f, 5.0f, 0.0f)
 							[
 								SNew(SButton)
 								.Text(LOCTEXT("DeleteSocket", "Delete Socket").ToString())
@@ -329,9 +347,10 @@ void SSocketManager::Construct(const FArguments& InArgs)
 
 SSocketManager::~SSocketManager()
 {
-	if( StaticMeshEditorPtr.IsValid() )
+	TSharedPtr<IStaticMeshEditor> StaticMeshEditorPinned = StaticMeshEditorPtr.Pin();
+	if (StaticMeshEditorPinned.IsValid())
 	{
-		StaticMeshEditorPtr.Pin()->UnregisterOnPostUndo( this );
+		StaticMeshEditorPinned->UnregisterOnPostUndo(this);
 	}
 
 	RemovePropertyChangeListenerFromSockets();
@@ -373,8 +392,11 @@ TSharedRef< ITableRow > SSocketManager::MakeWidgetFromOption( TSharedPtr<SocketL
 
 void SSocketManager::CreateSocket()
 {
-	if (UStaticMesh* CurrentStaticMesh = StaticMeshEditorPtr.Pin()->GetStaticMesh())
+	TSharedPtr<IStaticMeshEditor> StaticMeshEditorPinned = StaticMeshEditorPtr.Pin();
+	if (StaticMeshEditorPinned.IsValid())
 	{
+		UStaticMesh* CurrentStaticMesh = StaticMeshEditorPinned->GetStaticMesh();
+
 		const FScopedTransaction Transaction( LOCTEXT( "CreateSocket", "Create Socket" ) );
 
 		UStaticMeshSocket* NewSocket = ConstructObject<UStaticMeshSocket>(UStaticMeshSocket::StaticClass(), CurrentStaticMesh);
@@ -412,13 +434,15 @@ void SSocketManager::CreateSocket()
 
 void SSocketManager::DuplicateSelectedSocket()
 {
+	TSharedPtr<IStaticMeshEditor> StaticMeshEditorPinned = StaticMeshEditorPtr.Pin();
+
 	UStaticMeshSocket* SelectedSocket = GetSelectedSocket();
 
-	if(SelectedSocket)
+	if(StaticMeshEditorPinned.IsValid() && SelectedSocket)
 	{
 		const FScopedTransaction Transaction( LOCTEXT( "SocketManager_DuplicateSocket", "Duplicate Socket" ) );
 
-		UStaticMesh* CurrentStaticMesh = StaticMeshEditorPtr.Pin()->GetStaticMesh();
+		UStaticMesh* CurrentStaticMesh = StaticMeshEditorPinned->GetStaticMesh();
 
 		UStaticMeshSocket* NewSocket = DuplicateObject(SelectedSocket, CurrentStaticMesh);
 
@@ -460,8 +484,10 @@ void SSocketManager::DeleteSelectedSocket()
 	{
 		const FScopedTransaction Transaction( LOCTEXT( "DeleteSocket", "Delete Socket" ) );
 
-		if (UStaticMesh* CurrentStaticMesh = StaticMeshEditorPtr.Pin()->GetStaticMesh())
+		TSharedPtr<IStaticMeshEditor> StaticMeshEditorPinned = StaticMeshEditorPtr.Pin();
+		if (StaticMeshEditorPinned.IsValid())
 		{
+			UStaticMesh* CurrentStaticMesh = StaticMeshEditorPinned->GetStaticMesh();
 			CurrentStaticMesh->PreEditChange(NULL);
 			UStaticMeshSocket* SelectedSocket = SocketListView->GetSelectedItems()[ 0 ]->Socket;
 			SelectedSocket->OnPropertyChanged().RemoveAll( this );
@@ -480,37 +506,41 @@ void SSocketManager::RefreshSocketList()
 		// The static mesh might not be the same one we built the SocketListView with
 		// check it here and update it if necessary.
 		bool bIsSameStaticMesh = true;
-		UStaticMesh* CurrentStaticMesh = StaticMeshEditorPtr.Pin()->GetStaticMesh();
-		if(CurrentStaticMesh != StaticMesh.Get())
+		TSharedPtr<IStaticMeshEditor> StaticMeshEditorPinned = StaticMeshEditorPtr.Pin();
+		if (StaticMeshEditorPinned.IsValid())
 		{
-			StaticMesh = CurrentStaticMesh;
-			bIsSameStaticMesh = false;
-		}
-		// Only rebuild the socket list if it differs from the static meshes socket list
-		// This is done so that an undo on a socket property doesn't cause the selected
-		// socket to be de-selected, thus hiding the socket properties on the detail view.
-		// NB: Also force a rebuild if the underlying StaticMesh has been changed.
-		if( StaticMesh->Sockets.Num() != SocketList.Num() || !bIsSameStaticMesh )
-		{
-			SocketList.Empty();
-			for(int32 i=0; i < StaticMesh->Sockets.Num(); i++)
+			UStaticMesh* CurrentStaticMesh = StaticMeshEditorPinned->GetStaticMesh();
+			if (CurrentStaticMesh != StaticMesh.Get())
+			{
+				StaticMesh = CurrentStaticMesh;
+				bIsSameStaticMesh = false;
+			}
+			// Only rebuild the socket list if it differs from the static meshes socket list
+			// This is done so that an undo on a socket property doesn't cause the selected
+			// socket to be de-selected, thus hiding the socket properties on the detail view.
+			// NB: Also force a rebuild if the underlying StaticMesh has been changed.
+			if (StaticMesh->Sockets.Num() != SocketList.Num() || !bIsSameStaticMesh)
 			{
-				UStaticMeshSocket* Socket = StaticMesh->Sockets[i];
-				SocketList.Add( MakeShareable( new SocketListItem(Socket) ) );
+				SocketList.Empty();
+				for (int32 i = 0; i < StaticMesh->Sockets.Num(); i++)
+				{
+					UStaticMeshSocket* Socket = StaticMesh->Sockets[i];
+					SocketList.Add(MakeShareable(new SocketListItem(Socket)));
+				}
+
+				SocketListView->RequestListRefresh();
 			}
 
-			SocketListView->RequestListRefresh();
-		}
+			// Set the socket on the detail view to keep it in sync with the sockets properties
+			if (SocketListView->GetSelectedItems().Num())
+			{
+				TArray< UObject* > ObjectList;
+				ObjectList.Add(SocketListView->GetSelectedItems()[0]->Socket);
+				SocketDetailsView->SetObjects(ObjectList, true);
+			}
 
-		// Set the socket on the detail view to keep it in sync with the sockets properties
-		if( SocketListView->GetSelectedItems().Num() )
-		{
-			TArray< UObject* > ObjectList;
-			ObjectList.Add( SocketListView->GetSelectedItems()[0]->Socket );
-			SocketDetailsView->SetObjects( ObjectList, true );
+			StaticMeshEditorPinned->RefreshViewport();
 		}
-
-		StaticMeshEditorPtr.Pin()->RefreshViewport();
 	}
 	else
 	{
@@ -575,8 +605,13 @@ FReply SSocketManager::DeleteSelectedSocket_Execute()
 
 FString SSocketManager::GetSocketHeaderText() const
 {
-	UStaticMesh* CurrentStaticMesh = StaticMeshEditorPtr.Pin()->GetStaticMesh();
-	return FString::Printf(*LOCTEXT("SocketHeader_Total", "Sockets (%d Total)").ToString(), (CurrentStaticMesh != NULL) ? CurrentStaticMesh->Sockets.Num() : 0);
+	UStaticMesh* CurrentStaticMesh = nullptr;
+	TSharedPtr<IStaticMeshEditor> StaticMeshEditorPinned = StaticMeshEditorPtr.Pin();
+	if (StaticMeshEditorPinned.IsValid())
+	{
+		CurrentStaticMesh = StaticMeshEditorPinned->GetStaticMesh();
+	}
+	return FString::Printf(*LOCTEXT("SocketHeader_Total", "Sockets (%d Total)").ToString(), (CurrentStaticMesh != nullptr) ? CurrentStaticMesh->Sockets.Num() : 0);
 }
 
 void SSocketManager::SocketName_TextChanged(const FText& InText)
@@ -631,7 +666,8 @@ float SSocketManager::GetWorldSpaceRollValue() const
 
 void SSocketManager::RotateSocket_WorldSpace()
 {
-	if( StaticMeshEditorPtr.IsValid() && SocketListView->GetSelectedItems().Num() )
+	TSharedPtr<IStaticMeshEditor> StaticMeshEditorPinned = StaticMeshEditorPtr.Pin();
+	if (StaticMeshEditorPinned.IsValid() && SocketListView->GetSelectedItems().Num())
 	{
 		UProperty* ChangedProperty = FindField<UProperty>( UStaticMeshSocket::StaticClass(), "RelativeRotation" );
 		UStaticMeshSocket* SelectedSocket = SocketListView->GetSelectedItems()[0]->Socket;
@@ -647,7 +683,13 @@ TSharedPtr<SWidget> SSocketManager::OnContextMenuOpening()
 {
 	const bool bShouldCloseWindowAfterMenuSelection = true;
 
-	FMenuBuilder MenuBuilder( bShouldCloseWindowAfterMenuSelection, StaticMeshEditorPtr.Pin()->GetToolkitCommands());
+	TSharedPtr<IStaticMeshEditor> StaticMeshEditorPinned = StaticMeshEditorPtr.Pin();
+	if (!StaticMeshEditorPinned.IsValid())
+	{
+		return TSharedPtr<SWidget>();
+	}
+
+	FMenuBuilder MenuBuilder( bShouldCloseWindowAfterMenuSelection, StaticMeshEditorPinned->GetToolkitCommands());
 
 	{
 		MenuBuilder.BeginSection("BasicOperations");
@@ -677,8 +719,10 @@ void SSocketManager::NotifyPostChange( const FPropertyChangedEvent& PropertyChan
 
 void SSocketManager::AddPropertyChangeListenerToSockets()
 {
-	if (UStaticMesh* CurrentStaticMesh = StaticMeshEditorPtr.Pin()->GetStaticMesh())
+	TSharedPtr<IStaticMeshEditor> StaticMeshEditorPinned = StaticMeshEditorPtr.Pin();
+	if (StaticMeshEditorPinned.IsValid())
 	{
+		UStaticMesh* CurrentStaticMesh = StaticMeshEditorPinned->GetStaticMesh();
 		for (int32 i = 0; i < CurrentStaticMesh->Sockets.Num(); ++i)
 		{
 			CurrentStaticMesh->Sockets[i]->OnPropertyChanged().AddSP(this, &SSocketManager::OnSocketPropertyChanged);
@@ -688,9 +732,10 @@ void SSocketManager::AddPropertyChangeListenerToSockets()
 
 void SSocketManager::RemovePropertyChangeListenerFromSockets()
 {
-	if( StaticMeshEditorPtr.IsValid() )
+	TSharedPtr<IStaticMeshEditor> StaticMeshEditorPinned = StaticMeshEditorPtr.Pin();
+	if (StaticMeshEditorPinned.IsValid())
 	{
-		UStaticMesh* CurrentStaticMesh = StaticMeshEditorPtr.Pin()->GetStaticMesh();
+		UStaticMesh* CurrentStaticMesh = StaticMeshEditorPinned->GetStaticMesh();
 		if (CurrentStaticMesh)
 		{
 			for (int32 i = 0; i < CurrentStaticMesh->Sockets.Num(); ++i)
@@ -703,7 +748,14 @@ void SSocketManager::RemovePropertyChangeListenerFromSockets()
 
 void SSocketManager::OnSocketPropertyChanged( const UStaticMeshSocket* Socket, const UProperty* ChangedProperty )
 {
-	if( ChangedProperty->GetName() == TEXT( "RelativeRotation" ) )
+	static FName RelativeRotationName(TEXT("RelativeRotation"));
+	static FName RelativeLocationName(TEXT("RelativeLocation"));
+
+	check(Socket != nullptr);
+
+	FName ChangedPropertyName = ChangedProperty->GetFName();
+
+	if ( ChangedPropertyName == RelativeRotationName )
 	{
 		const UStaticMeshSocket* SelectedSocket = GetSelectedSocket();
 
@@ -712,6 +764,50 @@ void SSocketManager::OnSocketPropertyChanged( const UStaticMeshSocket* Socket, c
 			WorldSpaceRotation.Set( Socket->RelativeRotation.Pitch, Socket->RelativeRotation.Yaw, Socket->RelativeRotation.Roll );
 		}
 	}
+
+	TSharedPtr<IStaticMeshEditor> StaticMeshEditorPinned = StaticMeshEditorPtr.Pin();
+	if (!StaticMeshEditorPinned.IsValid())
+	{
+		return;
+	}
+
+	if (ChangedPropertyName == RelativeRotationName || ChangedPropertyName == RelativeLocationName)
+	{
+		// If socket location or rotation is changed, update the position of any actors attached to it in instances of this mesh
+		UStaticMesh* CurrentStaticMesh = StaticMeshEditorPinned->GetStaticMesh();
+		if (CurrentStaticMesh != nullptr)
+		{
+			bool bUpdatedChild = false;
+
+			for (TObjectIterator<UStaticMeshComponent> It; It; ++It)
+			{
+				if (It->StaticMesh == CurrentStaticMesh)
+				{
+					const AActor* Actor = It->GetOwner();
+					if (Actor != nullptr)
+					{
+						const USceneComponent* Root = Actor->GetRootComponent();
+						if (Root != nullptr)
+						{
+							for (USceneComponent* Child : Root->AttachChildren)
+							{
+								if (Child != nullptr && Child->AttachSocketName == Socket->SocketName)
+								{
+									Child->UpdateComponentToWorld();
+									bUpdatedChild = true;
+								}
+							}
+						}
+					}
+				}
+			}
+
+			if (bUpdatedChild)
+			{
+				GUnrealEd->RedrawLevelEditingViewports();
+			}
+		}
+	}
 }
 
 void SSocketManager::PostUndo()
@@ -721,9 +817,10 @@ void SSocketManager::PostUndo()
 
 void SSocketManager::OnItemScrolledIntoView( TSharedPtr<SocketListItem> InItem, const TSharedPtr<ITableRow>& InWidget)
 {
-	if( DeferredRenameRequest.IsValid() )
+	TSharedPtr<SocketListItem> DeferredRenameRequestPinned = DeferredRenameRequest.Pin();
+	if( DeferredRenameRequestPinned.IsValid() )
 	{
-		DeferredRenameRequest.Pin()->OnRenameRequested.ExecuteIfBound();
+		DeferredRenameRequestPinned->OnRenameRequested.ExecuteIfBound();
 		DeferredRenameRequest.Reset();
 	}
 }
diff --git a/Engine/Source/Runtime/Engine/Private/StaticMesh.cpp b/Engine/Source/Runtime/Engine/Private/StaticMesh.cpp
--- a/Engine/Source/Runtime/Engine/Private/StaticMesh.cpp
+++ b/Engine/Source/Runtime/Engine/Private/StaticMesh.cpp
@@ -2418,7 +2418,7 @@ void UStaticMeshSocket::PostEditChangeProperty( FPropertyChangedEvent& PropertyC
 
 	if( PropertyChangedEvent.Property )
 	{
-		ChangedEvent.Broadcast( this, PropertyChangedEvent.Property );
+		ChangedEvent.Broadcast( this, PropertyChangedEvent.MemberProperty );
 	}
 }
 #endif
