#summary 	Added comments & edit conditions to brush builders
#ttp 332479 	EDITOR: BSP: CubeBuilder doesn't have comments or edit conditions
#detail 	Updated each (used) brush builder to contain comments for tooltips & edit conditions where appropriate.

diff --git a/Engine/Source/Editor/UnrealEd/Classes/Builders/ConeBuilder.h b/Engine/Source/Editor/UnrealEd/Classes/Builders/ConeBuilder.h
--- a/Engine/Source/Editor/UnrealEd/Classes/Builders/ConeBuilder.h
+++ b/Engine/Source/Editor/UnrealEd/Classes/Builders/ConeBuilder.h
@@ -12,27 +12,34 @@ class UConeBuilder : public UEditorBrushBuilder
 {
 	GENERATED_UCLASS_BODY()
 
+	/** Distance from base to tip of cone */
 	UPROPERTY(EditAnywhere, Category=BrushSettings, meta=(ClampMin = "0.000001"))
 	float Z;
 
-	UPROPERTY(EditAnywhere, Category=BrushSettings)
+	/** Distance from base to the tip of inner cone (when hollow) */
+	UPROPERTY(EditAnywhere, Category=BrushSettings, meta=(EditCondition="Hollow"))
 	float CapZ;
 
+	/** Radius of cone */
 	UPROPERTY(EditAnywhere, Category=BrushSettings, meta=(ClampMin = "0.000001"))
 	float OuterRadius;
 
-	UPROPERTY(EditAnywhere, Category=BrushSettings)
+	/** Radius of inner cone (when hollow) */
+	UPROPERTY(EditAnywhere, Category=BrushSettings, meta=(EditCondition="Hollow"))
 	float InnerRadius;
 
+	/** How many sides this cone should have */
 	UPROPERTY(EditAnywhere, Category=BrushSettings, meta=(ClampMin = "3", ClampMax = "500"))
 	int32 Sides;
 
 	UPROPERTY()
 	FName GroupName;
 
+	/** Whether to align the brush to a face */
 	UPROPERTY(EditAnywhere, Category=BrushSettings)
 	uint32 AlignToSide:1;
 
+	/** Whether this is a hollow or solid cone */
 	UPROPERTY(EditAnywhere, Category=BrushSettings)
 	uint32 Hollow:1;
 
diff --git a/Engine/Source/Editor/UnrealEd/Classes/Builders/CubeBuilder.h b/Engine/Source/Editor/UnrealEd/Classes/Builders/CubeBuilder.h
--- a/Engine/Source/Editor/UnrealEd/Classes/Builders/CubeBuilder.h
+++ b/Engine/Source/Editor/UnrealEd/Classes/Builders/CubeBuilder.h
@@ -12,24 +12,30 @@ class UCubeBuilder : public UEditorBrushBuilder
 {
 	GENERATED_UCLASS_BODY()
 
+	/** The size of the cube in the X dimension */
 	UPROPERTY(EditAnywhere, Category=BrushSettings, meta=(ClampMin = "0.000001"))
 	float X;
 
+	/** The size of the cube in the Y dimension */
 	UPROPERTY(EditAnywhere, Category=BrushSettings, meta=(ClampMin = "0.000001"))
 	float Y;
 
+	/** The size of the cube in the Z dimension */
 	UPROPERTY(EditAnywhere, Category=BrushSettings, meta=(ClampMin = "0.000001"))
 	float Z;
 
-	UPROPERTY(EditAnywhere, Category=BrushSettings)
+	/** The thickness of the cube wall when hollow */
+	UPROPERTY(EditAnywhere, Category=BrushSettings, meta=(EditCondition="Hollow"))
 	float WallThickness;
 
 	UPROPERTY()
 	FName GroupName;
 
+	/** Whether this is a hollow or solid cube */
 	UPROPERTY(EditAnywhere, Category=BrushSettings)
 	uint32 Hollow:1;
 
+	/** Whether extra internal faces should be generated for each cube face */
 	UPROPERTY(EditAnywhere, Category=BrushSettings)
 	uint32 Tessellated:1;
 
diff --git a/Engine/Source/Editor/UnrealEd/Classes/Builders/CurvedStairBuilder.h b/Engine/Source/Editor/UnrealEd/Classes/Builders/CurvedStairBuilder.h
--- a/Engine/Source/Editor/UnrealEd/Classes/Builders/CurvedStairBuilder.h
+++ b/Engine/Source/Editor/UnrealEd/Classes/Builders/CurvedStairBuilder.h
@@ -12,27 +12,34 @@ class UCurvedStairBuilder : public UEditorBrushBuilder
 {
 	GENERATED_UCLASS_BODY()
 
+	/** The radius of the inner curve of the stair */
 	UPROPERTY(EditAnywhere, Category=BrushSettings, meta=(ClampMin = "1"))
 	int32 InnerRadius;
 
+	/** The height of each step */
 	UPROPERTY(EditAnywhere, Category=BrushSettings, meta=(ClampMin = "1"))
 	int32 StepHeight;
 
+	/** The width of each step */
 	UPROPERTY(EditAnywhere, Category=BrushSettings, meta=(ClampMin = "1"))
 	int32 StepWidth;
 
+	/** The angle of the total arc described by this stair */
 	UPROPERTY(EditAnywhere, Category=BrushSettings, meta=(ClampMin = "1", ClampMax = "360"))
 	int32 AngleOfCurve;
 
+	/** The number of steps */
 	UPROPERTY(EditAnywhere, Category=BrushSettings, meta=(ClampMin = "1", ClampMax = "100"))
 	int32 NumSteps;
 
+	/** The distance below the first step */
 	UPROPERTY(EditAnywhere, Category=BrushSettings)
 	int32 AddToFirstStep;
 
 	UPROPERTY()
 	FName GroupName;
 
+	/** Whether the stair curves clockwise or counter-clockwise */
 	UPROPERTY(EditAnywhere, Category=BrushSettings)
 	uint32 CounterClockwise:1;
 
diff --git a/Engine/Source/Editor/UnrealEd/Classes/Builders/CylinderBuilder.h b/Engine/Source/Editor/UnrealEd/Classes/Builders/CylinderBuilder.h
--- a/Engine/Source/Editor/UnrealEd/Classes/Builders/CylinderBuilder.h
+++ b/Engine/Source/Editor/UnrealEd/Classes/Builders/CylinderBuilder.h
@@ -17,24 +17,30 @@ class UCylinderBuilder : public UEditorBrushBuilder
 {
 	GENERATED_UCLASS_BODY()
 
+	/** Distance from base to tip of cylinder */
 	UPROPERTY(EditAnywhere, Category=BrushSettings, meta=(ClampMin = "0.000001"))
 	float Z;
 
+	/** Radius of cylinder */
 	UPROPERTY(EditAnywhere, Category=BrushSettings, meta=(ClampMin = "0.000001"))
 	float OuterRadius;
 
-	UPROPERTY(EditAnywhere, Category=BrushSettings)
+	/** Radius of inner cylinder (when hollow) */
+	UPROPERTY(EditAnywhere, Category=BrushSettings, meta=(EditCondition="Hollow"))
 	float InnerRadius;
 
+	/** How many sides this cylinder should have */
 	UPROPERTY(EditAnywhere, Category=BrushSettings, meta=(ClampMin = "3", ClampMax = "500"))
 	int32 Sides;
 
 	UPROPERTY()
 	FName GroupName;
 
+	/** Whether to align the brush to a face */
 	UPROPERTY(EditAnywhere, Category=BrushSettings)
 	uint32 AlignToSide:1;
 
+	/** Whether this is a hollow or solid cylinder */
 	UPROPERTY(EditAnywhere, Category=BrushSettings)
 	uint32 Hollow:1;
 
diff --git a/Engine/Source/Editor/UnrealEd/Classes/Builders/LinearStairBuilder.h b/Engine/Source/Editor/UnrealEd/Classes/Builders/LinearStairBuilder.h
--- a/Engine/Source/Editor/UnrealEd/Classes/Builders/LinearStairBuilder.h
+++ b/Engine/Source/Editor/UnrealEd/Classes/Builders/LinearStairBuilder.h
@@ -12,18 +12,23 @@ class ULinearStairBuilder : public UEditorBrushBuilder
 {
 	GENERATED_UCLASS_BODY()
 
+	/** The length of each step */
 	UPROPERTY(EditAnywhere, Category=BrushSettings, meta=(ClampMin = "1"))
 	int32 StepLength;
 
+	/** The height of each step */
 	UPROPERTY(EditAnywhere, Category=BrushSettings, meta=(ClampMin = "1"))
 	int32 StepHeight;
 
+	/** The width of each step */
 	UPROPERTY(EditAnywhere, Category=BrushSettings, meta=(ClampMin = "1"))
 	int32 StepWidth;
 
+	/** The number of steps */
 	UPROPERTY(EditAnywhere, Category=BrushSettings, meta=(ClampMin = "2", ClampMax = "45"))
 	int32 NumSteps;
 
+	/** The distance below the first step */
 	UPROPERTY(EditAnywhere, Category=BrushSettings)
 	int32 AddToFirstStep;
 
diff --git a/Engine/Source/Editor/UnrealEd/Classes/Builders/SpiralStairBuilder.h b/Engine/Source/Editor/UnrealEd/Classes/Builders/SpiralStairBuilder.h
--- a/Engine/Source/Editor/UnrealEd/Classes/Builders/SpiralStairBuilder.h
+++ b/Engine/Source/Editor/UnrealEd/Classes/Builders/SpiralStairBuilder.h
@@ -12,33 +12,42 @@ class USpiralStairBuilder : public UEditorBrushBuilder
 {
 	GENERATED_UCLASS_BODY()
 
+	/** The radius of the inner curve of the stair */
 	UPROPERTY(EditAnywhere, Category=BrushSettings, meta=(ClampMin = "1"))
 	int32 InnerRadius;
 
+	/** The width of each step */
 	UPROPERTY(EditAnywhere, Category=BrushSettings, meta=(ClampMin = "1"))
 	int32 StepWidth;
 
+	/** The height of each step */
 	UPROPERTY(EditAnywhere, Category=BrushSettings, meta=(ClampMin = "1"))
 	int32 StepHeight;
 
+	/** The thickness of each step */
 	UPROPERTY(EditAnywhere, Category=BrushSettings, meta=(ClampMin = "1"))
 	int32 StepThickness;
 
+	/** The number of steps in one whole spiral rotation */
 	UPROPERTY(EditAnywhere, Category=BrushSettings, meta=(ClampMin = "3"))
 	int32 NumStepsPer360;
 
+	/** The total number of steps */
 	UPROPERTY(EditAnywhere, Category=BrushSettings, meta=(ClampMin = "1", ClampMax = "100"))
 	int32 NumSteps;
 
 	UPROPERTY()
 	FName GroupName;
 
+	/** Whether the underside of the spiral is sloped or stepped */
 	UPROPERTY(EditAnywhere, Category=BrushSettings)
 	uint32 SlopedCeiling:1;
 
+	/** Whether the surface of the spiral is sloped or stepped */
 	UPROPERTY(EditAnywhere, Category=BrushSettings)
 	uint32 SlopedFloor:1;
 
+	/** Whether the stair curves clockwise or counter-clockwise */
 	UPROPERTY(EditAnywhere, Category=BrushSettings)
 	uint32 CounterClockwise:1;
 
diff --git a/Engine/Source/Editor/UnrealEd/Classes/Builders/TetrahedronBuilder.h b/Engine/Source/Editor/UnrealEd/Classes/Builders/TetrahedronBuilder.h
--- a/Engine/Source/Editor/UnrealEd/Classes/Builders/TetrahedronBuilder.h
+++ b/Engine/Source/Editor/UnrealEd/Classes/Builders/TetrahedronBuilder.h
@@ -12,9 +12,11 @@ class UTetrahedronBuilder : public UEditorBrushBuilder
 {
 	GENERATED_UCLASS_BODY()
 
+	/** The radius of this sphere */
 	UPROPERTY(EditAnywhere, Category=BrushSettings, meta=(ClampMin = "0.000001"))
 	float Radius;
 
+	/** How many iterations this sphere uses to tessellate its geometry */
 	UPROPERTY(EditAnywhere, Category=BrushSettings, meta=(ClampMin = "1", ClampMax = "5",DisplayName="Tessellation"))
 	int32 SphereExtrapolation;
 
