#ttp 332412 - EDITOR: Simulate: CRASH: Saving changes made to a bsp brush during simulation will cause a crash when stopping simulation
#proj UE4
#branch UE4
#summary EditorUtilities::CopyActorProperties now ignores objects from the PIE world
#extra Experimented with this and found that it didn't just affect CSG brushes, but it also affected anything copying an actor or object back from the PIE world (as it was performing a shallow copy).
#extra There was a comment in EditorUtilities::CopyActorProperties saying that copying objects might not work correctly, so I've updated it to either:
1) Update a reference to an actor or actor component in the PIE world to point at the equivalent object in the editor world.
2) Do nothing if the object outer is detected as the PIE world.

diff --git a/Engine/Source/Editor/UnrealEd/Private/Editor.cpp b/Engine/Source/Editor/UnrealEd/Private/Editor.cpp
--- a/Engine/Source/Editor/UnrealEd/Private/Editor.cpp
+++ b/Engine/Source/Editor/UnrealEd/Private/Editor.cpp
@@ -6743,6 +6743,58 @@ namespace EditorUtilities
 	}
 
 
+	void CopySingleActorProperty( UObject* const InSourceObject, UObject* const InTargetObject, UProperty* const InProperty )
+	{
+		// Properties that are *object* properties are tricky
+		// Sometimes the object will be a reference to a PIE-world object, and copying that reference back to an actor CDO asset is not a good idea
+		// If the property is referencing an actor or actor component in the PIE world, then we can try and fix that reference up to the equivalent
+		// from the editor world; otherwise we have to skip it
+		bool bNeedsGenericCopy = true;
+		if( UObjectPropertyBase* const ObjectProperty = Cast<UObjectPropertyBase>(InProperty) )
+		{
+			UObject* const SourceObjectPropertyValue = ObjectProperty->GetObjectPropertyValue_InContainer(InSourceObject);
+			if( SourceObjectPropertyValue && SourceObjectPropertyValue->GetOutermost()->PackageFlags & PKG_PlayInEditor )
+			{
+				// Not all the code paths below actually copy the object, but even if they don't we need to claim that they
+				// did, as copying a reference to an object in a PIE world leads to crashes
+				bNeedsGenericCopy = false;
+
+				// REFERENCE an existing actor in the editor world from a REFERENCE in the PIE world
+				if( SourceObjectPropertyValue->IsA(AActor::StaticClass()) )
+				{
+					// We can try and fix-up an actor reference from the PIE world to instead be the version from the persistent world
+					AActor* const EditorWorldActor = GetEditorWorldCounterpartActor(Cast<AActor>(SourceObjectPropertyValue));
+					if( EditorWorldActor )
+					{
+						ObjectProperty->SetObjectPropertyValue_InContainer(InTargetObject, EditorWorldActor);
+					}
+				}
+				// REFERENCE an existing actor component in the editor world from a REFERENCE in the PIE world
+				else if( SourceObjectPropertyValue->IsA(UActorComponent::StaticClass()) && InTargetObject->IsA(AActor::StaticClass()) )
+				{
+					AActor* const TargetActor = Cast<AActor>(InTargetObject);
+					TArray<UActorComponent*> TargetComponents;
+					TargetActor->GetComponents(TargetComponents);
+
+					// We can try and fix-up an actor component reference from the PIE world to instead be the version from the persistent world
+					int32 TargetComponentIndex = 0;
+					UActorComponent* const EditorWorldComponent = FindMatchingComponentInstance(Cast<UActorComponent>(SourceObjectPropertyValue), TargetActor, TargetComponents, TargetComponentIndex);
+					if(EditorWorldComponent)
+					{
+						ObjectProperty->SetObjectPropertyValue_InContainer(InTargetObject, EditorWorldComponent);
+					}
+				}
+			}
+		}
+		
+		// Handle copying properties that either aren't an object, or aren't part of the PIE world
+		if( bNeedsGenericCopy )
+		{
+			InProperty->CopyCompleteValue_InContainer(InTargetObject, InSourceObject);
+		}
+	}
+
+
 	int32 CopyActorProperties( AActor* SourceActor, AActor* TargetActor, const ECopyOptions::Type Options )
 	{
 		check( SourceActor != NULL && TargetActor != NULL );
@@ -6762,10 +6814,6 @@ namespace EditorUtilities
 			TargetActor->GetArchetypeInstances(ArchetypeInstances);
 		}
 
-		// @todo simulate: Properties that are *object* properties are tricky.  We might not always want to copy those.  Sometimes the
-		// object will be a reference to a PIE-world object, and copying that reference back to an actor CDO asset is not a good idea.  In
-		// these cases, we could do a deep copy, or we could skip it.
-
 		bool bTransformChanged = false;
 
 		// Copy non-component properties from the old actor to the new actor
@@ -6812,7 +6860,7 @@ namespace EditorUtilities
 								}
 							}
 
-							Property->CopyCompleteValue_InContainer( TargetActor, SourceActor );
+							CopySingleActorProperty(SourceActor, TargetActor, Property);
 
 							if( Options & ECopyOptions::CallPostEditChangeProperty )
 							{
@@ -6936,7 +6984,7 @@ namespace EditorUtilities
 									}
 								}
 
-								Property->CopyCompleteValue_InContainer( TargetComponent, SourceComponent );
+								CopySingleActorProperty(SourceComponent, TargetComponent, Property);
 
 								if( Options & ECopyOptions::CallPostEditChangeProperty )
 								{
