Fixed compiler error.

diff --git a/Engine/Source/Editor/UnrealEd/Classes/CookOnTheSide/CookOnTheFlyServer.h b/Engine/Source/Editor/UnrealEd/Classes/CookOnTheSide/CookOnTheFlyServer.h
--- a/Engine/Source/Editor/UnrealEd/Classes/CookOnTheSide/CookOnTheFlyServer.h
+++ b/Engine/Source/Editor/UnrealEd/Classes/CookOnTheSide/CookOnTheFlyServer.h
@@ -287,7 +287,7 @@ public:
 private:
 	
 
-	bool UCookOnTheFlyServer::ShouldCook(const FString& InFileName, const FString &InPlatformName);
+	bool ShouldCook(const FString& InFileName, const FString &InPlatformName);
 
 	/**
 	 * Returns cooker output directory.
