#ttp 332089 - EDITOR: Player spawns at map origin instead of at camera location when selecting Play in New Editor Window or Standalone Game
#proj UE4
#branch UE4
#summary PIE no longer tries to use the player start if there is no player start available in the world
#extra This had already been fixed for PIE when running in a viewport, but not for the other PIE options.

diff --git a/Engine/Source/Editor/UnrealEd/Private/Kismet2/DebuggerCommands.cpp b/Engine/Source/Editor/UnrealEd/Private/Kismet2/DebuggerCommands.cpp
--- a/Engine/Source/Editor/UnrealEd/Private/Kismet2/DebuggerCommands.cpp
+++ b/Engine/Source/Editor/UnrealEd/Private/Kismet2/DebuggerCommands.cpp
@@ -118,6 +118,9 @@ protected:
 	 * @param PlatformName - The name of the platform to play the game on.
 	 */
 	static void LaunchOnDevice(const FString& DeviceId, const FString& DeviceName);
+
+	/** Get the player start location to use when starting PIE */
+	static EPlayModeLocations GetPlayModeLocation();
 };
 
 
@@ -991,7 +994,7 @@ void FInternalPlayWorldCommandCallbacks::PlayInViewport_Clicked( )
 
 	TSharedPtr<ILevelViewport> ActiveLevelViewport = LevelEditorModule.GetFirstActiveViewport();
 
-	const bool bAtPlayerStart = (GetDefault<ULevelEditorPlaySettings>()->LastExecutedPlayModeLocation == PlayLocation_DefaultPlayerStart) && (GEditor->CheckForPlayerStart() != nullptr);
+	const bool bAtPlayerStart = (GetPlayModeLocation() == PlayLocation_DefaultPlayerStart);
 	const bool bSimulateInEditor = false;
 
 	// Make sure we can find a path to the view port.  This will fail in cases where the view port widget
@@ -1050,7 +1053,7 @@ void FInternalPlayWorldCommandCallbacks::PlayInEditorFloating_Clicked( )
 	{
 		RecordLastExecutedPlayMode();
 
-		const bool bAtPlayerStart = (GetDefault<ULevelEditorPlaySettings>()->LastExecutedPlayModeLocation == PlayLocation_DefaultPlayerStart);
+		const bool bAtPlayerStart = (GetPlayModeLocation() == PlayLocation_DefaultPlayerStart);
 		const bool bSimulateInEditor = false;
 
 		const FVector* StartLoc = NULL;
@@ -1107,7 +1110,7 @@ void FInternalPlayWorldCommandCallbacks::PlayInNewProcess_Clicked( bool MobilePr
 		const FVector* StartLoc = NULL;
 		const FRotator* StartRot = NULL;
 
-		const bool bAtPlayerStart = (GetDefault<ULevelEditorPlaySettings>()->LastExecutedPlayModeLocation == PlayLocation_DefaultPlayerStart);
+		const bool bAtPlayerStart = (GetPlayModeLocation() == PlayLocation_DefaultPlayerStart);
 
 		if (!bAtPlayerStart)
 		{
@@ -1540,4 +1543,10 @@ void FInternalPlayWorldCommandCallbacks::LaunchOnDevice( const FString& DeviceId
 }
 
 
+EPlayModeLocations FInternalPlayWorldCommandCallbacks::GetPlayModeLocation()
+{
+	// We can't use PlayLocation_DefaultPlayerStart without a player start position
+	return GEditor->CheckForPlayerStart() ? GetDefault<ULevelEditorPlaySettings>()->LastExecutedPlayModeLocation : PlayLocation_CurrentCameraLocation;
+}
+
 #undef LOCTEXT_NAMESPACE
