#ttp 332086 - EDITOR: BP: Renaming component variable does update for events using the component
#proj UE4
#branch UE4
#summary Renaming a component now updates any UK2Node_ComponentBoundEvent nodes which reference it

diff --git a/Engine/Source/Editor/UnrealEd/Private/Kismet2/BlueprintEditorUtils.cpp b/Engine/Source/Editor/UnrealEd/Private/Kismet2/BlueprintEditorUtils.cpp
--- a/Engine/Source/Editor/UnrealEd/Private/Kismet2/BlueprintEditorUtils.cpp
+++ b/Engine/Source/Editor/UnrealEd/Private/Kismet2/BlueprintEditorUtils.cpp
@@ -194,40 +194,42 @@ static void RenameVariableReferences(UBlueprint* Blueprint, UClass* VariableClas
 	Blueprint->GetAllGraphs(AllGraphs);
 
 	// Update any graph nodes that reference the old variable name to instead reference the new name
-	for (TArray<UEdGraph*>::TConstIterator GraphIt(AllGraphs); GraphIt; ++GraphIt)
+	for(UEdGraph* CurrentGraph : AllGraphs)
 	{
-		const UEdGraph* CurrentGraph = *GraphIt;
-
-		TArray<UK2Node_Variable*> GraphNodes;
-		CurrentGraph->GetNodesOfClass(GraphNodes);
-
-		for (TArray<UK2Node_Variable*>::TConstIterator NodeIt(GraphNodes); NodeIt; ++NodeIt)
+		for(UEdGraphNode* GraphNode : CurrentGraph->Nodes)
 		{
-			UK2Node_Variable* CurrentNode = *NodeIt;
-
-			UClass* NodeRefClass = CurrentNode->VariableReference.GetMemberParentClass(Blueprint->GeneratedClass);
-			if (!NodeRefClass->IsChildOf(VariableClass))
+			if(UK2Node_Variable* const VariableNode = Cast<UK2Node_Variable>(GraphNode))
 			{
+				UClass* const NodeRefClass = VariableNode->VariableReference.GetMemberParentClass(Blueprint->GeneratedClass);
+				if(NodeRefClass && NodeRefClass->IsChildOf(VariableClass) && OldVarName == VariableNode->GetVarName())
+				{
+					VariableNode->Modify();
+					if(VariableNode->VariableReference.IsSelfContext())
+					{
+						VariableNode->VariableReference.SetSelfMember(NewVarName);
+					}
+					else
+					{
+						VariableNode->VariableReference.SetExternalMember(NewVarName, NodeRefClass);
+					}
+
+					if(UEdGraphPin* const Pin = VariableNode->FindPin(OldVarName.ToString()))
+					{
+						Pin->Modify();
+						Pin->PinName = NewVarName.ToString();
+					}
+				}
 				continue;
 			}
-
-			if (OldVarName == CurrentNode->GetVarName())
+			
+			if(UK2Node_ComponentBoundEvent* const ComponentBoundEventNode = Cast<UK2Node_ComponentBoundEvent>(GraphNode))
 			{
-				CurrentNode->Modify();
-				if (CurrentNode->VariableReference.IsSelfContext())
+				if(OldVarName == ComponentBoundEventNode->ComponentPropertyName)
 				{
-					CurrentNode->VariableReference.SetSelfMember(NewVarName);
-				}
-				else
-				{
-					CurrentNode->VariableReference.SetExternalMember(NewVarName, NodeRefClass);
-				}
-
-				if (UEdGraphPin* Pin = CurrentNode->FindPin(OldVarName.ToString()))
-				{
-					Pin->Modify();
-					Pin->PinName = NewVarName.ToString();
+					ComponentBoundEventNode->Modify();
+					ComponentBoundEventNode->ComponentPropertyName = NewVarName;
 				}
+				continue;
 			}
 		}
 	}
