CIS fix - Added explicit cast to EPlayModeLocations in GetPlayModeLocation to remove ambiguity

diff --git a/Engine/Source/Editor/UnrealEd/Private/Kismet2/DebuggerCommands.cpp b/Engine/Source/Editor/UnrealEd/Private/Kismet2/DebuggerCommands.cpp
--- a/Engine/Source/Editor/UnrealEd/Private/Kismet2/DebuggerCommands.cpp
+++ b/Engine/Source/Editor/UnrealEd/Private/Kismet2/DebuggerCommands.cpp
@@ -1546,7 +1546,9 @@ void FInternalPlayWorldCommandCallbacks::LaunchOnDevice( const FString& DeviceId
 EPlayModeLocations FInternalPlayWorldCommandCallbacks::GetPlayModeLocation()
 {
 	// We can't use PlayLocation_DefaultPlayerStart without a player start position
-	return GEditor->CheckForPlayerStart() ? GetDefault<ULevelEditorPlaySettings>()->LastExecutedPlayModeLocation : PlayLocation_CurrentCameraLocation;
+	return GEditor->CheckForPlayerStart() 
+		? static_cast<EPlayModeLocations>(GetDefault<ULevelEditorPlaySettings>()->LastExecutedPlayModeLocation) 
+		: PlayLocation_CurrentCameraLocation;
 }
 
 #undef LOCTEXT_NAMESPACE
