#ttp 332679 - CRASH: MAIN: EDITOR: Editor crashes when double clicking "Empty Level" after opening another new level upon editor launch
#branch UE4
#proj Editor.UnrealEd
#summary Invalidate all level viewport hit proxies when opening a new level, so that click events cannot access old data on the first tick in a new world.
#add Added UEditorEngine::InvalidateAllViewportClientHitProxies() to iterate through all active level viewport clients and invalidate their hit proxies.
#add Added call to InvalidateAllViewportClientHitProxies() in UEditorEngine::NewMap() and UEditorEngine::Map_Load().

diff --git a/Engine/Source/Editor/UnrealEd/Classes/Editor/EditorEngine.h b/Engine/Source/Editor/UnrealEd/Classes/Editor/EditorEngine.h
--- a/Engine/Source/Editor/UnrealEd/Classes/Editor/EditorEngine.h
+++ b/Engine/Source/Editor/UnrealEd/Classes/Editor/EditorEngine.h
@@ -677,7 +677,6 @@ public:
 	 */
 	bool UpdateSingleViewportClient(FEditorViewportClient* InViewportClient, const bool bInAllowNonRealtimeViewportToDraw, bool bLinkedOrthoMovement );
 
-
 	/** Used for generating status bar text */
 	enum EMousePositionType
 	{
@@ -2546,6 +2545,9 @@ protected:
 
 	/** Called when Matinee is opened */
 	virtual void OnOpenMatinee(){};
+
+	/** Invalidate all viewport client hit proxies immediately */
+	void InvalidateAllViewportClientHitProxies();
 };
 
 
diff --git a/Engine/Source/Editor/UnrealEd/Private/Editor.cpp b/Engine/Source/Editor/UnrealEd/Private/Editor.cpp
--- a/Engine/Source/Editor/UnrealEd/Private/Editor.cpp
+++ b/Engine/Source/Editor/UnrealEd/Private/Editor.cpp
@@ -1626,6 +1626,16 @@ bool UEditorEngine::UpdateSingleViewportClient(FEditorViewportClient* InViewport
 	return bUpdatedNonRealtimeViewport;
 }
 
+void UEditorEngine::InvalidateAllViewportClientHitProxies()
+{
+	for (const auto* LevelViewportClient : LevelViewportClients)
+	{
+		if (LevelViewportClient->Viewport != nullptr)
+		{
+			LevelViewportClient->Viewport->InvalidateHitProxy();
+		}
+	}
+}
 
 void UEditorEngine::PostEditChangeProperty(FPropertyChangedEvent& PropertyChangedEvent)
 {
diff --git a/Engine/Source/Editor/UnrealEd/Private/EditorServer.cpp b/Engine/Source/Editor/UnrealEd/Private/EditorServer.cpp
--- a/Engine/Source/Editor/UnrealEd/Private/EditorServer.cpp
+++ b/Engine/Source/Editor/UnrealEd/Private/EditorServer.cpp
@@ -1786,6 +1786,9 @@ UWorld* UEditorEngine::NewMap()
 	// Clear the transaction buffer so the user can't remove the builder brush
 	GUnrealEd->ResetTransaction( CleanseText );
 
+	// Invalidate all the level viewport hit proxies
+	InvalidateAllViewportClientHitProxies();
+
 	return Context.World();
 }
 
@@ -2160,6 +2163,9 @@ bool UEditorEngine::Map_Load(const TCHAR* Str, FOutputDevice& Ar)
 					{
 						GEngine->WorldAdded( Context.World() );
 					}
+
+					// Invalidate all the level viewport hit proxies
+					InvalidateAllViewportClientHitProxies();
 				}
 			}
 			else
