#UE4 - added code to check that online PIE is active
- moved check for logins outside of login function so we can do the normal path more often

diff --git a/Engine/Source/Editor/UnrealEd/Classes/Editor/EditorEngine.h b/Engine/Source/Editor/UnrealEd/Classes/Editor/EditorEngine.h
--- a/Engine/Source/Editor/UnrealEd/Classes/Editor/EditorEngine.h
+++ b/Engine/Source/Editor/UnrealEd/Classes/Editor/EditorEngine.h
@@ -2203,7 +2203,10 @@ public:
 	bool IsPlayingViaLauncher() const { return bPlayUsingLauncher && !bIsPlayWorldQueued; }
 
 	/** @return true if the editor is able to launch PIE with online platform support */
-	virtual bool SupportsOnlinePIE() const;
+	bool SupportsOnlinePIE() const;
+
+    /** @return true if there are active PIE instances logged into an online platform */
+	bool IsPlayingWithOnlinePIE() const { return NumOnlinePIEInstances > 0; }
 
 	/**
 	 * Ensures the assets specified are loaded and adds them to the global selection set
@@ -2534,6 +2537,9 @@ private:
 
 	FPIEInstanceWindowSwitch PIEInstanceWindowSwitchDelegate;
 
+	/** Number of currently running instances logged into an online platform */
+	int32 NumOnlinePIEInstances;
+
 protected:
 
 	// @todo gmp: temp hack for Rocket demo
diff --git a/Engine/Source/Editor/UnrealEd/Private/Editor.cpp b/Engine/Source/Editor/UnrealEd/Private/Editor.cpp
--- a/Engine/Source/Editor/UnrealEd/Private/Editor.cpp
+++ b/Engine/Source/Editor/UnrealEd/Private/Editor.cpp
@@ -203,6 +203,7 @@ UEditorEngine::UEditorEngine(const class FPostConstructInitializeProperties& PCI
 	bDisableDeltaModification = false;
 	bPlayOnLocalPcSession = false;
 	bAllowMultiplePIEWorlds = true;
+	NumOnlinePIEInstances = 0;
 }
 
 
diff --git a/Engine/Source/Editor/UnrealEd/Private/PlayLevel.cpp b/Engine/Source/Editor/UnrealEd/Private/PlayLevel.cpp
--- a/Engine/Source/Editor/UnrealEd/Private/PlayLevel.cpp
+++ b/Engine/Source/Editor/UnrealEd/Private/PlayLevel.cpp
@@ -460,6 +460,7 @@ void UEditorEngine::TeardownPlaySession(FWorldContext &PieWorldContext)
 	{
 		FName OnlineIdentifier = FName(*FString::Printf(TEXT(":%d"), PieWorldContext.ContextHandle));
 		IOnlineSubsystem::Destroy(OnlineIdentifier);
+		NumOnlinePIEInstances--;
 	}
 }
 
@@ -1782,7 +1783,13 @@ void UEditorEngine::PlayInEditor( UWorld* InWorld, bool bInSimulateInEditor )
 	}
 	else
 	{
-		if (SupportsOnlinePIE())
+		bool bHasRequiredLogins = PlayInSettings->PlayNumberOfClients <= PIELogins.Num();
+		if (!bHasRequiredLogins)
+		{
+			UE_LOG(LogOnline, Verbose, TEXT("Not enough login credentials to launch all PIE instances, modify [/Script/UnrealEd.UnrealEdEngine].PIELogins"));
+		}
+
+		if (SupportsOnlinePIE() && bHasRequiredLogins)
 		{
 			// Make sure all instances of PIE are logged in before creating/launching worlds
 			LoginPIEInstances(bAnyBlueprintErrors, bStartInSpectatorMode, PIEStartTime);
@@ -1928,12 +1935,6 @@ void UEditorEngine::LoginPIEInstances(bool bAnyBlueprintErrors, bool bStartInSpe
 	int32 NextX = 0;
 	int32 NextY = 0;
 
-	bool bHasRequiredLogins = PlayInSettings->PlayNumberOfClients <= PIELogins.Num();
-	if (!bHasRequiredLogins)
-	{
-		UE_LOG(LogOnline, Verbose, TEXT("Not enough login credentials to launch all PIE instances, modify [/Script/UnrealEd.UnrealEdEngine].PIELogins"));
-	}
-
 	// Server
 	{
 		FWorldContext &PieWorldContext = CreateNewWorldContext(EWorldType::PIE);
@@ -1949,6 +1950,8 @@ void UEditorEngine::LoginPIEInstances(bool bAnyBlueprintErrors, bool bStartInSpe
 		FName OnlineIdentifier = FName(*FString::Printf(TEXT(":%d"), PieWorldContext.ContextHandle));
 		IOnlineSubsystem* OnlineSub = IOnlineSubsystem::Get(OnlineIdentifier);
 		IOnlineIdentityPtr IdentityInt = OnlineSub->GetIdentityInterface();
+		check(IdentityInt.IsValid());
+		NumOnlinePIEInstances++;
 
 		if (!PlayInSettings->PlayNetDedicated)
 		{
@@ -1956,26 +1959,18 @@ void UEditorEngine::LoginPIEInstances(bool bAnyBlueprintErrors, bool bStartInSpe
 			DataStruct.NextY = NextY;
 			GetMultipleInstancePositions(DataStruct.SettingsIndex, NextX, NextY);
 
-			if (bHasRequiredLogins)
-			{
-				// Login to online platform before creating world
-				FOnlineAccountCredentials AccountCreds;
-				AccountCreds.Id = PIELogins[ClientNum].Id;
-				AccountCreds.Token = PIELogins[ClientNum].Token;
-				AccountCreds.Type = PIELogins[ClientNum].Type;
-
-				FOnLoginCompleteDelegate Delegate;
-				Delegate.BindUObject(this, &UEditorEngine::OnLoginPIEComplete, DataStruct);
-
-				// Login first and continue the flow later
-				IdentityInt->AddOnLoginCompleteDelegate(0, Delegate);
-				IdentityInt->Login(0, AccountCreds);
-			}
-			else
-			{
-				// No logins, but continue anyway
-				CreatePIEWorldFromLogin(PieWorldContext, EPlayNetMode::PIE_ListenServer, DataStruct);
-			}
+			// Login to online platform before creating world
+			FOnlineAccountCredentials AccountCreds;
+			AccountCreds.Id = PIELogins[ClientNum].Id;
+			AccountCreds.Token = PIELogins[ClientNum].Token;
+			AccountCreds.Type = PIELogins[ClientNum].Type;
+
+			FOnLoginCompleteDelegate Delegate;
+			Delegate.BindUObject(this, &UEditorEngine::OnLoginPIEComplete, DataStruct);
+
+			// Login first and continue the flow later
+			IdentityInt->AddOnLoginCompleteDelegate(0, Delegate);
+			IdentityInt->Login(0, AccountCreds);
 
 			ClientNum++;
 		}
@@ -2003,28 +1998,22 @@ void UEditorEngine::LoginPIEInstances(bool bAnyBlueprintErrors, bool bStartInSpe
 		GetMultipleInstancePositions(DataStruct.SettingsIndex, NextX, NextY);
 		DataStruct.bIsServer = false;
 
-		if (bHasRequiredLogins)
-		{
-			FName OnlineIdentifier = FName(*FString::Printf(TEXT(":%d"), PieWorldContext.ContextHandle));
-			IOnlineIdentityPtr IdentityInt = Online::GetIdentityInterface(OnlineIdentifier);
+		FName OnlineIdentifier = FName(*FString::Printf(TEXT(":%d"), PieWorldContext.ContextHandle));
+		IOnlineIdentityPtr IdentityInt = Online::GetIdentityInterface(OnlineIdentifier);
+		check(IdentityInt.IsValid());
+		NumOnlinePIEInstances++;
 
-			FOnlineAccountCredentials AccountCreds;
-			AccountCreds.Id = PIELogins[ClientNum].Id;
-			AccountCreds.Token = PIELogins[ClientNum].Token;
-			AccountCreds.Type = PIELogins[ClientNum].Type;
+		FOnlineAccountCredentials AccountCreds;
+		AccountCreds.Id = PIELogins[ClientNum].Id;
+		AccountCreds.Token = PIELogins[ClientNum].Token;
+		AccountCreds.Type = PIELogins[ClientNum].Type;
 
-			FOnLoginCompleteDelegate Delegate;
-			Delegate.BindUObject(this, &UEditorEngine::OnLoginPIEComplete, DataStruct);
+		FOnLoginCompleteDelegate Delegate;
+		Delegate.BindUObject(this, &UEditorEngine::OnLoginPIEComplete, DataStruct);
 
-			IdentityInt->ClearOnLoginCompleteDelegate(0, Delegate);
-			IdentityInt->AddOnLoginCompleteDelegate(0, Delegate);
-			IdentityInt->Login(0, AccountCreds);
-		}
-		else
-		{
-			// No logins, but continue anyway
-			CreatePIEWorldFromLogin(PieWorldContext, EPlayNetMode::PIE_Client, DataStruct);
-		}
+		IdentityInt->ClearOnLoginCompleteDelegate(0, Delegate);
+		IdentityInt->AddOnLoginCompleteDelegate(0, Delegate);
+		IdentityInt->Login(0, AccountCreds);
 	}
 
 	// Restore window settings
