UE4: Added a UFactory for World assets.

diff --git a/Engine/Source/Editor/UnrealEd/Classes/Factories/WorldFactory.h b/Engine/Source/Editor/UnrealEd/Classes/Factories/WorldFactory.h
new file mode 100644
--- /dev/null
+++ b/Engine/Source/Editor/UnrealEd/Classes/Factories/WorldFactory.h
@@ -0,0 +1,21 @@
+// Copyright 1998-2013 Epic Games, Inc. All Rights Reserved.
+
+/**
+ *
+ */
+
+#include "WorldFactory.generated.h"
+
+UCLASS(MinimalAPI)
+class UWorldFactory : public UFactory
+{
+	GENERATED_UCLASS_BODY()
+
+	TEnumAsByte<EWorldType::Type> WorldType;
+
+	// Begin UFactory Interface
+	virtual bool ConfigureProperties() OVERRIDE;
+	virtual UObject* FactoryCreateNew(UClass* Class,UObject* InParent,FName Name,EObjectFlags Flags,UObject* Context,FFeedbackContext* Warn) OVERRIDE;
+	// Begin UFactory Interface	
+};
+
diff --git a/Engine/Source/Editor/UnrealEd/Private/EditorServer.cpp b/Engine/Source/Editor/UnrealEd/Private/EditorServer.cpp
--- a/Engine/Source/Editor/UnrealEd/Private/EditorServer.cpp
+++ b/Engine/Source/Editor/UnrealEd/Private/EditorServer.cpp
@@ -27,6 +27,7 @@
 #include "Editor/PropertyEditor/Public/IDetailsView.h"
 #include "Toolkits/AssetEditorManager.h"
 #include "AssetRegistryModule.h"
+#include "AssetToolsModule.h"
 #include "SnappingUtils.h"
 
 #include "TargetPlatform.h"
@@ -1748,11 +1749,20 @@ UWorld* UEditorEngine::NewMap()
 	const FText CleanseText = LOCTEXT("LoadingMap_Template", "New Map");
 	EditorDestroyWorld( Context, CleanseText );
 
-	// Create a new world.
-	Context.SetCurrentWorld(UWorld::CreateWorld(EWorldType::Editor, true ));
-	GWorld = Context.World();
-	GWorld->SetFlags(RF_Public|RF_Standalone);
-	
+	// Create a new world
+	UWorldFactory* Factory = ConstructObject<UWorldFactory>(UWorldFactory::StaticClass());
+	Factory->WorldType = EWorldType::Editor;
+	FString AssetName = TEXT("NewMap");
+	FString PackagePath = TEXT("/Temp");
+	FAssetToolsModule& AssetToolsModule = FModuleManager::Get().LoadModuleChecked<FAssetToolsModule>("AssetTools");
+	UWorld* NewWorld = Cast<UWorld>(AssetToolsModule.Get().CreateAsset(AssetName, PackagePath, UWorld::StaticClass(), Factory));
+	if ( NewWorld )
+	{
+		Context.SetCurrentWorld(NewWorld);
+		GWorld = NewWorld;
+		NewWorld->AddToRoot();
+	}
+
 	NoteSelectionChange();
 
 	// Starting a new map will wipe existing actors and add some defaults actors to the scene, so we need
@@ -1768,7 +1778,7 @@ UWorld* UEditorEngine::NewMap()
 		FLevelEditorModule& LevelEditor = FModuleManager::GetModuleChecked<FLevelEditorModule>("LevelEditor");
 
 		// Notify slate level editors of the map change
-		LevelEditor.BroadcastMapChanged( GWorld, EMapChangeType::NewMap );
+		LevelEditor.BroadcastMapChanged( NewWorld, EMapChangeType::NewMap );
 	}
 
 	// Move the brush to the origin.
@@ -1789,7 +1799,7 @@ UWorld* UEditorEngine::NewMap()
 	// Invalidate all the level viewport hit proxies
 	InvalidateAllViewportClientHitProxies();
 
-	return Context.World();
+	return NewWorld;
 }
 
 
diff --git a/Engine/Source/Editor/UnrealEd/Private/Factories/WorldFactory.cpp b/Engine/Source/Editor/UnrealEd/Private/Factories/WorldFactory.cpp
new file mode 100644
--- /dev/null
+++ b/Engine/Source/Editor/UnrealEd/Private/Factories/WorldFactory.cpp
@@ -0,0 +1,31 @@
+// Copyright 1998-2013 Epic Games, Inc. All Rights Reserved.
+
+#include "UnrealEd.h"
+
+#define LOCTEXT_NAMESPACE "WorldFactory"
+
+UWorldFactory::UWorldFactory(const class FPostConstructInitializeProperties& PCIP)
+	: Super(PCIP)
+{
+	bCreateNew = FParse::Param(FCommandLine::Get(), TEXT("WorldAssets"));
+	SupportedClass = UWorld::StaticClass();
+	WorldType = EWorldType::Inactive;
+}
+
+bool UWorldFactory::ConfigureProperties()
+{
+	return true;
+}
+
+UObject* UWorldFactory::FactoryCreateNew(UClass* Class, UObject* InParent, FName Name, EObjectFlags Flags, UObject* Context, FFeedbackContext* Warn)
+{
+	// Create a new world.
+	const bool bAddToRoot = false;
+	const bool bInformEngineOfWorld = true;
+	UWorld* NewWorld = UWorld::CreateWorld(WorldType, bInformEngineOfWorld, Name, Cast<UPackage>(InParent), bAddToRoot);
+	NewWorld->SetFlags(Flags);
+
+	return NewWorld;
+}
+
+#undef LOCTEXT_NAMESPACE
diff --git a/Engine/Source/Runtime/Engine/Classes/Engine/World.h b/Engine/Source/Runtime/Engine/Classes/Engine/World.h
--- a/Engine/Source/Runtime/Engine/Classes/Engine/World.h
+++ b/Engine/Source/Runtime/Engine/Classes/Engine/World.h
@@ -1780,7 +1780,7 @@ public:
 	/**
 	 * Static function that creates a new UWorld and returns a pointer to it
 	 */
-	static UWorld* CreateWorld( const EWorldType::Type InWorldType, bool bInformEngineOfWorld );
+	static UWorld* CreateWorld( const EWorldType::Type InWorldType, bool bInformEngineOfWorld, FName WorldName = NAME_None, UPackage* InWorldPackage = NULL, bool bAddToRoot = true );
 
 	/** 
 	 * Destroy this World instance 
diff --git a/Engine/Source/Runtime/Engine/Private/World.cpp b/Engine/Source/Runtime/Engine/Private/World.cpp
--- a/Engine/Source/Runtime/Engine/Private/World.cpp
+++ b/Engine/Source/Runtime/Engine/Private/World.cpp
@@ -850,10 +850,14 @@ void UWorld::DestroyWorld( bool bInformEngineOfWorld )
 	}
 }
 
-UWorld* UWorld::CreateWorld( const EWorldType::Type InWorldType, bool bInformEngineOfWorld )
+UWorld* UWorld::CreateWorld( const EWorldType::Type InWorldType, bool bInformEngineOfWorld, FName WorldName, UPackage* InWorldPackage, bool bAddToRoot )
 {
 	// Create a new package unless we're a commandlet in which case we keep the dummy world in the transient package.
-	UPackage*	WorldPackage		= IsRunningCommandlet() ? GetTransientPackage() : CreatePackage( NULL, NULL );
+	UPackage* WorldPackage = InWorldPackage;
+	if ( !WorldPackage )
+	{
+		WorldPackage = IsRunningCommandlet() ? GetTransientPackage() : CreatePackage( NULL, NULL );
+	}
 
 	if (InWorldType == EWorldType::PIE)
 	{
@@ -868,15 +872,19 @@ UWorld* UWorld::CreateWorld( const EWorldType::Type InWorldType, bool bInformEng
 	}
 
 	// Create new UWorld, ULevel and UModel.
-	UWorld* NewWorld = new( WorldPackage				, TEXT("NewWorld")			) UWorld(FPostConstructInitializeProperties(),FURL(NULL));
+	const FString WorldNameString = (WorldName != NAME_None) ? WorldName.ToString() : TEXT("NewWorld");
+	UWorld* NewWorld = new( WorldPackage				, *WorldNameString			) UWorld(FPostConstructInitializeProperties(),FURL(NULL));
 	NewWorld->WorldType = InWorldType;
 	NewWorld->InitializeNewWorld(UWorld::InitializationValues().ShouldSimulatePhysics(false).EnableTraceCollision(true).CreateNavigation(InWorldType == EWorldType::Editor));
 
 	// Clear the dirty flag set during SpawnActor and UpdateLevelComponents
 	WorldPackage->SetDirtyFlag(false);
 
-	// Add to root set so it doesn't get garbage collected.
-	NewWorld->AddToRoot();
+	if ( bAddToRoot )
+	{
+		// Add to root set so it doesn't get garbage collected.
+		NewWorld->AddToRoot();
+	}
 
 	// Tell the engine we are adding a world (unless we are asked not to)
 	if( ( GEngine ) && ( bInformEngineOfWorld == true ) )
