UE4: Fixed a crash when creating a new map in the editor. Not using asset tools to generate the world asset in this case since /Temp isnt a valid location to make an asset

diff --git a/Engine/Source/Editor/UnrealEd/Private/EditorServer.cpp b/Engine/Source/Editor/UnrealEd/Private/EditorServer.cpp
--- a/Engine/Source/Editor/UnrealEd/Private/EditorServer.cpp
+++ b/Engine/Source/Editor/UnrealEd/Private/EditorServer.cpp
@@ -27,7 +27,6 @@
 #include "Editor/PropertyEditor/Public/IDetailsView.h"
 #include "Toolkits/AssetEditorManager.h"
 #include "AssetRegistryModule.h"
-#include "AssetToolsModule.h"
 #include "SnappingUtils.h"
 
 #include "TargetPlatform.h"
@@ -1752,10 +1751,9 @@ UWorld* UEditorEngine::NewMap()
 	// Create a new world
 	UWorldFactory* Factory = ConstructObject<UWorldFactory>(UWorldFactory::StaticClass());
 	Factory->WorldType = EWorldType::Editor;
-	FString AssetName = TEXT("NewMap");
-	FString PackagePath = TEXT("/Temp");
-	FAssetToolsModule& AssetToolsModule = FModuleManager::Get().LoadModuleChecked<FAssetToolsModule>("AssetTools");
-	UWorld* NewWorld = Cast<UWorld>(AssetToolsModule.Get().CreateAsset(AssetName, PackagePath, UWorld::StaticClass(), Factory));
+	UPackage* Pkg = CreatePackage(NULL, TEXT("/Temp/NewMap"));
+	EObjectFlags Flags = RF_Public | RF_Standalone;
+	UWorld* NewWorld = Cast<UWorld>(Factory->FactoryCreateNew(UWorld::StaticClass(), Pkg, TEXT("NewMap"), Flags, NULL, GWarn));
 	if ( NewWorld )
 	{
 		Context.SetCurrentWorld(NewWorld);
