UE4: Added some check()s to enforce assumptions

diff --git a/Engine/Source/Editor/UnrealEd/Private/EditorActorFolders.cpp b/Engine/Source/Editor/UnrealEd/Private/EditorActorFolders.cpp
--- a/Engine/Source/Editor/UnrealEd/Private/EditorActorFolders.cpp
+++ b/Engine/Source/Editor/UnrealEd/Private/EditorActorFolders.cpp
@@ -87,12 +87,14 @@ void FActorFolders::Housekeeping()
 void FActorFolders::OnLevelActorListChanged()
 {
 	Housekeeping();
+	check(GWorld);
 	RebuildFolderListForWorld(*GWorld);
 }
 
 void FActorFolders::OnMapChange(uint32 MapChangeFlags)
 {
 	Housekeeping();
+	check(GWorld);
 	RebuildFolderListForWorld(*GWorld);
 }
 
diff --git a/Engine/Source/Editor/UnrealEd/Private/EditorServer.cpp b/Engine/Source/Editor/UnrealEd/Private/EditorServer.cpp
--- a/Engine/Source/Editor/UnrealEd/Private/EditorServer.cpp
+++ b/Engine/Source/Editor/UnrealEd/Private/EditorServer.cpp
@@ -1751,15 +1751,12 @@ UWorld* UEditorEngine::NewMap()
 	// Create a new world
 	UWorldFactory* Factory = ConstructObject<UWorldFactory>(UWorldFactory::StaticClass());
 	Factory->WorldType = EWorldType::Editor;
-	UPackage* Pkg = CreatePackage(NULL, TEXT("/Temp/NewMap"));
+	UPackage* Pkg = CreatePackage( NULL, NULL );
 	EObjectFlags Flags = RF_Public | RF_Standalone;
-	UWorld* NewWorld = Cast<UWorld>(Factory->FactoryCreateNew(UWorld::StaticClass(), Pkg, TEXT("NewMap"), Flags, NULL, GWarn));
-	if ( NewWorld )
-	{
-		Context.SetCurrentWorld(NewWorld);
-		GWorld = NewWorld;
-		NewWorld->AddToRoot();
-	}
+	UWorld* NewWorld = CastChecked<UWorld>(Factory->FactoryCreateNew(UWorld::StaticClass(), Pkg, TEXT("NewWorld"), Flags, NULL, GWarn));
+	Context.SetCurrentWorld(NewWorld);
+	GWorld = NewWorld;
+	NewWorld->AddToRoot();
 
 	NoteSelectionChange();
 
