Minor GWorld cleanup

diff --git a/Engine/Source/Editor/ContentBrowser/Private/AssetContextMenu.cpp b/Engine/Source/Editor/ContentBrowser/Private/AssetContextMenu.cpp
--- a/Engine/Source/Editor/ContentBrowser/Private/AssetContextMenu.cpp
+++ b/Engine/Source/Editor/ContentBrowser/Private/AssetContextMenu.cpp
@@ -564,14 +564,17 @@ struct WorldReferenceGenerator : public FFindReferencedAssets
 		const bool bIncludeDefaults = false;
 		const bool bReverseReferenceGraph = true;
 
-		// Generate the reference graph for GWorld
-		FReferencedAssets* WorldReferencer = new(Referencers) FReferencedAssets(GWorld);
-		FFindAssetsArchive(GWorld, WorldReferencer->AssetList, &ReferenceGraph, MaxRecursionDepth, bIncludeClasses, bIncludeDefaults, bReverseReferenceGraph);
+
+		UWorld* World = GWorld;
+
+		// Generate the reference graph for the world
+		FReferencedAssets* WorldReferencer = new(Referencers)FReferencedAssets(World);
+		FFindAssetsArchive(World, WorldReferencer->AssetList, &ReferenceGraph, MaxRecursionDepth, bIncludeClasses, bIncludeDefaults, bReverseReferenceGraph);
 
 		// Also include all the streaming levels in the results
-		for( int32 LevelIndex = 0 ; LevelIndex < GWorld->StreamingLevels.Num() ; ++LevelIndex )
+		for (int32 LevelIndex = 0; LevelIndex < World->StreamingLevels.Num(); ++LevelIndex)
 		{
-			ULevelStreaming* StreamingLevel = GWorld->StreamingLevels[LevelIndex];
+			ULevelStreaming* StreamingLevel = World->StreamingLevels[LevelIndex];
 			if( StreamingLevel != NULL )
 			{
 				ULevel* Level = StreamingLevel->GetLoadedLevel();
diff --git a/Engine/Source/Editor/LandscapeEditor/Private/LandscapeEditorDetailCustomization_NewLandscape.cpp b/Engine/Source/Editor/LandscapeEditor/Private/LandscapeEditorDetailCustomization_NewLandscape.cpp
--- a/Engine/Source/Editor/LandscapeEditor/Private/LandscapeEditorDetailCustomization_NewLandscape.cpp
+++ b/Engine/Source/Editor/LandscapeEditor/Private/LandscapeEditorDetailCustomization_NewLandscape.cpp
@@ -813,7 +813,7 @@ FReply FLandscapeEditorDetailCustomization_NewLandscape::OnCreateButtonClicked()
 
 		const float ComponentSize = QuadsPerComponent;
 		const FVector Offset = FTransform(LandscapeEdMode->UISettings->NewLandscape_Rotation, FVector::ZeroVector, LandscapeEdMode->UISettings->NewLandscape_Scale).TransformVector(FVector(-ComponentCountX * ComponentSize / 2, -ComponentCountY * ComponentSize / 2, 0));
-		ALandscape* Landscape = GWorld->SpawnActor<ALandscape>(LandscapeEdMode->UISettings->NewLandscape_Location + Offset, LandscapeEdMode->UISettings->NewLandscape_Rotation);
+		ALandscape* Landscape = LandscapeEdMode->GetWorld()->SpawnActor<ALandscape>(LandscapeEdMode->UISettings->NewLandscape_Location + Offset, LandscapeEdMode->UISettings->NewLandscape_Rotation);
 
 		//if (LandscapeEdMode->NewLandscapePreviewMode == ENewLandscapePreviewMode::ImportLandscape)
 		{
diff --git a/Engine/Source/Editor/SceneOutliner/Private/SSceneOutliner.cpp b/Engine/Source/Editor/SceneOutliner/Private/SSceneOutliner.cpp
--- a/Engine/Source/Editor/SceneOutliner/Private/SSceneOutliner.cpp
+++ b/Engine/Source/Editor/SceneOutliner/Private/SSceneOutliner.cpp
@@ -3057,9 +3057,10 @@ namespace SceneOutliner
 						}
 					}
 
-					if (GUnrealEd->CanDeleteSelectedActors( GWorld, true, false ))
+					UWorld* World = GWorld;
+					if (GUnrealEd->CanDeleteSelectedActors( World, true, false ))
 					{
-						GEditor->edactDeleteSelected( GWorld );
+						GEditor->edactDeleteSelected( World );
 					}
 				}
 
diff --git a/Engine/Source/Editor/UnrealEd/Private/Editor.cpp b/Engine/Source/Editor/UnrealEd/Private/Editor.cpp
--- a/Engine/Source/Editor/UnrealEd/Private/Editor.cpp
+++ b/Engine/Source/Editor/UnrealEd/Private/Editor.cpp
@@ -760,15 +760,16 @@ void UEditorEngine::AddReferencedObjects(UObject* InThis, FReferenceCollector& C
 
 void UEditorEngine::Tick( float DeltaSeconds, bool bIdleMode )
 {
-	check( GWorld );
-	check( GWorld != PlayWorld || bIsSimulatingInEditor );
+	UWorld* CurrentGWorld = GWorld;
+	check( CurrentGWorld );
+	check( CurrentGWorld != PlayWorld || bIsSimulatingInEditor );
 
 	// Always ensure we've got adequate slack for any worlds that are going to get created in this frame so that
 	// our EditorContext reference doesn't get invalidated
 	WorldList.Reserve(WorldList.Num() + 10);
 
 	FWorldContext& EditorContext = GetEditorWorldContext();
-	check( GWorld == EditorContext.World() );
+	check( CurrentGWorld == EditorContext.World() );
 
 	// was there a reregister requested last frame?
 	if (bHasPendingGlobalReregister)
@@ -1676,9 +1677,10 @@ void UEditorEngine::Cleanse( bool ClearSelection, bool Redraw, const FText& Tran
 
 void UEditorEngine::EditorClearComponents()
 {
-	if( GWorld != NULL )
+	UWorld* World = GWorld;
+	if (World != NULL)
 	{
-		GWorld->ClearWorldComponents();
+		World->ClearWorldComponents();
 	}
 }
 
@@ -3804,8 +3806,9 @@ void UEditorEngine::OpenMatinee(AMatineeActor* MatineeActor, bool bWarnUser)
 void UEditorEngine::UpdateReflectionCaptures()
 {
 	// Update sky light first because it's considered direct lighting, sky diffuse will be visible in reflection capture indirect specular
-	GWorld->UpdateAllSkyCaptures();
-	GWorld->UpdateAllReflectionCaptures();
+	UWorld* World = GWorld;
+	World->UpdateAllSkyCaptures();
+	World->UpdateAllReflectionCaptures();
 }
 
 void UEditorEngine::UpdateSkyCaptures()
diff --git a/Engine/Source/Editor/UnrealEd/Private/EditorActorFolders.cpp b/Engine/Source/Editor/UnrealEd/Private/EditorActorFolders.cpp
--- a/Engine/Source/Editor/UnrealEd/Private/EditorActorFolders.cpp
+++ b/Engine/Source/Editor/UnrealEd/Private/EditorActorFolders.cpp
@@ -31,9 +31,10 @@ FActorFolders::FActorFolders()
 	GEngine->OnLevelActorAdded().AddRaw(this, &FActorFolders::OnLevelActorAdded);
 	FEditorDelegates::MapChange.AddRaw( this, &FActorFolders::OnMapChange );
 
-	if (GWorld)
+	UWorld* World = GWorld;
+	if (World)
 	{
-		RebuildFolderListForWorld(*GWorld);
+		RebuildFolderListForWorld(*World);
 	}
 }
 
diff --git a/Engine/Source/Editor/UnrealEd/Private/Factories/ActorFactory.cpp b/Engine/Source/Editor/UnrealEd/Private/Factories/ActorFactory.cpp
--- a/Engine/Source/Editor/UnrealEd/Private/Factories/ActorFactory.cpp
+++ b/Engine/Source/Editor/UnrealEd/Private/Factories/ActorFactory.cpp
@@ -386,7 +386,7 @@ void UActorFactoryEmitter::PostSpawnActor(UObject* Asset, AActor* NewActor)
 	NewEmitter->SetTemplate(ParticleSystem);
 
 	// if we're created by Kismet on the server during gameplay, we need to replicate the emitter
-	if (GWorld->HasBegunPlay() && GWorld->GetNetMode() != NM_Client)
+	if (NewEmitter->GetWorld()->HasBegunPlay() && NewEmitter->GetWorld()->GetNetMode() != NM_Client)
 	{
 		NewEmitter->SetReplicates(true);
 		NewEmitter->bAlwaysRelevant = true;
diff --git a/Engine/Source/Editor/UnrealEd/Private/Factories/EditorFactories.cpp b/Engine/Source/Editor/UnrealEd/Private/Factories/EditorFactories.cpp
--- a/Engine/Source/Editor/UnrealEd/Private/Factories/EditorFactories.cpp
+++ b/Engine/Source/Editor/UnrealEd/Private/Factories/EditorFactories.cpp
@@ -946,9 +946,10 @@ UObject* UPackageFactory::FactoryCreateText( UClass* Class, UObject* InParent, F
 
 	UPackage* TopLevelPackage = NULL;
 	UPackage* RootMapPackage = NULL;
-	if (GWorld)
+	UWorld* World = GWorld;
+	if (World)
 	{
-		RootMapPackage = GWorld->GetOutermost();
+		RootMapPackage = World->GetOutermost();
 	}
 
 	if (RootMapPackage)
