#UE4 - Added OBJ static mesh import support.

diff --git a/Engine/Source/Editor/DestructibleMeshEditor/Private/SDestructibleMeshEditorViewport.cpp b/Engine/Source/Editor/DestructibleMeshEditor/Private/SDestructibleMeshEditorViewport.cpp
--- a/Engine/Source/Editor/DestructibleMeshEditor/Private/SDestructibleMeshEditorViewport.cpp
+++ b/Engine/Source/Editor/DestructibleMeshEditor/Private/SDestructibleMeshEditorViewport.cpp
@@ -416,7 +416,7 @@ void FDestructibleMeshEditorViewportClient::ImportFBXChunks()
 			FlushRenderingCommands();
 
 			UnFbx::FFbxImporter* FFbxImporter = UnFbx::FFbxImporter::GetInstance();
-			if (FFbxImporter->ImportFromFile(*ImportFilename))
+			if (FFbxImporter->ImportFromFile( *ImportFilename, FPaths::GetExtension( ImportFilename ) ) )
 			{
 				TArray<FbxNode*> FbxMeshArray;
 				FFbxImporter->FillFbxMeshArray(FFbxImporter->Scene->GetRootNode(), FbxMeshArray, FFbxImporter);
diff --git a/Engine/Source/Editor/Matinee/Private/MatineeMenus.cpp b/Engine/Source/Editor/Matinee/Private/MatineeMenus.cpp
--- a/Engine/Source/Editor/Matinee/Private/MatineeMenus.cpp
+++ b/Engine/Source/Editor/Matinee/Private/MatineeMenus.cpp
@@ -3695,13 +3695,13 @@ void FMatinee::OnMenuImport()
 			FEditorDirectories::Get().SetLastDirectory(ELastDirectory::GENERIC_IMPORT, FPaths::GetPath(FileName)); // Save path as default for next time.
 		
 			const FString FileExtension = FPaths::GetExtension(FileName);
-			const bool bIsFBX = FCString::Stricmp(*FileExtension, TEXT("FBX")) == 0;
+			const bool bIsFBX = FileExtension.Equals( TEXT("FBX"), ESearchCase::IgnoreCase );
 
 			if (bIsFBX)
 			{
 				// Import the Matinee information from the FBX document.
 				UnFbx::FFbxImporter* FFbxImporter = UnFbx::FFbxImporter::GetInstance();
-				if (FFbxImporter->ImportFromFile(*ImportFilename))
+				if (FFbxImporter->ImportFromFile(ImportFilename, FileExtension ))
 				{
 					FFbxImporter->SetProcessUnknownCameras(false);
 					
diff --git a/Engine/Source/Editor/UnrealEd/Private/Fbx/FbxFactory.cpp b/Engine/Source/Editor/UnrealEd/Private/Fbx/FbxFactory.cpp
--- a/Engine/Source/Editor/UnrealEd/Private/Fbx/FbxFactory.cpp
+++ b/Engine/Source/Editor/UnrealEd/Private/Fbx/FbxFactory.cpp
@@ -45,6 +45,8 @@ UFbxFactory::UFbxFactory(const class FPostConstructInitializeProperties& PCIP)
 {
 	SupportedClass = NULL;
 	Formats.Add(TEXT("fbx;FBX meshes and animations"));
+	Formats.Add(TEXT("obj;OBJ Static meshes"));
+	//Formats.Add(TEXT("dae;Collada meshes and animations"));
 
 	bCreateNew = false;
 	bText = false;
@@ -188,9 +190,18 @@ UObject* UFbxFactory::FactoryCreateBinary
 	
 	UnFbx::FFbxLoggerSetter Logger(FbxImporter);
 
+	EFBXImportType ForcedImportType = FBXIT_StaticMesh;
+
+	bool bIsObjFormat = false;
+	if( FString(Type).Equals(TEXT("obj"), ESearchCase::IgnoreCase ) )
+	{
+		bIsObjFormat = true;
+	}
+
+
 	bool bShowImportDialog = bShowOption && !GIsAutomationTesting;
 	bool bImportAll = false;
-	UnFbx::FBXImportOptions* ImportOptions = GetImportOptions(FbxImporter, ImportUI, bShowImportDialog, InParent->GetPathName(), bOperationCanceled, bImportAll);
+	UnFbx::FBXImportOptions* ImportOptions = GetImportOptions(FbxImporter, ImportUI, bShowImportDialog, InParent->GetPathName(), bOperationCanceled, bImportAll, bIsObjFormat, bIsObjFormat, ForcedImportType );
 	bOutOperationCanceled = bOperationCanceled;
 	
 	if( bImportAll )
@@ -205,7 +216,7 @@ UObject* UFbxFactory::FactoryCreateBinary
 	if (ImportOptions)
 	{
 		Warn->BeginSlowTask( NSLOCTEXT("FbxFactory", "BeginImportingFbxMeshTask", "Importing FBX mesh"), true );
-		if ( !FbxImporter->ImportFromFile( *UFactory::CurrentFilename ) )
+		if ( !FbxImporter->ImportFromFile( *UFactory::CurrentFilename, Type ) )
 		{
 			// Log the error message and fail the import.
 			Warn->Log(ELogVerbosity::Error, FbxImporter->GetErrorMessage() );
diff --git a/Engine/Source/Editor/UnrealEd/Private/Fbx/FbxMainImport.cpp b/Engine/Source/Editor/UnrealEd/Private/Fbx/FbxMainImport.cpp
--- a/Engine/Source/Editor/UnrealEd/Private/Fbx/FbxMainImport.cpp
+++ b/Engine/Source/Editor/UnrealEd/Private/Fbx/FbxMainImport.cpp
@@ -53,7 +53,7 @@ TSharedPtr<FFbxImporter> FFbxImporter::StaticInstance;
 
 
 
-FBXImportOptions* GetImportOptions( UnFbx::FFbxImporter* FbxImporter, UFbxImportUI* ImportUI, bool bShowOptionDialog, const FString& FullPath, bool& OutOperationCanceled, bool& bOutImportAll, bool bForceImportType, EFBXImportType ImportType )
+FBXImportOptions* GetImportOptions( UnFbx::FFbxImporter* FbxImporter, UFbxImportUI* ImportUI, bool bShowOptionDialog, const FString& FullPath, bool& OutOperationCanceled, bool& bOutImportAll, bool bIsObjFormat, bool bForceImportType, EFBXImportType ImportType )
 {
 	OutOperationCanceled = false;
 
@@ -101,6 +101,7 @@ FBXImportOptions* GetImportOptions( UnFbx::FFbxImporter* FbxImporter, UFbxImport
 			.WidgetWindow(Window)
 			.FullPath(FullPath)
 			.ForcedImportType( bForceImportType ? TOptional<EFBXImportType>( ImportType ) : TOptional<EFBXImportType>() )
+			.IsObjFormat( bIsObjFormat )
 		);
 
 		// @todo: we can make this slow as showing progress bar later
@@ -665,13 +666,10 @@ bool FFbxImporter::ImportFile(FString Filename)
 //-------------------------------------------------------------------------
 //
 //-------------------------------------------------------------------------
-bool FFbxImporter::ImportFromFile(const TCHAR* Filename)
+bool FFbxImporter::ImportFromFile(const FString& Filename, const FString& Type)
 {
 	bool Result = true;
-	// Converts the FBX data to Z-up, X-forward, Y-left.  Unreal is the same except with Y-right, 
-	// but the conversion to left-handed coordinates is not working properly
-	FbxAxisSystem::EFrontVector FrontVector = (FbxAxisSystem::EFrontVector)-FbxAxisSystem::eParityOdd;
-	const FbxAxisSystem UnrealZUp(FbxAxisSystem::eZAxis, FrontVector, FbxAxisSystem::eRightHanded);
+
 
 	switch (CurPhase)
 	{
@@ -689,23 +687,34 @@ bool FFbxImporter::ImportFromFile(const TCHAR* Filename)
 			break;
 		}
 	case IMPORTED:
-		// convert axis to Z-up
-		FbxRootNodeUtility::RemoveAllFbxRoots( Scene );
-		UnrealZUp.ConvertScene( Scene );
-
-		// Convert the scene's units to what is used in this program, if needed.
-		// The base unit used in both FBX and Unreal is centimeters.  So unless the units 
-		// are already in centimeters (ie: scalefactor 1.0) then it needs to be converted
-		//if( FbxScene->GetGlobalSettings().GetSystemUnit().GetScaleFactor() != 1.0 )
-		//{
-		//	KFbxSystemUnit::cm.ConvertScene( FbxScene );
-		//}
-
-
-		// convert name to unreal-supported format
-		// actually, crashes...
-		//KFbxSceneRenamer renamer(FbxScene);
-		//renamer.ResolveNameClashing(false,false,true,true,true,KString(),"TestBen",false,false);
+		{
+			static const FString Obj(TEXT("obj"));
+
+			// The imported axis system is unknown for obj files
+			if( !Type.Equals( Obj, ESearchCase::IgnoreCase ) )
+			{
+				FbxAxisSystem::EFrontVector FrontVector = (FbxAxisSystem::EFrontVector) - FbxAxisSystem::eParityOdd;
+				const FbxAxisSystem UnrealZUp(FbxAxisSystem::eZAxis, FrontVector, FbxAxisSystem::eRightHanded);
+
+				if( Scene->GetGlobalSettings().GetAxisSystem() != UnrealZUp )
+				{
+					// Converts the FBX data to Z-up, X-forward, Y-left.  Unreal is the same except with Y-right, 
+					// but the conversion to left-handed coordinates is not working properly
+
+					// convert axis to Z-up
+					FbxRootNodeUtility::RemoveAllFbxRoots(Scene);
+					UnrealZUp.ConvertScene(Scene);
+				}
+			}
+
+			// Convert the scene's units to what is used in this program, if needed.
+			// The base unit used in both FBX and Unreal is centimeters.  So unless the units 
+			// are already in centimeters (ie: scalefactor 1.0) then it needs to be converted
+			//if( FbxScene->GetGlobalSettings().GetSystemUnit().GetScaleFactor() != 1.0 )
+			//{
+			//	KFbxSystemUnit::cm.ConvertScene( FbxScene );
+			//}
+		}
 		
 	default:
 		break;
diff --git a/Engine/Source/Editor/UnrealEd/Private/Fbx/FbxOptionWindow.cpp b/Engine/Source/Editor/UnrealEd/Private/Fbx/FbxOptionWindow.cpp
--- a/Engine/Source/Editor/UnrealEd/Private/Fbx/FbxOptionWindow.cpp
+++ b/Engine/Source/Editor/UnrealEd/Private/Fbx/FbxOptionWindow.cpp
@@ -13,17 +13,20 @@ class SImportTypeButton : public SCompoundWidget
 public:
 	SLATE_BEGIN_ARGS( SImportTypeButton )
 	:	_DefaultImportType(FBXIT_StaticMesh)
+	,	_ForceImportType( false )
 	,	_OnSelectionChanged()
 	{}
 		SLATE_ARGUMENT( EFBXImportType, DefaultImportType )
+		SLATE_ARGUMENT( bool, ForceImportType )
 		SLATE_EVENT( FOnImportTypeChanged, OnSelectionChanged )
 	SLATE_END_ARGS()
 
 	void Construct( const FArguments& InArgs )
 	{
-		DetectedImportType = InArgs._DefaultImportType;
+		DefaultImportType = InArgs._DefaultImportType;
 		CurrentChoice = InArgs._DefaultImportType;
 		OnSelectionChanged = InArgs._OnSelectionChanged;
+		bForceImportType = InArgs._ForceImportType;
 
 		const UEnum* ImportTypeEnum = FindObject<UEnum>(ANY_PACKAGE, TEXT("EFBXImportType"), true);
 		check(ImportTypeEnum != nullptr);
@@ -54,17 +57,17 @@ public:
 
 				+SVerticalBox::Slot().AutoHeight()
 				[
-			SNew(SUniformGridPanel)
-			+ SUniformGridPanel::Slot(0, 0)
-			[
+					SNew(SUniformGridPanel)
+					+ SUniformGridPanel::Slot(0, 0)
+					[
 						CreateRadioButton(*ImportTypeEnum, FBXIT_StaticMesh, FEditorStyle::GetBrush("FBXIcon.StaticMesh"), 0)
-			]
-			+ SUniformGridPanel::Slot(1, 0)
-			[
+					]
+					+ SUniformGridPanel::Slot(1, 0)
+					[
 						CreateRadioButton(*ImportTypeEnum, FBXIT_SkeletalMesh, FEditorStyle::GetBrush("FBXIcon.SkeletalMesh"), 1)
-			]
-			+ SUniformGridPanel::Slot(2, 0)
-			[
+					]
+					+ SUniformGridPanel::Slot(2, 0)
+					[
 						CreateRadioButton(*ImportTypeEnum, FBXIT_Animation, FEditorStyle::GetBrush("FBXIcon.Animation"), 2)
 					]
 				]
@@ -91,6 +94,7 @@ private:
 			.IsChecked( this, &SImportTypeButton::IsRadioChecked, RadioButtonChoice )
 			.OnCheckStateChanged( this, &SImportTypeButton::OnRadioChanged, RadioButtonChoice )
 			.ToolTipText(Enum.GetToolTipText(RadioButtonChoice))
+			.IsEnabled( !bForceImportType || RadioButtonChoice == DefaultImportType )
 			[
 				SNew(SHorizontalBox)
 
@@ -108,7 +112,7 @@ private:
 				.VAlign(VAlign_Center)
 				.HAlign(HAlign_Center)
 				.Padding(6, 2)
-			[
+				[
 					SNew(STextBlock)
 					.Text(Enum.GetEnumText(RadioButtonChoice))
 					.ColorAndOpacity(FLinearColor(0.72f, 0.72f, 0.72f, 1.f))
@@ -125,7 +129,7 @@ private:
 
 	EVisibility IsWarningVisible() const
 	{
-		return (CurrentChoice == DetectedImportType)
+		return (CurrentChoice == DefaultImportType)
 			? EVisibility::Collapsed
 			: EVisibility::Visible;
 	}
@@ -143,9 +147,10 @@ private:
 		}
 	}
 
-	EFBXImportType DetectedImportType;
+	EFBXImportType DefaultImportType;
 	EFBXImportType CurrentChoice;
 	FOnImportTypeChanged OnSelectionChanged;
+	bool bForceImportType;
 };
 
 DECLARE_DELEGATE_OneParam( FOnAnimImportLengthOptionChanged, EFBXAnimationLengthImportType );
@@ -264,6 +269,7 @@ void SFbxOptionWindow::Construct(const FArguments& InArgs)
 {
 	ImportUI = InArgs._ImportUI;
 	WidgetWindow = InArgs._WidgetWindow;
+	bIsObjFormat = InArgs._IsObjFormat;
 
 	check (ImportUI);
 
@@ -275,10 +281,10 @@ void SFbxOptionWindow::Construct(const FArguments& InArgs)
 		StaticMeshLODGroups.Add(MakeShareable(new FString(StaticMeshLODGroupDisplayNames[GroupIndex].ToString())));
 	}
 
-	bReimport = InArgs._ForcedImportType.IsSet();
+	bForceImportType = InArgs._ForcedImportType.IsSet();
 
 	// Force the import type
-	if( bReimport )
+	if( bForceImportType )
 	{
 		ImportUI->MeshTypeToImport = InArgs._ForcedImportType.GetValue();
 	}
@@ -313,9 +319,8 @@ void SFbxOptionWindow::Construct(const FArguments& InArgs)
 			[
 				SNew(SImportTypeButton)
 				.DefaultImportType(ImportUI->MeshTypeToImport)
-				// Do not allow users to change the import type if it is forced to a specific type
-				.IsEnabled( !bReimport )
 				.OnSelectionChanged(this, &SFbxOptionWindow::SetImportType)
+				.ForceImportType( bForceImportType )
 			]
 
 			+SVerticalBox::Slot()
@@ -533,6 +538,7 @@ TSharedRef<SWidget> SFbxOptionWindow::ConstructMiscOption()
 		SNew(SCheckBox)
 		.IsChecked(ImportUI->bOverrideFullName? true: false)
 		.OnCheckStateChanged(this, &SFbxOptionWindow::SetGeneral_OverrideFullName)
+		.IsEnabled( !bIsObjFormat )
 		[
 			SNew(STextBlock)
 			.Text(LOCTEXT("FbxOptionWindow_NameOverride", "Override FullName"))
@@ -545,14 +551,21 @@ TSharedRef<SWidget> SFbxOptionWindow::ConstructMiscOption()
 TSharedRef<SWidget> SFbxOptionWindow::ConstructNormalImportOptions()
 {
 	NormalImportOptions.Empty();
-	for( int32 i = 0; i < FBXNIM_MAX; ++i )
+	for( int32 NormalMethod = 0; NormalMethod < FBXNIM_MAX; ++NormalMethod )
 	{
-		NormalImportOptions.Add( TSharedPtr<EFBXNormalImportMethod>() );
+		if( !bIsObjFormat || NormalMethod != FBXNIM_ImportNormalsAndTangents  )
+		{
+			NormalImportOptions.Add( TSharedPtr<EFBXNormalImportMethod>() );
+		}
 	}
 
 	NormalImportOptions[FBXNIM_ComputeNormals] = MakeShareable( new EFBXNormalImportMethod( FBXNIM_ComputeNormals ) );
 	NormalImportOptions[FBXNIM_ImportNormals] = MakeShareable( new EFBXNormalImportMethod( FBXNIM_ImportNormals ) );
-	NormalImportOptions[FBXNIM_ImportNormalsAndTangents] = MakeShareable( new EFBXNormalImportMethod( FBXNIM_ImportNormalsAndTangents ) );
+
+	if( !bIsObjFormat )
+	{
+		NormalImportOptions[FBXNIM_ImportNormalsAndTangents] = MakeShareable( new EFBXNormalImportMethod( FBXNIM_ImportNormalsAndTangents ) );
+	}
 
 	return 
 	SNew( SHorizontalBox )
@@ -653,6 +666,7 @@ TSharedRef<SWidget> SFbxOptionWindow::ConstructStaticMeshAdvanced()
 		SNew(SCheckBox)
 		.IsChecked(ImportUI->StaticMeshImportData->bImportMeshLODs? true: false)
 		.OnCheckStateChanged(this, &SFbxOptionWindow::SetStaticMesh_ImportMeshLODs)
+		.IsEnabled( !bIsObjFormat )
 		[
 			SNew(STextBlock)
 			.Text(LOCTEXT("FbxOptionWindow_ImportMeshLODs", "Import Mesh LODs"))
@@ -664,6 +678,7 @@ TSharedRef<SWidget> SFbxOptionWindow::ConstructStaticMeshAdvanced()
 		SNew(SCheckBox)
 		.IsChecked(ImportUI->StaticMeshImportData->bReplaceVertexColors? true: false)
 		.OnCheckStateChanged(this, &SFbxOptionWindow::SetStaticMesh_ReplaceVertexColor)
+		.IsEnabled( !bIsObjFormat )
 		[
 			SNew(STextBlock)
 			.Text(LOCTEXT("FbxOptionWindow_ReplaceVertexColors", "Replace Vertex Colors"))
@@ -686,12 +701,14 @@ TSharedRef<SWidget> SFbxOptionWindow::ConstructStaticMeshAdvanced()
 		SNew(SCheckBox)
 		.IsChecked(ImportUI->StaticMeshImportData->bOneConvexHullPerUCX? true: false)
 		.OnCheckStateChanged(this, &SFbxOptionWindow::SetStaticMesh_OneConvexHullPerUCX)
+		.IsEnabled( !bIsObjFormat )
 		[
 			SNew(STextBlock)
 			.Text(LOCTEXT("FbxOptionWindow_OneConvexHullPerUCX", "One Convex Hull Per UCX"))
 		]
 	];
 
+
 	return NewBox;
 }
 
@@ -850,7 +867,7 @@ TSharedRef<SWidget> SFbxOptionWindow::ConstructSkeletalMeshAdvanced()
 		SNew(SCheckBox)
 		.IsChecked(ImportUI->bCreatePhysicsAsset? true: false)
 		.OnCheckStateChanged(this, &SFbxOptionWindow::SetSkeletalMesh_CreatePhysicsAsset)
-		.IsEnabled(!bReimport)
+		.IsEnabled(!bForceImportType)
 		[
 			SNew(STextBlock)
 			.Text(LOCTEXT("FbxOptionWindow_CreatePhysicsAsset", "Create Physics Asset"))
@@ -905,6 +922,7 @@ TSharedRef<SWidget> SFbxOptionWindow::ConstructMaterialOption()
 		SNew(SCheckBox)
 		.IsChecked(ImportUI->bImportMaterials? true: false)
 		.OnCheckStateChanged(this, &SFbxOptionWindow::SetMaterial_ImportMaterials)
+		.IsEnabled( !bIsObjFormat )
 		[
 			SNew(STextBlock)
 			.Text(LOCTEXT("FbxOptionWindow_ImportMaterials", "Import Materials"))
@@ -916,6 +934,7 @@ TSharedRef<SWidget> SFbxOptionWindow::ConstructMaterialOption()
 		SNew(SCheckBox)
 		.IsChecked(ImportUI->bImportTextures? true: false)
 		.OnCheckStateChanged(this, &SFbxOptionWindow::SetMaterial_ImportTextures)
+		.IsEnabled( !bIsObjFormat )
 		[
 			SNew(STextBlock)
 			.Text(LOCTEXT("FbxOptionWindow_ImportTextures", "Import Textures"))
@@ -927,6 +946,7 @@ TSharedRef<SWidget> SFbxOptionWindow::ConstructMaterialOption()
 		SNew(SCheckBox)
 		.IsChecked(ImportUI->TextureImportData->bInvertNormalMaps? true: false)
 		.OnCheckStateChanged(this, &SFbxOptionWindow::SetMaterial_InvertNormalMaps)
+		.IsEnabled( !bIsObjFormat )
 		[
 			SNew(STextBlock)
 			.Text(LOCTEXT("FbxOptionWindow_InvertNormalMaps", "Invert Normal Maps"))
@@ -940,13 +960,6 @@ TSharedRef<SWidget> SFbxOptionWindow::ConstructAnimationOption()
 {
 	TSharedRef<SVerticalBox> NewBox = SNew(SVerticalBox);
 
-// 	NewBox->AddSlot().AutoHeight() .Padding(4)
-// 	[
-// 		SNew(STextBlock)
-// 		.TextStyle( FEditorStyle::Get(), "FBXSmallFont" ) 
-// 		.Text( LOCTEXT("FbxOptionWindow_Animation", "Animation") )
-// 	];
-
 	NewBox->AddSlot().AutoHeight() .Padding(2, 4)
 	[
 		SNew(SCheckBox)
@@ -1022,7 +1035,7 @@ TSharedRef<SWidget> SFbxOptionWindow::ConstructSkeletonOptionForMesh()
 		.ContentPadding(1.f)
 		.OnGetMenuContent( this, &SFbxOptionWindow::MakeSkeletonPickerMenu )
 		.HasDownArrow(true)
-		.IsEnabled( !bReimport )
+		.IsEnabled( !bForceImportType )
 		.ButtonContent()
 		[
 			SNew(STextBlock)
@@ -1069,7 +1082,7 @@ TSharedRef<SWidget> SFbxOptionWindow::ConstructSkeletonOptionForAnim()
 		.ContentPadding(1.f)
 		.OnGetMenuContent( this, &SFbxOptionWindow::MakeSkeletonPickerMenu )
 		.HasDownArrow(true)
-		.IsEnabled( !bReimport )
+		.IsEnabled( !bForceImportType )
 		.ButtonContent()
 		[
 			SNew(STextBlock)
@@ -1409,7 +1422,7 @@ bool SFbxOptionWindow::CanImport()  const
 
 bool SFbxOptionWindow::ShouldShowPhysicsAssetPicker() const
 {
-	return (!bReimport && ImportUI->bCreatePhysicsAsset == false);
+	return (!bForceImportType && ImportUI->bCreatePhysicsAsset == false);
 }
 
 EFBXNormalImportMethod SFbxOptionWindow::GetCurrentNormalImportMethod() const
diff --git a/Engine/Source/Editor/UnrealEd/Private/FbxImporter.h b/Engine/Source/Editor/UnrealEd/Private/FbxImporter.h
--- a/Engine/Source/Editor/UnrealEd/Private/FbxImporter.h
+++ b/Engine/Source/Editor/UnrealEd/Private/FbxImporter.h
@@ -217,7 +217,7 @@ public:
 	FbxString	ConvertToFbxString(const FString& String);
 };
 
-FBXImportOptions* GetImportOptions( class FFbxImporter* FbxImporter, UFbxImportUI* ImportUI, bool bShowOptionDialog, const FString& FullPath, bool& OutOperationCanceled, bool& OutImportAll, bool bForceImportType = false, EFBXImportType ImportType = FBXIT_StaticMesh );
+FBXImportOptions* GetImportOptions( class FFbxImporter* FbxImporter, UFbxImportUI* ImportUI, bool bShowOptionDialog, const FString& FullPath, bool& OutOperationCanceled, bool& OutImportAll, bool bIsObjFormat, bool bForceImportType = false, EFBXImportType ImportType = FBXIT_StaticMesh );
 void ApplyImportUIToImportOptions(UFbxImportUI* ImportUI, FBXImportOptions& InOutImportOptions);
 
 struct FImportedMaterialData
@@ -288,7 +288,7 @@ public:
 	 * @param Filename FBX file name to import.
 	 * @returns true on success.
 	 */
-	UNREALED_API bool ImportFromFile(const TCHAR* Filename);
+	UNREALED_API bool ImportFromFile(const FString& Filename, const FString& Type);
 
 	/**
 	 * Retrieve the FBX loader's error message explaining its failure to read a given FBX file.
diff --git a/Engine/Source/Editor/UnrealEd/Private/FbxMeshUtils.cpp b/Engine/Source/Editor/UnrealEd/Private/FbxMeshUtils.cpp
--- a/Engine/Source/Editor/UnrealEd/Private/FbxMeshUtils.cpp
+++ b/Engine/Source/Editor/UnrealEd/Private/FbxMeshUtils.cpp
@@ -297,7 +297,7 @@ namespace FbxMeshUtils
 		ImportOptions->bImportMaterials = false;
 		ImportOptions->bImportTextures = false;
 
-		if ( !FFbxImporter->ImportFromFile( *Filename ) )
+		if ( !FFbxImporter->ImportFromFile( *Filename, FPaths::GetExtension( Filename ) ) )
 		{
 			// Log the error message and fail the import.
 			// @todo verify if the message works
@@ -401,7 +401,7 @@ namespace FbxMeshUtils
 			ImportOptions->bImportTextures = false;
 			ImportOptions->bImportAnimations = false;
 
-			if ( !FFbxImporter->ImportFromFile( *Filename ) )
+			if ( !FFbxImporter->ImportFromFile( *Filename, FPaths::GetExtension( Filename ) ) )
 			{
 				// Log the error message and fail the import.
 				FFbxImporter->AddTokenizedErrorMessage(FTokenizedMessage::Create(EMessageSeverity::Error, LOCTEXT("FBXImport_ParseFailed", "FBX file parsing failed.")));
diff --git a/Engine/Source/Editor/UnrealEd/Private/FbxOptionWindow.h b/Engine/Source/Editor/UnrealEd/Private/FbxOptionWindow.h
--- a/Engine/Source/Editor/UnrealEd/Private/FbxOptionWindow.h
+++ b/Engine/Source/Editor/UnrealEd/Private/FbxOptionWindow.h
@@ -13,12 +13,14 @@ public:
 		, _WidgetWindow()
 		, _FullPath()
 		, _ForcedImportType()
+		, _IsObjFormat(false)
 		{}
 
 		SLATE_ARGUMENT( UFbxImportUI*, ImportUI )
 		SLATE_ARGUMENT( TSharedPtr<SWindow>, WidgetWindow )
 		SLATE_ARGUMENT( FString, FullPath )
 		SLATE_ARGUMENT( TOptional<EFBXImportType>, ForcedImportType )
+		SLATE_ARGUMENT( bool, IsObjFormat )
 	SLATE_END_ARGS()
 
 public:
@@ -81,19 +83,6 @@ public:
 	{}
 		
 private:
-	UFbxImportUI*	ImportUI;
-	bool			bShouldImport;
-	bool			bShouldImportAll;
-	bool			bReimport;
-	FString			ErrorMessage;
-	SVerticalBox::FSlot* CustomBox;
-	TWeakPtr< SWindow > WidgetWindow;
-	TSharedPtr< SButton > ImportButton;
-	TArray<FName> StaticMeshLODGroupNames;
-	TArray<FText> StaticMeshLODGroupDisplayNames;
-	TArray< TSharedPtr<FString> > StaticMeshLODGroups;
-	TArray< TSharedPtr<EFBXNormalImportMethod> > NormalImportOptions;
-
 	// sub options
 	TSharedRef<SWidget> ConstructStaticMeshBasic();
 	TSharedRef<SWidget> ConstructStaticMeshAdvanced();
@@ -162,4 +151,19 @@ private:
 	bool CanImport() const;
 	bool ShouldShowPhysicsAssetPicker() const;
 	EFBXNormalImportMethod GetCurrentNormalImportMethod() const;
+
+private:
+	UFbxImportUI*	ImportUI;
+	FString			ErrorMessage;
+	SVerticalBox::FSlot* CustomBox;
+	TWeakPtr< SWindow > WidgetWindow;
+	TSharedPtr< SButton > ImportButton;
+	TArray<FName> StaticMeshLODGroupNames;
+	TArray<FText> StaticMeshLODGroupDisplayNames;
+	TArray< TSharedPtr<FString> > StaticMeshLODGroups;
+	TArray< TSharedPtr<EFBXNormalImportMethod> > NormalImportOptions;
+	bool			bShouldImport;
+	bool			bShouldImportAll;
+	bool			bForceImportType;
+	bool			bIsObjFormat;
 };
diff --git a/Engine/Source/Editor/UnrealEd/Private/SkeletalMeshEdit.cpp b/Engine/Source/Editor/UnrealEd/Private/SkeletalMeshEdit.cpp
--- a/Engine/Source/Editor/UnrealEd/Private/SkeletalMeshEdit.cpp
+++ b/Engine/Source/Editor/UnrealEd/Private/SkeletalMeshEdit.cpp
@@ -25,7 +25,7 @@ UAnimSequence * UEditorEngine::ImportFbxAnimation( USkeleton * Skeleton, UObject
 
 	const bool bPrevImportMorph = FFbxImporter->ImportOptions->bImportMorph;
 	FFbxImporter->ImportOptions->bImportMorph = bImportMorphTracks;
-	if ( !FFbxImporter->ImportFromFile( InFilename ) )
+	if ( !FFbxImporter->ImportFromFile( InFilename, FPaths::GetExtension( InFilename ) ) )
 	{
 		// Log the error message and fail the import.
 		FFbxImporter->FlushToTokenizedErrorMessage(EMessageSeverity::Error);
@@ -118,7 +118,7 @@ bool UEditorEngine::ReimportFbxAnimation( USkeleton * Skeleton, UAnimSequence *
 		FbxImporter->ImportOptions->ResetForReimportAnimation();	
 	}
 
-	if ( !FbxImporter->ImportFromFile( InFilename ) )
+	if ( !FbxImporter->ImportFromFile( InFilename, FPaths::GetExtension( InFilename ) ) )
 	{
 		// Log the error message and fail the import.
 		FbxImporter->FlushToTokenizedErrorMessage(EMessageSeverity::Error);
