#UE4 - Fix reimport not working on OBJ files

diff --git a/Engine/Source/Editor/UnrealEd/Private/Factories/EditorFactories.cpp b/Engine/Source/Editor/UnrealEd/Private/Factories/EditorFactories.cpp
--- a/Engine/Source/Editor/UnrealEd/Private/Factories/EditorFactories.cpp
+++ b/Engine/Source/Editor/UnrealEd/Private/Factories/EditorFactories.cpp
@@ -5652,9 +5652,9 @@ EReimportResult::Type UReimportFbxStaticMeshFactory::Reimport( UObject* Obj )
 	{
 		const FString Filename = FReimportManager::ResolveImportFilename(ImportData->SourceFilePath, Mesh);
 		const FString FileExtension = FPaths::GetExtension(Filename);
-		const bool bIsFBX = FCString::Stricmp(*FileExtension, TEXT("FBX")) == 0;
+		const bool bIsValidFile = FileExtension.Equals( TEXT("fbx"), ESearchCase::IgnoreCase ) || FileExtension.Equals( "obj",  ESearchCase::IgnoreCase );
 
-		if ( !bIsFBX )
+		if ( !bIsValidFile )
 		{
 			return EReimportResult::Failed;
 		}
@@ -5677,7 +5677,7 @@ EReimportResult::Type UReimportFbxStaticMeshFactory::Reimport( UObject* Obj )
 
 		CurrentFilename = Filename;
 
-		if ( FFbxImporter->ImportFromFile( *Filename ) )
+		if ( FFbxImporter->ImportFromFile( *Filename, FPaths::GetExtension( Filename ) ) )
 		{
 			if (FFbxImporter->ReimportStaticMesh(Mesh, ImportData))
 			{
@@ -5831,7 +5831,7 @@ EReimportResult::Type UReimportFbxSkeletalMeshFactory::Reimport( UObject* Obj )
 		}
 		CurrentFilename = Filename;
 
-		if ( FFbxImporter->ImportFromFile( *Filename ) )
+		if ( FFbxImporter->ImportFromFile( *Filename, FPaths::GetExtension( Filename ) ) )
 		{
 			if ( FFbxImporter->ReimportSkeletalMesh(SkeletalMesh, ImportData) )
 			{
