#UE4 - Disabled toolbar customization experimental feature as it isn't even remotely usable.

diff --git a/Engine/Source/Editor/UnrealEd/Classes/Settings/EditorExperimentalSettings.h b/Engine/Source/Editor/UnrealEd/Classes/Settings/EditorExperimentalSettings.h
--- a/Engine/Source/Editor/UnrealEd/Classes/Settings/EditorExperimentalSettings.h
+++ b/Engine/Source/Editor/UnrealEd/Classes/Settings/EditorExperimentalSettings.h
@@ -53,7 +53,7 @@ public:
 	TEnumAsByte<EConsoleForGamepadLabels::Type> ConsoleForGamepadLabels;
 
 	/** Allows for customization of toolbars and menus throughout the editor */
-	UPROPERTY(EditAnywhere, config, Category=UserInterface, meta=(DisplayName="Toolbar Customization"))
+	UPROPERTY(config, Category=UserInterface, meta=(DisplayName="Toolbar Customization"))
 	bool bToolbarCustomization;
 
 	/** Break on Exceptions allows you to trap Access Nones and other exceptional events in Blueprints. */
