#ttp 331620 - EDITOR: The SlateWidgetStyle assets are missing the ôSaveö and ôFind in CBö button in the toolbar
#proj UE4
#branch UE4
#summary You can now provide a function to get the objects you want to show in the details view when creating a FSimpleAssetEditor
#extra FAssetTypeActions_SlateWidgetStyle::OpenAssetEditor was passing in the sub-objects that it wanted to show in the details view, however, these objects were not assets, and so the asset toolbar wasn't added for them.
#change I added a way to provide a callback to FSimpleAssetEditor that allows you to extract an array of objects to edit in the details view from the list of assets provided to the function. This lets the toolbar work (as we're editing assets now), as well as keeping the details view UI working as it did before.
#extra If no callback is provided, the assets are themselves edited in the details view.
#extra This also fixes this crash from AnswerHub because you always have the asset to edit, even if the sub-object is invalid (you just get an empty details view): https://answers.unrealengine.com/questions/21281/question-bug-copying-style-assets-from-strategy-ga.html

diff --git a/Engine/Source/Developer/AssetTools/Private/AssetTypeActions/AssetTypeActions_SlateWidgetStyle.cpp b/Engine/Source/Developer/AssetTools/Private/AssetTypeActions/AssetTypeActions_SlateWidgetStyle.cpp
--- a/Engine/Source/Developer/AssetTools/Private/AssetTypeActions/AssetTypeActions_SlateWidgetStyle.cpp
+++ b/Engine/Source/Developer/AssetTools/Private/AssetTypeActions/AssetTypeActions_SlateWidgetStyle.cpp
@@ -24,17 +24,24 @@ void FAssetTypeActions_SlateWidgetStyle::GetActions( const TArray<UObject*>& InO
 
 void FAssetTypeActions_SlateWidgetStyle::OpenAssetEditor( const TArray<UObject*>& InObjects, TSharedPtr<IToolkitHost> EditWithinLevelEditor )
 {
-	TArray<UObject*> SubObjects;
-	for (auto ObjIt = InObjects.CreateConstIterator(); ObjIt; ++ObjIt)
+	struct Local
 	{
-		auto Style = Cast< USlateWidgetStyleAsset >(*ObjIt);
-		if ( Style != NULL && Style->CustomStyle != NULL )
+		static TArray<UObject*> GetSubObjects(const TArray<UObject*>& InObjects)
 		{
-			SubObjects.Add( Style->CustomStyle );
+			TArray<UObject*> SubObjects;
+			for(UObject* Object : InObjects)
+			{
+				auto Style = Cast<USlateWidgetStyleAsset>(Object);
+				if(Style && Style->CustomStyle)
+				{
+					SubObjects.Add(Style->CustomStyle);
+				}
+			}
+			return SubObjects;
 		}
-	}
+	};
 
-	FSimpleAssetEditor::CreateEditor(EToolkitMode::Standalone, EditWithinLevelEditor, SubObjects);
+	FSimpleAssetEditor::CreateEditor(EToolkitMode::Standalone, EditWithinLevelEditor, InObjects, FSimpleAssetEditor::FGetDetailsViewObjects::CreateStatic(&Local::GetSubObjects));
 }
 
 void FAssetTypeActions_SlateWidgetStyle::ExecuteEdit(TArray<TWeakObjectPtr< USlateWidgetStyleAsset >> Styles)
diff --git a/Engine/Source/Editor/UnrealEd/Private/Toolkits/SimpleAssetEditor.cpp b/Engine/Source/Editor/UnrealEd/Private/Toolkits/SimpleAssetEditor.cpp
--- a/Engine/Source/Editor/UnrealEd/Private/Toolkits/SimpleAssetEditor.cpp
+++ b/Engine/Source/Editor/UnrealEd/Private/Toolkits/SimpleAssetEditor.cpp
@@ -39,7 +39,7 @@ FSimpleAssetEditor::~FSimpleAssetEditor()
 }
 
 
-void FSimpleAssetEditor::InitEditor( const EToolkitMode::Type Mode, const TSharedPtr< class IToolkitHost >& InitToolkitHost, const TArray<UObject*>& ObjectsToEdit )
+void FSimpleAssetEditor::InitEditor( const EToolkitMode::Type Mode, const TSharedPtr< class IToolkitHost >& InitToolkitHost, const TArray<UObject*>& ObjectsToEdit, FGetDetailsViewObjects GetDetailsViewObjects )
 {
 	const bool bIsUpdatable = false;
 	const bool bAllowFavorites = true;
@@ -83,16 +83,19 @@ void FSimpleAssetEditor::InitEditor( const EToolkitMode::Type Mode, const TShare
 		PropertiesTab = SpawnToolkitTab( PropertiesTabId, TabInitializationPayload, EToolkitTabSpot::Details );
 	}*/
 	
+	// Get the list of objects to edit the details of
+	const TArray<UObject*> ObjectsToEditInDetailsView = ( GetDetailsViewObjects.IsBound() ) ? GetDetailsViewObjects.Execute( ObjectsToEdit ) : ObjectsToEdit;
+
 	// Ensure all objects are transactable for undo/redo in the details panel
-	for( int32 ObjectIndex = 0; ObjectIndex < ObjectsToEdit.Num(); ++ObjectIndex )
+	for( UObject* ObjectToEditInDetailsView : ObjectsToEditInDetailsView )
 	{
-		ObjectsToEdit[ObjectIndex]->SetFlags(RF_Transactional);
+		ObjectToEditInDetailsView->SetFlags(RF_Transactional);
 	}
 
 	if( DetailsView.IsValid() )
 	{
 		// Make sure details window is pointing to our object
-		DetailsView->SetObjects( ObjectsToEdit );
+		DetailsView->SetObjects( ObjectsToEditInDetailsView );
 	}
 }
 
@@ -194,21 +197,21 @@ FString FSimpleAssetEditor::GetWorldCentricTabPrefix() const
 	return LOCTEXT("WorldCentricTabPrefix", "Generic Asset ").ToString();
 }
 
-TSharedRef<FSimpleAssetEditor> FSimpleAssetEditor::CreateEditor( const EToolkitMode::Type Mode, const TSharedPtr< IToolkitHost >& InitToolkitHost, UObject* ObjectToEdit )
+TSharedRef<FSimpleAssetEditor> FSimpleAssetEditor::CreateEditor( const EToolkitMode::Type Mode, const TSharedPtr< IToolkitHost >& InitToolkitHost, UObject* ObjectToEdit, FGetDetailsViewObjects GetDetailsViewObjects )
 {
 	TSharedRef< FSimpleAssetEditor > NewEditor( new FSimpleAssetEditor() );
 
 	TArray<UObject*> ObjectsToEdit;
 	ObjectsToEdit.Add( ObjectToEdit );
-	NewEditor->InitEditor( Mode, InitToolkitHost, ObjectsToEdit );
+	NewEditor->InitEditor( Mode, InitToolkitHost, ObjectsToEdit, GetDetailsViewObjects );
 
 	return NewEditor;
 }
 
-TSharedRef<FSimpleAssetEditor> FSimpleAssetEditor::CreateEditor( const EToolkitMode::Type Mode, const TSharedPtr< IToolkitHost >& InitToolkitHost, const TArray<UObject*>& ObjectsToEdit )
+TSharedRef<FSimpleAssetEditor> FSimpleAssetEditor::CreateEditor( const EToolkitMode::Type Mode, const TSharedPtr< IToolkitHost >& InitToolkitHost, const TArray<UObject*>& ObjectsToEdit, FGetDetailsViewObjects GetDetailsViewObjects )
 {
 	TSharedRef< FSimpleAssetEditor > NewEditor( new FSimpleAssetEditor() );
-	NewEditor->InitEditor( Mode, InitToolkitHost, ObjectsToEdit );
+	NewEditor->InitEditor( Mode, InitToolkitHost, ObjectsToEdit, GetDetailsViewObjects );
 	return NewEditor;
 }
 
diff --git a/Engine/Source/Editor/UnrealEd/Public/Toolkits/SimpleAssetEditor.h b/Engine/Source/Editor/UnrealEd/Public/Toolkits/SimpleAssetEditor.h
--- a/Engine/Source/Editor/UnrealEd/Public/Toolkits/SimpleAssetEditor.h
+++ b/Engine/Source/Editor/UnrealEd/Public/Toolkits/SimpleAssetEditor.h
@@ -6,11 +6,12 @@
 #include "Toolkits/AssetEditorToolkit.h"
 #include "Editor/PropertyEditor/Public/PropertyEditorModule.h"
 
-
-
 class UNREALED_API FSimpleAssetEditor : public FAssetEditorToolkit
 {
 public:
+	/** Delegate that, given an array of assets, returns an array of objects to use in the details view of an FSimpleAssetEditor */
+	DECLARE_DELEGATE_RetVal_OneParam(TArray<UObject*>, FGetDetailsViewObjects, const TArray<UObject*>&);
+
 	virtual void RegisterTabSpawners(const TSharedRef<class FTabManager>& TabManager) OVERRIDE;
 	virtual void UnregisterTabSpawners(const TSharedRef<class FTabManager>& TabManager) OVERRIDE;
 
@@ -20,9 +21,10 @@ public:
 	 *
 	 * @param	Mode					Asset editing mode for this editor (standalone or world-centric)
 	 * @param	InitToolkitHost			When Mode is WorldCentric, this is the level editor instance to spawn this editor within
-	 * @param	ObjectToEdit			The object to edit
+	 * @param	ObjectsToEdit			The object to edit
+	 * @param	GetDetailsViewObjects	If bound, a delegate to get the array of objects to use in the details view; uses ObjectsToEdit if not bound
 	 */
-	void InitEditor( const EToolkitMode::Type Mode, const TSharedPtr< class IToolkitHost >& InitToolkitHost, const TArray<UObject*>& ObjectToEdit );
+	void InitEditor( const EToolkitMode::Type Mode, const TSharedPtr< class IToolkitHost >& InitToolkitHost, const TArray<UObject*>& ObjectsToEdit, FGetDetailsViewObjects GetDetailsViewObjects );
 
 	/** Destructor */
 	virtual ~FSimpleAssetEditor();
@@ -54,7 +56,7 @@ private:
 	static const FName PropertiesTabId;
 
 public:
-	static TSharedRef<FSimpleAssetEditor> CreateEditor( const EToolkitMode::Type Mode, const TSharedPtr< IToolkitHost >& InitToolkitHost, UObject* ObjectToEdit );
+	static TSharedRef<FSimpleAssetEditor> CreateEditor( const EToolkitMode::Type Mode, const TSharedPtr< IToolkitHost >& InitToolkitHost, UObject* ObjectToEdit, FGetDetailsViewObjects GetDetailsViewObjects = FGetDetailsViewObjects() );
 
-	static TSharedRef<FSimpleAssetEditor> CreateEditor( const EToolkitMode::Type Mode, const TSharedPtr< IToolkitHost >& InitToolkitHost, const TArray<UObject*>& ObjectsToEdit );
+	static TSharedRef<FSimpleAssetEditor> CreateEditor( const EToolkitMode::Type Mode, const TSharedPtr< IToolkitHost >& InitToolkitHost, const TArray<UObject*>& ObjectsToEdit, FGetDetailsViewObjects GetDetailsViewObjects = FGetDetailsViewObjects() );
 };
