TTP #333047 Mac: UE4Game: Warning: /usr/bin/atos is moving and will be removed from a future OS X release.

diff --git a/Engine/Source/Editor/UnrealEd/Private/SourceCodeNavigation.cpp b/Engine/Source/Editor/UnrealEd/Private/SourceCodeNavigation.cpp
--- a/Engine/Source/Editor/UnrealEd/Private/SourceCodeNavigation.cpp
+++ b/Engine/Source/Editor/UnrealEd/Private/SourceCodeNavigation.cpp
@@ -512,7 +512,7 @@ void FSourceCodeNavigationImpl::NavigateToFunctionSource( ISourceCodeAccessor* S
 							if(FunctionSymbolName == SymbolName)
 							{
 								uint64 Address = SymbolEntry.n_value;
-								FString AtoSCommand = FString::Printf(TEXT("-d -o %s 0x%x"), *FullModulePath, Address);
+								FString AtoSCommand = FString::Printf(TEXT("-nowarning -d -o %s 0x%x"), *FullModulePath, Address);
 								int32 ReturnCode = 0;
 								FString Results;
 								FPlatformProcess::ExecProcess( TEXT("/usr/bin/atos"), *AtoSCommand, &ReturnCode, &Results, NULL );
diff --git a/Engine/Source/Runtime/Core/Private/Apple/ApplePlatformStackWalk.cpp b/Engine/Source/Runtime/Core/Private/Apple/ApplePlatformStackWalk.cpp
--- a/Engine/Source/Runtime/Core/Private/Apple/ApplePlatformStackWalk.cpp
+++ b/Engine/Source/Runtime/Core/Private/Apple/ApplePlatformStackWalk.cpp
@@ -178,7 +178,7 @@ void FApplePlatformStackWalk::ProgramCounterToSymbolInfo( uint64 ProgramCounter,
 			FCStringAnsi::Sprintf(FBase, "%p", (void*)DylibInfo.dli_fbase);
 			FCStringAnsi::Sprintf(Address, "%p", (void*)ProgramCounter);
 			
-			execl("/usr/bin/atos", "-d", "-l", FBase, "-o", DylibInfo.dli_fname, Address, NULL);
+			execl("/usr/bin/atos", "-nowarning", "-d", "-l", FBase, "-o", DylibInfo.dli_fname, Address, NULL);
 			close(FileDesc[1]);
 		}
 		else
