TTP #333274: Making it so it is safe to "promote to variable" from self pin links.

diff --git a/Engine/Source/Editor/UnrealEd/Private/Kismet2/BlueprintEditorUtils.cpp b/Engine/Source/Editor/UnrealEd/Private/Kismet2/BlueprintEditorUtils.cpp
--- a/Engine/Source/Editor/UnrealEd/Private/Kismet2/BlueprintEditorUtils.cpp
+++ b/Engine/Source/Editor/UnrealEd/Private/Kismet2/BlueprintEditorUtils.cpp
@@ -2972,9 +2972,15 @@ bool FBlueprintEditorUtils::AddMemberVariable(UBlueprint* Blueprint, const FName
 	}
 	else if ((NewVarType.PinCategory == K2Schema->PC_Object) || (NewVarType.PinCategory == K2Schema->PC_Interface))
 	{
+		if (NewVar.VarType.PinSubCategory == K2Schema->PSC_Self)
+		{
+			NewVar.VarType.PinSubCategory.Empty();
+			NewVar.VarType.PinSubCategoryObject = *Blueprint->GeneratedClass;
+		}
+
 		// if it's a PC_Object, then it should have an associated UClass object
-		check(NewVarType.PinSubCategoryObject.IsValid());
-		const UClass* ClassObject = Cast<UClass>(NewVarType.PinSubCategoryObject.Get());
+		check(NewVar.VarType.PinSubCategoryObject.IsValid());
+		const UClass* ClassObject = Cast<UClass>(NewVar.VarType.PinSubCategoryObject.Get());
 		check(ClassObject != NULL);
 
 		if (ClassObject->IsChildOf(AActor::StaticClass()))
