Fixed importing to handle NaN issue.

diff --git a/Engine/Source/Editor/UnrealEd/Private/SkeletalMeshEdit.cpp b/Engine/Source/Editor/UnrealEd/Private/SkeletalMeshEdit.cpp
--- a/Engine/Source/Editor/UnrealEd/Private/SkeletalMeshEdit.cpp
+++ b/Engine/Source/Editor/UnrealEd/Private/SkeletalMeshEdit.cpp
@@ -807,9 +807,16 @@ namespace AnimationTransformDebug
 		TArray<FTransform>	SourceGlobalTransform;
 		TArray<FTransform>	SourceParentGlobalTransform;
 
-		FAnimationTransformDebugData(int32 InTrackIndex, int32 InBoneIndex, FName InBoneName)
-			: TrackIndex(InTrackIndex), BoneIndex(InBoneIndex), BoneName(InBoneName)
+		FAnimationTransformDebugData()
+			: TrackIndex(INDEX_NONE), BoneIndex(INDEX_NONE), BoneName(NAME_None)
 		{}
+
+		void SetTrackData(int32 InTrackIndex, int32 InBoneIndex, FName InBoneName)
+		{
+			TrackIndex = InTrackIndex;
+			BoneIndex = InBoneIndex;
+			BoneName = InBoneName;
+		}
 	};
 
 	void OutputAnimationTransformDebugData(TArray<AnimationTransformDebug::FAnimationTransformDebugData> &TransformDebugData, int32 TotalNumKeys, const FReferenceSkeleton & RefSkeleton)
@@ -970,19 +977,15 @@ bool UnFbx::FFbxImporter::ImportAnimation(USkeleton * Skeleton, UAnimSequence *
 
 		if (BoneTreeIndex!=INDEX_NONE)
 		{
-			//add new track
-			int32 NewTrackIdx = DestSeq->RawAnimationData.AddZeroed(1);
-			DestSeq->AnimationTrackNames.Add(BoneName);
-
-			AnimationTransformDebug::FAnimationTransformDebugData NewDebugData(NewTrackIdx, BoneTreeIndex, BoneName);
-			FRawAnimSequenceTrack& RawTrack = DestSeq->RawAnimationData[NewTrackIdx];
-			// add mapping to skeleton bone track
-			DestSeq->TrackToSkeletonMapTable.Add(FTrackToSkeletonMap(BoneTreeIndex));
+			bool bSuccess = true;
 
+			FRawAnimSequenceTrack RawTrack;
 			RawTrack.PosKeys.Empty();
 			RawTrack.RotKeys.Empty();
 			RawTrack.ScaleKeys.Empty();
 
+			AnimationTransformDebug::FAnimationTransformDebugData NewDebugData;
+
 			FbxNode* Link = SortedLinks[SourceTrackIdx];
 			FbxNode * LinkParent = Link->GetParent();
 			
@@ -991,8 +994,25 @@ bool UnFbx::FFbxImporter::ImportAnimation(USkeleton * Skeleton, UAnimSequence *
 			{
 				// save global trasnform
 				FbxAMatrix GlobalMatrix = Link->EvaluateGlobalTransform(CurTime);
+				// we'd like to verify this before going to Transform. 
+				// currently transform has tons of NaN check, so it will crash there
+				FMatrix GlobalUEMatrix = Converter.ConvertMatrix(GlobalMatrix);
+				if (GlobalUEMatrix.ContainsNaN())
+				{
+					bSuccess = false;
+					AddTokenizedErrorMessage(FTokenizedMessage::Create(EMessageSeverity::Error, FText::Format(LOCTEXT("Error_InvalidTransform",
+						"Track {0} contains invalid transform. Could not import the track."), FText::FromName(BoneName))));
+					break;
+				}
+
 				FTransform GlobalTransform =  Converter.ConvertTransform(GlobalMatrix);
-				
+				if (GlobalTransform.ContainsNaN())
+				{
+					bSuccess = false;
+					AddTokenizedErrorMessage(FTokenizedMessage::Create(EMessageSeverity::Error, FText::Format(LOCTEXT("Error_InvalidUnrealTransform",
+										"Track {0} did not yeild valid transform. Please report this to animation team."), FText::FromName(BoneName))));
+					break;
+				}
 				// debug data
 				NewDebugData.SourceGlobalTransform.Add(GlobalTransform);
 
@@ -1018,7 +1038,13 @@ bool UnFbx::FFbxImporter::ImportAnimation(USkeleton * Skeleton, UAnimSequence *
 					NewDebugData.SourceParentGlobalTransform.Add(FTransform::Identity);
 				}
 
-				ensure(LocalTransform.ContainsNaN() == false);
+				if (LocalTransform.ContainsNaN())
+				{
+					bSuccess = false;
+					AddTokenizedErrorMessage(FTokenizedMessage::Create(EMessageSeverity::Error, FText::Format(LOCTEXT("Error_InvalidUnrealLocalTransform",
+										"Track {0} did not yeild valid local transform. Please report this to animation team."), FText::FromName(BoneName))));
+					break;
+				}
 
 				RawTrack.ScaleKeys.Add(LocalTransform.GetScale3D());
 				RawTrack.PosKeys.Add(LocalTransform.GetTranslation());
@@ -1028,7 +1054,18 @@ bool UnFbx::FFbxImporter::ImportAnimation(USkeleton * Skeleton, UAnimSequence *
 				++NumKeysForTrack;
 			}
 
-			TransformDebugData.Add(NewDebugData);
+			if (bSuccess)
+			{
+				//add new track
+				int32 NewTrackIdx = DestSeq->RawAnimationData.Add(RawTrack);
+				DestSeq->AnimationTrackNames.Add(BoneName);
+
+				NewDebugData.SetTrackData(NewTrackIdx, BoneTreeIndex, BoneName);
+
+				// add mapping to skeleton bone track
+				DestSeq->TrackToSkeletonMapTable.Add(FTrackToSkeletonMap(BoneTreeIndex));
+				TransformDebugData.Add(NewDebugData);
+			}
 		}
 
 		TotalNumKeys = FMath::Max( TotalNumKeys, NumKeysForTrack );
