#ttp 333088 - EDITOR: BSP: CRASH: Undoing a convert brush can lead to a crash
#branch UE4
#proj Editor.UnrealEd
#summary Disallowed conversion of brushes to other types.
#change Amended the condition in FSelectedActorInfo::HasConvertableAsset so that brushes are deemed to be "not convertable".

diff --git a/Engine/Source/Editor/UnrealEd/Public/AssetSelection.h b/Engine/Source/Editor/UnrealEd/Public/AssetSelection.h
--- a/Engine/Source/Editor/UnrealEd/Public/AssetSelection.h
+++ b/Engine/Source/Editor/UnrealEd/Public/AssetSelection.h
@@ -140,7 +140,7 @@ struct FSelectedActorInfo
 	 */
 	bool HasConvertableAsset() const
 	{
-		return NumSelected && !bHaveBuilderBrush;
+		return NumSelected && !bHaveBuilderBrush && !bHaveBrush;
 	}
 
 };
