#ttp 331969 - EDITOR: Vertices on geometry show up in game view
#summary Geometry vertices are no longer visible in "Game View" viewports
#proj Editor
#branch UE4

diff --git a/Engine/Source/Editor/UnrealEd/Private/EditorModes.cpp b/Engine/Source/Editor/UnrealEd/Private/EditorModes.cpp
--- a/Engine/Source/Editor/UnrealEd/Private/EditorModes.cpp
+++ b/Engine/Source/Editor/UnrealEd/Private/EditorModes.cpp
@@ -631,7 +631,8 @@ void FEdMode::Render(const FSceneView* View,FViewport* Viewport,FPrimitiveDrawIn
 		}
 	}
 
-	if (GEditorModeTools().ShouldDrawBrushVertices())
+	const bool bIsInGameView = !Viewport->GetClient() || Viewport->GetClient()->IsInGameView();
+	if (GEditorModeTools().ShouldDrawBrushVertices() && !bIsInGameView)
 	{
 		UTexture2D* VertexTexture = GetVertexTexture();
 		const float TextureSizeX = VertexTexture->GetSizeX() * 0.170f;
diff --git a/Engine/Source/Editor/UnrealEd/Public/EditorViewportClient.h b/Engine/Source/Editor/UnrealEd/Public/EditorViewportClient.h
--- a/Engine/Source/Editor/UnrealEd/Public/EditorViewportClient.h
+++ b/Engine/Source/Editor/UnrealEd/Public/EditorViewportClient.h
@@ -837,7 +837,10 @@ public:
 
 	void SetGameView(bool bGameViewEnable);
 
-	bool IsInGameView() const { return bInGameViewMode; }
+	/** 
+	 * Returns true if this viewport is excluding non-game elements from its display
+	 */
+	virtual bool IsInGameView() const OVERRIDE { return bInGameViewMode; }
 
 	/**
 	 * Aspect ratio bar display settings
diff --git a/Engine/Source/Runtime/Engine/Public/UnrealClient.h b/Engine/Source/Runtime/Engine/Public/UnrealClient.h
--- a/Engine/Source/Runtime/Engine/Public/UnrealClient.h
+++ b/Engine/Source/Runtime/Engine/Public/UnrealClient.h
@@ -682,6 +682,11 @@ public:
 	 */
 	virtual bool IsOrtho() const { return false; }
 
+	/** 
+	 * Returns true if this viewport is excluding non-game elements from its display
+	 */
+	virtual bool IsInGameView() const { return false; }
+
 	/**
  	 * Sets GWorld to the appropriate world for this client
 	 * 
