#ttp 331634 - EDITOR: PERSONA: Unable to change camera speed in Persona
#summary Fixed not being able to set camera movement speed in toolkit viewports
#change Moved anything to do with camera speed in the level editor viewport settings into the level editor viewport client and added a transient camera speed setting to all other viewports.
#proj Editor
#branch UE4

diff --git a/Engine/Source/Editor/GeometryMode/Private/GeometryModifiers.cpp b/Engine/Source/Editor/GeometryMode/Private/GeometryModifiers.cpp
--- a/Engine/Source/Editor/GeometryMode/Private/GeometryModifiers.cpp
+++ b/Engine/Source/Editor/GeometryMode/Private/GeometryModifiers.cpp
@@ -587,7 +587,7 @@ bool UGeomModifier_Edit::InputDelta(FLevelEditorViewportClient* InViewportClient
 		// Only apply camera speed modifiers to the drag if we aren't zooming in an ortho viewport.
 		if( !InViewportClient->IsOrtho() || !(InViewport->KeyState(EKeys::LeftMouseButton) && InViewport->KeyState(EKeys::RightMouseButton)) )
 		{
-			const float CameraSpeed = InViewportClient->GetCameraSpeed(GetDefault<ULevelEditorViewportSettings>()->CameraSpeed);
+			const float CameraSpeed = InViewportClient->GetCameraSpeed();
 			CameraDelta *= CameraSpeed;
 		}
 
diff --git a/Engine/Source/Editor/UnrealEd/Private/EditorViewportClient.cpp b/Engine/Source/Editor/UnrealEd/Private/EditorViewportClient.cpp
--- a/Engine/Source/Editor/UnrealEd/Private/EditorViewportClient.cpp
+++ b/Engine/Source/Editor/UnrealEd/Private/EditorViewportClient.cpp
@@ -143,15 +143,34 @@ FMatrix FViewportCameraTransform::ComputeOrbitMatrix() const
 /**The Maximum Mouse/Camera Speeds Setting supported */
 const uint32 FEditorViewportClient::MaxCameraSpeeds = 8;
 
-/**
+float FEditorViewportClient::GetCameraSpeed() const
+{
+	return GetCameraSpeed(GetCameraSpeedSetting());
+}
 
- * Static: Utility to get the actual mouse/camera speed (multiplier) based on the passed in setting.
- *
- * @param SpeedSetting	The desired speed setting
- */
 float FEditorViewportClient::GetCameraSpeed(int32 SpeedSetting) const
 {
-	return 1.0f;
+	//previous mouse speed values were as follows...
+	//(note: these were previously all divided by 4 when used be the viewport)
+	//#define MOVEMENTSPEED_SLOW			4	~ 1
+	//#define MOVEMENTSPEED_NORMAL			12	~ 3
+	//#define MOVEMENTSPEED_FAST			32	~ 8
+	//#define MOVEMENTSPEED_VERYFAST		64	~ 16
+
+	const int32 SpeedToUse = FMath::Clamp<int32>(SpeedSetting, 1, MaxCameraSpeeds);
+	const float Speed[] = { 0.03125f, 0.09375f, 0.33f, 1.f, 3.f, 8.f, 16.f, 32.f };
+
+	return Speed[SpeedToUse - 1];
+}
+
+void FEditorViewportClient::SetCameraSpeedSetting(int32 SpeedSetting)
+{
+	CameraSpeedSetting = SpeedSetting;
+}
+
+int32 FEditorViewportClient::GetCameraSpeedSetting() const
+{
+	return CameraSpeedSetting;
 }
 
 float const FEditorViewportClient::SafePadding = 0.075f;
@@ -218,6 +237,7 @@ FEditorViewportClient::FEditorViewportClient(FPreviewScene* InPreviewScene)
 	, bShowStats(false)
 	, PreviewScene(InPreviewScene)
 	, bCameraLock(false)
+	, CameraSpeedSetting(4)
 {
 
 	ViewState.Allocate();
@@ -1167,7 +1187,7 @@ void FEditorViewportClient::UpdateCameraMovement( float DeltaTime )
 
 		// We'll combine the regular camera speed scale (controlled by viewport toolbar setting) with
 		// the flight camera speed scale (controlled by mouse wheel).
-		const float CameraSpeed = GetCameraSpeed(GetDefault<ULevelEditorViewportSettings>()->CameraSpeed);
+		const float CameraSpeed = GetCameraSpeed();
 		const float FinalCameraSpeedScale = FlightCameraSpeedScale * CameraSpeed;
 
 		// Only allow FOV recoil if flight camera mode is currently inactive.
@@ -1468,7 +1488,7 @@ void FEditorViewportClient::UpdateMouseDelta()
 							{
 								if( !IsOrtho())
 								{
-									const float CameraSpeed = GetCameraSpeed(GetDefault<ULevelEditorViewportSettings>()->CameraSpeed);
+									const float CameraSpeed = GetCameraSpeed();
 									Drag *= CameraSpeed;
 								}
 								MoveViewportCamera( Drag, Rot );
diff --git a/Engine/Source/Editor/UnrealEd/Private/LevelEditorViewport.cpp b/Engine/Source/Editor/UnrealEd/Private/LevelEditorViewport.cpp
--- a/Engine/Source/Editor/UnrealEd/Private/LevelEditorViewport.cpp
+++ b/Engine/Source/Editor/UnrealEd/Private/LevelEditorViewport.cpp
@@ -4416,21 +4416,14 @@ bool FLevelEditorViewportClient::IsAspectRatioConstrained() const
 	}
 }
 
-float FLevelEditorViewportClient::GetCameraSpeed(int32 SpeedSetting) const
+int32 FLevelEditorViewportClient::GetCameraSpeedSetting() const
 {
-	
-	//previous mouse speed values were as follows...
-	//(note: these were previously all divided by 4 when used be the viewport)
-	//#define MOVEMENTSPEED_SLOW			4	~ 1
-	//#define MOVEMENTSPEED_NORMAL			12	~ 3
-	//#define MOVEMENTSPEED_FAST			32	~ 8
-	//#define MOVEMENTSPEED_VERYFAST		64	~ 16
-
-
-	const int32 SpeedToUse = FMath::Clamp<int32>(SpeedSetting, 1, MaxCameraSpeeds);
-	const float Speed[] = {0.03125f, 0.09375f, 0.33f, 1.f, 3.f, 8.f, 16.f, 32.f};
+	return GetDefault<ULevelEditorViewportSettings>()->CameraSpeed;
+}
 
-	return Speed[SpeedToUse - 1];
+void FLevelEditorViewportClient::SetCameraSpeedSetting(int32 SpeedSetting)
+{
+	GetMutableDefault<ULevelEditorViewportSettings>()->CameraSpeed = SpeedSetting;
 }
 
 bool FLevelEditorViewportClient::OverrideHighResScreenshotCaptureRegion(FIntRect& OutCaptureRegion)
diff --git a/Engine/Source/Editor/UnrealEd/Private/STransformViewportToolbar.cpp b/Engine/Source/Editor/UnrealEd/Private/STransformViewportToolbar.cpp
--- a/Engine/Source/Editor/UnrealEd/Private/STransformViewportToolbar.cpp
+++ b/Engine/Source/Editor/UnrealEd/Private/STransformViewportToolbar.cpp
@@ -599,20 +599,36 @@ FText STransformViewportToolBar::GetScaleGridLabel() const
 
 FText STransformViewportToolBar::GetCameraSpeedLabel() const
 {
-	return FText::AsNumber( GetDefault<ULevelEditorViewportSettings>()->CameraSpeed );
+	auto ViewportPin = Viewport.Pin();
+	if (ViewportPin.IsValid() && ViewportPin->GetViewportClient().IsValid())
+	{
+		return FText::AsNumber( ViewportPin->GetViewportClient()->GetCameraSpeedSetting() );
+	}
+
+	return FText();
 }	
 
 float STransformViewportToolBar::GetCamSpeedSliderPosition() const
 {
-	const int32 SpeedSetting = GetDefault<ULevelEditorViewportSettings>()->CameraSpeed;	
-	float SliderPos = (SpeedSetting - 1) / ((float)FEditorViewportClient::MaxCameraSpeeds - 1);
+	float SliderPos = 0.f;
+
+	auto ViewportPin = Viewport.Pin();
+	if (ViewportPin.IsValid() && ViewportPin->GetViewportClient().IsValid())
+	{
+		SliderPos = (ViewportPin->GetViewportClient()->GetCameraSpeedSetting() - 1) / ((float)FEditorViewportClient::MaxCameraSpeeds - 1);
+	}
+
 	return SliderPos;
 }
 
 void STransformViewportToolBar::OnSetCamSpeed(float NewValue)
 {
-	int32 SpeedSetting = NewValue * ((float)FEditorViewportClient::MaxCameraSpeeds - 1) + 1;
-	GetMutableDefault<ULevelEditorViewportSettings>()->CameraSpeed = SpeedSetting;
+	auto ViewportPin = Viewport.Pin();
+	if (ViewportPin.IsValid() && ViewportPin->GetViewportClient().IsValid())
+	{
+		const int32 SpeedSetting = NewValue * ((float)FEditorViewportClient::MaxCameraSpeeds - 1) + 1;
+		ViewportPin->GetViewportClient()->SetCameraSpeedSetting(SpeedSetting);
+	}
 }
 
 /**
diff --git a/Engine/Source/Editor/UnrealEd/Public/EditorViewportClient.h b/Engine/Source/Editor/UnrealEd/Public/EditorViewportClient.h
--- a/Engine/Source/Editor/UnrealEd/Public/EditorViewportClient.h
+++ b/Engine/Source/Editor/UnrealEd/Public/EditorViewportClient.h
@@ -826,12 +826,24 @@ public:
 
 	void SetRequiredCursorOverride( bool WantOverride, EMouseCursor::Type RequiredCursor = EMouseCursor::Default ); 
 
-	/**
-	 * Utility to get the actually mouse/camera speed (multiplier) based on the passed in setting.
-	 *
-	 * @param SpeedSetting	The desired speed steeing
-	 */
-	virtual float GetCameraSpeed(int32 SpeedSetting) const;
+	/** Get the camera speed for this viewport */
+	float GetCameraSpeed() const;
+
+	/** Get the camera speed for this viewport based on the specified speed setting */
+	float GetCameraSpeed(int32 SpeedSetting) const;
+
+	/** Set the speed setting for the camera in this viewport */
+	virtual void SetCameraSpeedSetting(int32 SpeedSetting);
+
+	/** Get the camera speed setting for this viewport */
+	virtual int32 GetCameraSpeedSetting() const;
+
+protected:
+
+	/** Camera speed setting */
+	int32 CameraSpeedSetting;
+
+public:
 
 	void DrawBoundingBox(FBox &Box, FCanvas* InCanvas, const FSceneView* InView, const FViewport* InViewport, const FLinearColor& InColor, const bool bInDrawBracket, const FString &InLabelText) ;
 
diff --git a/Engine/Source/Editor/UnrealEd/Public/LevelEditorViewport.h b/Engine/Source/Editor/UnrealEd/Public/LevelEditorViewport.h
--- a/Engine/Source/Editor/UnrealEd/Public/LevelEditorViewport.h
+++ b/Engine/Source/Editor/UnrealEd/Public/LevelEditorViewport.h
@@ -171,7 +171,8 @@ public:
 	virtual FSceneInterface* GetScene() const OVERRIDE;
 	virtual FLinearColor GetBackgroundColor() const OVERRIDE;
 	virtual bool IsAspectRatioConstrained() const OVERRIDE;
-	virtual float GetCameraSpeed(int32 SpeedSetting) const OVERRIDE;
+	virtual int32 GetCameraSpeedSetting() const OVERRIDE;
+	virtual void SetCameraSpeedSetting(int32 SpeedSetting) OVERRIDE;
 
 	virtual bool OverrideHighResScreenshotCaptureRegion(FIntRect& OutCaptureRegion) OVERRIDE;
 
