Corrected tooltip text in Scene Outliner when attaching actors.
#ttp 333148 - Regression: EDITOR: Scene Outliner: When attaching actors the tooltip text is incorrectly displayed as "Attach %s to %s"
#branch UE4
#proj Editor.UnrealEd
#change Changed %s to {0} and {1} in FText::Format string.

diff --git a/Engine/Source/Editor/UnrealEd/Public/DragAndDrop/ActorDragDropGraphEdOp.h b/Engine/Source/Editor/UnrealEd/Public/DragAndDrop/ActorDragDropGraphEdOp.h
--- a/Engine/Source/Editor/UnrealEd/Public/DragAndDrop/ActorDragDropGraphEdOp.h
+++ b/Engine/Source/Editor/UnrealEd/Public/DragAndDrop/ActorDragDropGraphEdOp.h
@@ -50,7 +50,7 @@ public:
 			}
 		case ToolTip_CompatibleAttach:
 			{
-				CurrentHoverText = FText::Format( LOCTEXT("ToolTipCompatibleAttach", "Attach %s to %s"), FText::FromString(Actors[0].Get()->GetActorLabel()), ParamText);
+				CurrentHoverText = FText::Format( LOCTEXT("ToolTipCompatibleAttach", "Attach {0} to {1}"), FText::FromString(Actors[0].Get()->GetActorLabel()), ParamText);
 				CurrentIconBrush = FEditorStyle::GetBrush(TEXT("Graph.ConnectorFeedback.OK"));
 				break;
 			}
@@ -68,25 +68,25 @@ public:
 			}
 		case ToolTip_CompatibleMultipleAttach:
 			{
-				CurrentHoverText = FText::Format( LOCTEXT("ToolTipCompatibleMultipleAttach", "Attach multiple objects to %s"), ParamText);
+				CurrentHoverText = FText::Format( LOCTEXT("ToolTipCompatibleMultipleAttach", "Attach multiple objects to {0}"), ParamText);
 				CurrentIconBrush = FEditorStyle::GetBrush(TEXT("Graph.ConnectorFeedback.OK"));
 				break;
 			}
 		case ToolTip_IncompatibleMultipleAttach:
 			{
-				CurrentHoverText = FText::Format( LOCTEXT("ToolTipIncompatibleMultipleAttach", "Cannot attach multiple objects to %s"), ParamText);
+				CurrentHoverText = FText::Format( LOCTEXT("ToolTipIncompatibleMultipleAttach", "Cannot attach multiple objects to {0}"), ParamText);
 				CurrentIconBrush = FEditorStyle::GetBrush(TEXT("Graph.ConnectorFeedback.Error"));
 				break;
 			}
 		case ToolTip_CompatibleDetach:
 			{
-				CurrentHoverText = FText::Format( LOCTEXT("ToolTipCompatibleDetach", "Detach %s from %s"), FText::FromString(Actors[0].Get()->GetActorLabel()), ParamText);
+				CurrentHoverText = FText::Format( LOCTEXT("ToolTipCompatibleDetach", "Detach {0} from {1}"), FText::FromString(Actors[0].Get()->GetActorLabel()), ParamText);
 				CurrentIconBrush = FEditorStyle::GetBrush(TEXT("Graph.ConnectorFeedback.OK"));
 				break;
 			}
 		case ToolTip_CompatibleMultipleDetach:
 			{
-				CurrentHoverText = FText::Format( LOCTEXT("ToolTipCompatibleDetachMultiple", "Detach multiple objects from %s"), ParamText);
+				CurrentHoverText = FText::Format( LOCTEXT("ToolTipCompatibleDetachMultiple", "Detach multiple objects from {0}"), ParamText);
 				CurrentIconBrush = FEditorStyle::GetBrush(TEXT("Graph.ConnectorFeedback.OK"));
 				break;
 			}
