Can use arrow keys to move through the pin type selector's menu and component's selector menu.

Added a utility widget to help keep focus on a widget and still be able to navigate a listview.

Moved the menu content into a sub-widget that handles the keyboard actions to get the user from the search box to focus on the tree view and handles the enter key.

When using arrow keys in the Pin Type Selector or Component list, item selection will not be acted upon until enter is pressed.

Listviews/Treeviews will no longer skip the first item when using keyboard selection

#ttp 318318 - EDITOR: BP: Variable type selection menu does not support keyboard navigation
#ttp 329065 - BLUEPRINTS: Want to use arrow keys and 'Enter' to select component from the "Add Component" combo box

diff --git a/Engine/Source/Editor/KismetWidgets/Private/SPinTypeSelector.cpp b/Engine/Source/Editor/KismetWidgets/Private/SPinTypeSelector.cpp
--- a/Engine/Source/Editor/KismetWidgets/Private/SPinTypeSelector.cpp
+++ b/Engine/Source/Editor/KismetWidgets/Private/SPinTypeSelector.cpp
@@ -139,7 +139,6 @@ TSharedRef<ITableRow> SPinTypeSelector::GenerateTypeTreeRow(FPinTypeTreeItem InI
 
 	return SNew( STableRow<FPinTypeTreeItem>, OwnerTree )
 		.ToolTip( IDocumentation::Get()->CreateToolTip( FText::FromString( Tooltip ), NULL, *BigTooltipDocLink, *Description) )
-		.ShowSelection(!InItem->bReadOnly)
 		[
 			SNew(SHorizontalBox)
 			+SHorizontalBox::Slot()
@@ -163,8 +162,14 @@ TSharedRef<ITableRow> SPinTypeSelector::GenerateTypeTreeRow(FPinTypeTreeItem InI
 		];
 }
 
-void SPinTypeSelector::OnTypeSelectionChanged(FPinTypeTreeItem Selection, ESelectInfo::Type /*SelectInfo*/)
+void SPinTypeSelector::OnTypeSelectionChanged(FPinTypeTreeItem Selection, ESelectInfo::Type SelectInfo)
 {
+	// When the user is navigating, do not act upon the selection change
+	if(SelectInfo == ESelectInfo::OnNavigation)
+	{
+		return;
+	}
+
 	// Only handle selection for non-read only items, since STreeViewItem doesn't actually support read-only
  	if( Selection.IsValid() )
  	{
@@ -197,7 +202,11 @@ void SPinTypeSelector::OnTypeSelectionChanged(FPinTypeTreeItem Selection, ESelec
 			{
 				const bool bIsExpanded = TypeTreeView->IsItemExpanded(Selection);
 				TypeTreeView->SetItemExpansion(Selection, !bIsExpanded);
-				TypeTreeView->ClearSelection();
+
+				if(SelectInfo == ESelectInfo::OnMouseClick)
+				{
+					TypeTreeView->ClearSelection();
+				}
 			}
 		}
  	}
@@ -215,35 +224,49 @@ TSharedRef<SWidget>	SPinTypeSelector::GetMenuContent()
 
 	if( !MenuContent.IsValid() )
 	{
-		MenuContent = SNew( SVerticalBox )
-			+SVerticalBox::Slot()
-			.AutoHeight()
-			.Padding(4.f, 4.f, 4.f, 4.f)
+		// Pre-build the tree view and search box as it is needed as a parameter for the context menu's container.
+		SAssignNew(TypeTreeView, SPinTypeTreeView)
+			.TreeItemsSource(&FilteredTypeTreeRoot)
+			.SelectionMode(ESelectionMode::Single)
+			.OnGenerateRow(this, &SPinTypeSelector::GenerateTypeTreeRow)
+			.OnSelectionChanged(this, &SPinTypeSelector::OnTypeSelectionChanged)
+			.OnGetChildren(this, &SPinTypeSelector::GetTypeChildren);
+
+		SAssignNew(FilterTextBox, SSearchBox)
+			.OnTextChanged( this, &SPinTypeSelector::OnFilterTextChanged )
+			.OnTextCommitted( this, &SPinTypeSelector::OnFilterTextCommitted );
+
+		MenuContent = SNew(SListViewSelectorDropdownMenu<FPinTypeTreeItem>, FilterTextBox, TypeTreeView)
 			[
-				SAssignNew(FilterTextBox, SSearchBox)
-				.OnTextChanged( this, &SPinTypeSelector::OnFilterTextChanged )
-				.OnTextCommitted( this, &SPinTypeSelector::OnFilterTextCommitted )
-			]
-			+SVerticalBox::Slot()
-			.AutoHeight()
-			.Padding(4.f, 4.f, 4.f, 4.f)
-			[
-				SNew(SBox)
-				.HeightOverride(TreeViewHeight)
-				.WidthOverride(TreeViewWidth)
+				SNew( SVerticalBox )
+				+SVerticalBox::Slot()
+				.AutoHeight()
+				.Padding(4.f, 4.f, 4.f, 4.f)
 				[
-					SAssignNew(TypeTreeView, SPinTypeTreeView)
-					.TreeItemsSource(&FilteredTypeTreeRoot)
-					.SelectionMode(ESelectionMode::Single)
-					.OnGenerateRow(this, &SPinTypeSelector::GenerateTypeTreeRow)
-					.OnSelectionChanged(this, &SPinTypeSelector::OnTypeSelectionChanged)
-					.OnGetChildren(this, &SPinTypeSelector::GetTypeChildren)
+					FilterTextBox.ToSharedRef()
 				]
+				+SVerticalBox::Slot()
+					.AutoHeight()
+					.Padding(4.f, 4.f, 4.f, 4.f)
+					[
+						SNew(SBox)
+						.HeightOverride(TreeViewHeight)
+						.WidthOverride(TreeViewWidth)
+						[
+							TypeTreeView.ToSharedRef()
+						]
+					]
 			];
 			
 
 			TypeComboButton->SetMenuContentWidgetToFocus(FilterTextBox);
 	}
+	else
+	{
+		// Clear the selection in such a way as to also clear the keyboard selector
+		TypeTreeView->SetSelection(NULL, ESelectInfo::OnNavigation);
+		TypeTreeView->ClearExpandedItems();
+	}
 
 	// Clear the filter text box with each opening
 	if( FilterTextBox.IsValid() )
@@ -263,19 +286,31 @@ void SPinTypeSelector::OnFilterTextChanged(const FText& NewText)
 
 	GetChildrenMatchingSearch(NewText.ToString(), TypeTreeRoot, FilteredTypeTreeRoot);
 	TypeTreeView->RequestTreeRefresh();
+
+	// Select the first non-category item
+	auto SelectedItems = TypeTreeView->GetSelectedItems();
+	if(FilteredTypeTreeRoot.Num() > 0)
+	{
+		// Categories have children, we don't want to select categories
+		if(FilteredTypeTreeRoot[0]->Children.Num() > 0)
+		{
+			TypeTreeView->SetSelection(FilteredTypeTreeRoot[0]->Children[0], ESelectInfo::OnNavigation);
+		}
+		else
+		{
+			TypeTreeView->SetSelection(FilteredTypeTreeRoot[0], ESelectInfo::OnNavigation);
+		}
+	}
 }
 
 void SPinTypeSelector::OnFilterTextCommitted(const FText& NewText, ETextCommit::Type CommitInfo)
 {
-	if( FilteredTypeTreeRoot.Num() == 1 && CommitInfo == ETextCommit::OnEnter)
+	if(CommitInfo == ETextCommit::OnEnter)
 	{
-		if(FilteredTypeTreeRoot[0]->Children.Num() == 0)
-		{
-			OnTypeSelectionChanged(FilteredTypeTreeRoot[0], ESelectInfo::Direct);
-		}
-		else if(FilteredTypeTreeRoot[0]->Children.Num() == 1)
+		auto SelectedItems = TypeTreeView->GetSelectedItems();
+		if(SelectedItems.Num() > 0)
 		{
-			OnTypeSelectionChanged(FilteredTypeTreeRoot[0]->Children[0], ESelectInfo::Direct);
+			TypeTreeView->SetSelection(SelectedItems[0]);
 		}
 	}
 }
diff --git a/Engine/Source/Editor/UnrealEd/Private/SComponentClassCombo.cpp b/Engine/Source/Editor/UnrealEd/Private/SComponentClassCombo.cpp
--- a/Engine/Source/Editor/UnrealEd/Private/SComponentClassCombo.cpp
+++ b/Engine/Source/Editor/UnrealEd/Private/SComponentClassCombo.cpp
@@ -13,11 +13,23 @@
 
 void SComponentClassCombo::Construct(const FArguments& InArgs)
 {
+	PrevSelectedIndex = INDEX_NONE;
 	OnComponentClassSelected = InArgs._OnComponentClassSelected;
 
 	UpdateComponentClassList();
 	GenerateFilteredComponentList(FString());
 
+	SAssignNew(ComponentClassListView, SListView<FComponentClassComboEntryPtr>)
+		.ListItemsSource( &FilteredComponentClassList )
+		.OnSelectionChanged( this, &SComponentClassCombo::OnAddComponentSelectionChanged )
+		.OnGenerateRow( this, &SComponentClassCombo::GenerateAddComponentRow )
+		.SelectionMode(ESelectionMode::Single);
+
+	SAssignNew(SearchBox, SSearchBox)
+		.HintText( LOCTEXT( "BlueprintAddComponentSearchBoxHint", "Search Components" ) )
+		.OnTextChanged( this, &SComponentClassCombo::OnSearchBoxTextChanged )
+		.OnTextCommitted( this, &SComponentClassCombo::OnSearchBoxTextCommitted );
+
 	// Create the Construct arguments for the parent class (SComboButton)
 	SComboButton::FArguments Args;
 	Args.ButtonContent()
@@ -27,32 +39,31 @@ void SComponentClassCombo::Construct(const FArguments& InArgs)
 	]
 	.MenuContent()
 	[
-		SNew(SVerticalBox)
-		+SVerticalBox::Slot()
-		.AutoHeight()
-		[
-			SAssignNew(SearchBox, SSearchBox)
-			.HintText( LOCTEXT( "BlueprintAddComponentSearchBoxHint", "Search Components" ) )
-			.OnTextChanged( this, &SComponentClassCombo::OnSearchBoxTextChanged )
-		]
-		+SVerticalBox::Slot()
-		.AutoHeight()
-		[
-			SNew(SSeparator)
-			.Orientation(EOrientation::Orient_Horizontal)
-		]
-		+SVerticalBox::Slot()
-		.MaxHeight(400)
+
+		SNew(SListViewSelectorDropdownMenu<FComponentClassComboEntryPtr>, SearchBox, ComponentClassListView)
 		[
-			SAssignNew(ComponentClassListView, SListView<FComponentClassComboEntryPtr>)
-			.ListItemsSource( &FilteredComponentClassList )
-			.OnSelectionChanged( this, &SComponentClassCombo::OnAddComponentSelectionChanged )
-			.OnGenerateRow( this, &SComponentClassCombo::GenerateAddComponentRow )
-			.SelectionMode(ESelectionMode::Single)
+			SNew(SVerticalBox)
+			+SVerticalBox::Slot()
+			.AutoHeight()
+			[
+				SearchBox.ToSharedRef()
+			]
+			+SVerticalBox::Slot()
+			.AutoHeight()
+			[
+				SNew(SSeparator)
+				.Orientation(EOrientation::Orient_Horizontal)
+			]
+			+SVerticalBox::Slot()
+			.MaxHeight(400)
+			[
+				ComponentClassListView.ToSharedRef()
+			]
 		]
 	]
 	.IsFocusable(true)
-	.ContentPadding(FMargin(2));
+	.ContentPadding(FMargin(2))
+	.OnComboBoxOpened(this, &SComponentClassCombo::ClearSelection);
 
 	SComboButton::Construct(Args);
 
@@ -62,8 +73,16 @@ void SComponentClassCombo::Construct(const FArguments& InArgs)
 
 void SComponentClassCombo::ClearSelection()
 {
-	SearchBox->SetText(FText::GetEmpty());
-	ComponentClassListView->ClearSelection();
+	PrevSelectedIndex = INDEX_NONE;
+
+	// Scroll to the first item if there is one, to reset the scroll position
+	if(FilteredComponentClassList.Num())
+	{
+		ComponentClassListView->RequestScrollIntoView(FilteredComponentClassList[0]);
+	}
+
+	// Clear the selection in such a way as to also clear the keyboard selector
+	ComponentClassListView->SetSelection(NULL, ESelectInfo::OnNavigation);
 }
 
 void SComponentClassCombo::GenerateFilteredComponentList(const FString& InSearchText)
@@ -108,9 +127,21 @@ void SComponentClassCombo::OnSearchBoxTextChanged( const FText& InSearchText )
 	ComponentClassListView->RequestListRefresh();
 }
 
+void SComponentClassCombo::OnSearchBoxTextCommitted(const FText& NewText, ETextCommit::Type CommitInfo)
+{
+	if(CommitInfo == ETextCommit::OnEnter)
+	{
+		auto SelectedItems = ComponentClassListView->GetSelectedItems();
+		if(SelectedItems.Num() > 0)
+		{
+			ComponentClassListView->SetSelection(SelectedItems[0]);
+		}
+	}
+}
+
 void SComponentClassCombo::OnAddComponentSelectionChanged( FComponentClassComboEntryPtr InItem, ESelectInfo::Type SelectInfo )
 {
-	if ( InItem.IsValid() )
+	if ( InItem.IsValid() && InItem->IsClass() && SelectInfo != ESelectInfo::OnNavigation)
 	{
 		// We don't want the item to remain selected
 		ClearSelection();
@@ -123,6 +154,32 @@ void SComponentClassCombo::OnAddComponentSelectionChanged( FComponentClassComboE
 			SetIsOpen( false, false );
 		}
 	}
+	else if ( InItem.IsValid() && SelectInfo != ESelectInfo::OnMouseClick )
+	{
+		int32 SelectedIdx = INDEX_NONE;
+		FilteredComponentClassList.Find(InItem, SelectedIdx);
+		
+		if(SelectedIdx != INDEX_NONE)
+		{
+			if(!InItem->IsClass())
+			{
+				int32 SelectionDirection = SelectedIdx - PrevSelectedIndex;
+
+				// Update the previous selected index
+				PrevSelectedIndex = SelectedIdx;
+
+				if(SelectedIdx + SelectionDirection >= 0 && SelectedIdx + SelectionDirection < FilteredComponentClassList.Num())
+				{
+					ComponentClassListView->SetSelection(FilteredComponentClassList[SelectedIdx + SelectionDirection], ESelectInfo::OnNavigation);
+				}
+			}
+			else
+			{
+				// Update the previous selected index
+				PrevSelectedIndex = SelectedIdx;
+			}
+		}
+	}
 }
 
 TSharedRef<ITableRow> SComponentClassCombo::GenerateAddComponentRow( FComponentClassComboEntryPtr Entry, const TSharedRef<STableViewBase> &OwnerTable ) const
@@ -133,6 +190,7 @@ TSharedRef<ITableRow> SComponentClassCombo::GenerateAddComponentRow( FComponentC
 	{
 		return 
 			SNew( STableRow< TSharedPtr<FString> >, OwnerTable )
+				.ShowSelection(false)
 			[
 				SNew(STextBlock)
 				.Text(Entry->GetHeadingText())
@@ -142,6 +200,7 @@ TSharedRef<ITableRow> SComponentClassCombo::GenerateAddComponentRow( FComponentC
 	{
 		return 
 			SNew( STableRow< TSharedPtr<FString> >, OwnerTable )
+				.ShowSelection(false)
 			[
 				SNew(SSeparator)
 			];
diff --git a/Engine/Source/Editor/UnrealEd/Public/SComponentClassCombo.h b/Engine/Source/Editor/UnrealEd/Public/SComponentClassCombo.h
--- a/Engine/Source/Editor/UnrealEd/Public/SComponentClassCombo.h
+++ b/Engine/Source/Editor/UnrealEd/Public/SComponentClassCombo.h
@@ -81,6 +81,9 @@ public:
 	 */
 	void OnSearchBoxTextChanged( const FText& InSearchText );
 
+	/** Callback when the user commits the text in the searchbox */
+	void OnSearchBoxTextCommitted(const FText& NewText, ETextCommit::Type CommitInfo);
+
 	void OnAddComponentSelectionChanged( FComponentClassComboEntryPtr InItem, ESelectInfo::Type SelectInfo );
 
 	TSharedRef<ITableRow> GenerateAddComponentRow( FComponentClassComboEntryPtr Entry, const TSharedRef<STableViewBase> &OwnerTable ) const;
@@ -109,4 +112,7 @@ private:
 
 	/** The component list control - part of the combo drop down */
 	TSharedPtr< SListView<FComponentClassComboEntryPtr> > ComponentClassListView;
+
+	/** Cached selection index used to skip over unselectable items */
+	int32 PrevSelectedIndex;
 };
\ No newline at end of file
diff --git a/Engine/Source/Runtime/Slate/Public/SListViewSelectorDropdownMenu.h b/Engine/Source/Runtime/Slate/Public/SListViewSelectorDropdownMenu.h
new file mode 100644
--- /dev/null
+++ b/Engine/Source/Runtime/Slate/Public/SListViewSelectorDropdownMenu.h
@@ -0,0 +1,67 @@
+// Copyright 1998-2014 Epic Games, Inc. All Rights Reserved.
+
+#pragma once
+
+/** This is a container widget to help refocus to a listview widget from a searchbox or other widgets that are used in conjunction.
+	Will refocus when the up or down arrows are pressed, and will commit a selection when enter is pressed regardless of where focus is. */
+template <typename ListType>
+class SListViewSelectorDropdownMenu : public SCompoundWidget
+{
+public:
+	SLATE_BEGIN_ARGS( SListViewSelectorDropdownMenu )
+		{}
+		SLATE_DEFAULT_SLOT( FArguments, Content )
+	SLATE_END_ARGS()
+public:
+	/**
+	 * @param InDefaultfocusWidget		The default widget to give focus to when the listview does not handle an action
+	 * @param InTargetListView			If the SListViewSelectorDropdownMenu is handling the KeyDown event, focus will be applied to the TargetListView for certain keys that it can handle
+	 */
+	void Construct(const FArguments& InArgs, TSharedPtr< SWidget> InDefaultFocusWidget, TSharedPtr< SListView< ListType > > InTargetListView)
+	{
+		check(InDefaultFocusWidget.IsValid());
+		check(InTargetListView.IsValid());
+
+		TargetListView = InTargetListView;
+		DefaultFocusWidget = InDefaultFocusWidget;
+
+		this->ChildSlot
+		[
+			InArgs._Content.Widget
+		];
+	}
+
+	// SWidget interface
+	virtual FReply OnKeyDown(const FGeometry& MyGeometry, const FKeyboardEvent& KeyboardEvent) OVERRIDE
+	{
+		TSharedPtr< SListView< ListType > > TargetListViewPtr = TargetListView.Pin();
+
+		if (KeyboardEvent.GetKey() == EKeys::Up || KeyboardEvent.GetKey() == EKeys::Down)
+		{
+			// Deliver focus to the tree view, so the user can use the arrow keys to move through the items
+			return TargetListViewPtr->OnKeyDown(FindChildGeometry(MyGeometry, TargetListViewPtr.ToSharedRef()), KeyboardEvent);
+		}
+		else if(KeyboardEvent.GetKey() == EKeys::Enter)
+		{
+			// If there is anything selected, re-select it "direct" so that the menu will act upon the selection
+			if(TargetListViewPtr->GetNumItemsSelected() > 0)
+			{
+				TargetListViewPtr->SetSelection(TargetListViewPtr->GetSelectedItems()[0]);
+			}
+			return FReply::Handled();
+		}
+		else
+		{
+			TSharedPtr< SWidget > DefaultFocusWidgetPtr = DefaultFocusWidget.Pin();
+			return DefaultFocusWidgetPtr->OnKeyDown(FindChildGeometry(MyGeometry, TargetListViewPtr.ToSharedRef()), KeyboardEvent).SetKeyboardFocus(DefaultFocusWidgetPtr.ToSharedRef(), EKeyboardFocusCause::SetDirectly); //FReply::Unhandled().SetKeyboardFocus(TargetListViewPtr.ToSharedRef(), EKeyboardFocusCause::SetDirectly);
+		}
+		return FReply::Unhandled();
+	}
+	// End of SWidget interface
+
+private:
+	/** The type tree view widget this is handling keyboard input for */
+	TWeakPtr< SListView<ListType> > TargetListView;
+	/** Widget to revert focus back to when this widget does not handle (or forward) a key input */
+	TWeakPtr< SWidget > DefaultFocusWidget;
+};
\ No newline at end of file
diff --git a/Engine/Source/Runtime/Slate/Public/Slate.h b/Engine/Source/Runtime/Slate/Public/Slate.h
--- a/Engine/Source/Runtime/Slate/Public/Slate.h
+++ b/Engine/Source/Runtime/Slate/Public/Slate.h
@@ -206,6 +206,7 @@ DECLARE_DWORD_ACCUMULATOR_STAT_EXTERN( TEXT("Num Dynamic Textures"), STAT_SlateN
 #include "STutorialWrapper.h"
 #include "ScrollyZoomy.h"
 #include "SSafeZone.h"
+#include "SListViewSelectorDropdownMenu.h"
 
 // Docking Framework
 #include "DockingFramework/WorkspaceItem.h"
diff --git a/Engine/Source/Runtime/Slate/Public/TableView/SListView.h b/Engine/Source/Runtime/Slate/Public/TableView/SListView.h
--- a/Engine/Source/Runtime/Slate/Public/TableView/SListView.h
+++ b/Engine/Source/Runtime/Slate/Public/TableView/SListView.h
@@ -250,7 +250,8 @@ public:
 			}
 			else if ( InKeyboardEvent.GetKey() == EKeys::Down )
 			{
-				int32 SelectionIndex = 0;
+				// Begin at INDEX_NONE so the first item will get selected
+				int32 SelectionIndex = INDEX_NONE;
 				if( TListTypeTraits<ItemType>::IsPtrValid(SelectorItem) )
 				{
 					SelectionIndex = ItemsSourceRef.Find( TListTypeTraits<ItemType>::NullableItemTypeConvertToItemType( SelectorItem ) );
@@ -1233,6 +1234,9 @@ protected:
 	{
 		if ( SelectionMode != ESelectionMode::None )
 		{
+			// Must be set before signaling selection changes because sometimes new items will be selected that need to stomp this value
+			SelectorItem = ItemToSelect;
+
 			if ( SelectionMode == ESelectionMode::Multi && ( InKeyboardEvent.IsShiftDown() || InKeyboardEvent.IsControlDown() ) )
 			{
 				// Range select.
@@ -1261,8 +1265,6 @@ protected:
 			{
 				this->RequestScrollIntoView( ItemToSelect );
 			}
-
-			SelectorItem = ItemToSelect;
 		}
 	}
 
