Initializing an uninitialized variable in the DecoratedDragDropOp class that can cause a crash.

diff --git a/Engine/Source/Editor/UnrealEd/Public/DragAndDrop/DecoratedDragDropOp.h b/Engine/Source/Editor/UnrealEd/Public/DragAndDrop/DecoratedDragDropOp.h
--- a/Engine/Source/Editor/UnrealEd/Public/DragAndDrop/DecoratedDragDropOp.h
+++ b/Engine/Source/Editor/UnrealEd/Public/DragAndDrop/DecoratedDragDropOp.h
@@ -17,6 +17,7 @@ public:
 
 	FDecoratedDragDropOp()
 		: CurrentIconBrush(nullptr)
+		, DefaultHoverIcon(nullptr)
 	{ }
 
 	/** Overridden to provide public access */
@@ -41,20 +42,22 @@ public:
 			.Content()
 			[			
 				SNew(SHorizontalBox)
-				+SHorizontalBox::Slot()
+				
+				+ SHorizontalBox::Slot()
 				.AutoWidth()
 				.Padding( 0.0f, 0.0f, 3.0f, 0.0f )
 				[
 					SNew( SImage )
 					.Image( this, &FDecoratedDragDropOp::GetIcon )
 				]
-				+SHorizontalBox::Slot()
-					.AutoWidth()
-					.VAlign( VAlign_Center )
-					[
-						SNew(STextBlock) 
-						.Text( this, &FDecoratedDragDropOp::GetHoverText )
-					]
+
+				+ SHorizontalBox::Slot()
+				.AutoWidth()
+				.VAlign( VAlign_Center )
+				[
+					SNew(STextBlock) 
+					.Text( this, &FDecoratedDragDropOp::GetHoverText )
+				]
 			];
 	}
 
