#summary Moving some files from Slate module to UnrealEd

diff --git a/Engine/Source/Editor/KismetWidgets/Private/SPinTypeSelector.cpp b/Engine/Source/Editor/KismetWidgets/Private/SPinTypeSelector.cpp
--- a/Engine/Source/Editor/KismetWidgets/Private/SPinTypeSelector.cpp
+++ b/Engine/Source/Editor/KismetWidgets/Private/SPinTypeSelector.cpp
@@ -4,6 +4,7 @@
 #include "ClassIconFinder.h"
 #include "SPinTypeSelector.h"
 #include "IDocumentation.h"
+#include "Editor/UnrealEd/Public/SListViewSelectorDropdownMenu.h"
 
 #define LOCTEXT_NAMESPACE "PinTypeSelector"
 
diff --git a/Engine/Source/Editor/UnrealEd/Private/SComponentClassCombo.cpp b/Engine/Source/Editor/UnrealEd/Private/SComponentClassCombo.cpp
--- a/Engine/Source/Editor/UnrealEd/Private/SComponentClassCombo.cpp
+++ b/Engine/Source/Editor/UnrealEd/Private/SComponentClassCombo.cpp
@@ -7,6 +7,7 @@
 #include "ClassIconFinder.h"
 #include "BlueprintGraphDefinitions.h"
 #include "IDocumentation.h"
+#include "SListViewSelectorDropdownMenu.h"
 
 #define LOCTEXT_NAMESPACE "ComponentClassCombo"
 
diff --git a/Engine/Source/Editor/UnrealEd/Public/SListViewSelectorDropdownMenu.h b/Engine/Source/Editor/UnrealEd/Public/SListViewSelectorDropdownMenu.h
new file mode 100644
--- /dev/null
+++ b/Engine/Source/Editor/UnrealEd/Public/SListViewSelectorDropdownMenu.h
@@ -0,0 +1,67 @@
+// Copyright 1998-2014 Epic Games, Inc. All Rights Reserved.
+
+#pragma once
+
+/** This is a container widget to help refocus to a listview widget from a searchbox or other widgets that are used in conjunction.
+	Will refocus when the up or down arrows are pressed, and will commit a selection when enter is pressed regardless of where focus is. */
+template <typename ListType>
+class SListViewSelectorDropdownMenu : public SCompoundWidget
+{
+public:
+	SLATE_BEGIN_ARGS( SListViewSelectorDropdownMenu )
+		{}
+		SLATE_DEFAULT_SLOT( FArguments, Content )
+	SLATE_END_ARGS()
+public:
+	/**
+	 * @param InDefaultfocusWidget		The default widget to give focus to when the listview does not handle an action
+	 * @param InTargetListView			If the SListViewSelectorDropdownMenu is handling the KeyDown event, focus will be applied to the TargetListView for certain keys that it can handle
+	 */
+	void Construct(const FArguments& InArgs, TSharedPtr< SWidget> InDefaultFocusWidget, TSharedPtr< SListView< ListType > > InTargetListView)
+	{
+		check(InDefaultFocusWidget.IsValid());
+		check(InTargetListView.IsValid());
+
+		TargetListView = InTargetListView;
+		DefaultFocusWidget = InDefaultFocusWidget;
+
+		this->ChildSlot
+		[
+			InArgs._Content.Widget
+		];
+	}
+
+	// SWidget interface
+	virtual FReply OnKeyDown(const FGeometry& MyGeometry, const FKeyboardEvent& KeyboardEvent) OVERRIDE
+	{
+		TSharedPtr< SListView< ListType > > TargetListViewPtr = TargetListView.Pin();
+
+		if (KeyboardEvent.GetKey() == EKeys::Up || KeyboardEvent.GetKey() == EKeys::Down)
+		{
+			// Deliver focus to the tree view, so the user can use the arrow keys to move through the items
+			return TargetListViewPtr->OnKeyDown(FindChildGeometry(MyGeometry, TargetListViewPtr.ToSharedRef()), KeyboardEvent);
+		}
+		else if(KeyboardEvent.GetKey() == EKeys::Enter)
+		{
+			// If there is anything selected, re-select it "direct" so that the menu will act upon the selection
+			if(TargetListViewPtr->GetNumItemsSelected() > 0)
+			{
+				TargetListViewPtr->SetSelection(TargetListViewPtr->GetSelectedItems()[0]);
+			}
+			return FReply::Handled();
+		}
+		else
+		{
+			TSharedPtr< SWidget > DefaultFocusWidgetPtr = DefaultFocusWidget.Pin();
+			return DefaultFocusWidgetPtr->OnKeyDown(FindChildGeometry(MyGeometry, TargetListViewPtr.ToSharedRef()), KeyboardEvent).SetKeyboardFocus(DefaultFocusWidgetPtr.ToSharedRef(), EKeyboardFocusCause::SetDirectly); //FReply::Unhandled().SetKeyboardFocus(TargetListViewPtr.ToSharedRef(), EKeyboardFocusCause::SetDirectly);
+		}
+		return FReply::Unhandled();
+	}
+	// End of SWidget interface
+
+private:
+	/** The type tree view widget this is handling keyboard input for */
+	TWeakPtr< SListView<ListType> > TargetListView;
+	/** Widget to revert focus back to when this widget does not handle (or forward) a key input */
+	TWeakPtr< SWidget > DefaultFocusWidget;
+};
\ No newline at end of file
diff --git a/Engine/Source/Runtime/Slate/Public/SListViewSelectorDropdownMenu.h b/Engine/Source/Runtime/Slate/Public/SListViewSelectorDropdownMenu.h
deleted file mode 100644
diff --git a/Engine/Source/Runtime/Slate/Public/Slate.h b/Engine/Source/Runtime/Slate/Public/Slate.h
--- a/Engine/Source/Runtime/Slate/Public/Slate.h
+++ b/Engine/Source/Runtime/Slate/Public/Slate.h
@@ -206,7 +206,6 @@ DECLARE_DWORD_ACCUMULATOR_STAT_EXTERN( TEXT("Num Dynamic Textures"), STAT_SlateN
 #include "STutorialWrapper.h"
 #include "ScrollyZoomy.h"
 #include "SSafeZone.h"
-#include "SListViewSelectorDropdownMenu.h"
 
 // Docking Framework
 #include "DockingFramework/WorkspaceItem.h"
