#ue4 User Defined Structure - minior updates: error message, initial member field, etc..

diff --git a/Engine/Source/Editor/Kismet/Private/UserDefinedStructureEditor.cpp b/Engine/Source/Editor/Kismet/Private/UserDefinedStructureEditor.cpp
--- a/Engine/Source/Editor/Kismet/Private/UserDefinedStructureEditor.cpp
+++ b/Engine/Source/Editor/Kismet/Private/UserDefinedStructureEditor.cpp
@@ -180,6 +180,23 @@ public:
 		return NULL;
 	}
 
+	FString GetStatusTooltip() const
+	{
+		auto StructureDetailsSP = StructureDetails.Pin();
+		if (StructureDetailsSP.IsValid())
+		{
+			if (auto Struct = StructureDetailsSP->GetUserDefinedStruct())
+			{
+				switch (Struct->Status.GetValue())
+				{
+				case EUserDefinedStructureStatus::UDSS_Error:
+					return Struct->ErrorMessage;
+				}
+			}
+		}
+		return FString();
+	}
+
 	/** IDetailCustomNodeBuilder Interface*/
 	virtual void SetOnRebuildChildren( FSimpleDelegate InOnRegenerateChildren ) OVERRIDE 
 	{
@@ -502,14 +519,15 @@ void FUserDefinedStructureLayout::GenerateChildContent( IDetailChildrenBuilder&
 			[
 				SNew(SImage)
 				.Image( this, &FUserDefinedStructureLayout::OnGetStructureStatus )
+				.ToolTipText(this, &FUserDefinedStructureLayout::GetStatusTooltip )
 			]
 			+SHorizontalBox::Slot()
-				.HAlign(HAlign_Right)
-				[
-					SNew(SButton)
-					.Text(LOCTEXT("NewStructureField", "New Variable").ToString())
-					.OnClicked(this, &FUserDefinedStructureLayout::OnAddNewField)
-				]
+			.HAlign(HAlign_Right)
+			[
+				SNew(SButton)
+				.Text(LOCTEXT("NewStructureField", "New Variable").ToString())
+				.OnClicked(this, &FUserDefinedStructureLayout::OnAddNewField)
+			]
 		];
 
 	auto StructureDetailsSP = StructureDetails.Pin();
diff --git a/Engine/Source/Editor/KismetCompiler/Private/UserDefinedStructureCompilerUtils.cpp b/Engine/Source/Editor/KismetCompiler/Private/UserDefinedStructureCompilerUtils.cpp
--- a/Engine/Source/Editor/KismetCompiler/Private/UserDefinedStructureCompilerUtils.cpp
+++ b/Engine/Source/Editor/KismetCompiler/Private/UserDefinedStructureCompilerUtils.cpp
@@ -102,6 +102,16 @@ struct FUserDefinedStructureCompilerInner
 		StructToClean->DestructorLink = NULL;
 		StructToClean->ScriptObjectReferences.Empty();
 		StructToClean->PropertyLink = NULL;
+		StructToClean->ErrorMessage.Empty();
+	}
+
+	static void LogError(UUserDefinedStruct* Struct, FCompilerResultsLog& MessageLog, const FString& ErrorMsg)
+	{
+		MessageLog.Error(*ErrorMsg);
+		if (Struct && Struct->ErrorMessage.IsEmpty())
+		{
+			Struct->ErrorMessage = ErrorMsg;
+		}
 	}
 
 	static void CreateVariables(UUserDefinedStruct* Struct, const class UEdGraphSchema_K2* Schema, FCompilerResultsLog& MessageLog)
@@ -119,7 +129,7 @@ struct FUserDefinedStructureCompilerInner
 			FString ErrorMsg;
 			if(!FStructureEditorUtils::CanHaveAMemberVariableOfType(Struct, VarType, &ErrorMsg))
 			{
-				MessageLog.Error(*FString::Printf(*LOCTEXT("StructureGeneric_Error", "Structure: %s Error: %s").ToString(), *Struct->GetFullName(), *ErrorMsg));
+				LogError(Struct, MessageLog, FString::Printf(*LOCTEXT("StructureGeneric_Error", "Structure: %s Error: %s").ToString(), *Struct->GetFullName(), *ErrorMsg));
 				continue;
 			}
 
@@ -130,7 +140,7 @@ struct FUserDefinedStructureCompilerInner
 			}
 			else
 			{
-				MessageLog.Error(*FString::Printf(*LOCTEXT("VariableInvalidType_Error", "The variable %s declared in %s has an invalid type %s").ToString(),
+				LogError(Struct, MessageLog, FString::Printf(*LOCTEXT("VariableInvalidType_Error", "The variable %s declared in %s has an invalid type %s").ToString(),
 					*VarDesc.VarName.ToString(), *Struct->GetName(), *UEdGraphSchema_K2::TypeToString(VarType)));
 				continue;
 			}
@@ -163,7 +173,7 @@ struct FUserDefinedStructureCompilerInner
 
 		if (Struct->GetStructureSize() <= 0)
 		{
-			MessageLog.Error(*FString::Printf(*LOCTEXT("StructurEmpty_Error", "Structure '%s' is empty ").ToString(), *Struct->GetFullName()));
+			LogError(Struct, MessageLog, FString::Printf(*LOCTEXT("StructurEmpty_Error", "Structure '%s' is empty ").ToString(), *Struct->GetFullName()));
 		}
 
 		const bool bNoErrorsDuringCompilation = (ErrorNum == MessageLog.NumErrors);
diff --git a/Engine/Source/Editor/UnrealEd/Private/Kismet2/StructureEditorUtils.cpp b/Engine/Source/Editor/UnrealEd/Private/Kismet2/StructureEditorUtils.cpp
--- a/Engine/Source/Editor/UnrealEd/Private/Kismet2/StructureEditorUtils.cpp
+++ b/Engine/Source/Editor/UnrealEd/Private/Kismet2/StructureEditorUtils.cpp
@@ -39,6 +39,11 @@ UUserDefinedStruct* FStructureEditorUtils::CreateUserDefinedStruct(UObject* InPa
 		Struct->Bind();
 		Struct->StaticLink(true);
 		Struct->Status = UDSS_Error;
+
+		{
+			const UEdGraphSchema_K2* K2Schema = GetDefault<UEdGraphSchema_K2>();
+			AddVariable(Struct, FEdGraphPinType(K2Schema->PC_Boolean, FString(), NULL, false, false));
+		}
 	}
 
 	return Struct;
@@ -89,42 +94,40 @@ FStructureEditorUtils::EStructureError FStructureEditorUtils::IsStructureValid(c
 
 		for (const UProperty * P = Struct->PropertyLink; P; P = P->PropertyLinkNext)
 		{
+			const UStructProperty* StructProp = Cast<const UStructProperty>(P);
+			if (NULL == StructProp)
 			{
-				const UStructProperty* StructProp = Cast<const UStructProperty>(P);
-				if (NULL == StructProp)
+				if (const UArrayProperty* ArrayProp = Cast<const UArrayProperty>(P))
 				{
-					if (const UArrayProperty* ArrayProp = Cast<const UArrayProperty>(P))
-					{
-						StructProp = Cast<const UStructProperty>(ArrayProp->Inner);
-					}
+					StructProp = Cast<const UStructProperty>(ArrayProp->Inner);
 				}
+			}
 
-				if (StructProp)
+			if (StructProp)
+			{
+				if ((NULL == StructProp->Struct) || (FallbackStruct == StructProp->Struct))
 				{
-					if ((NULL == StructProp->Struct) || (FallbackStruct == StructProp->Struct))
+					if (OutMsg)
 					{
-						if (OutMsg)
-						{
-							*OutMsg = FString::Printf(*LOCTEXT("StructureUnknownProperty", "Struct unknown (deleted?). Parent '%s' Property: '%s'").ToString(), 
-								*Struct->GetFullName(), *StructProp->GetName());
-						}
-						return EStructureError::FallbackStruct;
+						*OutMsg = FString::Printf(*LOCTEXT("StructureUnknownProperty", "Struct unknown (deleted?). Parent '%s' Property: '%s'").ToString(),
+							*Struct->GetFullName(), *StructProp->GetName());
 					}
+					return EStructureError::FallbackStruct;
+				}
 
-					FString OutMsgInner;
-					const EStructureError Result = IsStructureValid(
-						StructProp->Struct, 
-						RecursionParent ? RecursionParent : Struct, 
-						OutMsg ? &OutMsgInner : NULL);
-					if (EStructureError::Ok != Result) 
+				FString OutMsgInner;
+				const EStructureError Result = IsStructureValid(
+					StructProp->Struct,
+					RecursionParent ? RecursionParent : Struct,
+					OutMsg ? &OutMsgInner : NULL);
+				if (EStructureError::Ok != Result)
+				{
+					if (OutMsg)
 					{
-						if (OutMsg)
-						{
-							*OutMsg = FString::Printf(*LOCTEXT("StructurePropertyErrorTemplate", "Struct '%s' Property '%s' Error ( %s )").ToString(), 
-								*Struct->GetFullName(), *StructProp->GetName(), *OutMsgInner);
-						}
-						return Result;
+						*OutMsg = FString::Printf(*LOCTEXT("StructurePropertyErrorTemplate", "Struct '%s' Property '%s' Error ( %s )").ToString(),
+							*Struct->GetFullName(), *StructProp->GetName(), *OutMsgInner);
 					}
+					return Result;
 				}
 			}
 		}
@@ -371,9 +374,16 @@ FString FStructureEditorUtils::GetVariableDisplayName(const UUserDefinedStruct*
 
 bool FStructureEditorUtils::UserDefinedStructEnabled()
 {
-	bool bUseUserDefinedStructure = false;
-	GConfig->GetBool( TEXT("UserDefinedStructure"), TEXT("bUseUserDefinedStructure"), bUseUserDefinedStructure, GEditorIni );
-	return bUseUserDefinedStructure;
+	struct FUserDefinedStructEnabled
+	{
+		bool bUseUserDefinedStructure;
+		FUserDefinedStructEnabled() : bUseUserDefinedStructure(false)
+		{
+			GConfig->GetBool(TEXT("UserDefinedStructure"), TEXT("bUseUserDefinedStructure"), bUseUserDefinedStructure, GEditorIni);
+		}
+	};
+	static FUserDefinedStructEnabled Helper;
+	return Helper.bUseUserDefinedStructure;
 }
 
 bool FStructureEditorUtils::Fill_MakeStructureDefaultValue(const UUserDefinedStruct* Struct, uint8* StructData)
diff --git a/Engine/Source/Runtime/Engine/Classes/Engine/UserDefinedStruct.h b/Engine/Source/Runtime/Engine/Classes/Engine/UserDefinedStruct.h
--- a/Engine/Source/Runtime/Engine/Classes/Engine/UserDefinedStruct.h
+++ b/Engine/Source/Runtime/Engine/Classes/Engine/UserDefinedStruct.h
@@ -33,6 +33,9 @@ public:
 	UPROPERTY(Transient)
 	TWeakObjectPtr<UUserDefinedStruct> PrimaryStruct;
 
+	UPROPERTY()
+	FString ErrorMessage;
+
 	UPROPERTY()
 	UObject* EditorData;
 #endif // WITH_EDITORONLY_DATA
diff --git a/Engine/Source/Runtime/Engine/Private/UserDefinedStruct.cpp b/Engine/Source/Runtime/Engine/Private/UserDefinedStruct.cpp
--- a/Engine/Source/Runtime/Engine/Private/UserDefinedStruct.cpp
+++ b/Engine/Source/Runtime/Engine/Private/UserDefinedStruct.cpp
@@ -19,12 +19,11 @@ void UUserDefinedStruct::Serialize(FArchive& Ar)
 
 	if (Ar.IsLoading() && (EUserDefinedStructureStatus::UDSS_UpToDate == Status))
 	{
-		FString ErrorMsg;
-		const FStructureEditorUtils::EStructureError Result = FStructureEditorUtils::IsStructureValid(this, NULL, &ErrorMsg);
+		const FStructureEditorUtils::EStructureError Result = FStructureEditorUtils::IsStructureValid(this, NULL, &ErrorMessage);
 		if (FStructureEditorUtils::EStructureError::Ok != Result)
 		{
 			Status = EUserDefinedStructureStatus::UDSS_Error;
-			UE_LOG(LogClass, Log, TEXT("UUserDefinedStruct.Serialize '%s' validation: %s"), *GetName(), *ErrorMsg);
+			UE_LOG(LogClass, Log, TEXT("UUserDefinedStruct.Serialize '%s' validation: %s"), *GetName(), *ErrorMessage);
 		}
 	}
 }
