Do not allow non-Actor blueprints to be reparented in to the Actor hierarchy.  It almost certainly would spell doom.
TTP #333544

diff --git a/Engine/Source/Editor/UnrealEd/Private/Kismet2/BlueprintEditorUtils.cpp b/Engine/Source/Editor/UnrealEd/Private/Kismet2/BlueprintEditorUtils.cpp
--- a/Engine/Source/Editor/UnrealEd/Private/Kismet2/BlueprintEditorUtils.cpp
+++ b/Engine/Source/Editor/UnrealEd/Private/Kismet2/BlueprintEditorUtils.cpp
@@ -5026,20 +5026,21 @@ TSharedRef<SWidget> FBlueprintEditorUtils::ConstructBlueprintParentClassPicker(
 		{
 			// Don't allow conversion outside of the Actor hierarchy
 			Filter->AllowedChildrenOfClasses.Add( AActor::StaticClass() );
-		}
-	}
 
-	if(!bIsLevelScriptActor)
-	{
-		// Don't allow non-LevelScriptActor->LevelScriptActor conversion
-		Filter->DisallowedChildrenOfClasses.Add( ALevelScriptActor::StaticClass() );
+			// Don't allow non-LevelScriptActor->LevelScriptActor conversion
+			Filter->DisallowedChildrenOfClasses.Add( ALevelScriptActor::StaticClass() );
+		}
 	}
-
-	if (bIsAnimBlueprint)
+	else if (bIsAnimBlueprint)
 	{
 		// If it's an anim blueprint, do not allow conversion to non anim
 		Filter->AllowedChildrenOfClasses.Add( UAnimInstance::StaticClass() );
 	}
+	else
+	{
+		Filter->DisallowedChildrenOfClasses.Add( AActor::StaticClass() );
+	}
+
 
 	for( auto BlueprintIter = Blueprints.CreateConstIterator(); BlueprintIter; ++BlueprintIter )
 	{
