Fixed padding on "New Game Mode" button in the world settings details and added a tooltip

diff --git a/Engine/Source/Editor/UnrealEd/Public/GameModeInfoCustomizer.h b/Engine/Source/Editor/UnrealEd/Public/GameModeInfoCustomizer.h
--- a/Engine/Source/Editor/UnrealEd/Public/GameModeInfoCustomizer.h
+++ b/Engine/Source/Editor/UnrealEd/Public/GameModeInfoCustomizer.h
@@ -126,18 +126,27 @@ public:
 				.SelectedClass(this, &FGameModeInfoCustomizer::GetCurrentGameModeClass)
 				.OnSetClass(FOnSetClass::CreateSP(this, &FGameModeInfoCustomizer::SetCurrentGameModeClass))
 			]
+
 			+SHorizontalBox::Slot()
 			.AutoWidth()
+			.Padding( 2.0f, 1.0f, 2.0f, 0.0f )
 			[
-				PropertyCustomizationHelpers::MakeBrowseButton(FSimpleDelegate::CreateSP(this, &FGameModeInfoCustomizer::OnBrowseGameModeClicked), FText(), CanBrowseAtrribute)
+				SNew(SButton)
+				.HAlign(HAlign_Center)
+				.VAlign(VAlign_Center)
+				.ContentPadding(FMargin(2.0f,0.0f))
+				.OnClicked(this, &FGameModeInfoCustomizer::OnClickNewGameMode)
+				.ToolTipText(LOCTEXT("NewGameMode_ToolTip", "Create a new Game Mode"))
+				[
+					SNew( STextBlock )
+					.Text(LOCTEXT("NewGameMode", "New.."))
+					.Font( IDetailLayoutBuilder::GetDetailFont() )
+				]
 			]
-			+SHorizontalBox::Slot()
+			+ SHorizontalBox::Slot()
 			.AutoWidth()
 			[
-				SNew(SButton)
-				.Text(LOCTEXT("NewGameMode", "New.."))
-				.ContentPadding( FMargin(4.f, 0.f) )
-				.OnClicked(this, &FGameModeInfoCustomizer::OnClickNewGameMode)
+				PropertyCustomizationHelpers::MakeBrowseButton(FSimpleDelegate::CreateSP(this, &FGameModeInfoCustomizer::OnBrowseGameModeClicked), FText(), CanBrowseAtrribute)
 			]
 		];
 
