Disabled KillZ level rendering in ortho views when bWorldBoundsCheck is disabled

diff --git a/Engine/Source/Editor/UnrealEd/Private/EditorComponents.cpp b/Engine/Source/Editor/UnrealEd/Private/EditorComponents.cpp
--- a/Engine/Source/Editor/UnrealEd/Private/EditorComponents.cpp
+++ b/Engine/Source/Editor/UnrealEd/Private/EditorComponents.cpp
@@ -476,7 +476,7 @@ void FEditorCommonDrawHelper::DrawOldGrid(const FSceneView* View,FPrimitiveDrawI
 			DrawOriginAxisLine( &StartY, &EndY, &StartY.Z, &EndY.Z, View, PDI, AxisColors[1] );
 		}
 
-		if( bDrawKillZ && ( bIsOrthoXZ || bIsOrthoYZ ) )
+		if( bDrawKillZ && ( bIsOrthoXZ || bIsOrthoYZ ) && GWorld->GetWorldSettings()->bEnableWorldBoundsChecks )
 		{
 			float KillZ = GWorld->GetWorldSettings()->KillZ;
 
