In-game Slate UI is now visible in editor viewports when in immersive mode.

Tweaked the layout so hiding the viewport border does not hide all of the child widgets inside of it.

TTP# 333385 - TOOLS: VIEWPORT: PIE: Game slate UI hidden while in fullscreen mode PIE

diff --git a/Engine/Source/Editor/UnrealEd/Private/SEditorViewport.cpp b/Engine/Source/Editor/UnrealEd/Private/SEditorViewport.cpp
--- a/Engine/Source/Editor/UnrealEd/Private/SEditorViewport.cpp
+++ b/Engine/Source/Editor/UnrealEd/Private/SEditorViewport.cpp
@@ -39,15 +39,15 @@ void SEditorViewport::Construct( const FArguments& InArgs )
 			.ShowEffectWhenDisabled( false )
 			.EnableGammaCorrection( false ) // Scene rendering handles this
 			[
-				SNew( SBorder )
-				.BorderImage( this, &SEditorViewport::OnGetViewportBorderBrush )
-				.BorderBackgroundColor( this, &SEditorViewport::OnGetViewportBorderColorAndOpacity )
-				.Visibility( this, &SEditorViewport::OnGetViewportContentVisibility )
-				.Padding(0.0f)
-				.ShowEffectWhenDisabled( false )
+				SAssignNew( ViewportOverlay, SOverlay )
+				+SOverlay::Slot()
 				[
-					SAssignNew( ViewportOverlay, SOverlay )
-			
+					SNew( SBorder )
+					.BorderImage( this, &SEditorViewport::OnGetViewportBorderBrush )
+					.BorderBackgroundColor( this, &SEditorViewport::OnGetViewportBorderColorAndOpacity )
+					.Visibility( this, &SEditorViewport::OnGetViewportContentVisibility )
+					.Padding(0.0f)
+					.ShowEffectWhenDisabled( false )
 				]
 			]
 		]
