BSP actor is now deselected when none of its surfaces are selected.
#ttp 331756 - UE4: unable to deselect a BSP brush once selected
#branch UE4
#proj Editor.UnrealEd
#add When toggling the selection state of a BSP surface, if afterwards there are no surfaces selected, the BSP actor will also be deselected (in ClickHandlers::ClickSurface).

diff --git a/Engine/Source/Editor/UnrealEd/Private/LevelViewportClickHandlers.cpp b/Engine/Source/Editor/UnrealEd/Private/LevelViewportClickHandlers.cpp
--- a/Engine/Source/Editor/UnrealEd/Private/LevelViewportClickHandlers.cpp
+++ b/Engine/Source/Editor/UnrealEd/Private/LevelViewportClickHandlers.cpp
@@ -899,6 +899,12 @@ namespace ClickHandlers
 				}
 				Model->ModifySurf( iSurf, false );
 				Surf.PolyFlags ^= PF_Selected;
+
+				// If there are no surfaces selected now, deselect the actor
+				if (!Model->HasSelectedSurfaces())
+				{
+					GEditor->SelectActor(Surf.Actor, false, true);
+				}
 			}
 			GEditor->NoteSelectionChange();
 		}
diff --git a/Engine/Source/Runtime/Engine/Private/Model.cpp b/Engine/Source/Runtime/Engine/Private/Model.cpp
--- a/Engine/Source/Runtime/Engine/Private/Model.cpp
+++ b/Engine/Source/Runtime/Engine/Private/Model.cpp
@@ -345,6 +345,19 @@ void UModel::ModifySelectedSurfs( bool UpdateMaster )
 			ModifySurf( i, UpdateMaster );
 
 }
+
+bool UModel::HasSelectedSurfaces() const
+{
+	for (const auto& Surface : Surfs)
+	{
+		if (Surface.PolyFlags & PF_Selected)
+		{
+			return true;
+		}
+	}
+
+	return false;
+}
 #endif // WITH_EDITOR
 
 bool UModel::Rename( const TCHAR* InName, UObject* NewOuter, ERenameFlags Flags )
diff --git a/Engine/Source/Runtime/Engine/Public/Model.h b/Engine/Source/Runtime/Engine/Public/Model.h
--- a/Engine/Source/Runtime/Engine/Public/Model.h
+++ b/Engine/Source/Runtime/Engine/Public/Model.h
@@ -497,6 +497,7 @@ public:
 	ENGINE_API void ModifySelectedSurfs( bool UpdateMaster );
 	void ModifyAllSurfs( bool UpdateMaster );
 	ENGINE_API void ModifySurf( int32 InIndex, bool UpdateMaster );
+	ENGINE_API bool HasSelectedSurfaces() const;
 #endif
 
 	ENGINE_API float FindNearestVertex
