Disable stats render setting in editor viewports when PIE is active.
#ttp 333479 - Scene rendered two times in PIE
#branch UE4
#proj Editor.UnrealEd
#change SEditorViewport::HandleViewportStatsEnabled now does nothing if PIE is active.  (Game viewport clients handle stats differently and don't require activating via this mechanism).

diff --git a/Engine/Source/Editor/UnrealEd/Private/SEditorViewport.cpp b/Engine/Source/Editor/UnrealEd/Private/SEditorViewport.cpp
--- a/Engine/Source/Editor/UnrealEd/Private/SEditorViewport.cpp
+++ b/Engine/Source/Editor/UnrealEd/Private/SEditorViewport.cpp
@@ -440,7 +440,7 @@ void SEditorViewport::OnCycleCoordinateSystem()
 
 void SEditorViewport::HandleViewportStatsEnabled()
 {
-	if(Client.Get() == GCurrentLevelEditingViewportClient)
+	if(!GIsPlayInEditorWorld && Client.Get() == GCurrentLevelEditingViewportClient)
 	{
 		Client->SetShowStats( true );
 		Client->SetRealtime( true );
