Can now click on selected items in the pin type selector and components dropdown

diff --git a/Engine/Source/Editor/KismetWidgets/Private/SPinTypeSelector.cpp b/Engine/Source/Editor/KismetWidgets/Private/SPinTypeSelector.cpp
--- a/Engine/Source/Editor/KismetWidgets/Private/SPinTypeSelector.cpp
+++ b/Engine/Source/Editor/KismetWidgets/Private/SPinTypeSelector.cpp
@@ -138,9 +138,10 @@ TSharedRef<ITableRow> SPinTypeSelector::GenerateTypeTreeRow(FPinTypeTreeItem InI
 	const FString OrgTooltip = InItem->GetToolTip();
 	const FString Tooltip = !OrgTooltip.IsEmpty() ? OrgTooltip : Description;
 
-	return SNew( STableRow<FPinTypeTreeItem>, OwnerTree )
+	return SNew( SComboRow<FPinTypeTreeItem>, OwnerTree )
 		.ToolTip( IDocumentation::Get()->CreateToolTip( FText::FromString( Tooltip ), NULL, *BigTooltipDocLink, *Description) )
-		[
+		.RowContent
+		(
 			SNew(SHorizontalBox)
 			+SHorizontalBox::Slot()
 			.AutoWidth()
@@ -160,7 +161,7 @@ TSharedRef<ITableRow> SPinTypeSelector::GenerateTypeTreeRow(FPinTypeTreeItem InI
 				.HighlightText(SearchText)
 				.Font( bHasChildren ? FEditorStyle::GetFontStyle(TEXT("Kismet.TypePicker.CategoryFont")) : FEditorStyle::GetFontStyle(TEXT("Kismet.TypePicker.NormalFont")) )
 			]
-		];
+		);
 }
 
 void SPinTypeSelector::OnTypeSelectionChanged(FPinTypeTreeItem Selection, ESelectInfo::Type SelectInfo)
diff --git a/Engine/Source/Editor/UnrealEd/Private/SComponentClassCombo.cpp b/Engine/Source/Editor/UnrealEd/Private/SComponentClassCombo.cpp
--- a/Engine/Source/Editor/UnrealEd/Private/SComponentClassCombo.cpp
+++ b/Engine/Source/Editor/UnrealEd/Private/SComponentClassCombo.cpp
@@ -258,10 +258,10 @@ TSharedRef<ITableRow> SComponentClassCombo::GenerateAddComponentRow( FComponentC
 		FString DocExcerptName = FString::Printf(TEXT("object'%s'"), *Entry->GetComponentClass()->GetName());
 
 		return
-			SNew( STableRow< TSharedPtr<FString> >, OwnerTable )
+			SNew( SComboRow< TSharedPtr<FString> >, OwnerTable )
 			.ToolTip( IDocumentation::Get()->CreateToolTip(Entry->GetComponentClass()->GetToolTipText(), NULL, DocLink, DocExcerptName) )
-			.Padding(1.0f)
-			[
+			.RowContent
+			(
 				SNew(SHorizontalBox)
 				+SHorizontalBox::Slot()
 				.AutoWidth()
@@ -271,6 +271,7 @@ TSharedRef<ITableRow> SComponentClassCombo::GenerateAddComponentRow( FComponentC
 					.Size(FVector2D(8.0f,1.0f))
 				]
 				+SHorizontalBox::Slot()
+				.Padding(1.0f)
 				.AutoWidth()
 				[
 					SNew(SImage)
@@ -290,7 +291,7 @@ TSharedRef<ITableRow> SComponentClassCombo::GenerateAddComponentRow( FComponentC
 					.HighlightText(this, &SComponentClassCombo::GetCurrentSearchString)
 					.Text( FriendlyComponentName )
 				]
-			];
+			);
 	}
 }
 
diff --git a/Engine/Source/Editor/UnrealEd/Public/SListViewSelectorDropdownMenu.h b/Engine/Source/Editor/UnrealEd/Public/SListViewSelectorDropdownMenu.h
--- a/Engine/Source/Editor/UnrealEd/Public/SListViewSelectorDropdownMenu.h
+++ b/Engine/Source/Editor/UnrealEd/Public/SListViewSelectorDropdownMenu.h
@@ -53,7 +53,7 @@ public:
 		else
 		{
 			TSharedPtr< SWidget > DefaultFocusWidgetPtr = DefaultFocusWidget.Pin();
-			return DefaultFocusWidgetPtr->OnKeyDown(FindChildGeometry(MyGeometry, TargetListViewPtr.ToSharedRef()), KeyboardEvent).SetKeyboardFocus(DefaultFocusWidgetPtr.ToSharedRef(), EKeyboardFocusCause::SetDirectly); //FReply::Unhandled().SetKeyboardFocus(TargetListViewPtr.ToSharedRef(), EKeyboardFocusCause::SetDirectly);
+			return DefaultFocusWidgetPtr->OnKeyDown(FindChildGeometry(MyGeometry, TargetListViewPtr.ToSharedRef()), KeyboardEvent).SetKeyboardFocus(DefaultFocusWidgetPtr.ToSharedRef(), EKeyboardFocusCause::OtherWidgetLostFocus);
 		}
 		return FReply::Unhandled();
 	}
