Limit which SoundWave exporter is offered based on properties of the SoundWave

diff --git a/Engine/Source/Editor/UnrealEd/Classes/Exporters/SoundExporterOGG.h b/Engine/Source/Editor/UnrealEd/Classes/Exporters/SoundExporterOGG.h
--- a/Engine/Source/Editor/UnrealEd/Classes/Exporters/SoundExporterOGG.h
+++ b/Engine/Source/Editor/UnrealEd/Classes/Exporters/SoundExporterOGG.h
@@ -9,6 +9,7 @@ class USoundExporterOGG : public UExporter
     GENERATED_UCLASS_BODY()
 	// Begin UExporter Interface
 	virtual bool ExportBinary( UObject* Object, const TCHAR* Type, FArchive& Ar, FFeedbackContext* Warn, int32 FileIndex = 0, uint32 PortFlags=0 ) OVERRIDE;
+	virtual bool SupportsObject(UObject* Object) const OVERRIDE;
 	// End UExporter Interface
 };
 
diff --git a/Engine/Source/Editor/UnrealEd/Classes/Exporters/SoundExporterWAV.h b/Engine/Source/Editor/UnrealEd/Classes/Exporters/SoundExporterWAV.h
--- a/Engine/Source/Editor/UnrealEd/Classes/Exporters/SoundExporterWAV.h
+++ b/Engine/Source/Editor/UnrealEd/Classes/Exporters/SoundExporterWAV.h
@@ -9,6 +9,7 @@ class USoundExporterWAV : public UExporter
     GENERATED_UCLASS_BODY()
 	// Begin UExporter Interface
 	virtual bool ExportBinary( UObject* Object, const TCHAR* Type, FArchive& Ar, FFeedbackContext* Warn, int32 FileIndex = 0, uint32 PortFlags=0 ) OVERRIDE;
+	virtual bool SupportsObject(UObject* Object) const OVERRIDE;
 	// End UExporter Interface
 };
 
diff --git a/Engine/Source/Editor/UnrealEd/Classes/Exporters/SoundSurroundExporterWAV.h b/Engine/Source/Editor/UnrealEd/Classes/Exporters/SoundSurroundExporterWAV.h
--- a/Engine/Source/Editor/UnrealEd/Classes/Exporters/SoundSurroundExporterWAV.h
+++ b/Engine/Source/Editor/UnrealEd/Classes/Exporters/SoundSurroundExporterWAV.h
@@ -17,6 +17,7 @@ class USoundSurroundExporterWAV : public UExporter
 	virtual int32 GetFileCount( void ) const OVERRIDE;
 	virtual FString GetUniqueFilename( const TCHAR* Filename, int32 FileIndex ) OVERRIDE;
 	virtual bool ExportBinary( UObject* Object, const TCHAR* Type, FArchive& Ar, FFeedbackContext* Warn, int32 FileIndex = 0, uint32 PortFlags=0 ) OVERRIDE;
+	virtual bool SupportsObject(UObject* Object) const OVERRIDE;
 	// End UExporter Interface
 };
 
diff --git a/Engine/Source/Editor/UnrealEd/Private/EditorExporters.cpp b/Engine/Source/Editor/UnrealEd/Private/EditorExporters.cpp
--- a/Engine/Source/Editor/UnrealEd/Private/EditorExporters.cpp
+++ b/Engine/Source/Editor/UnrealEd/Private/EditorExporters.cpp
@@ -62,6 +62,17 @@ USoundExporterWAV::USoundExporterWAV(const class FPostConstructInitializePropert
 
 }
 
+bool USoundExporterWAV::SupportsObject(UObject* Object) const
+{
+	bool bSupportsObject = false;
+	if (Super::SupportsObject(Object))
+	{
+		USoundWave* SoundWave = CastChecked<USoundWave>(Object);
+		bSupportsObject = (SoundWave->NumChannels <= 2);
+	}
+	return bSupportsObject;
+}
+
 bool USoundExporterWAV::ExportBinary( UObject* Object, const TCHAR* Type, FArchive& Ar, FFeedbackContext* Warn, int32 FileIndex, uint32 PortFlags )
 {
 	USoundWave* Sound = CastChecked<USoundWave>( Object );
@@ -86,12 +97,23 @@ USoundExporterOGG::USoundExporterOGG(const class FPostConstructInitializePropert
 
 }
 
+bool USoundExporterOGG::SupportsObject(UObject* Object) const
+{
+	bool bSupportsObject = false;
+	if (Super::SupportsObject(Object))
+	{
+		USoundWave* SoundWave = CastChecked<USoundWave>(Object);
+		bSupportsObject = (SoundWave->GetCompressedData("OGG") != NULL);
+	}
+	return bSupportsObject;
+}
+
+
 bool USoundExporterOGG::ExportBinary( UObject* Object, const TCHAR* Type, FArchive& Ar, FFeedbackContext* Warn, int32 FileIndex, uint32 PortFlags )
 {
 	USoundWave* Sound = CastChecked<USoundWave>( Object );
 
-	static FName NAME_OGG(TEXT("OGG"));
-	FByteBulkData* Bulk = Sound->GetCompressedData(NAME_OGG);
+	FByteBulkData* Bulk = Sound->GetCompressedData("OGG");
 	if (Bulk)
 	{
 		Ar.Serialize(Bulk->Lock(LOCK_READ_ONLY), Bulk->GetBulkDataSize());
@@ -116,6 +138,17 @@ USoundSurroundExporterWAV::USoundSurroundExporterWAV(const class FPostConstructI
 
 }
 
+bool USoundSurroundExporterWAV::SupportsObject(UObject* Object) const
+{
+	bool bSupportsObject = false;
+	if (Super::SupportsObject(Object))
+	{
+		USoundWave* SoundWave = CastChecked<USoundWave>(Object);
+		bSupportsObject = (SoundWave->NumChannels > 2);
+	}
+	return bSupportsObject;
+}
+
 int32 USoundSurroundExporterWAV::GetFileCount( void ) const
 {
 	return( SPEAKER_Count );
diff --git a/Engine/Source/Editor/UnrealEd/Private/ObjectTools.cpp b/Engine/Source/Editor/UnrealEd/Private/ObjectTools.cpp
--- a/Engine/Source/Editor/UnrealEd/Private/ObjectTools.cpp
+++ b/Engine/Source/Editor/UnrealEd/Private/ObjectTools.cpp
@@ -2767,7 +2767,7 @@ namespace ObjectTools
 				UExporter* Exporter = Exporters[ExporterIndex];
 				if( Exporter->SupportedClass )
 				{
-					const bool bObjectIsSupported = ObjectToExport->IsA( Exporter->SupportedClass );
+					const bool bObjectIsSupported = Exporter->SupportsObject(ObjectToExport);
 					if ( bObjectIsSupported )
 					{
 						// Get a string representing of the exportable types.
@@ -2931,7 +2931,7 @@ namespace ObjectTools
 				for( int32 ExporterIndex = 0 ; ExporterIndex < Exporters.Num(); ++ExporterIndex )
 				{
 					UExporter* Exporter = Exporters[ExporterIndex];
-					if( Exporter->SupportedClass && ObjectToExport->IsA(Exporter->SupportedClass) )
+					if( Exporter->SupportsObject(ObjectToExport) )
 					{
 						check( Exporter->FormatExtension.Num() == Exporter->FormatDescription.Num() );
 						for( int32 FormatIndex = 0 ; FormatIndex < Exporter->FormatExtension.Num() ; ++FormatIndex )
diff --git a/Engine/Source/Runtime/Engine/Classes/Exporters/Exporter.h b/Engine/Source/Runtime/Engine/Classes/Exporters/Exporter.h
--- a/Engine/Source/Runtime/Engine/Classes/Exporters/Exporter.h
+++ b/Engine/Source/Runtime/Engine/Classes/Exporters/Exporter.h
@@ -62,6 +62,8 @@ class UExporter : public UObject
 	ENGINE_API virtual void Serialize( FArchive& Ar ) OVERRIDE;
 	// End UObject interface.
 
+	// Returns whether this exporter supports the specific object
+	ENGINE_API virtual bool SupportsObject(UObject* Object) const;
 
 	/**
 	 * Export object to text
diff --git a/Engine/Source/Runtime/Engine/Private/UnrealExporter.cpp b/Engine/Source/Runtime/Engine/Private/UnrealExporter.cpp
--- a/Engine/Source/Runtime/Engine/Private/UnrealExporter.cpp
+++ b/Engine/Source/Runtime/Engine/Private/UnrealExporter.cpp
@@ -39,6 +39,13 @@ void UExporter::Serialize( FArchive& Ar )
 	Ar << PreferredFormatIndex;
 }
 
+// Returns whether this exporter supports the specific object
+bool UExporter::SupportsObject(UObject* Object) const
+{
+	return (SupportedClass && Object->IsA(SupportedClass));
+}
+
+
 UExporter* UExporter::FindExporter( UObject* Object, const TCHAR* FileType )
 {
 	check(Object);
@@ -51,7 +58,7 @@ UExporter* UExporter::FindExporter( UObject* Object, const TCHAR* FileType )
 		if(Default)
 		{
 			check( Default->FormatExtension.Num() == Default->FormatDescription.Num() );
-			if( Default->SupportedClass && Object->IsA(Default->SupportedClass) )
+			if (Default->SupportsObject(Object))
 			{
 				for( int32 i=0; i<Default->FormatExtension.Num(); i++ )
 				{
