Fix import of .po files using /r/n line endings
Portable Object(.po) localization files using windows line endings(\\r\\n) will now import as expected.

diff --git a/Engine/Source/Editor/UnrealEd/Private/Commandlets/InternationalizationExportCommandlet.cpp b/Engine/Source/Editor/UnrealEd/Private/Commandlets/InternationalizationExportCommandlet.cpp
--- a/Engine/Source/Editor/UnrealEd/Private/Commandlets/InternationalizationExportCommandlet.cpp
+++ b/Engine/Source/Editor/UnrealEd/Private/Commandlets/InternationalizationExportCommandlet.cpp
@@ -516,9 +516,10 @@ bool FPortableObjectFormatDOM::FromString( const FString& InStr )
 
 	bool bSuccess = true;
 
+	FString ParseString = InStr.Replace(TEXT("\r\n"), NewLineDelimiter);
 
 	TArray<FString> LinesToProcess;
-	InStr.ParseIntoArray( &LinesToProcess, NewLineDelimiter, false );
+	ParseString.ParseIntoArray( &LinesToProcess, NewLineDelimiter, false );
 
 	TSharedRef<FPortableObjectEntry> ProcessedEntry = MakeShareable( new FPortableObjectEntry );
 	bool bHasMsgId = false;
