#ue4 BP, UDS: Tooltips for structures and variables

diff --git a/Engine/Source/Developer/AssetTools/Private/AssetTypeActions/AssetTypeActions_Struct.cpp b/Engine/Source/Developer/AssetTools/Private/AssetTypeActions/AssetTypeActions_Struct.cpp
--- a/Engine/Source/Developer/AssetTools/Private/AssetTypeActions/AssetTypeActions_Struct.cpp
+++ b/Engine/Source/Developer/AssetTools/Private/AssetTypeActions/AssetTypeActions_Struct.cpp
@@ -27,4 +27,18 @@ void FAssetTypeActions_Struct::OpenAssetEditor( const TArray<UObject*>& InObject
 	}
 }
 
+FText FAssetTypeActions_Struct::GetAssetDescription(const FAssetData& AssetData) const
+{
+	if (const FString* pDescription = AssetData.TagsAndValues.Find(TEXT("Tooltip")))
+	{
+		if (!pDescription->IsEmpty())
+		{
+			const FString DescriptionStr(*pDescription);
+			return FText::FromString(DescriptionStr.Replace(TEXT("\\n"), TEXT("\n")));
+		}
+	}
+
+	return FText::GetEmpty();
+}
+
 #undef LOCTEXT_NAMESPACE
diff --git a/Engine/Source/Developer/AssetTools/Private/AssetTypeActions/AssetTypeActions_Struct.h b/Engine/Source/Developer/AssetTools/Private/AssetTypeActions/AssetTypeActions_Struct.h
--- a/Engine/Source/Developer/AssetTools/Private/AssetTypeActions/AssetTypeActions_Struct.h
+++ b/Engine/Source/Developer/AssetTools/Private/AssetTypeActions/AssetTypeActions_Struct.h
@@ -11,5 +11,6 @@ public:
 	virtual UClass* GetSupportedClass() const OVERRIDE { return UUserDefinedStruct::StaticClass(); }
 	virtual uint32 GetCategories() OVERRIDE { return EAssetTypeCategories::Misc; }
 
+	virtual FText GetAssetDescription(const FAssetData& AssetData) const OVERRIDE;
 	virtual void OpenAssetEditor( const TArray<UObject*>& InObjects, TSharedPtr<class IToolkitHost> EditWithinLevelEditor = TSharedPtr<IToolkitHost>() ) OVERRIDE;
 };
diff --git a/Engine/Source/Editor/BlueprintGraph/Private/EdGraphSchema_K2.cpp b/Engine/Source/Editor/BlueprintGraph/Private/EdGraphSchema_K2.cpp
--- a/Engine/Source/Editor/BlueprintGraph/Private/EdGraphSchema_K2.cpp
+++ b/Engine/Source/Editor/BlueprintGraph/Private/EdGraphSchema_K2.cpp
@@ -84,15 +84,16 @@ UEdGraphSchema_K2::FPinTypeTreeInfo::FPinTypeTreeInfo(const FString& CategoryNam
 
 struct FGatherStructTypesFromAssetsHelper
 {
+	typedef TSharedPtr<UEdGraphSchema_K2::FPinTypeTreeInfo> FPinTypeTreeInfoPtr;
 	struct FCompareChildren
 	{
-		FORCEINLINE bool operator()(const TSharedPtr<UEdGraphSchema_K2::FPinTypeTreeInfo> A, const TSharedPtr<UEdGraphSchema_K2::FPinTypeTreeInfo> B) const
+		FORCEINLINE bool operator()(const FPinTypeTreeInfoPtr A, const FPinTypeTreeInfoPtr B) const
 		{
 			return (A->GetDescription() < B->GetDescription());
 		}
 	};
 
-	static void Gather(const FString& CategoryName, const UEdGraphSchema_K2* Schema, TArray< TSharedPtr<UEdGraphSchema_K2::FPinTypeTreeInfo> >& OutChildren)
+	static void Gather(const FString& CategoryName, const UEdGraphSchema_K2* Schema, TArray<FPinTypeTreeInfoPtr>& OutChildren)
 	{
 		const FAssetRegistryModule& AssetRegistryModule = FModuleManager::LoadModuleChecked<FAssetRegistryModule>(TEXT("AssetRegistry"));
 		TArray<FAssetData> AssetData;
@@ -102,8 +103,10 @@ struct FGatherStructTypesFromAssetsHelper
 			const FAssetData& Asset = AssetData[AssetIndex];
 			if (Asset.IsValid() && !Asset.IsAssetLoaded())
 			{
-				TSharedPtr<UEdGraphSchema_K2::FPinTypeTreeInfo> TypeTreeInfo = 
-					MakeShareable(new UEdGraphSchema_K2::FPinTypeTreeInfo(CategoryName, Asset.ToStringReference(), Asset.ObjectPath.ToString()));
+				const FString* pDescription = Asset.TagsAndValues.Find(TEXT("Tooltip"));
+				const FString Tooltip = (pDescription && !pDescription->IsEmpty()) ? *pDescription : Asset.ObjectPath.ToString();
+
+				FPinTypeTreeInfoPtr TypeTreeInfo = MakeShareable(new UEdGraphSchema_K2::FPinTypeTreeInfo(CategoryName, Asset.ToStringReference(), Tooltip));
 				TypeTreeInfo->FriendlyName = Asset.AssetName.ToString();
 				OutChildren.Add(TypeTreeInfo);
 			}
diff --git a/Engine/Source/Editor/KismetCompiler/Private/UserDefinedStructureCompilerUtils.cpp b/Engine/Source/Editor/KismetCompiler/Private/UserDefinedStructureCompilerUtils.cpp
--- a/Engine/Source/Editor/KismetCompiler/Private/UserDefinedStructureCompilerUtils.cpp
+++ b/Engine/Source/Editor/KismetCompiler/Private/UserDefinedStructureCompilerUtils.cpp
@@ -148,6 +148,7 @@ struct FUserDefinedStructureCompilerInner
 			NewProperty->SetPropertyFlags(CPF_Edit);
 			NewProperty->SetMetaData(TEXT("DisplayName"), *VarDesc.FriendlyName);
 			NewProperty->SetMetaData(TEXT("Category"), *VarDesc.Category);
+			NewProperty->SetMetaData(FBlueprintMetadata::MD_Tooltip, *VarDesc.ToolTip);
 			NewProperty->RepNotifyFunc = NAME_None;
 
 			if (!VarDesc.DefaultValue.IsEmpty())
@@ -165,6 +166,8 @@ struct FUserDefinedStructureCompilerInner
 		check(Struct);
 		const int32 ErrorNum = MessageLog.NumErrors;
 
+		Struct->SetMetaData(FBlueprintMetadata::MD_Tooltip, *FStructureEditorUtils::GetTooltip(Struct));
+
 		CreateVariables(Struct, K2Schema, MessageLog);
 
 		Struct->Bind();
diff --git a/Engine/Source/Editor/UnrealEd/Classes/UserDefinedStructure/UserDefinedStructEditorData.h b/Engine/Source/Editor/UnrealEd/Classes/UserDefinedStructure/UserDefinedStructEditorData.h
--- a/Engine/Source/Editor/UnrealEd/Classes/UserDefinedStructure/UserDefinedStructEditorData.h
+++ b/Engine/Source/Editor/UnrealEd/Classes/UserDefinedStructure/UserDefinedStructEditorData.h
@@ -41,6 +41,9 @@ struct FStructVariableDescription
 	UPROPERTY()
 	FString CurrentDefaultValue;
 
+	UPROPERTY()
+	FString ToolTip;
+
 	UNREALED_API bool SetPinType(const struct FEdGraphPinType& VarType);
 
 	UNREALED_API FEdGraphPinType ToPinType() const;
@@ -60,6 +63,9 @@ public:
 	UPROPERTY()
 	TArray<FStructVariableDescription> VariablesDescriptions;
 
+	UPROPERTY()
+	FString ToolTip;
+
 	// UObject interface.
 	virtual void PostEditUndo() OVERRIDE;
 	virtual void PostLoadSubobjects(struct FObjectInstancingGraph* OuterInstanceGraph) OVERRIDE;
diff --git a/Engine/Source/Editor/UnrealEd/Private/Kismet2/StructureEditorUtils.cpp b/Engine/Source/Editor/UnrealEd/Private/Kismet2/StructureEditorUtils.cpp
--- a/Engine/Source/Editor/UnrealEd/Private/Kismet2/StructureEditorUtils.cpp
+++ b/Engine/Source/Editor/UnrealEd/Private/Kismet2/StructureEditorUtils.cpp
@@ -534,6 +534,65 @@ const TArray<FStructVariableDescription>& FStructureEditorUtils::GetVarDesc(cons
 	return CastChecked<const UUserDefinedStructEditorData>(Struct->EditorData)->VariablesDescriptions;
 }
 
+FString FStructureEditorUtils::GetTooltip(const UUserDefinedStruct* Struct)
+{
+	const auto StructEditorData = Struct ? Cast<const UUserDefinedStructEditorData>(Struct->EditorData) : NULL;
+	return StructEditorData ? StructEditorData->ToolTip : FString();
+}
+
+bool FStructureEditorUtils::ChangeTooltip(UUserDefinedStruct* Struct, const FString& InTooltip)
+{
+	auto StructEditorData = Struct ? Cast<UUserDefinedStructEditorData>(Struct->EditorData) : NULL;
+	if (StructEditorData && (InTooltip != StructEditorData->ToolTip))
+	{
+		const FScopedTransaction Transaction(LOCTEXT("ChangeTooltip", "Change UDS Tooltip"));
+		StructEditorData->Modify();
+		StructEditorData->ToolTip = InTooltip;
+
+		Struct->SetMetaData(FBlueprintMetadata::MD_Tooltip, *StructEditorData->ToolTip);
+
+		return true;
+	}
+	return false;
+}
+
+FString FStructureEditorUtils::GetVariableTooltip(const UUserDefinedStruct* Struct, FGuid VarGuid)
+{
+	if (Struct)
+	{
+		const FStructVariableDescription* VarDesc = GetVarDesc(Struct).FindByPredicate(FFindByGuidHelper<FStructVariableDescription>(VarGuid));
+		if (VarDesc)
+		{
+			return VarDesc->ToolTip;
+		}
+	}
+
+	return FString();
+}
+
+bool FStructureEditorUtils::ChangeVariableTooltip(UUserDefinedStruct* Struct, FGuid VarGuid, const FString& InTooltip)
+{
+	if (Struct)
+	{
+		FStructVariableDescription* VarDesc = GetVarDesc(Struct).FindByPredicate(FFindByGuidHelper<FStructVariableDescription>(VarGuid));
+		if (VarDesc && (InTooltip != VarDesc->ToolTip))
+		{
+			const FScopedTransaction Transaction(LOCTEXT("ChangeVariableTooltip", "Change UDS Variable Tooltip"));
+			ModifyStructData(Struct);
+			VarDesc->ToolTip = InTooltip;
+
+			auto Property = FindField<UProperty>(Struct, VarDesc->VarName);
+			if (Property)
+			{
+				Property->SetMetaData(FBlueprintMetadata::MD_Tooltip, *VarDesc->ToolTip);
+			}
+
+			return true;
+		}
+	}
+	return false;
+}
+
 void FStructureEditorUtils::ModifyStructData(UUserDefinedStruct* Struct)
 {
 	UUserDefinedStructEditorData* EditorData = Struct ? Cast<UUserDefinedStructEditorData>(Struct->EditorData) : NULL;
diff --git a/Engine/Source/Editor/UnrealEd/Public/Kismet2/StructureEditorUtils.h b/Engine/Source/Editor/UnrealEd/Public/Kismet2/StructureEditorUtils.h
--- a/Engine/Source/Editor/UnrealEd/Public/Kismet2/StructureEditorUtils.h
+++ b/Engine/Source/Editor/UnrealEd/Public/Kismet2/StructureEditorUtils.h
@@ -84,6 +84,10 @@ public:
 
 	static void CompileStructure(UUserDefinedStruct* Struct);
 
+	static FString GetTooltip(const UUserDefinedStruct* Struct);
+
+	static bool ChangeTooltip(UUserDefinedStruct* Struct, const FString& InTooltip);
+
 	//VARIABLE
 	static bool AddVariable(UUserDefinedStruct* Struct, const FEdGraphPinType& VarType);
 
@@ -99,6 +103,10 @@ public:
 
 	static FString GetVariableDisplayName(const UUserDefinedStruct* Struct, FGuid VarGuid);
 
+	static FString GetVariableTooltip(const UUserDefinedStruct* Struct, FGuid VarGuid);
+	
+	static bool ChangeVariableTooltip(UUserDefinedStruct* Struct, FGuid VarGuid, const FString& InTooltip);
+
 	//MISC
 	static TArray<FStructVariableDescription>& GetVarDesc(UUserDefinedStruct* Struct);
 
diff --git a/Engine/Source/Runtime/Engine/Classes/Engine/UserDefinedStruct.h b/Engine/Source/Runtime/Engine/Classes/Engine/UserDefinedStruct.h
--- a/Engine/Source/Runtime/Engine/Classes/Engine/UserDefinedStruct.h
+++ b/Engine/Source/Runtime/Engine/Classes/Engine/UserDefinedStruct.h
@@ -44,8 +44,8 @@ public:
 	// UObject interface.
 	virtual void Serialize(FArchive& Ar) OVERRIDE;
 	virtual void PostDuplicate(bool bDuplicateForPIE) OVERRIDE;
+	virtual void GetAssetRegistryTags(TArray<FAssetRegistryTag>& OutTags) const OVERRIDE;
 	// End of UObject interface.
-
 #endif	// WITH_EDITOR
 
 	virtual void SerializeTaggedProperties( FArchive& Ar, uint8* Data, UStruct* DefaultsStruct, uint8* Defaults ) const OVERRIDE;
diff --git a/Engine/Source/Runtime/Engine/Private/UserDefinedStruct.cpp b/Engine/Source/Runtime/Engine/Private/UserDefinedStruct.cpp
--- a/Engine/Source/Runtime/Engine/Private/UserDefinedStruct.cpp
+++ b/Engine/Source/Runtime/Engine/Private/UserDefinedStruct.cpp
@@ -38,6 +38,13 @@ void UUserDefinedStruct::PostDuplicate(bool bDuplicateForPIE)
 	}
 }
 
+void UUserDefinedStruct::GetAssetRegistryTags(TArray<FAssetRegistryTag>& OutTags) const
+{
+	Super::GetAssetRegistryTags(OutTags);
+
+	OutTags.Add(FAssetRegistryTag(TEXT("Tooltip"), FStructureEditorUtils::GetTooltip(this), FAssetRegistryTag::TT_Hidden));
+}
+
 #endif	// WITH_EDITOR
 
 void UUserDefinedStruct::SerializeTaggedProperties(FArchive& Ar, uint8* Data, UStruct* DefaultsStruct, uint8* Defaults) const
