Allow GameMode to have no Pawn specified

diff --git a/Engine/Source/Editor/UnrealEd/Public/GameModeInfoCustomizer.h b/Engine/Source/Editor/UnrealEd/Public/GameModeInfoCustomizer.h
--- a/Engine/Source/Editor/UnrealEd/Public/GameModeInfoCustomizer.h
+++ b/Engine/Source/Editor/UnrealEd/Public/GameModeInfoCustomizer.h
@@ -44,10 +44,12 @@ public:
 		// Find the metaclass of this property
 		UClass* MetaClass = UObject::StaticClass();
 		const UClass* GameModeClass = GetCurrentGameModeClass();
+		bool bAllowNone = false;
 		if (GameModeClass != NULL)
 		{
 			UClassProperty* ClassProp = FindFieldChecked<UClassProperty>(GameModeClass, DefaultClassPropertyName);
 			MetaClass = ClassProp->MetaClass;
+			bAllowNone = !(ClassProp->PropertyFlags & CPF_NoClear);
 		}
 
 		TAttribute<bool> CanBrowseAtrribute = TAttribute<bool>::Create( TAttribute<bool>::FGetter::CreateSP( this, &FGameModeInfoCustomizer::CanBrowseDefaultClass, DefaultClassPropertyName) ) ;
@@ -66,7 +68,7 @@ public:
 			.AutoWidth()
 			[
 				SNew(SClassPropertyEntryBox)
-				.AllowNone(false)
+				.AllowNone(bAllowNone)
 				.MetaClass(MetaClass)
 				.IsEnabled(this, &FGameModeInfoCustomizer::AllowModifyGameMode)
 				.SelectedClass(this, &FGameModeInfoCustomizer::OnGetDefaultClass, DefaultClassPropertyName)
@@ -217,7 +219,7 @@ public:
 	void OnSetDefaultClass(const UClass* NewDefaultClass, FName ClassPropertyName)
 	{
 		const UClass* GameModeClass = GetCurrentGameModeClass();
-		if (GameModeClass != NULL && NewDefaultClass != NULL && AllowModifyGameMode())
+		if (GameModeClass != NULL && AllowModifyGameMode())
 		{
 			UClassProperty* ClassProp = FindFieldChecked<UClassProperty>(GameModeClass, ClassPropertyName);
 			const UClass** DefaultClassPtr = ClassProp->ContainerPtrToValuePtr<const UClass*>(GetCurrentGameModeCDO());
diff --git a/Engine/Source/Runtime/Engine/Classes/GameFramework/GameMode.h b/Engine/Source/Runtime/Engine/Classes/GameFramework/GameMode.h
--- a/Engine/Source/Runtime/Engine/Classes/GameFramework/GameMode.h
+++ b/Engine/Source/Runtime/Engine/Classes/GameFramework/GameMode.h
@@ -146,7 +146,7 @@ public:
 	FString OptionsString;
 
 	/** The default pawn class used by players. */
-	UPROPERTY(EditAnywhere, noclear, BlueprintReadWrite, Category=GameMode)
+	UPROPERTY(EditAnywhere, BlueprintReadWrite, Category=GameMode)
 	TSubclassOf<class APawn>  DefaultPawnClass;
 
 	/** HUD class this game uses. */
diff --git a/Engine/Source/Runtime/Engine/Private/GameMode.cpp b/Engine/Source/Runtime/Engine/Private/GameMode.cpp
--- a/Engine/Source/Runtime/Engine/Private/GameMode.cpp
+++ b/Engine/Source/Runtime/Engine/Private/GameMode.cpp
@@ -441,13 +441,13 @@ void AGameMode::RestartPlayer(AController* NewPlayer)
 		}
 	}
 	// try to create a pawn to use of the default class for this player
-	if (NewPlayer->GetPawn() == NULL)
+	if (NewPlayer->GetPawn() == NULL && GetDefaultPawnClassForController(NewPlayer) != NULL)
 	{
 		NewPlayer->SetPawn(SpawnDefaultPawnFor(NewPlayer, StartSpot));
 	}
+
 	if (NewPlayer->GetPawn() == NULL)
 	{
-		UE_LOG(LogGameMode, Warning, TEXT("failed to spawn player at %s"), *StartSpot->GetName());
 		NewPlayer->FailedToSpawnPawn();
 	}
 	else
