ttp#333554 CRASH: Assertion failed: FindObject<UBlueprint>(Outer, *NewBPName.ToString()) Function: UE4Editor_UnrealEd!FKismetEditorUtilities::CreateBlueprint()

diff --git a/Engine/Source/Editor/UnrealEd/Private/Kismet2/Kismet2.cpp b/Engine/Source/Editor/UnrealEd/Private/Kismet2/Kismet2.cpp
--- a/Engine/Source/Editor/UnrealEd/Private/Kismet2/Kismet2.cpp
+++ b/Engine/Source/Editor/UnrealEd/Private/Kismet2/Kismet2.cpp
@@ -1208,7 +1208,7 @@ const UK2Node_ActorBoundEvent* FKismetEditorUtilities::FindBoundEventForActor(AA
 	const UK2Node_ActorBoundEvent* Node = NULL;
 	if(Actor != NULL && EventName != NAME_None)
 	{
-		ULevelScriptBlueprint* LSB = Actor->GetLevel()->GetLevelScriptBlueprint();
+		ULevelScriptBlueprint* LSB = Actor->GetLevel()->GetLevelScriptBlueprint(true);
 		if(LSB != NULL)
 		{
 			TArray<UK2Node_ActorBoundEvent*> EventNodes;
