Redraw all of the viewports whenever we update the geometry of a BSP/Brush object so they remain in sync.
#ttp 333422 BSP does not update when editing from a 2d viewport.
#proj Editor
#branch UE4

diff --git a/Engine/Source/Editor/GeometryMode/Private/GeometryModifiers.cpp b/Engine/Source/Editor/GeometryMode/Private/GeometryModifiers.cpp
--- a/Engine/Source/Editor/GeometryMode/Private/GeometryModifiers.cpp
+++ b/Engine/Source/Editor/GeometryMode/Private/GeometryModifiers.cpp
@@ -596,6 +596,8 @@ bool UGeomModifier_Edit::InputDelta(FLevelEditorViewportClient* InViewportClient
 
 	EndTrans();
 
+	GEditor->RedrawLevelEditingViewports(true);
+
 	return true;
 }
 
diff --git a/Engine/Source/Editor/UnrealEd/Private/Editor.cpp b/Engine/Source/Editor/UnrealEd/Private/Editor.cpp
--- a/Engine/Source/Editor/UnrealEd/Private/Editor.cpp
+++ b/Engine/Source/Editor/UnrealEd/Private/Editor.cpp
@@ -1596,11 +1596,12 @@ bool UEditorEngine::UpdateSingleViewportClient(FEditorViewportClient* InViewport
 		// Redraw any linked ortho viewports that need to be updated this frame.
 		else if( InViewportClient->IsOrtho() && bLinkedOrthoMovement && InViewportClient->IsVisible() )
 		{
-			if( InViewportClient->bNeedsLinkedRedraw  )
+			if( InViewportClient->bNeedsLinkedRedraw || InViewportClient->bNeedsRedraw )
 			{
 				// Redraw this viewport
 				InViewportClient->Viewport->Draw();
 				InViewportClient->bNeedsLinkedRedraw = false;
+				InViewportClient->bNeedsRedraw = false;
 			}
 			else
 			{
