Fixed an icon on the asset drag/drop decorator being stretched vertically

This addresses TTP#333653 (CRITICAL: Regression: EDITOR: Content Browser: Green check mark is really big when moving an asset)

diff --git a/Engine/Source/Editor/UnrealEd/Private/DragAndDrop/AssetDragDropOp.cpp b/Engine/Source/Editor/UnrealEd/Private/DragAndDrop/AssetDragDropOp.cpp
--- a/Engine/Source/Editor/UnrealEd/Private/DragAndDrop/AssetDragDropOp.cpp
+++ b/Engine/Source/Editor/UnrealEd/Private/DragAndDrop/AssetDragDropOp.cpp
@@ -109,7 +109,7 @@ TSharedPtr<SWidget> FAssetDragDropOp::GetDefaultDecorator() const
 			// Right slot is for optional tooltip
 			+ SHorizontalBox::Slot()
 			.AutoWidth()
-			.VAlign(VAlign_Fill)
+			.VAlign(VAlign_Center)
 			[
 				SNew(SBox)
 				.Visibility(this, &FAssetDragDropOp::GetTooltipVisibility)
