Added world composition tiles cooking when persistent map specified as deafult map in ini files

diff --git a/Engine/Source/Editor/UnrealEd/Private/Commandlets/CookCommandlet.cpp b/Engine/Source/Editor/UnrealEd/Private/Commandlets/CookCommandlet.cpp
--- a/Engine/Source/Editor/UnrealEd/Private/Commandlets/CookCommandlet.cpp
+++ b/Engine/Source/Editor/UnrealEd/Private/Commandlets/CookCommandlet.cpp
@@ -853,22 +853,35 @@ void UCookCommandlet::CollectFilesToCook(TArray<FString>& FilesInPath)
 		FConfigCacheIni::LoadLocalIniFile(PlatformEngineIni, TEXT("Engine"), true, *Platforms[Index]->IniPlatformName());
 
 		// get the server and game default maps and cook them
+		TArray<FString> DefaultMaps;
 		FString Obj;
 		if (PlatformEngineIni.GetString(TEXT("/Script/EngineSettings.GameMapsSettings"), TEXT("GameDefaultMap"), Obj))
 		{
-			FilesInPath.AddUnique(Obj);
+			DefaultMaps.AddUnique(Obj);
 		}
 		if (PlatformEngineIni.GetString(TEXT("/Script/EngineSettings.GameMapsSettings"), TEXT("ServerDefaultMap"), Obj))
 		{
-			FilesInPath.AddUnique(Obj);
+			DefaultMaps.AddUnique(Obj);
 		}
 		if (PlatformEngineIni.GetString(TEXT("/Script/EngineSettings.GameMapsSettings"), TEXT("GlobalDefaultGameMode"), Obj))
 		{
-			FilesInPath.AddUnique(Obj);
+			DefaultMaps.AddUnique(Obj);
 		}
 		if (PlatformEngineIni.GetString(TEXT("/Script/EngineSettings.GameMapsSettings"), TEXT("GlobalDefaultServerGameMode"), Obj))
 		{
-			FilesInPath.AddUnique(Obj);
+			DefaultMaps.AddUnique(Obj);
+		}
+
+		for (FString DefaultMap : DefaultMaps)
+		{
+			// Check if this map is related to world composition
+			if (UWorldComposition::CollectTilesToCook(DefaultMap, FilesInPath))
+			{
+				// Entry has been handled by world composition, no further processing required
+				continue;
+			}
+
+			FilesInPath.AddUnique(DefaultMap);
 		}
 	}
 
