ttp333760 EDITOR: Regression: ORION: Copying particular blueprints to the developer folder causes compiler errors

diff --git a/Engine/Source/Editor/UnrealEd/Private/Kismet2/BlueprintEditorUtils.cpp b/Engine/Source/Editor/UnrealEd/Private/Kismet2/BlueprintEditorUtils.cpp
--- a/Engine/Source/Editor/UnrealEd/Private/Kismet2/BlueprintEditorUtils.cpp
+++ b/Engine/Source/Editor/UnrealEd/Private/Kismet2/BlueprintEditorUtils.cpp
@@ -3890,15 +3890,15 @@ void FBlueprintEditorUtils::ConformImplementedEvents(UBlueprint* Blueprint)
 				// If the event is loaded and is not a custom event
 				if(!EventNode->HasAnyFlags(RF_NeedLoad|RF_NeedPostLoad) && EventNode->bOverrideFunction)
 				{
-					if(Blueprint->GeneratedClass)
+					const bool bEventNodeUsedByInterface = ImplementedInterfaceClasses.Contains(EventNode->EventSignatureClass);
+					if (Blueprint->GeneratedClass && !bEventNodeUsedByInterface)
 					{
 						// See if the generated class implements an event with the given function signature
 						const UFunction* TargetFunction = Blueprint->GeneratedClass->FindFunctionByName(EventNode->EventSignatureName);
-						if(TargetFunction || EventGraphNames.Contains(EventNode->EventSignatureName))
+						if (TargetFunction || EventGraphNames.Contains(EventNode->EventSignatureName))
 						{
 							// The generated class implements the event but the function signature is not up-to-date
-							if(!Blueprint->GeneratedClass->IsChildOf(EventNode->EventSignatureClass)
-								&& !ImplementedInterfaceClasses.Contains(EventNode->EventSignatureClass))
+							if (!Blueprint->GeneratedClass->IsChildOf(EventNode->EventSignatureClass))
 							{
 								FFormatNamedArguments Args;
 								Args.Add(TEXT("NodeTitle"), EventNode->GetNodeTitle(ENodeTitleType::ListView));
