UE4: PkgInfo commandlet now returns the correct file version.

diff --git a/Engine/Source/Editor/UnrealEd/Private/Commandlets/PackageUtilities.cpp b/Engine/Source/Editor/UnrealEd/Private/Commandlets/PackageUtilities.cpp
--- a/Engine/Source/Editor/UnrealEd/Private/Commandlets/PackageUtilities.cpp
+++ b/Engine/Source/Editor/UnrealEd/Private/Commandlets/PackageUtilities.cpp
@@ -835,7 +835,7 @@ void FPkgInfoReporter_Log::GeneratePackageReport( ULinkerLoad* InLinker/*=NULL*/
 	UE_LOG(LogPackageUtilities, Warning, TEXT("--------------------------------------------") );
 
 	UE_LOG(LogPackageUtilities, Warning, TEXT("\t         Filename: %s"), *Linker->Filename);
-	UE_LOG(LogPackageUtilities, Warning, TEXT("\t     File Version: %i"), Linker->UE3Ver() );
+	UE_LOG(LogPackageUtilities, Warning, TEXT("\t     File Version: %i"), Linker->UE4Ver() );
 	UE_LOG(LogPackageUtilities, Warning, TEXT("\t   Engine Version: %s"), *Linker->Summary.EngineVersion.ToString());
 	UE_LOG(LogPackageUtilities, Warning, TEXT("\t     PackageFlags: %X"), Linker->Summary.PackageFlags );
 	UE_LOG(LogPackageUtilities, Warning, TEXT("\t        NameCount: %d"), Linker->Summary.NameCount );
