Move a NoRedist -> NotForLicensees reference for Simplygon

diff --git a/Engine/Source/Editor/UnrealEd/UnrealEd.Build.cs b/Engine/Source/Editor/UnrealEd/UnrealEd.Build.cs
--- a/Engine/Source/Editor/UnrealEd/UnrealEd.Build.cs
+++ b/Engine/Source/Editor/UnrealEd/UnrealEd.Build.cs
@@ -217,8 +217,8 @@ public class UnrealEd : ModuleRules
 		SetupModulePhysXAPEXSupport(Target);
 
 		if ((UEBuildConfiguration.bCompileSimplygon == true) &&
-			(Directory.Exists(UEBuildConfiguration.UEThirdPartyDirectory + "NoRedist") == true) &&
-			(Directory.Exists(UEBuildConfiguration.UEThirdPartyDirectory + "NoRedist/Simplygon") == true))
+			(Directory.Exists(UEBuildConfiguration.UEThirdPartyDirectory + "NotForLicensees") == true) &&
+			(Directory.Exists(UEBuildConfiguration.UEThirdPartyDirectory + "NotForLicensees/Simplygon") == true))
 		{
 			AddThirdPartyPrivateStaticDependencies(Target, "Simplygon");
 		}
