Source code access is now performed on the main thread

Split the process of updating source code symbols and launching the current source code accessor accross different tasks in the task graph, forcing the launch to occur on the main thread.

TTP# 333862 - Editor: Source Access: Assert/Crash when using "Goto code defintiion" in Blueprint graph

diff --git a/Engine/Source/Editor/UnrealEd/Private/SourceCodeNavigation.cpp b/Engine/Source/Editor/UnrealEd/Private/SourceCodeNavigation.cpp
--- a/Engine/Source/Editor/UnrealEd/Private/SourceCodeNavigation.cpp
+++ b/Engine/Source/Editor/UnrealEd/Private/SourceCodeNavigation.cpp
@@ -196,9 +196,8 @@ public:
 	 * @param	FunctionSymbolName	The function to navigate tool (e.g. "MyClass::MyFunction")
 	 * @param	FunctionModuleName	The module to search for this function's symbols (e.g. "GameName-Win64-Debug")
 	 * @param	bIgnoreLineNumber	True if we should just go to the source file and not a specific line within the file
-	 * @param	Application	The running Slate application reference.
 	 */
-	void NavigateToFunctionSource( ISourceCodeAccessor* SourceCodeAccessor, const FString& FunctionSymbolName, const FString& FunctionModuleName, const bool bIgnoreLineNumber, FSlateApplication& Application );
+	void NavigateToFunctionSource( const FString& FunctionSymbolName, const FString& FunctionModuleName, const bool bIgnoreLineNumber );
 
 	/**
 	 * Gathers all functions within a C++ class using debug symbols
@@ -317,8 +316,11 @@ void FSourceCodeNavigationImpl::SetupModuleSymbols()
 }
 
 
-void FSourceCodeNavigationImpl::NavigateToFunctionSource( ISourceCodeAccessor* SourceCodeAccessor, const FString& FunctionSymbolName, const FString& FunctionModuleName, const bool bIgnoreLineNumber, FSlateApplication& Application )
+void FSourceCodeNavigationImpl::NavigateToFunctionSource( const FString& FunctionSymbolName, const FString& FunctionModuleName, const bool bIgnoreLineNumber )
 {
+	ISourceCodeAccessModule& SourceCodeAccessModule = FModuleManager::LoadModuleChecked<ISourceCodeAccessModule>("SourceCodeAccess");
+	ISourceCodeAccessor& SourceCodeAccessor = SourceCodeAccessModule.GetAccessor();
+
 #if PLATFORM_WINDOWS
 	// We'll need the current process handle in order to call into DbgHelp.  This must be the same
 	// process handle that was passed to SymInitialize() earlier.
@@ -336,11 +338,6 @@ void FSourceCodeNavigationImpl::NavigateToFunctionSource( ISourceCodeAccessor* S
 		FullyQualifiedSymbolName = FString::Printf( TEXT( "%s!%s" ), *FunctionModuleName, *FunctionSymbolName );
 	}
 
-
-	// Make sure debug symbols are loaded and ready
-	SetupModuleSymbols();
-
-
 	// Ask DbgHelp to locate information about this symbol by name
 	// NOTE:  Careful!  This function is not thread safe, but we're calling it from a separate thread!
 	if( SymGetSymFromName64( ProcessHandle, TCHAR_TO_ANSI( *FullyQualifiedSymbolName ), SymbolInfoPtr ) )
@@ -371,7 +368,7 @@ void FSourceCodeNavigationImpl::NavigateToFunctionSource( ISourceCodeAccessor* S
 				SourceColumnNumber );
 
 			// Open this source file in our IDE and take the user right to the line number
-			SourceCodeAccessor->OpenFileAtLine( SourceFileName, SourceLineNumber, SourceColumnNumber );
+			SourceCodeAccessor.OpenFileAtLine( SourceFileName, SourceLineNumber, SourceColumnNumber );
 		}
 #if !NO_LOGGING
 		else
@@ -530,7 +527,7 @@ void FSourceCodeNavigationImpl::NavigateToFunctionSource( ISourceCodeAccessor* S
 											int32 FileNameLen = ColonIndex-FileNamePos;
 											FString FileName = Results.Mid(FileNamePos, FileNameLen);
 											FString LineNumber = Results.Mid(ColonIndex + 1, 1);
-											SourceCodeAccessor->OpenFileAtLine( FileName, FCString::Atoi(*LineNumber), 0 );
+											SourceCodeAccessor.OpenFileAtLine( FileName, FCString::Atoi(*LineNumber), 0 );
 										}
 									}
 								}
@@ -564,43 +561,60 @@ void FSourceCodeNavigation::NavigateToFunctionSourceAsync( const FString& Functi
 	// @todo editcode: This will potentially take a long time to execute.  We need a way to tell the async task
 	//				   system that this may block internally and it should always have it's own thread.  Also
 	//				   we may need a way to kill these tasks on shutdown, or when an action is cancelled/overridden
-	//				   by the user interactively
-	ISourceCodeAccessModule& SourceCodeAccessModule = FModuleManager::LoadModuleChecked<ISourceCodeAccessModule>("SourceCodeAccess");
-	ISourceCodeAccessor* SourceCodeAccessor = &SourceCodeAccessModule.GetAccessor();
-	
-	FSlateApplication& Application = FSlateApplication::Get();
-
+	//				   by the user interactively	
 	struct FNavigateFunctionParams
 	{
-		ISourceCodeAccessor* SourceCodeAccessor;
 		FString FunctionSymbolName;
 		FString FunctionModuleName;
 		bool bIgnoreLineNumber;
-		FSlateApplication* Application;
 	};
 
-	TSharedRef< FNavigateFunctionParams, ESPMode::ThreadSafe > NavigateFunctionParams( new FNavigateFunctionParams() );
-	NavigateFunctionParams->SourceCodeAccessor = SourceCodeAccessor;
+	TSharedRef< FNavigateFunctionParams > NavigateFunctionParams( new FNavigateFunctionParams() );
 	NavigateFunctionParams->FunctionSymbolName = FunctionSymbolName;
 	NavigateFunctionParams->FunctionModuleName = FunctionModuleName;
 	NavigateFunctionParams->bIgnoreLineNumber = bIgnoreLineNumber;
-	NavigateFunctionParams->Application = &Application;
 
 	struct FLocal
 	{
-		/** Wrapper function to asynchronously look-up symbols and navigate to source code.  We do this
+		/** Wrapper functions to asynchronously look-up symbols and navigate to source code.  We do this
 	        asynchronously because symbol look-up can often take some time */
-		static void NavigateToFunctionSourceTaskWrapper( ENamedThreads::Type CurrentThread, const FGraphEventRef& MyCompletionGraphEvent, TSharedRef< FNavigateFunctionParams, ESPMode::ThreadSafe > Params )
+
+		static void PreloadSymbolsTaskWrapper( ENamedThreads::Type CurrentThread, const FGraphEventRef& MyCompletionGraphEvent )
+		{
+			// Make sure debug symbols are loaded and ready
+			FSourceCodeNavigationImpl::Get().SetupModuleSymbols();
+		}
+
+		static void NavigateToFunctionSourceTaskWrapper( ENamedThreads::Type CurrentThread, const FGraphEventRef& MyCompletionGraphEvent, TSharedRef< FNavigateFunctionParams > Params, TSharedPtr< SNotificationItem > CompileNotificationPtr )
 		{
 			// Call the navigate function!
-			FSourceCodeNavigationImpl::Get().NavigateToFunctionSource( Params->SourceCodeAccessor, Params->FunctionSymbolName, Params->FunctionModuleName, Params->bIgnoreLineNumber, *Params->Application );
+			FSourceCodeNavigationImpl::Get().NavigateToFunctionSource( Params->FunctionSymbolName, Params->FunctionModuleName, Params->bIgnoreLineNumber );
+
+			// clear the notification
+			CompileNotificationPtr->SetCompletionState( SNotificationItem::CS_Success );
+			CompileNotificationPtr->ExpireAndFadeout();
 		}
 	};
 
-	// Kick off asynchronous task
+	FNotificationInfo Info( LOCTEXT("ReadingSymbols", "Reading C++ Symbols") );
+	Info.Image = FEditorStyle::GetBrush(TEXT("LevelEditor.RecompileGameCode"));
+	Info.ExpireDuration = 2.0f;
+	Info.bFireAndForget = false;
+
+	TSharedPtr< SNotificationItem > CompileNotificationPtr = FSlateNotificationManager::Get().AddNotification(Info);
+	if (CompileNotificationPtr.IsValid())
+	{
+		CompileNotificationPtr->SetCompletionState(SNotificationItem::CS_Pending);
+	}
+
+	// Kick off asynchronous task to load symbols
+	FGraphEventRef PreloadSymbolsAsyncResult(
+		FDelegateGraphTask::CreateAndDispatchWhenReady( FDelegateGraphTask::FDelegate::CreateStatic(&FLocal::PreloadSymbolsTaskWrapper ), TEXT( "EditorSourceCodeNavigation" ) ) );
+
+	// add a dependent task to run on the main thread when symbols are loaded
 	FGraphEventRef UnusedAsyncResult(
 		FDelegateGraphTask::CreateAndDispatchWhenReady( FDelegateGraphTask::FDelegate::CreateStatic(
-				&FLocal::NavigateToFunctionSourceTaskWrapper, NavigateFunctionParams ), TEXT( "EditorSourceCodeNavigation" ) ) );
+				&FLocal::NavigateToFunctionSourceTaskWrapper, NavigateFunctionParams, CompileNotificationPtr ), TEXT( "EditorSourceCodeNavigation" ), PreloadSymbolsAsyncResult, ENamedThreads::GameThread, ENamedThreads::GameThread ));
 }
 
 
