Editor: Adjusted normal map autodetection to consider alpha=0 pixels as black for the averaging (to avoid misidentifying predominantly blue sprites as normal maps)

diff --git a/Engine/Source/Editor/UnrealEd/Private/NormalMapIdentification.cpp b/Engine/Source/Editor/UnrealEd/Private/NormalMapIdentification.cpp
--- a/Engine/Source/Editor/UnrealEd/Private/NormalMapIdentification.cpp
+++ b/Engine/Source/Editor/UnrealEd/Private/NormalMapIdentification.cpp
@@ -34,6 +34,9 @@ namespace
 	// This is used in the comparison with "mid-gray"
 	const float ColorComponentNearlyZeroThreshold = (2.0f / 255.0f);
 
+	// This is used when comparing alpha to zero to avoid picking up sprites
+	const float AlphaComponentNearlyZeroThreshold = (1.0f / 255.0f);
+
 	// These values are chosen to make the threshold colors (from uint8 textures)
 	// discard the top most and bottom most two values, i.e. 0, 1, 254 and 255 on
 	// the assumption that these are likely invalid values for a general normal map
@@ -187,6 +190,13 @@ public:
 				FLinearColor ColorSample = Sampler.DoSampleColor( X, Y );
 				if ( !ColorSample.IsAlmostBlack() )
 				{
+					if (FMath::IsNearlyZero(ColorSample.A, AlphaComponentNearlyZeroThreshold))
+					{
+						AverageColor += FLinearColor::Transparent;
+						NumSamplesTaken++;
+						continue;
+					}
+
 					// Scale and bias, if required, to get a signed vector
 					float Vx = Sampler.ScaleAndBiasComponent( ColorSample.R );
 					float Vy = Sampler.ScaleAndBiasComponent( ColorSample.G );
