Reduced transform diff change to be more generous - it's too strict and with lots of conversion.

diff --git a/Engine/Source/Editor/UnrealEd/Private/SkeletalMeshEdit.cpp b/Engine/Source/Editor/UnrealEd/Private/SkeletalMeshEdit.cpp
--- a/Engine/Source/Editor/UnrealEd/Private/SkeletalMeshEdit.cpp
+++ b/Engine/Source/Editor/UnrealEd/Private/SkeletalMeshEdit.cpp
@@ -855,7 +855,8 @@ namespace AnimationTransformDebug
 				check(Data.RecalculatedParentTransform.Num() == Key+1);
 
 				FTransform GlobalTransform = Data.RecalculatedLocalTransform[Key] * Data.RecalculatedParentTransform[Key];
-				if(GlobalTransform.Equals(Data.SourceGlobalTransform[Key]) == false)
+				// makes more generous on the threshold. 
+				if(GlobalTransform.Equals(Data.SourceGlobalTransform[Key], 0.1f) == false)
 				{
 					// so that we don't spawm with this message
 					if(bShouldOutputToMessageLog)
