- #if/def out of IsValidBindPose check which causes more issues to us
- Refined/tweaked warning messages.

diff --git a/Engine/Source/Editor/UnrealEd/Private/Fbx/FbxSkeletalMeshImport.cpp b/Engine/Source/Editor/UnrealEd/Private/Fbx/FbxSkeletalMeshImport.cpp
--- a/Engine/Source/Editor/UnrealEd/Private/Fbx/FbxSkeletalMeshImport.cpp
+++ b/Engine/Source/Editor/UnrealEd/Private/Fbx/FbxSkeletalMeshImport.cpp
@@ -704,21 +704,69 @@ bool UnFbx::FFbxImporter::ImportBone(TArray<FbxNode*>& NodeArray, FSkeletalMeshI
 			return false;
 		}
 
+#define CHECK_VALID_BIND_POSE 0
+
 		// get bind pose
 		const int32 PoseCount = Scene->GetPoseCount();
 		for (int32 PoseIndex = 0; PoseIndex < PoseCount; PoseIndex++)
 		{
 			FbxPose* CurrentPose = Scene->GetPose(PoseIndex);
+
+			// current pose is bind pose, 
 			if (CurrentPose && CurrentPose->IsBindPose())
 			{
-				if(CurrentPose->IsValidBindPose(NodeArray[0]))
+#if !CHECK_VALID_BIND_POSE
+				PoseArray.Add(CurrentPose);
+#else
+				// IsValidBindPose doesn't work reliably
+				// It checks all the parent chain(regardless root given), and if the parent doesn't have correct bind pose, it fails
+				// It causes more false positive issues than the real issue we have to worry about
+				// If you'd like to try this, set CHECK_VALID_BIND_POSE to 1, and try the error message
+				// when Autodesk fixes this bug, then we might be able to re-open this
+				FString PoseName = CurrentPose->GetName();
+				// all error report status
+				FbxStatus Status;
+				FbxUserNotification Notification(SdkManager, FbxString(), FbxString());
+				Notification.InitAccumulator();
+
+				// it does not make any difference of checking with different node
+				// it is possible pose 0 -> node array 2, but isValidBindPose function returns true even with node array 0
+				FbxNode * Current = NodeArray[0];
+				if (CurrentPose->IsValidBindPoseVerbose(Current, &Notification, 0.0001, &Status))
 				{
 					PoseArray.Add(CurrentPose);
+					UE_LOG(LogFbx, Warning, TEXT("Valid bind pose for Pose (%s) - %s"), *PoseName, *FString(Current->GetName()));
 				}
 				else
 				{
+					FString ErrorString = Status.GetErrorString();
+
+					UE_LOG(LogFbx, Warning, TEXT("Not valid bind pose for Pose (%s) - Node %s : %s"), *PoseName, *FString(Current->GetName()), *ErrorString );
+					for(int32 NotificationIndex=0; NotificationIndex<Notification.GetNbEntries(); ++NotificationIndex)
+					{
+						const FbxAccumulatorEntry* Entry = Notification.GetEntryAt(NotificationIndex);
+						if(Entry)
+						{
+							FString Name = (const char*)(Entry->GetName());
+							FString Description = (const char*)(Entry->GetDescription());
+							FString Detail;
+
+							for(int32 DetailCount =0; DetailCount< Entry->GetDetailsCount(); ++DetailCount)
+							{
+								Detail += FString((const char*)(*Entry->GetDetail(DetailCount)));
+							}
+
+							UE_LOG(LogFbx, Warning, TEXT("\t(%d) %s : %s (%s)"), NotificationIndex+1, *Name, *Description, *Detail );
+						}
+					}
+
+					Notification.ClearAccumulator();
+
+					// add to tokenized error message
 					AddTokenizedErrorMessage(FTokenizedMessage::Create(EMessageSeverity::Warning, FText::Format(LOCTEXT("FbxSkeletaLMeshimport_InvalidBindPose", "FBX bind pose '{0}' is invalid.  It will be ignored"), FText::FromString(Scene->GetPose(PoseIndex)->GetName()))));
 				}
+#endif // CHECK_VALID_BIND_POSE
+
 			}
 		}
 
@@ -923,7 +971,10 @@ bool UnFbx::FFbxImporter::ImportBone(TArray<FbxNode*>& NodeArray, FSkeletalMeshI
 	
 	if(bAnyLinksNotInBindPose)
 	{
-		AddTokenizedErrorMessage(FTokenizedMessage::Create(EMessageSeverity::Warning, FText::Format(LOCTEXT("FbxSkeletaLMeshimport_BonesAreMissingFromBindPose", "Warning: The following bones are missing from the bind pose.  If they are not in the correct orientation after importing, please set the \"Use T0 as ref pose\" option or add them to the bind pose and reimport the skeletal mesh \n\n'{0}'"), FText::FromString(LinksWithoutBindPoses))));
+		AddTokenizedErrorMessage(FTokenizedMessage::Create(EMessageSeverity::Warning, FText::Format(LOCTEXT("FbxSkeletaLMeshimport_BonesAreMissingFromBindPose", "Warning: The following bones are missing from the bind pose.  \
+																																								 This can happen for bones that are not vert weighted. \
+																																								 If they are not in the correct orientation after importing, please set \
+																																								 the \"Use T0 as ref pose\" option or add them to the bind pose and reimport the skeletal mesh \n\n{0}"), FText::FromString(LinksWithoutBindPoses))));
 	}
 	
 	return true;
diff --git a/Engine/Source/Editor/UnrealEd/Private/SkeletalMeshEdit.cpp b/Engine/Source/Editor/UnrealEd/Private/SkeletalMeshEdit.cpp
--- a/Engine/Source/Editor/UnrealEd/Private/SkeletalMeshEdit.cpp
+++ b/Engine/Source/Editor/UnrealEd/Private/SkeletalMeshEdit.cpp
@@ -353,15 +353,22 @@ void UnFbx::FFbxImporter::FillAndVerifyBoneNames(USkeleton* Skeleton, TArray<Fbx
 	}
 
 	// make sure all bone names are included, if not warn user
+	FString BoneNames;
 	for (int32 I = 0; I < TrackNum; ++I)
 	{
 		FName RawBoneName = OutRawBoneNames[I];
 		if ( RefSkeleton.FindBoneIndex(RawBoneName) == INDEX_NONE)
 		{
-			// warn user
-			AddTokenizedErrorMessage(FTokenizedMessage::Create(EMessageSeverity::Warning, FText::Format(LOCTEXT("FBXImport_NoBone", "Could not find bone ({0}) in selected Skeleton({1}) to patch animation from. The data will be lost."), FText::FromName(RawBoneName), FText::FromString(Skeleton->GetFullName()))));
+			BoneNames += RawBoneName.ToString();
+			BoneNames += TEXT("  \n");
 		}
 	}
+
+	if (BoneNames.IsEmpty() == false)
+	{
+		// warn user
+		AddTokenizedErrorMessage(FTokenizedMessage::Create(EMessageSeverity::Warning, FText::Format(LOCTEXT("FBXImport_MissingBone", "The following bones exist in the imported animation, but not in the Skeleton asset {0}.  Any animation on these bones will not be imported: \n\n {1}"), FText::FromString(Skeleton->GetName()), FText::FromString(BoneNames) )));
+	}
 }
 //-------------------------------------------------------------------------
 //
