Give the "showcategories" metadata more granular control (can now unhide sub categories). Also, a commandlet that'll dump hidden functions (so we can diff as changes are made).

diff --git a/Engine/Source/Editor/UnrealEd/Classes/Commandlets/DumpHiddenCategoriesCommandlet.h b/Engine/Source/Editor/UnrealEd/Classes/Commandlets/DumpHiddenCategoriesCommandlet.h
new file mode 100644
--- /dev/null
+++ b/Engine/Source/Editor/UnrealEd/Classes/Commandlets/DumpHiddenCategoriesCommandlet.h
@@ -0,0 +1,15 @@
+// Copyright 1998-2014 Epic Games, Inc. All Rights Reserved.
+
+#pragma once
+#include "DumpHiddenCategoriesCommandlet.generated.h"
+
+UCLASS()
+class UDumpHiddenCategoriesCommandlet : public UCommandlet
+{
+	GENERATED_UCLASS_BODY()
+
+public:		
+	// Begin UCommandlet Interface
+	virtual int32 Main(FString const& Params) OVERRIDE;
+	// End UCommandlet Interface
+};
diff --git a/Engine/Source/Editor/UnrealEd/Private/Commandlets/DumpHiddenCategoriesCommandlet.cpp b/Engine/Source/Editor/UnrealEd/Private/Commandlets/DumpHiddenCategoriesCommandlet.cpp
new file mode 100644
--- /dev/null
+++ b/Engine/Source/Editor/UnrealEd/Private/Commandlets/DumpHiddenCategoriesCommandlet.cpp
@@ -0,0 +1,243 @@
+// Copyright 1998-2014 Epic Games, Inc. All Rights Reserved.
+
+#include "UnrealEd.h"
+//#include "DumpHiddenCategoriesCommandlet.h"
+#include "ObjectEditorUtils.h"
+
+/*******************************************************************************
+ * Static Helpers
+*******************************************************************************/
+
+//------------------------------------------------------------------------------
+static FString GetIndentString(uint32 IndentCount)
+{
+	FString IndentString;
+	while (IndentCount > 0)
+	{
+		IndentString += "\t";
+		--IndentCount;
+	}
+	return IndentString;
+}
+
+//------------------------------------------------------------------------------
+static int32 GetHideCategories(uint32 Indent, UClass* Class, FString& JsonOut)
+{
+	FString IndentString = GetIndentString(Indent);
+	JsonOut += IndentString + TEXT("\"HiddenCategories\" : [");
+
+	TArray<FString> HideCategories;
+	Class->GetHideCategories(HideCategories);
+
+	for (FString const& Category : HideCategories)
+	{
+		JsonOut += TEXT("\n\t") + IndentString + TEXT("\"") + Category + TEXT("\",");
+	}
+	if (HideCategories.Num() > 0)
+	{
+		JsonOut.RemoveAt(JsonOut.Len() - 1); // remove the last comma
+	}
+	JsonOut += TEXT("\n") + IndentString + TEXT("]");
+
+	return HideCategories.Num();
+}
+
+//------------------------------------------------------------------------------
+static int32 GetShowCategories(uint32 Indent, UClass* Class, FString& JsonOut)
+{
+	FString IndentString = GetIndentString(Indent);
+	JsonOut += IndentString + TEXT("\"ShownCategories\" : [");
+
+	TArray<FString> ShowCategories;
+	Class->GetShowCategories(ShowCategories);
+
+	for (FString const& Category : ShowCategories)
+	{
+		JsonOut += TEXT("\n\t") + IndentString + TEXT("\"") + Category + TEXT("\",");
+	}
+	if (ShowCategories.Num() > 0)
+	{
+		JsonOut.RemoveAt(JsonOut.Len() - 1); // remove the last comma
+	}
+	JsonOut += TEXT("\n") + IndentString + TEXT("]");
+
+	return ShowCategories.Num();
+}
+
+//------------------------------------------------------------------------------
+static int32 GetHiddenFunctions(uint32 Indent, UClass* Class, bool bShowFunLibs, FString& JsonOut)
+{
+	int32 HiddenFuncCount = 0;
+
+	FString IndentString = GetIndentString(Indent);
+	JsonOut += IndentString + TEXT("\"HiddenFunctions\" : [");
+
+	FString& OutputString = JsonOut;
+	UClass*  CallingClass = Class;
+	auto FindHiddenFuncs = [&IndentString, &CallingClass, &HiddenFuncCount, &OutputString] (UClass* FunctionClass) 
+	{
+		for (TFieldIterator<UFunction> FunctionIt(FunctionClass, EFieldIteratorFlags::IncludeSuper); FunctionIt; ++FunctionIt)
+		{
+			UFunction* Function = *FunctionIt;
+			if (FObjectEditorUtils::IsFunctionHiddenFromClass(Function, CallingClass))
+			{
+				++HiddenFuncCount;
+				OutputString += TEXT("\n\t") + IndentString + TEXT("\"") + Function->GetPathName() + TEXT("\",");
+			}
+		}
+
+		/*for (TFieldIterator<UObjectProperty> PropIt(FuncClass, EFieldIteratorFlags::IncludeSuper); PropIt; ++PropIt)
+		{
+			UClass* PropClass = PropIt->PropertyClass;
+			for (TFieldIterator<UFunction> FunctionIt(PropClass, EFieldIteratorFlags::IncludeSuper); FunctionIt; ++FunctionIt)
+			{
+				if (FObjectEditorUtils::IsFunctionHiddenFromClass(*FunctionIt, ThisClass))
+				{
+					++HiddenFuncCount;
+					OutputString += TEXT("\n\t\t\t\"") + FunctionIt->GetPathName() + TEXT("\",");
+				}
+			}
+		}*/
+	};
+
+	FindHiddenFuncs(CallingClass); // find all this class's functions
+
+	if (bShowFunLibs)
+	{
+		for (TObjectIterator<UClass> ClassIt; ClassIt; ++ClassIt) // find all functions in each function library
+		{
+			UClass* TestClass = *ClassIt;
+			// if this is a skeleton class, don't bother
+			if (TestClass->HasAnyFlags(RF_Transient) && TestClass->HasAnyClassFlags(CLASS_CompiledFromBlueprint))
+			{
+				continue;
+			}
+
+			if (TestClass->IsChildOf(UBlueprintFunctionLibrary::StaticClass()))
+			{
+				FindHiddenFuncs(TestClass);
+			}
+		}
+	}
+	if (HiddenFuncCount > 0)
+	{
+		OutputString.RemoveAt(OutputString.Len() - 1); // remove the last comma
+	}
+	OutputString += TEXT("\n") + IndentString + TEXT("]");
+
+	return HiddenFuncCount;
+}
+
+//------------------------------------------------------------------------------
+static int32 GetHiddenProperties(uint32 Indent, UClass* Class, FString& JsonOut)
+{
+	int32 HiddenPropCount = 0;
+
+	FString IndentString = GetIndentString(Indent);
+	JsonOut += IndentString + TEXT("\"HiddenProperties\" : [");
+
+	for (TFieldIterator<UProperty> PropertyIt(Class, EFieldIteratorFlags::IncludeSuper); PropertyIt; ++PropertyIt)
+	{
+		UProperty* Property = *PropertyIt;
+		if (FObjectEditorUtils::IsVariableCategoryHiddenFromClass(Property, Class))
+		{
+			++HiddenPropCount;
+			JsonOut += TEXT("\n\t") + IndentString + TEXT("\"") + Property->GetPathName() + TEXT("\",");
+		}
+	}
+	if (HiddenPropCount > 0)
+	{
+		JsonOut.RemoveAt(JsonOut.Len() - 1); // remove the last comma
+	}
+	JsonOut += TEXT("\n") + IndentString + TEXT("]");
+
+	return HiddenPropCount;
+}
+
+/*******************************************************************************
+ * UDumpHiddenCategoriesCommandlet
+*******************************************************************************/
+
+//------------------------------------------------------------------------------
+UDumpHiddenCategoriesCommandlet::UDumpHiddenCategoriesCommandlet(class FPostConstructInitializeProperties const& PCIP)
+	: Super(PCIP)
+{
+}
+
+//------------------------------------------------------------------------------
+int32 UDumpHiddenCategoriesCommandlet::Main(FString const& Params)
+{
+	TArray<FString> Tokens, Switches;
+	ParseCommandLine(*Params, Tokens, Switches);
+
+	bool const bShowHiddenLibraryFuncs = Switches.Contains(TEXT("IncludeFuncLibs"));
+	
+	FString CommandletSaveDir = FPaths::GameSavedDir() + TEXT("Commandlets/");
+	IFileManager::Get().MakeDirectory(*CommandletSaveDir);
+
+	FString Filename = FString::Printf(TEXT("HiddenCategoryDump_%s.json"), FPlatformTime::StrTimestamp());
+	Filename = Filename.Replace(TEXT(" "), TEXT("_"));
+	Filename = Filename.Replace(TEXT("/"), TEXT("-"));
+	Filename = Filename.Replace(TEXT(":"), TEXT("."));
+	Filename = FPaths::GetCleanFilename(Filename);
+
+	FString FilePath = CommandletSaveDir / *Filename;
+	FArchive* FileOut = IFileManager::Get().CreateFileWriter(*FilePath);
+	if (FileOut != nullptr)
+	{
+		FileOut->Serialize(TCHAR_TO_ANSI(TEXT("{")), 1);
+
+		bool bFirstEntry = true;
+		for (TObjectIterator<UClass> ClassIt; ClassIt; ++ClassIt)
+		{
+			UClass* ThisClass = *ClassIt;
+
+			FString ClassEntry;
+			if (!bFirstEntry)
+			{
+				ClassEntry += TEXT(",");
+			}
+			ClassEntry += TEXT("\n\t\"") + ThisClass->GetName() + TEXT("\": {");
+
+
+			FString HideCategoriesJson;
+			if (GetHideCategories(2u, ThisClass, HideCategoriesJson) > 0)
+			{
+				ClassEntry += TEXT("\n") + HideCategoriesJson;
+			}
+			else 
+			{
+				// no need to make an entry for a class that has no "HideCategories" meta data (it can see all)
+				continue;
+			}
+
+			FString ShowCategoriesJson;
+			if (GetShowCategories(2u, ThisClass, ShowCategoriesJson) > 0)
+			{
+				ClassEntry += TEXT(",\n") + ShowCategoriesJson;
+			}
+
+			FString HiddenFunctionJson;
+			if (GetHiddenFunctions(2u, ThisClass, bShowHiddenLibraryFuncs, HiddenFunctionJson) > 0)
+			{
+				ClassEntry += TEXT(",\n") + HiddenFunctionJson;
+			}
+			
+			FString HiddenPropertyJson;
+			if (GetHiddenProperties(2u, ThisClass, HiddenPropertyJson) > 0)
+			{
+				ClassEntry += TEXT(",\n") + HiddenPropertyJson;
+			}
+
+			ClassEntry += TEXT("\n\t}");
+			FileOut->Serialize(TCHAR_TO_ANSI(*ClassEntry), ClassEntry.Len());
+			bFirstEntry = false;
+		}
+		FileOut->Serialize(TCHAR_TO_ANSI(TEXT("\n}")), 2);
+
+		FileOut->Close();
+		delete FileOut;
+		FileOut = nullptr;
+	}
+	return 0;
+}
\ No newline at end of file
diff --git a/Engine/Source/Programs/UnrealHeaderTool/Private/HeaderParser.cpp b/Engine/Source/Programs/UnrealHeaderTool/Private/HeaderParser.cpp
--- a/Engine/Source/Programs/UnrealHeaderTool/Private/HeaderParser.cpp
+++ b/Engine/Source/Programs/UnrealHeaderTool/Private/HeaderParser.cpp
@@ -4117,10 +4117,12 @@ void FHeaderParser::CompileClassDeclaration(FClasses& AllClasses)
 
 	// Get categories inherited from the parent.
 	TArray<FString> HideCategories;
+	TArray<FString> ShowSubCatgories;
 	TArray<FString> HideFunctions;
 	TArray<FString> AutoExpandCategories;
 	TArray<FString> AutoCollapseCategories;
 	Class->GetHideCategories        (HideCategories);
+	Class->GetShowCategories        (ShowSubCatgories);
 	Class->GetHideFunctions			(HideFunctions);
 	Class->GetAutoExpandCategories  (AutoExpandCategories);
 	Class->GetAutoCollapseCategories(AutoCollapseCategories);
@@ -4325,7 +4327,26 @@ void FHeaderParser::CompileClassDeclaration(FClasses& AllClasses)
 
 			for (const FString& Value : PropSpecifier.Values)
 			{
-				HideCategories.Remove(Value);
+				// if we didn't find this specific category path in the HideCategories metadata
+				if (HideCategories.Remove(Value) == 0)
+				{
+					TArray<FString> SubCategoryList;
+					Value.ParseIntoArray(&SubCategoryList, TEXT("|"), true);
+
+					FString SubCategoryPath;
+					// look to see if any of the parent paths are excluded in the HideCategories list
+					for (int32 CategoryPathIndex = 0; CategoryPathIndex < SubCategoryList.Num()-1; ++CategoryPathIndex)
+					{
+						SubCategoryPath += SubCategoryList[CategoryPathIndex];
+						// if we're hiding a parent category, then we need to flag this sub category for show
+						if (HideCategories.Contains(SubCategoryPath))
+						{
+							ShowSubCatgories.AddUnique(Value);
+							break;
+						}
+						SubCategoryPath += "|";
+					}
+				}
 			}
 		}
 		else if (Specifier == TEXT("hideCategories"))
@@ -4421,6 +4442,7 @@ void FHeaderParser::CompileClassDeclaration(FClasses& AllClasses)
 	if (ClassGroupNames       .Num()) { MetaData.Add("ClassGroupNames",        FString::Join(ClassGroupNames,        TEXT(" "))); }
 	if (AutoCollapseCategories.Num()) { MetaData.Add("AutoCollapseCategories", FString::Join(AutoCollapseCategories, TEXT(" "))); }
 	if (HideCategories        .Num()) { MetaData.Add("HideCategories",         FString::Join(HideCategories,         TEXT(" "))); }
+	if (ShowSubCatgories      .Num()) { MetaData.Add("ShowCategories",         FString::Join(ShowSubCatgories,       TEXT(" "))); }
 	if (HideFunctions         .Num()) { MetaData.Add("HideFunctions",          FString::Join(HideFunctions,          TEXT(" "))); }
 	if (AutoExpandCategories  .Num()) { MetaData.Add("AutoExpandCategories",   FString::Join(AutoExpandCategories,   TEXT(" "))); }
 
@@ -6092,6 +6114,10 @@ ECompilationResult::Type FHeaderParser::ParseHeaderForOneClass(FClasses& AllClas
 		{
 			Class->SetMetaData(TEXT("HideCategories"), *SuperClass->GetMetaData("HideCategories"));
 		}
+		if (SuperClass->HasMetaData(TEXT("ShowCategories")))
+		{
+			Class->SetMetaData(TEXT("ShowCategories"), *SuperClass->GetMetaData("ShowCategories"));
+		}
 		if (SuperClass->HasMetaData(TEXT("HideFunctions")))
 		{
 			Class->SetMetaData(TEXT("HideFunctions"), *SuperClass->GetMetaData("HideFunctions"));
diff --git a/Engine/Source/Runtime/CoreUObject/Private/UObject/Class.cpp b/Engine/Source/Runtime/CoreUObject/Private/UObject/Class.cpp
--- a/Engine/Source/Runtime/CoreUObject/Private/UObject/Class.cpp
+++ b/Engine/Source/Runtime/CoreUObject/Private/UObject/Class.cpp
@@ -2958,6 +2958,7 @@ void UClass::PurgeClass(bool bRecompilingOnLoad)
 	{
 		// this is not safe to do at COL time. The meta data is not loaded yet, so if we attempt to load it, we recursively load the package and that will fail
 		RemoveMetaData("HideCategories");
+		RemoveMetaData("ShowCategories");
 		RemoveMetaData("HideFunctions");
 		RemoveMetaData("AutoExpandCategories");
 		RemoveMetaData("AutoCollapseCategories");
@@ -3291,26 +3292,62 @@ void UClass::GetHideCategories(TArray<FString>& OutHideCategories) const
 	}
 }
 
+void UClass::GetShowCategories(TArray<FString>& OutShowCategories) const
+{
+	static const FName NAME_ShowCategories(TEXT("ShowCategories"));
+	if (HasMetaData(NAME_ShowCategories))
+	{
+		const FString& ShowCategories = GetMetaData(NAME_ShowCategories);
+		ShowCategories.ParseIntoArray(&OutShowCategories, TEXT(" "), true);
+	}
+}
+
 bool UClass::IsCategoryHidden(const FString& InCategory) const
 {
 	bool bHidden = false;
-	static const FName NAME_HideCategories(TEXT("HideCategories"));
+	static FName const NAME_HideCategories(TEXT("HideCategories"));
+	static FName const NAME_ShowCategories(TEXT("ShowCategories"));
 	if (HasMetaData(NAME_HideCategories))
 	{
-		const FString& HideCategories = GetMetaData(NAME_HideCategories);
+		FString ShowCategories;
+		if (HasMetaData(NAME_ShowCategories))
+		{
+			ShowCategories = GetMetaData(NAME_ShowCategories);
+		}
+
+		FString const& HideCategories = GetMetaData(NAME_HideCategories);
 		bHidden = !!FCString::StrfindDelim(*HideCategories, *InCategory, TEXT(" "));
 		if (!bHidden)
 		{
 			TArray<FString> SubCategoryList;
 			InCategory.ParseIntoArray(&SubCategoryList, TEXT("|"), true);
-			for (const FString& SubCategory : SubCategoryList)
+
+			FString SubCategoryPath;
+			for (int32 SubCatIndex = 0; SubCatIndex < SubCategoryList.Num(); ++SubCatIndex)
 			{
+				FString SubCategory = SubCategoryList[SubCatIndex];
+				SubCategory = SubCategory.Replace(TEXT(" "), TEXT(""));
+				SubCategoryPath += SubCategory;
+
 				if (!!FCString::StrfindDelim(*HideCategories, *SubCategory, TEXT(" ")))
 				{
 					bHidden = true;
-					break;
 				}
-			}		
+				else if (!!FCString::StrfindDelim(*HideCategories, *SubCategoryPath, TEXT(" ")))
+				{
+					bHidden = true;
+				}
+
+				if (bHidden && !ShowCategories.IsEmpty() && !!FCString::StrfindDelim(*ShowCategories, *SubCategoryPath, TEXT(" ")))
+				{
+					bHidden = false;
+				}
+				SubCategoryPath += "|";
+			}	
+		}
+		else if (!ShowCategories.IsEmpty() && !!FCString::StrfindDelim(*ShowCategories, *InCategory, TEXT(" ")))
+		{
+			bHidden = false;
 		}
 	}
 	return bHidden;
diff --git a/Engine/Source/Runtime/CoreUObject/Public/UObject/Class.h b/Engine/Source/Runtime/CoreUObject/Public/UObject/Class.h
--- a/Engine/Source/Runtime/CoreUObject/Public/UObject/Class.h
+++ b/Engine/Source/Runtime/CoreUObject/Public/UObject/Class.h
@@ -1628,6 +1628,7 @@ public:
 #if WITH_EDITOR || HACK_HEADER_GENERATOR 
 	// Editor only properties
 	void GetHideCategories(TArray<FString>& OutHideCategories) const;
+	void GetShowCategories(TArray<FString>& OutShowCategories) const;
 	bool IsCategoryHidden(const FString& InCategory) const;
 	void GetHideFunctions(TArray<FString>& OutHideFunctions) const;
 	bool IsFunctionHidden(const TCHAR* InFunction) const;
