Prevent divide by 0 error when zooming out level viewport too far.

#ttp 333853 - CRASH: EDITOR: Level Viewport: Zooming out too far

#branch UE4

#change Skip rendering grid section when zoomed out too far

diff --git a/Engine/Source/Editor/UnrealEd/Private/EditorComponents.cpp b/Engine/Source/Editor/UnrealEd/Private/EditorComponents.cpp
--- a/Engine/Source/Editor/UnrealEd/Private/EditorComponents.cpp
+++ b/Engine/Source/Editor/UnrealEd/Private/EditorComponents.cpp
@@ -504,14 +504,6 @@ void FEditorCommonDrawHelper::DrawGridSection(float ViewportGridY,FVector* A,FVe
 	// todo
 	static int32 Exponent = GEditor->IsGridSizePowerOfTwo() ? 8 : 10;
 
-	const FMatrix InvViewProjMatrix = View->ViewMatrices.ProjMatrix.Inverse() * View->ViewMatrices.ViewMatrix.Inverse();
-	int32 FirstLine = FMath::Trunc(InvViewProjMatrix.TransformPosition(FVector(-1,-1,0.5f)).Component(Axis) / ViewportGridY);
-	int32 LastLine = FMath::Trunc(InvViewProjMatrix.TransformPosition(FVector(+1,+1,0.5f)).Component(Axis) / ViewportGridY);
-	if( FirstLine > LastLine )
-	{
-		Exchange(FirstLine,LastLine);
-	}
-
 	const float SizeX = View->ViewRect.Width();
 	const float Zoom = (1.0f / View->ViewMatrices.ProjMatrix.M[0][0]) * 2.0f / SizeX;
 	const float Dist = SizeX * Zoom / ViewportGridY;
@@ -520,7 +512,6 @@ void FEditorCommonDrawHelper::DrawGridSection(float ViewportGridY,FVector* A,FVe
 	static float Tweak = 4.0f;
 
 	float IncValue = FMath::LogX(Exponent, Dist / (float)(SizeX / Tweak));
-	
 	int32 IncScale = 1;
 
 	for(float x = 0; x < IncValue; ++x)
@@ -528,6 +519,12 @@ void FEditorCommonDrawHelper::DrawGridSection(float ViewportGridY,FVector* A,FVe
 		IncScale *= Exponent;
 	}
 
+	if (IncScale == 0)
+	{
+		// Prevent divide by zero
+		return;
+	}
+
 	// 0 excluded for hard transitions .. 0.5f for very soft transitions
 	const float TransitionRegion = 0.5f;
 
@@ -550,6 +547,14 @@ void FEditorCommonDrawHelper::DrawGridSection(float ViewportGridY,FVector* A,FVe
 	FLinearColor MajorColor = FMath::Lerp(Background, FLinearColor::White, 0.05f);
 	FLinearColor MinorColor = FMath::Lerp(Background, FLinearColor::White, 0.02f);
 
+	const FMatrix InvViewProjMatrix = View->ViewMatrices.ProjMatrix.Inverse() * View->ViewMatrices.ViewMatrix.Inverse();
+	int32 FirstLine = FMath::Trunc(InvViewProjMatrix.TransformPosition(FVector(-1, -1, 0.5f)).Component(Axis) / ViewportGridY);
+	int32 LastLine = FMath::Trunc(InvViewProjMatrix.TransformPosition(FVector(+1, +1, 0.5f)).Component(Axis) / ViewportGridY);
+	if (FirstLine > LastLine)
+	{
+		Exchange(FirstLine, LastLine);
+	}
+
 	// Draw major and minor grid lines
 	const int32 FirstLineClamped = FMath::Max<int32>(FirstLine - 1,-HALF_WORLD_MAX/ViewportGridY) / IncScale;
 	const int32 LastLineClamped = FMath::Min<int32>(LastLine + 1, +HALF_WORLD_MAX/ViewportGridY) / IncScale;
