Added support for changed objects being recooked for cook on the side.
Converted cookonthefly to use fnames instead of string compares for platform names.
Stopped cook on the fly from loading packages if they are already loaded.
Removed some cabbage code.

diff --git a/Engine/Source/Editor/UnrealEd/Classes/Commandlets/CookCommandlet.h b/Engine/Source/Editor/UnrealEd/Classes/Commandlets/CookCommandlet.h
--- a/Engine/Source/Editor/UnrealEd/Classes/Commandlets/CookCommandlet.h
+++ b/Engine/Source/Editor/UnrealEd/Classes/Commandlets/CookCommandlet.h
@@ -46,7 +46,7 @@ class UCookCommandlet
 	 *
 	 * @return true on success, false otherwise.
 	 */
-	bool CookOnTheFly( bool BindAnyPort, int32 Timeout = 180, bool bForceClose = false );
+	bool CookOnTheFly( FGuid InstanceId, int32 Timeout = 180, bool bForceClose = false );
 
 	/**
 	 * Returns cooker output directory.
@@ -101,9 +101,6 @@ public:
 
 private:
 
-	/** Holds the instance identifier. */
-	FGuid InstanceId;
-
 	/** Holds the sandbox file wrapper to handle sandbox path conversion. */
 	TAutoPtr<class FSandboxPlatformFile> SandboxFile;
 
diff --git a/Engine/Source/Editor/UnrealEd/Classes/CookOnTheSide/CookOnTheFlyServer.h b/Engine/Source/Editor/UnrealEd/Classes/CookOnTheSide/CookOnTheFlyServer.h
--- a/Engine/Source/Editor/UnrealEd/Classes/CookOnTheSide/CookOnTheFlyServer.h
+++ b/Engine/Source/Editor/UnrealEd/Classes/CookOnTheSide/CookOnTheFlyServer.h
@@ -28,10 +28,6 @@ private:
 	bool bAsyncSaveLoad;
 	/** Should we auto tick (used in the editor) */
 	bool bIsAutoTick;
-
-	/** Holds the instance identifier. */
-	FGuid InstanceId;
-
 	/** Holds the sandbox file wrapper to handle sandbox path conversion. */
 	TAutoPtr<class FSandboxPlatformFile> SandboxFile;
 	/** Cook on the fly server uses the NetworkFileServer */
@@ -40,29 +36,29 @@ private:
 	/** We hook this up to a delegate to avoid reloading textures and whatnot */
 	TSet<FString> PackagesToNotReload;
 
-
+private:
 
 	/** Array which has been made thread safe :) */
-	template<typename Type>
-	struct FThreadSafeQueue
+	template<typename Type, typename SynchronizationObjectType, typename ScopeLockType>
+	struct FUnsyncronizedQueue
 	{
 	private:
-		mutable FCriticalSection	SynchronizationObject; // made this mutable so this class can have const functions and still be thread safe
+		mutable SynchronizationObjectType	SynchronizationObject; // made this mutable so this class can have const functions and still be thread safe
 		TArray<Type>		Items;
 	public:
 		void Enqueue(const Type& Item)
 		{
-			FScopeLock ScopeLock(&SynchronizationObject);
+			ScopeLockType ScopeLock(&SynchronizationObject);
 			Items.Add(Item);
 		}
 		void EnqueueUnique( const Type& Item )
 		{
-			FScopeLock ScopeLock(&SynchronizationObject);
+			ScopeLockType ScopeLock(&SynchronizationObject);
 			Items.AddUnique(Item);
 		}
 		bool Dequeue(Type* Result)
 		{
-			FScopeLock ScopeLock(&SynchronizationObject);
+			ScopeLockType ScopeLock(&SynchronizationObject);
 			if (Items.Num())
 			{
 				*Result = Items[0];
@@ -73,65 +69,70 @@ private:
 		}
 		void DequeueAll(TArray<Type>& Results)
 		{
-			FScopeLock ScopeLock(&SynchronizationObject);
+			ScopeLockType ScopeLock(&SynchronizationObject);
 			Results += Items;
 			Items.Empty();
 		}
 
 		bool HasItems() const
 		{
-			FScopeLock ScopeLock( &SynchronizationObject );
+			ScopeLockType ScopeLock( &SynchronizationObject );
 			return Items.Num() > 0;
 		}
 
 		void Remove( const Type& Item ) 
 		{
-			FScopeLock ScopeLock( &SynchronizationObject );
+			ScopeLockType ScopeLock( &SynchronizationObject );
 			Items.Remove( Item );
 		}
-	};
 
-	/** Normal queue to match the interface of FThreadSafeQueue :) */
-	template<typename Type>
-	struct FQueue
-	{
-	private:
-		TArray<Type>		Items;
-	public:
-		void Enqueue(const Type& Item)
-		{
-			Items.Add(Item);
-		}
-		void EnqueueUnique( const Type& Item )
+		void CopyItems( TArray<Type> &InItems ) const
 		{
-			Items.AddUnique(Item);
-		}
-		bool Dequeue(Type* Result)
-		{
-			if (Items.Num())
-			{
-				*Result = Items[0];
-				Items.RemoveAt(0);
-				return true;
-			}
-			return false;
-		}
-		void DequeueAll(TArray<Type>& Results)
-		{
-			Results += Items;
-			Items.Empty();
+			ScopeLockType ScopeLock( &SynchronizationObject );
+			InItems = Items;
 		}
 
-		bool HasItems() const
+		int Num() const 
 		{
-			return Items.Num() > 0;
+			ScopeLockType ScopeLock( &SynchronizationObject );
+			return Items.Num();
 		}
+	};
 
 
-		void Remove( const Type& Item ) 
-		{
-			Items.Remove( Item );
-		}
+
+	struct FDummyCriticalSection
+	{
+	public:
+		FORCEINLINE void Lock() { }
+		FORCEINLINE void Unlock() { }
+	};
+
+	struct FDummyScopeLock
+	{
+	public:
+		FDummyScopeLock( FDummyCriticalSection * ) { }
+	};
+
+public:
+
+
+	template<typename Type>
+	struct FThreadSafeQueue : public FUnsyncronizedQueue<Type, FCriticalSection, FScopeLock>
+	{
+		/**
+		 * Don't add any functions here, this is just a overqualified typedef
+		 * Add functions / functionality to the FUnsyncronizedQueue
+		 */
+	};
+
+	template<typename Type>
+	struct FQueue : public FUnsyncronizedQueue<Type, FDummyCriticalSection, FDummyScopeLock>
+	{
+		/**
+		 * Don't add any functions here, this is just a overqualified typedef
+		 * Add functions / functionality to the FUnsyncronizedQueue
+		 */
 	};
 
 public:
@@ -140,12 +141,12 @@ public:
 	{
 	public:
 		FFilePlatformRequest() { }
-		FFilePlatformRequest( const FString &InFilename, const FString &InPlatformname ) : Filename( FPaths::ConvertRelativePathToFull( InFilename) ), Platformname( InPlatformname ) 
+		FFilePlatformRequest( const FString &InFilename, const FName &InPlatformname ) : Filename( InFilename ), Platformname( InPlatformname ) 
 		{
 		}
 
 		FString Filename;
-		FString Platformname;
+		FName Platformname;
 
 		bool IsValid()  const
 		{
@@ -154,7 +155,6 @@ public:
 
 		void Clear()
 		{
-
 			Filename = TEXT("");
 			Platformname = TEXT("");
 		}
@@ -171,7 +171,7 @@ public:
 
 		FORCEINLINE FString ToString() const
 		{
-			return FString::Printf( TEXT("%s %s"), *Filename, *Platformname );
+			return FString::Printf( TEXT("%s %s"), *Filename, *Platformname.ToString() );
 		}
 
 	};
@@ -179,28 +179,43 @@ public:
 private:
 
 	FThreadSafeQueue<FFilePlatformRequest> FileRequests; // list of requested files
-	FQueue<FFilePlatformRequest> UnsolicitedFileRequests;
-	FThreadSafeQueue<FFilePlatformRequest> UnsolicitedCookedPackages; // list of files which haven't been requested but we think should cook based on previous requests
-
+	FQueue<FFilePlatformRequest> UnsolicitedFileRequests; // list of files which haven't been requested but we think should cook based on previous requests
+	FThreadSafeQueue<FFilePlatformRequest> UnsolicitedCookedPackages; // list of files which weren't requested but were cooked (based on UnsolicitedFileRequests)
 
 
 	/** Helper list of all files which have been cooked */
 	struct FThreadSafeFilenameSet
 	{
-		FCriticalSection	SynchronizationObject;
+	private:
+		mutable FCriticalSection	SynchronizationObject;
 		TSet<FFilePlatformRequest> FilesProcessed;
-
+	public:
 		void Add(const FFilePlatformRequest &Filename)
 		{
 			FScopeLock ScopeLock(&SynchronizationObject);
 			check(Filename.IsValid());
 			FilesProcessed.Add(Filename);
 		}
-		bool Exists(const FFilePlatformRequest &Filename)
+		bool Exists(const FFilePlatformRequest &Filename) const
 		{
 			FScopeLock ScopeLock(&SynchronizationObject);
 			return FilesProcessed.Contains(Filename);
 		}
+		int RemoveAll( const FString &Filename )
+		{
+			FScopeLock ScopeLock( &SynchronizationObject );
+			int NumRemoved = 0;
+			// for ( TSet<FFilePlatformRequest>::TIterator It( FilesProcessed ); It; ++It )
+			for ( auto It = FilesProcessed.CreateIterator(); It; ++It )
+			{
+				if ( It->Filename == Filename )
+				{
+					It.RemoveCurrent();
+					++NumRemoved;
+				}
+			}
+			return NumRemoved;
+		}
 	};
 	FThreadSafeFilenameSet ThreadSafeFilenameSet; // set of files which have been cooked when needing to recook a file the entry will need to be removed from here
 
@@ -246,6 +261,17 @@ public:
 	 * @return true on success, false otherwise.
 	 */
 	bool StartNetworkFileServer( bool BindAnyPort );
+
+	/**
+	 * Broadcast our the fileserver presence on the network
+	 */
+	bool BroadcastFileserverPresence( const FGuid &InstanceId );
+	/** 
+	 * Stop the network file server
+	 *
+	 */
+	void EndNetworkFileServer();
+
 	
 	/**
 	 * Handles cook package requests until there are no more requests, then returns
@@ -267,12 +293,6 @@ public:
 
 	bool HasRecompileShaderRequests() const { return RecompileRequests.HasItems(); }
 
-	/** 
-	 * Stop the network file server
-	 *
-	 */
-	void EndNetworkFileServer();
-
 
 	uint32 NumConnections() const;
 
@@ -280,14 +300,30 @@ public:
 
 
 	virtual void BeginDestroy() OVERRIDE;
+	
+
+
+
+	/**
+	 * Callbacks from editor 
+	 */
+
+	void OnObjectModified( UObject *ObjectMoving );
+	void OnObjectPropertyChanged(UObject* ObjectBeingModified);
+	void OnObjectUpdated( UObject *Object );
 
+	/**
+	 * Marks a package as dirty for cook
+	 * causes package to be recooked on next request (and all dependent packages which are currently cooked)
+	 */
+	void MarkPackageDirtyForCooker( UPackage *Package );
 
 
 
 private:
 	
 
-	bool ShouldCook(const FString& InFileName, const FString &InPlatformName);
+	bool ShouldCook(const FString& InFileName, const FName &InPlatformName);
 
 	/**
 	 * Returns cooker output directory.
@@ -328,9 +364,7 @@ private:
 	 *
 	 *	@return	bool			true if packages was cooked
 	 */
-	bool SaveCookedPackage( UPackage* Package, uint32 SaveFlags, bool& bOutWasUpToDate, TArray<FString> &TargetPlatformNames );
-
-	// bool ShouldCook(const FString& InFilename, const FString &InPlatformname = TEXT(""));
+	bool SaveCookedPackage( UPackage* Package, uint32 SaveFlags, bool& bOutWasUpToDate, TArray<FName> &TargetPlatformNames );
 
 	// Callback for handling a network file request.
 	void HandleNetworkFileServerFileRequest( const FString& Filename, const FString &Platformname, TArray<FString>& UnsolicitedFiles );
diff --git a/Engine/Source/Editor/UnrealEd/Private/Commandlets/CookCommandlet.cpp b/Engine/Source/Editor/UnrealEd/Private/Commandlets/CookCommandlet.cpp
--- a/Engine/Source/Editor/UnrealEd/Private/Commandlets/CookCommandlet.cpp
+++ b/Engine/Source/Editor/UnrealEd/Private/Commandlets/CookCommandlet.cpp
@@ -59,7 +59,7 @@ UCookCommandlet::UCookCommandlet( const class FPostConstructInitializeProperties
 /* UCookCommandlet interface
  *****************************************************************************/
 
-bool UCookCommandlet::CookOnTheFly( bool BindAnyPort, int32 Timeout, bool bForceClose )
+bool UCookCommandlet::CookOnTheFly( FGuid InstanceId, int32 Timeout, bool bForceClose )
 {
 	UCookOnTheFlyServer *CookOnTheFlyServer = ConstructObject<UCookOnTheFlyServer>( UCookOnTheFlyServer::StaticClass() );
 
@@ -82,10 +82,20 @@ bool UCookCommandlet::CookOnTheFly( bool BindAnyPort, int32 Timeout, bool bForce
 
 	CookOnTheFlyServer->Initialize( bCompressed, bIterativeCooking, bSkipEditorContent );
 
+	bool BindAnyPort = InstanceId.IsValid();
 
 	if ( CookOnTheFlyServer->StartNetworkFileServer(BindAnyPort) == false )
+	{
 		return false;
+	}
 
+	if ( InstanceId.IsValid() )
+	{
+		if ( CookOnTheFlyServer->BroadcastFileserverPresence(InstanceId) == false )
+		{
+			return false;
+		}
+	}
 
 	// Garbage collection should happen when either
 	//	1. We have cooked a map
@@ -105,15 +115,14 @@ bool UCookCommandlet::CookOnTheFly( bool BindAnyPort, int32 Timeout, bool bForce
 	FDateTime LastConnectionTime = FDateTime::UtcNow();
 	bool bHadConnection = false;
 
-	
-
+	bool bCookedAMapSinceLastGC = false;
 	while (!GIsRequestingExit)
 	{
 		uint32 TickResults = 0;
 		static const float CookOnTheSideTimeSlice = 10.0f;
 		TickResults = CookOnTheFlyServer->TickCookOnTheSide(CookOnTheSideTimeSlice, NonMapPackageCountSinceLastGC);
 
-		bool bCookedAMapSinceLastGC = TickResults & UCookOnTheFlyServer::COSR_CookedMap;
+		bCookedAMapSinceLastGC |= TickResults & UCookOnTheFlyServer::COSR_CookedMap;
 		if ( TickResults & (UCookOnTheFlyServer::COSR_CookedMap | UCookOnTheFlyServer::COSR_CookedPackage))
 		{
 			LastCookActionTime = FPlatformTime::Seconds();
@@ -130,7 +139,13 @@ bool UCookCommandlet::CookOnTheFly( bool BindAnyPort, int32 Timeout, bool bForce
 					bShouldGC = (NonMapPackageCountSinceLastGC > PackagesPerGC) || 
 						((FPlatformTime::Seconds() - LastCookActionTime) >= IdleTimeToGC);
 				}
-				bShouldGC |= bCookedAMapSinceLastGC;
+
+				// delay the gc until we process some unsolicited packages
+				if ( bCookedAMapSinceLastGC && (CookOnTheFlyServer->HasUnsolicitedCookRequests() == false) )
+				{
+					UE_LOG( LogCookCommandlet, Display, TEXT("Delaying map gc because we have unsolicited cook requests") );
+					bShouldGC |= bCookedAMapSinceLastGC;
+				}
 
 				if (bShouldGC)
 				{
@@ -497,6 +512,7 @@ int32 UCookCommandlet::Main(const FString& CmdLineParams)
 		FString InstanceIdString;
 		bool bForceClose = Switches.Contains(TEXT("FORCECLOSE"));
 
+		FGuid InstanceId;
 		if (FParse::Value(*Params, TEXT("InstanceId="), InstanceIdString))
 		{
 			if (!FGuid::Parse(InstanceIdString, InstanceId))
@@ -505,7 +521,13 @@ int32 UCookCommandlet::Main(const FString& CmdLineParams)
 			}
 		}
 
-		CookOnTheFly( InstanceId.IsValid() );
+		int32 Timeout = 180;
+		if (!FParse::Value(*Params, TEXT("timeout="), Timeout))
+		{
+			Timeout = 180;
+		}
+
+		CookOnTheFly( InstanceId, Timeout, bForceClose);
 	}
 	else
 	{
diff --git a/Engine/Source/Editor/UnrealEd/Private/CookOnTheFlyServer.cpp b/Engine/Source/Editor/UnrealEd/Private/CookOnTheFlyServer.cpp
--- a/Engine/Source/Editor/UnrealEd/Private/CookOnTheFlyServer.cpp
+++ b/Engine/Source/Editor/UnrealEd/Private/CookOnTheFlyServer.cpp
@@ -24,7 +24,7 @@
 
 
 
-DEFINE_LOG_CATEGORY_STATIC(LogCookCommandlet, Log, All);
+DEFINE_LOG_CATEGORY_STATIC(LogCookOnTheFly, Log, All);
 
 #define DEBUG_COOKONTHEFLY 0
 
@@ -103,11 +103,20 @@ bool UCookOnTheFlyServer::StartNetworkFileServer( bool BindAnyPort )
 	NetworkFileServer = FModuleManager::LoadModuleChecked<INetworkFileSystemModule>("NetworkFileSystem")
 		.CreateNetworkFileServer(BindAnyPort ? 0 : -1, &FileRequestDelegate, &RecompileShadersDelegate);
 
+
+	// loop while waiting for requests
+	GIsRequestingExit = false;
+	return true;
+}
+
+
+bool UCookOnTheFlyServer::BroadcastFileserverPresence( const FGuid &InstanceId )
+{
 	TArray<TSharedPtr<FInternetAddr> > AddressList;
 
 	if ((NetworkFileServer == NULL) || !NetworkFileServer->GetAddressList(AddressList))
 	{
-		UE_LOG(LogCookCommandlet, Error, TEXT("Failed to create network file server"));
+		UE_LOG(LogCookOnTheFly, Error, TEXT("Failed to create network file server"));
 
 		return false;
 	}
@@ -129,9 +138,6 @@ bool UCookOnTheFlyServer::StartNetworkFileServer( bool BindAnyPort )
 			MessageEndpoint->Publish(new FFileServerReady(AddressStringList, InstanceId), EMessageScope::Network);
 		}		
 	}
-
-	// loop while waiting for requests
-	GIsRequestingExit = false;
 	return true;
 }
 
@@ -147,10 +153,11 @@ uint32 UCookOnTheFlyServer::TickCookOnTheSide( const float TimeSlice, uint32 &Co
 
 	// This is all the target platforms which we needed to process requests for this iteration
 	// we use this in the unsolicited packages processing below
-	TArray<FString> AllTargetPlatformNames;
-	bool bCookedPackage = false;
+	TArray<FName> AllTargetPlatformNames;
+
 	while (!GIsRequestingExit)
 	{
+		bool bCookedPackage = false;
 		FFilePlatformRequest ToBuild;
 
 		bool bIsUnsolicitedRequest = false;
@@ -170,38 +177,72 @@ uint32 UCookOnTheFlyServer::TickCookOnTheSide( const float TimeSlice, uint32 &Co
 		}
 
 		check( ToBuild.IsValid() );
-		TArray<FString> TargetPlatformNames;
+		TArray<FName> TargetPlatformNames;
 		
-		if( ToBuild.Platformname.Len() )
+		if( ToBuild.Platformname != NAME_None )
 		{
-			TargetPlatformNames.Add( ToBuild.Platformname );
+			TargetPlatformNames.AddUnique( ToBuild.Platformname );
 		}
 
 		bool bWasUpToDate = false;
-		FString Name = FPackageName::FilenameToLongPackageName(*ToBuild.Filename);
+
 		if (!ThreadSafeFilenameSet.Exists(ToBuild))
 		{
+
 			// if we have no target platforms then we want to cook because this will cook for all target platforms in that case
-			bool bShouldCook = TargetPlatformNames.Num() > 0 ? false : true;
+			bool bShouldCook = TargetPlatformNames.Num() > 0 ? false : ShouldCook( ToBuild.Filename, NAME_None );
 			for ( int Index = 0; Index < TargetPlatformNames.Num(); ++Index )
 			{
 				bShouldCook |= ShouldCook( ToBuild.Filename, TargetPlatformNames[Index] );
 			}
 
+			
 			if ( bShouldCook ) // if we should cook the package then cook it otherwise add it to the list of already cooked packages below
 			{
-				//  if the package is already loaded then don't do any work :)
-				UPackage* Package = LoadPackage( NULL, *ToBuild.Filename, LOAD_None );
+				FString PackageName;	
+				UPackage *Package = NULL;
+				if ( FPackageName::TryConvertFilenameToLongPackageName(ToBuild.Filename, PackageName) )
+				{
+					Package = FindObject<UPackage>( ANY_PACKAGE, *PackageName );
+				}
+
+				UE_LOG( LogCookOnTheFly, Display, TEXT("Processing request %s"), *ToBuild.Filename)
+
+				//  if the package is already loaded then try to avoid reloading it :)
+				if ( ( Package == NULL ) || ( Package->IsFullyLoaded() == false ) )
+				{
+					Package = LoadPackage( NULL, *ToBuild.Filename, LOAD_None );
+				}
+				else
+				{
+					UE_LOG(LogCookOnTheFly, Display, TEXT("Package already loaded %s avoiding reload"), *ToBuild.Filename);
+				}
 
 				if( Package == NULL )
 				{
-					UE_LOG(LogCookCommandlet, Error, TEXT("Error loading %s!"), *ToBuild.Filename );
+					UE_LOG(LogCookOnTheFly, Error, TEXT("Error loading %s!"), *ToBuild.Filename );
 					Result |= COSR_ErrorLoadingPackage;
 				}
 				else
 				{
 					if( SaveCookedPackage(Package, SAVE_KeepGUID | SAVE_Async, bWasUpToDate, TargetPlatformNames) )
 					{
+						FString Name = Package->GetPathName();
+						FString PackageFilename(GetPackageFilename(Package));
+						FPaths::MakeStandardFilename(PackageFilename);
+
+						if ( PackageFilename != ToBuild.Filename )
+						{
+							// we have saved something which we didn't mean to load 
+							//  sounds unpossible.... but it is due to searching for files and such
+							//  mark the original request as processed (if this isn't actually the file they were requesting then it will fail)
+							//	and then also save our new request as processed so we don't do it again
+							ThreadSafeFilenameSet.Add( ToBuild );
+							UE_LOG( LogCookOnTheFly, Display, TEXT("Request for %s received saved %s"), *ToBuild.Filename, *PackageFilename );
+							ToBuild.Filename = PackageFilename;
+						}
+						// ToBuild.Filename = PackageFilename;
+
 						// Update flags used to determine garbage collection.
 						if (Package->ContainsMap())
 						{
@@ -238,12 +279,10 @@ uint32 UCookOnTheFlyServer::TickCookOnTheSide( const float TimeSlice, uint32 &Co
 						(bWasUpToDate == false))
 					{
 						UnsolicitedCookedPackages.EnqueueUnique( FileRequest );
-						UE_LOG(LogCookCommandlet, Display, TEXT("UnsolicitedCookedPackages: %s"), *FileRequest.Filename);
+						UE_LOG(LogCookOnTheFly, Display, TEXT("UnsolicitedCookedPackages: %s"), *FileRequest.Filename);
 					}
 				}
 
-				UnsolicitedCookedPackages.Remove( FileRequest );
-				
 				ThreadSafeFilenameSet.Add( FileRequest );
 
 
@@ -251,55 +290,86 @@ uint32 UCookOnTheFlyServer::TickCookOnTheSide( const float TimeSlice, uint32 &Co
 		}
 		else
 		{
-			UE_LOG(LogCookCommandlet, Display, TEXT("Package for platform already cooked %s, discarding request"), *ToBuild.Filename);
+			UE_LOG(LogCookOnTheFly, Display, TEXT("Package for platform already cooked %s, discarding request"), *ToBuild.Filename);
 		}
-		
 
-		if ( (FPlatformTime::Seconds() - StartTime) > TimeSlice)
-		{
-			break;
-		}
-	}
 
-	if ( bCookedPackage )
-	{
-		TArray<UObject *> ObjectsInOuter;
-		GetObjectsWithOuter(NULL, ObjectsInOuter, false);
-		for( int32 Index = 0; Index < ObjectsInOuter.Num() && !GIsRequestingExit; Index++ )
+		if ( bCookedPackage && !bIsUnsolicitedRequest)
 		{
-			UPackage* Package = Cast<UPackage>(ObjectsInOuter[Index]);
-
-			if (Package)
+			TArray<UObject *> ObjectsInOuter;
+			GetObjectsWithOuter(NULL, ObjectsInOuter, false);
+			for( int32 Index = 0; Index < ObjectsInOuter.Num() && !GIsRequestingExit; Index++ )
 			{
-				FString Name = Package->GetPathName();
-				FString PackageFilename(GetPackageFilename(Package));
-				FPaths::MakeStandardFilename(PackageFilename);
+				UPackage* Package = Cast<UPackage>(ObjectsInOuter[Index]);
 
-				const FString FullPath = FPaths::ConvertRelativePathToFull(PackageFilename);
-
-				bool AlreadyCooked = true;
-				for ( int Index = 0; Index < AllTargetPlatformNames.Num(); ++Index )
+				if (Package)
 				{
-					if ( !ThreadSafeFilenameSet.Exists(FFilePlatformRequest(FullPath, AllTargetPlatformNames[Index])) )
+					FString Name = Package->GetPathName();
+					FString PackageFilename(GetPackageFilename(Package));
+					FPaths::MakeStandardFilename(PackageFilename);
+
+					const FString FullPath = FPaths::ConvertRelativePathToFull(PackageFilename);
+
+					bool AlreadyCooked = true;
+					for ( int Index = 0; Index < AllTargetPlatformNames.Num(); ++Index )
 					{
-						AlreadyCooked = false;
-						break;
+						if ( !ThreadSafeFilenameSet.Exists(FFilePlatformRequest(PackageFilename, AllTargetPlatformNames[Index])) )
+						{
+							AlreadyCooked = false;
+							break;
+						}
 					}
-				}
-				if ( !AlreadyCooked )
-				{
-					for ( int Index = 0; Index < AllTargetPlatformNames.Num(); ++Index )
+					if ( !AlreadyCooked )
 					{
-						const FFilePlatformRequest Request( FullPath, AllTargetPlatformNames[Index] );
-						UnsolicitedFileRequests.EnqueueUnique(Request);
+						for ( int Index = 0; Index < AllTargetPlatformNames.Num(); ++Index )
+						{
+							const FFilePlatformRequest Request( PackageFilename, AllTargetPlatformNames[Index] );
+							UnsolicitedFileRequests.EnqueueUnique(Request);
+						}
 					}
 				}
 			}
 		}
+
+		if ( (FPlatformTime::Seconds() - StartTime) > TimeSlice)
+		{
+			break;
+		}
 	}
+
 	return Result;
 }
 
+void UCookOnTheFlyServer::OnObjectModified( UObject *ObjectMoving )
+{
+	OnObjectUpdated( ObjectMoving );
+}
+
+void UCookOnTheFlyServer::OnObjectPropertyChanged(UObject* ObjectBeingModified)
+{
+	OnObjectUpdated( ObjectBeingModified );
+}
+
+void UCookOnTheFlyServer::OnObjectUpdated( UObject *Object )
+{
+	// get the outer of the object
+	UPackage *Package = Object->GetOutermost();
+
+	MarkPackageDirtyForCooker( Package );
+}
+
+void UCookOnTheFlyServer::MarkPackageDirtyForCooker( UPackage *Package )
+{
+	// force that package to be recooked
+	const FString Name = Package->GetPathName();
+	FString PackageFilename(GetPackageFilename(Package));
+	FPaths::MakeStandardFilename(PackageFilename);
+	UE_LOG(LogCookOnTheFly, Display, TEXT("Modification detected to package %s"), *PackageFilename);
+	ThreadSafeFilenameSet.RemoveAll( PackageFilename );
+
+	// TODO: go through package dependencies of this package, if a found dependency is not loaded then mark it as need recooking.
+}
+
 
 void UCookOnTheFlyServer::EndNetworkFileServer()
 {
@@ -365,12 +435,12 @@ bool UCookOnTheFlyServer::GetPackageTimestamp( const FString& InFilename, FDateT
 
 bool UCookOnTheFlyServer::SaveCookedPackage( UPackage* Package, uint32 SaveFlags, bool& bOutWasUpToDate ) 
 {
-	TArray<FString> TargetPlatformNames; 
+	TArray<FName> TargetPlatformNames; 
 	return SaveCookedPackage( Package, SaveFlags, bOutWasUpToDate, TargetPlatformNames );
 }
 
 
-bool UCookOnTheFlyServer::ShouldCook(const FString& InFileName, const FString &InPlatformName)
+bool UCookOnTheFlyServer::ShouldCook(const FString& InFileName, const FName &InPlatformName)
 {
 	bool bDoCook = false;
 
@@ -385,7 +455,7 @@ bool UCookOnTheFlyServer::ShouldCook(const FString& InFileName, const FString &I
 
 		if (GetPackageTimestamp(FPaths::GetBaseFilename(PkgFilename, false), DependentTimeStamp) == false)
 		{
-			UE_LOG(LogCookCommandlet, Display, TEXT("Failed to find dependency timestamp for: %s"), *PkgFilename);
+			UE_LOG(LogCookOnTheFly, Display, TEXT("Failed to find dependency timestamp for: %s"), *PkgFilename);
 		}
 	}
 
@@ -398,9 +468,9 @@ bool UCookOnTheFlyServer::ShouldCook(const FString& InFileName, const FString &I
 
 	TArray<ITargetPlatform*> Platforms;
 
-	if ( InPlatformName.Len() > 0 )
+	if ( InPlatformName != NAME_None )
 	{
-		Platforms.Add( TPM.FindTargetPlatform( InPlatformName ) );
+		Platforms.Add( TPM.FindTargetPlatform( InPlatformName.ToString() ) );
 	}
 	else 
 	{
@@ -428,7 +498,7 @@ bool UCookOnTheFlyServer::ShouldCook(const FString& InFileName, const FString &I
 	return bDoCook;
 }
 
-bool UCookOnTheFlyServer::SaveCookedPackage( UPackage* Package, uint32 SaveFlags, bool& bOutWasUpToDate, TArray<FString> &TargetPlatformNames )
+bool UCookOnTheFlyServer::SaveCookedPackage( UPackage* Package, uint32 SaveFlags, bool& bOutWasUpToDate, TArray<FName> &TargetPlatformNames )
 {
 
 
@@ -452,7 +522,7 @@ bool UCookOnTheFlyServer::SaveCookedPackage( UPackage* Package, uint32 SaveFlags
 
 			if (GetPackageTimestamp(FPaths::GetBaseFilename(PkgFilename, false), DependentTimeStamp) == false)
 			{
-				UE_LOG(LogCookCommandlet, Display, TEXT("Failed to find depedency timestamp for: %s"), *PkgFilename);
+				UE_LOG(LogCookOnTheFly, Display, TEXT("Failed to find depedency timestamp for: %s"), *PkgFilename);
 			}
 		}
 
@@ -479,7 +549,7 @@ bool UCookOnTheFlyServer::SaveCookedPackage( UPackage* Package, uint32 SaveFlags
 		{
 			for ( int Index = 0; Index < TargetPlatformNames.Num(); ++Index )
 			{
-				const FString &TargetPlatformName = TargetPlatformNames[Index];
+				const FString TargetPlatformName = TargetPlatformNames[Index].ToString();
 
 
 				const TArray<ITargetPlatform*>& TargetPlatforms = TPM.GetTargetPlatforms();	
@@ -501,7 +571,7 @@ bool UCookOnTheFlyServer::SaveCookedPackage( UPackage* Package, uint32 SaveFlags
 
 			for ( int Index = 0; Index < Platforms.Num(); ++Index )
 			{
-				TargetPlatformNames.Add( Platforms[Index]->PlatformName() );
+				TargetPlatformNames.Add( FName(*Platforms[Index]->PlatformName()) );
 			}
 		}
 		
@@ -535,7 +605,7 @@ bool UCookOnTheFlyServer::SaveCookedPackage( UPackage* Package, uint32 SaveFlags
 					Package->FullyLoad();
 					if (!Package->IsFullyLoaded())
 					{
-						UE_LOG(LogCookCommandlet, Warning, TEXT("Package %s supposed to be fully loaded but isn't. RF_WasLoaded is %s"), 
+						UE_LOG(LogCookOnTheFly, Warning, TEXT("Package %s supposed to be fully loaded but isn't. RF_WasLoaded is %s"), 
 							*Package->GetName(), Package->HasAnyFlags(RF_WasLoaded) ? TEXT("set") : TEXT("not set"));
 					}
 					bPackageFullyLoaded = true;
@@ -549,7 +619,7 @@ bool UCookOnTheFlyServer::SaveCookedPackage( UPackage* Package, uint32 SaveFlags
 					Flags = World ? RF_NoFlags : RF_Standalone;
 				}
 
-				UE_LOG(LogCookCommandlet, Display, TEXT("Cooking %s -> %s"), *Package->GetName(), *PlatFilename);
+				UE_LOG(LogCookOnTheFly, Display, TEXT("Cooking %s -> %s"), *Package->GetName(), *PlatFilename);
 
 				bool bSwap = (!Target->IsLittleEndian()) ^ (!PLATFORM_LITTLE_ENDIAN);
 
@@ -570,7 +640,7 @@ bool UCookOnTheFlyServer::SaveCookedPackage( UPackage* Package, uint32 SaveFlags
 
 				if( PlatFilename.Len() >= PLATFORM_MAX_FILEPATH_LENGTH )
 				{
-					UE_LOG( LogCookCommandlet, Error, TEXT( "Couldn't save package, filename is too long :%s" ), *PlatFilename );
+					UE_LOG( LogCookOnTheFly, Error, TEXT( "Couldn't save package, filename is too long :%s" ), *PlatFilename );
 					bSavedCorrectly = false;
 				}
 				else
@@ -583,7 +653,7 @@ bool UCookOnTheFlyServer::SaveCookedPackage( UPackage* Package, uint32 SaveFlags
 			}
 			else
 			{
-				UE_LOG(LogCookCommandlet, Display, TEXT("Up to date: %s"), *PlatFilename);
+				UE_LOG(LogCookOnTheFly, Display, TEXT("Up to date: %s"), *PlatFilename);
 
 				bOutWasUpToDate = true;
 			}
@@ -602,7 +672,7 @@ void UCookOnTheFlyServer::MaybeMarkPackageAsAlreadyLoaded(UPackage *Package)
 	FString Name = Package->GetName();
 	if (PackagesToNotReload.Contains(Name))
 	{
-		UE_LOG(LogCookCommandlet, Verbose, TEXT("Marking %s already loaded."), *Name);
+		UE_LOG(LogCookOnTheFly, Verbose, TEXT("Marking %s already loaded."), *Name);
 		Package->PackageFlags |= PKG_ReloadingForCooker;
 	}
 }
@@ -719,7 +789,7 @@ void UCookOnTheFlyServer::CleanSandbox(const TArray<ITargetPlatform*>& Platforms
 
 						if (Diff < 0.0)
 						{
-							UE_LOG(LogCookCommandlet, Display, TEXT("Deleting out of date cooked file: %s"), *CookedFilename);
+							UE_LOG(LogCookOnTheFly, Display, TEXT("Deleting out of date cooked file: %s"), *CookedFilename);
 
 							IFileManager::Get().Delete(*CookedFilename);
 						}
@@ -732,7 +802,7 @@ void UCookOnTheFlyServer::CleanSandbox(const TArray<ITargetPlatform*>& Platforms
 		}
 	}
 
-	UE_LOG(LogCookCommandlet, Display, TEXT("Sandbox cleanup took %5.3f seconds"), SandboxCleanTime);
+	UE_LOG(LogCookOnTheFly, Display, TEXT("Sandbox cleanup took %5.3f seconds"), SandboxCleanTime);
 }
 
 void UCookOnTheFlyServer::GenerateAssetRegistry(const TArray<ITargetPlatform*>& Platforms)
@@ -744,7 +814,7 @@ void UCookOnTheFlyServer::GenerateAssetRegistry(const TArray<ITargetPlatform*>&
 	double GenerateAssetRegistryTime = 0.0;
 	{
 		SCOPE_SECONDS_COUNTER(GenerateAssetRegistryTime);
-		UE_LOG(LogCookCommandlet, Display, TEXT("Creating asset registry [is editor: %d]"), GIsEditor);
+		UE_LOG(LogCookOnTheFly, Display, TEXT("Creating asset registry [is editor: %d]"), GIsEditor);
 
 		// Perform a synchronous search of any .ini based asset paths (note that the per-game delegate may
 		// have already scanned paths on its own)
@@ -765,7 +835,7 @@ void UCookOnTheFlyServer::GenerateAssetRegistry(const TArray<ITargetPlatform*>&
 			// write it out to a memory archive
 			FArrayWriter SerializedAssetRegistry;
 			AssetRegistry.Serialize(SerializedAssetRegistry);
-			UE_LOG(LogCookCommandlet, Display, TEXT("Generated asset registry size is %5.2fkb"), (float)SerializedAssetRegistry.Num() / 1024.f);
+			UE_LOG(LogCookOnTheFly, Display, TEXT("Generated asset registry size is %5.2fkb"), (float)SerializedAssetRegistry.Num() / 1024.f);
 
 			// now save it in each cooked directory
 			FString RegistryFilename = FPaths::GameDir() / TEXT("AssetRegistry.bin");
@@ -779,7 +849,7 @@ void UCookOnTheFlyServer::GenerateAssetRegistry(const TArray<ITargetPlatform*>&
 			}
 		}
 	}
-	UE_LOG(LogCookCommandlet, Display, TEXT("Done creating registry. It took %5.2fs."), GenerateAssetRegistryTime);
+	UE_LOG(LogCookOnTheFly, Display, TEXT("Done creating registry. It took %5.2fs."), GenerateAssetRegistryTime);
 	
 }
 
@@ -803,7 +873,7 @@ void UCookOnTheFlyServer::GenerateLongPackageNames(TArray<FString>& FilesInPath)
 			}
 			else
 			{
-				UE_LOG(LogCookCommandlet, Warning, TEXT("Unable to generate long package name for %s"), *FileInPath);
+				UE_LOG(LogCookOnTheFly, Warning, TEXT("Unable to generate long package name for %s"), *FileInPath);
 			}
 		}
 	}
@@ -819,19 +889,33 @@ void UCookOnTheFlyServer::HandleNetworkFileServerFileRequest( const FString& Fil
 {
 	bool bIsCookable = FPackageName::IsPackageExtension(*FPaths::GetExtension(Filename, true));
 
+
+	FName PlatformFname = FName( *Platformname );
+
 	if (!bIsCookable)
 	{
 		TArray<FFilePlatformRequest> FileRequests;
 		UnsolicitedCookedPackages.DequeueAll(FileRequests);
 		for ( int Index=  0; Index < FileRequests.Num(); ++Index)
 		{
-			UnsolicitedFiles.Add( FileRequests[Index].Filename);
+			const FFilePlatformRequest &Request = FileRequests[Index];
+			if ( Request.Platformname == PlatformFname )
+			{
+				FString StandardFilename = FileRequests[Index].Filename;
+				FPaths::MakeStandardFilename( StandardFilename );
+				UnsolicitedFiles.Add( StandardFilename );
+			}
+			else
+			{
+				UnsolicitedCookedPackages.Enqueue( Request );
+			}
 		}
 		UPackage::WaitForAsyncFileWrites();
 		return;
 	}
-
-	FFilePlatformRequest FileRequest( Filename, Platformname );
+	FString StandardFileName = Filename;
+	FPaths::MakeStandardFilename( StandardFileName );
+	FFilePlatformRequest FileRequest( StandardFileName, PlatformFname );
 	FileRequests.Enqueue(FileRequest);
 
 	do
@@ -840,6 +924,8 @@ void UCookOnTheFlyServer::HandleNetworkFileServerFileRequest( const FString& Fil
 	}
 	while (!ThreadSafeFilenameSet.Exists(FileRequest));
 
+	UE_LOG( LogCookOnTheFly, Display, TEXT("Cook complete %s"), *FileRequest.Filename)
+
 	TArray<FFilePlatformRequest> FileRequests;
 	UnsolicitedCookedPackages.DequeueAll(FileRequests);
 	for ( int Index=  0; Index < FileRequests.Num(); ++Index)
@@ -849,23 +935,22 @@ void UCookOnTheFlyServer::HandleNetworkFileServerFileRequest( const FString& Fil
 		{
 			// don't do anything we don't want this guy
 		}
-		else if ( Request.Platformname == Platformname )
+		else if ( Request.Platformname == PlatformFname )
 		{
-			FString StandardFilename = FileRequests[Index].Filename;
-			FPaths::MakeStandardFilename( StandardFilename );
-			UnsolicitedFiles.Add( StandardFilename );
+			UE_LOG( LogCookOnTheFly, Display, TEXT("Returning unsolicited file %s"), *Request.Filename );
+			UnsolicitedFiles.Add( Request.Filename );
 		}
 		else
 		{
 			UnsolicitedCookedPackages.Enqueue( Request );
 		}
 	}
-	
+
 	UPackage::WaitForAsyncFileWrites();
 
 
 #if DEBUG_COOKONTHEFLY
-	UE_LOG( LogCookCommandlet, Display, TEXT("Processed file request %s"), *Filename );
+	UE_LOG( LogCookOnTheFly, Display, TEXT("Processed file request %s"), *Filename );
 #endif
 
 }
@@ -876,7 +961,7 @@ void UCookOnTheFlyServer::HandleNetworkFileServerRecompileShaders(const FShaderR
 	// if we aren't in the game thread, we need to push this over to the game thread and wait for it to finish
 	if (!IsInGameThread())
 	{
-		UE_LOG(LogCookCommandlet, Display, TEXT("Got a recompile request on non-game thread"));
+		UE_LOG(LogCookOnTheFly, Display, TEXT("Got a recompile request on non-game thread"));
 
 		// make a new request
 		FRecompileRequest* Request = new FRecompileRequest;
@@ -892,7 +977,7 @@ void UCookOnTheFlyServer::HandleNetworkFileServerRecompileShaders(const FShaderR
 			FPlatformProcess::Sleep(0);
 		}
 		delete Request;
-		UE_LOG(LogCookCommandlet, Display, TEXT("Completed recompile..."));
+		UE_LOG(LogCookOnTheFly, Display, TEXT("Completed recompile..."));
 
 		// at this point, we are done on the game thread, and ModifiedFiles will have been filled out
 		return;
diff --git a/Engine/Source/Editor/UnrealEd/Private/UnrealEdEngine.cpp b/Engine/Source/Editor/UnrealEd/Private/UnrealEdEngine.cpp
--- a/Engine/Source/Editor/UnrealEd/Private/UnrealEdEngine.cpp
+++ b/Engine/Source/Editor/UnrealEd/Private/UnrealEdEngine.cpp
@@ -99,6 +99,9 @@ void UUnrealEdEngine::Init(IEngineLoop* InEngineLoop)
 		CookServer = ConstructObject<UCookOnTheFlyServer>( UCookOnTheFlyServer::StaticClass() );
 		CookServer->Initialize( false, false, false, true, true );
 		CookServer->StartNetworkFileServer( false );
+
+		FCoreDelegates::OnObjectPropertyChanged.Add( FCoreDelegates::FOnObjectPropertyChanged::FDelegate::CreateUObject(CookServer, &UCookOnTheFlyServer::OnObjectPropertyChanged) );
+		FCoreDelegates::OnObjectModified.Add( FCoreDelegates::FOnObjectModified::FDelegate::CreateUObject(CookServer, &UCookOnTheFlyServer::OnObjectModified) );
 	}
 }
 
@@ -199,6 +202,12 @@ UUnrealEdEngine::~UUnrealEdEngine()
 
 void UUnrealEdEngine::FinishDestroy()
 {
+	if( CookServer)
+	{
+		FCoreDelegates::OnObjectPropertyChanged.RemoveAll( CookServer );
+		FCoreDelegates::OnObjectModified.RemoveAll( CookServer );
+	}
+
 	if(PackageAutoSaver.Get())
 	{
 		// We've finished shutting down, so disable the auto-save restore
diff --git a/Engine/Source/Runtime/Core/Private/Misc/CallbackDevice.cpp b/Engine/Source/Runtime/Core/Private/Misc/CallbackDevice.cpp
--- a/Engine/Source/Runtime/Core/Private/Misc/CallbackDevice.cpp
+++ b/Engine/Source/Runtime/Core/Private/Misc/CallbackDevice.cpp
@@ -28,7 +28,9 @@ FCoreDelegates::FOnHandleSystemError FCoreDelegates::OnHandleSystemError;
 FCoreDelegates::FOnPreObjectPropertyChanged FCoreDelegates::OnPreObjectPropertyChanged;
 FCoreDelegates::FOnObjectPropertyChanged FCoreDelegates::OnObjectPropertyChanged;
 FCoreDelegates::FOnActorLabelChanged FCoreDelegates::OnActorLabelChanged;
+
 #if WITH_EDITOR
+	FCoreDelegates::FOnObjectModified  FCoreDelegates::OnObjectModified;
 	FCoreDelegates::FOnAssetLoaded FCoreDelegates::OnAssetLoaded;
 	FSimpleMulticastDelegate FCoreDelegates::PreModal;
 	FSimpleMulticastDelegate FCoreDelegates::PostModal;
diff --git a/Engine/Source/Runtime/Core/Public/Misc/CallbackDevice.h b/Engine/Source/Runtime/Core/Public/Misc/CallbackDevice.h
--- a/Engine/Source/Runtime/Core/Public/Misc/CallbackDevice.h
+++ b/Engine/Source/Runtime/Core/Public/Misc/CallbackDevice.h
@@ -54,6 +54,9 @@ public:
 	DECLARE_MULTICAST_DELEGATE_OneParam(FOnPreObjectPropertyChanged, UObject*);
 
 	#if WITH_EDITOR
+	// Callback for all object modifications
+	DECLARE_MULTICAST_DELEGATE_OneParam(FOnObjectModified, UObject*);
+
 	// Callback for when an asset is loaded (Editor)
 	DECLARE_MULTICAST_DELEGATE_OneParam(FOnAssetLoaded, UObject*);
 	#endif	//WITH_EDITOR
@@ -115,6 +118,9 @@ public:
 	static FOnActorLabelChanged OnActorLabelChanged;
 
 #if WITH_EDITOR
+	// Called when any object is modified at all 
+	static FOnObjectModified OnObjectModified;
+
 	// Called when an asset is loaded
 	static FOnAssetLoaded OnAssetLoaded;
 
diff --git a/Engine/Source/Runtime/CoreUObject/Private/UObject/Obj.cpp b/Engine/Source/Runtime/CoreUObject/Private/UObject/Obj.cpp
--- a/Engine/Source/Runtime/CoreUObject/Private/UObject/Obj.cpp
+++ b/Engine/Source/Runtime/CoreUObject/Private/UObject/Obj.cpp
@@ -701,6 +701,9 @@ bool UObject::Modify( bool bAlwaysMarkDirty/*=true*/ )
 			MarkPackageDirty();
 		}
 	}
+#if WITH_EDITOR
+	FCoreDelegates::OnObjectModified.Broadcast(this);
+#endif
 	return bSavedToTransactionBuffer;
 }
 
diff --git a/Engine/Source/Runtime/CoreUObject/Private/UObject/SavePackage.cpp b/Engine/Source/Runtime/CoreUObject/Private/UObject/SavePackage.cpp
--- a/Engine/Source/Runtime/CoreUObject/Private/UObject/SavePackage.cpp
+++ b/Engine/Source/Runtime/CoreUObject/Private/UObject/SavePackage.cpp
@@ -2499,31 +2499,6 @@ bool UPackage::SavePackage( UPackage* InOuter, UObject* Base, EObjectFlags TopLe
 
 
 			UE_LOG_COOK_TIME(TEXT("TagPackageExports"));
-
-
-			// Kick off any Precaching required for the target platform to save these objects
-			// only need to do this if we are cooking a different platform then the one which is currently running
-			// TODO: if save package is cancelled then call ClearCache on each object
-			TArray<UObject*> CachedObjects;
-			if ( !!TargetPlatform )
-			{
-				TArray<UObject*> TagExpObjects;
-				GetObjectsWithAnyMarks(TagExpObjects, OBJECTMARK_TagExp);
-				for ( int Index =0; Index < TagExpObjects.Num(); ++Index)
-				{
-					UObject *ExpObject = TagExpObjects[Index];
-					if ( ExpObject->HasAnyMarks( OBJECTMARK_TagExp ) )
-					{
-						ExpObject->BeginCacheForCookedPlatformData( TargetPlatform );
-						CachedObjects.Add( ExpObject );
-					}
-				}
-			}
-
-			
-
-
-			UE_LOG_COOK_TIME(TEXT("BeginCacheForCookedPlatformData"));
 		
 			{
 				// set GIsSavingPackage here as it is now illegal to create any new object references; they potentially wouldn't be saved correctly
@@ -2544,6 +2519,33 @@ bool UPackage::SavePackage( UPackage* InOuter, UObject* Base, EObjectFlags TopLe
 
 				UE_LOG_COOK_TIME(TEXT("TagPackageExports"));
 
+
+
+				// Kick off any Precaching required for the target platform to save these objects
+				// only need to do this if we are cooking a different platform then the one which is currently running
+				// TODO: if save package is cancelled then call ClearCache on each object
+				TArray<UObject*> CachedObjects;
+				if ( !!TargetPlatform )
+				{
+					TArray<UObject*> TagExpObjects;
+					GetObjectsWithAnyMarks(TagExpObjects, OBJECTMARK_TagExp);
+					for ( int Index =0; Index < TagExpObjects.Num(); ++Index)
+					{
+						UObject *ExpObject = TagExpObjects[Index];
+						if ( ExpObject->HasAnyMarks( OBJECTMARK_TagExp ) )
+						{
+							ExpObject->BeginCacheForCookedPlatformData( TargetPlatform );
+							CachedObjects.Add( ExpObject );
+						}
+					}
+				}
+
+
+
+
+				UE_LOG_COOK_TIME(TEXT("BeginCacheForCookedPlatformData"));
+
+
 				GWarn->UpdateProgress( ++CurSaveStep, TotalSaveSteps );
 
 
@@ -3561,12 +3563,13 @@ bool UPackage::SavePackage( UPackage* InOuter, UObject* Base, EObjectFlags TopLe
 
 				GWarn->UpdateProgress( ++CurSaveStep, TotalSaveSteps );
 			
-			}
 
-			for ( int CachedObjectIndex = 0; CachedObjectIndex < CachedObjects.Num(); ++CachedObjectIndex )
-			{
-				CachedObjects[CachedObjectIndex]->ClearCachedCookedPlatformData(TargetPlatform);
+				for ( int CachedObjectIndex = 0; CachedObjectIndex < CachedObjects.Num(); ++CachedObjectIndex )
+				{
+					CachedObjects[CachedObjectIndex]->ClearCachedCookedPlatformData(TargetPlatform);
+				}
 			}
+
 		}
 		if( Success == true )
 		{
diff --git a/Engine/Source/Runtime/Engine/Private/Materials/Material.cpp b/Engine/Source/Runtime/Engine/Private/Materials/Material.cpp
--- a/Engine/Source/Runtime/Engine/Private/Materials/Material.cpp
+++ b/Engine/Source/Runtime/Engine/Private/Materials/Material.cpp
@@ -348,36 +348,6 @@ void UMaterialInterface::AssertDefaultMaterialsPostLoaded()
 
 void SerializeInlineShaderMaps(const TMap<const ITargetPlatform*,TArray<FMaterialResource*>>& PlatformMaterialResourcesToSave, FArchive& Ar, FMaterialResource* OutMaterialResourcesLoaded[][ERHIFeatureLevel::Num])
 {
-
-	/*if ( Ar.IsLoading() || (Ar.IsCooking() == false) )
-	{
-		// shouldn't ever load into this saved material resources array
-		// if we are not cooking then SerializeInlineShaderMaps will ignore the SavedMaterialResources array
-		TArray<FMaterialResource*> SavedMaterialResources;
-		SerializeInlineShaderMaps( SavedMaterialResources, Ar, MaterialResources );
-		check( SavedMaterialResources.Num() == 0);
-	}
-	else if ( Ar.IsSaving() )
-	{
-		check( Ar.IsCooking() );
-		const TArray<FMaterialResource*> *CachedMaterialResourcesForPlatform = CachedMaterialResourcesForCooking.Find(Ar.CookingTarget());
-		// check( CachedMaterialResourcesForPlatform != NULL || (!HasAnyMarks( OBJECTMARK_TagExp ))  );
-		if ( CachedMaterialResourcesForPlatform != NULL )
-		{
-			SerializeInlineShaderMaps(*CachedMaterialResourcesForPlatform, Ar, MaterialResources);
-		}
-		else
-		{
-			// if we are saving to a linker and we are cooking then we should have definatly have had a cached material for the target platform (see UPackage::SavePackage() -> UObject::BeginCacheResourceForPlatform
-			//  if this check is triggering most likely because some object didn't have BeginCacheResourceForPlatform called from SavePackage OR
-			//  the archive isn't actually saving and GetLinker is returning a valid value for some other reason
-			check( Ar.GetLinker() == NULL );
-		}
-
-	}*/
-
-
-
 	if (Ar.IsSaving())
 	{
 		int32 NumResourcesToSave = 0;
diff --git a/Engine/Source/Runtime/Engine/Private/Materials/MaterialInstance.cpp b/Engine/Source/Runtime/Engine/Private/Materials/MaterialInstance.cpp
--- a/Engine/Source/Runtime/Engine/Private/Materials/MaterialInstance.cpp
+++ b/Engine/Source/Runtime/Engine/Private/Materials/MaterialInstance.cpp
@@ -1543,30 +1543,6 @@ void UMaterialInstance::Serialize(FArchive& Ar)
 			StaticParameters.Serialize(Ar);
 
 			SerializeInlineShaderMaps( CachedMaterialResourcesForCooking, Ar, StaticPermutationMaterialResources );
-
-
-			/*if ( Ar.IsLoading() || (Ar.IsCooking() == false) )
-			{
-				// shouldn't ever load into this saved material resources array
-				// if we are not cooking then SerializeInlineShaderMaps will ignore the SavedMaterialResources array
-				TArray<FMaterialResource*> SavedMaterialResources;
-				SerializeInlineShaderMaps( SavedMaterialResources, Ar, StaticPermutationMaterialResources );
-				check( SavedMaterialResources.Num() == 0);
-			}
-			else if ( Ar.IsSaving() )
-			{
-				check( Ar.IsCooking() );
-				const TArray<FMaterialResource*> *CachedMaterialResourcesForPlatform = CachedMaterialResourcesForCooking.Find(Ar.CookingTarget());
-				// check( CachedMaterialResourcesForPlatform != NULL || (!HasAnyMarks( OBJECTMARK_TagExp ))  );
-				if ( CachedMaterialResourcesForPlatform != NULL )
-				{
-					SerializeInlineShaderMaps(*CachedMaterialResourcesForPlatform, Ar, StaticPermutationMaterialResources);
-				}
-				else
-				{
-					check( Ar.GetLinker() == NULL );
-				}
-			}*/
 		}
 		else
 		{
