Fixing it so references to reinstanced function libraries get fixed up in existing blueprint's bytecode.

diff --git a/Engine/Source/Editor/UnrealEd/Private/Kismet2/KismetReinstanceUtilities.cpp b/Engine/Source/Editor/UnrealEd/Private/Kismet2/KismetReinstanceUtilities.cpp
--- a/Engine/Source/Editor/UnrealEd/Private/Kismet2/KismetReinstanceUtilities.cpp
+++ b/Engine/Source/Editor/UnrealEd/Private/Kismet2/KismetReinstanceUtilities.cpp
@@ -136,6 +136,9 @@ void FBlueprintCompileReinstancer::GenerateFieldMappings(TMap<UObject*, UObject*
 		UFunction* NewFunction = ClassToReinstance->FindFunctionByName(Func.Key, EIncludeSuperFlag::ExcludeSuper);
 		FieldMapping.Add(Func.Value, NewFunction);
 	}
+
+	UObject* NewCDO = ClassToReinstance->GetDefaultObject();
+	FieldMapping.Add(OriginalCDO, NewCDO);
 }
 
 FBlueprintCompileReinstancer::~FBlueprintCompileReinstancer()
