No need to reconstuct every single node anymore on PIE (most cases where this was needed before is already taken care of).
(TTP #326008)

diff --git a/Engine/Source/Editor/UnrealEd/Private/PlayLevel.cpp b/Engine/Source/Editor/UnrealEd/Private/PlayLevel.cpp
--- a/Engine/Source/Editor/UnrealEd/Private/PlayLevel.cpp
+++ b/Engine/Source/Editor/UnrealEd/Private/PlayLevel.cpp
@@ -1616,7 +1616,6 @@ void UEditorEngine::PlayInEditor( UWorld* InWorld, bool bInSimulateInEditor )
 					UPackage* Package = Cast<UPackage>(Blueprint->GetOutermost());
 					const bool bIsPackageDirty = Package ? Package->IsDirty() : false;
 
-					FBlueprintEditorUtils::RefreshAllNodes(Blueprint);
 					Blueprint->BroadcastChanged();
 
 					UE_LOG(LogPlayLevel, Log, TEXT("[PIE] Compiling %s before PIE..."), *Blueprint->GetName());
