For automation testing, giving us a way to dump the blueprint palette into a readable file (json).

diff --git a/Engine/Source/Editor/BlueprintGraph/Public/K2ActionMenuBuilder.h b/Engine/Source/Editor/BlueprintGraph/Public/K2ActionMenuBuilder.h
--- a/Engine/Source/Editor/BlueprintGraph/Public/K2ActionMenuBuilder.h
+++ b/Engine/Source/Editor/BlueprintGraph/Public/K2ActionMenuBuilder.h
@@ -2,8 +2,7 @@
 
 #pragma once
 
-// Forward Declarations
-class UEdGraphSchema_K2;
+#include "EdGraphSchema_K2.h" // for FFunctionTargetInfo
 
 /*******************************************************************************
 * FBlueprintGraphActionListBuilder
@@ -36,24 +35,24 @@ public:
 * FK2ActionMenuBuilder
 *******************************************************************************/
 
-class FK2ActionMenuBuilder
+class BLUEPRINTGRAPH_API FK2ActionMenuBuilder
 {
 public:
 	/** Utility to create a component node add action */
-	static TSharedPtr<FEdGraphSchemaAction_K2AddComponent> CreateAddComponentAction(UObject* TemporaryOuter, const UBlueprint* Blueprint, TSubclassOf<UActorComponent> DestinationComponentType, UObject* Asset);
+	static TSharedPtr<struct FEdGraphSchemaAction_K2AddComponent> CreateAddComponentAction(UObject* TemporaryOuter, const UBlueprint* Blueprint, TSubclassOf<UActorComponent> DestinationComponentType, UObject* Asset);
 
 	/** Utility to create a node creating action */
-	BLUEPRINTGRAPH_API static TSharedPtr<FEdGraphSchemaAction_K2NewNode> AddNewNodeAction(FGraphActionListBuilderBase& ContextMenuBuilder, const FString& Category, const FText& MenuDesc, const FString& Tooltip, const int32 Grouping = 0, const FString& Keywords = FString());
+	static TSharedPtr<struct FEdGraphSchemaAction_K2NewNode> AddNewNodeAction(FGraphActionListBuilderBase& ContextMenuBuilder, const FString& Category, const FText& MenuDesc, const FString& Tooltip, const int32 Grouping = 0, const FString& Keywords = FString());
 
 	/** Utility that creates a 'new event' entry based on the function supplied. Can either override that function, or simply use the function as its signature. */
-	BLUEPRINTGRAPH_API static void AddEventForFunction(FGraphActionListBuilderBase& ActionMenuBuilder, class UFunction* Function);
+	static void AddEventForFunction(FGraphActionListBuilderBase& ActionMenuBuilder, class UFunction* Function);
 
 	/** 
 	 * Adds macro action to passed array of schema actions
 	 * @param [in,out]	ActionMenuBuilder		Menu builder to append macro actions to
 	 * @param [in]		MacroGraph				The macro graph to attach		
 	 */
-	BLUEPRINTGRAPH_API static void AttachMacroGraphAction(FGraphActionListBuilderBase& ActionMenuBuilder, UEdGraph* MacroGraph);
+	static void AttachMacroGraphAction(FGraphActionListBuilderBase& ActionMenuBuilder, UEdGraph* MacroGraph);
 
 	/**
 	 * Converts a UFunction into a single spawn info record and adds that to the specified array.
@@ -67,12 +66,12 @@ public:
 	 * @param	PlaceOnTop			Don't use category, place the node on top
 	 * @param [in,out]	ListBuilder	List builder.
 	 */
-	BLUEPRINTGRAPH_API static void AddSpawnInfoForFunction(const UFunction* Function, bool bIsParentContext, const FFunctionTargetInfo& TargetInfo, const struct FMemberReference& CallOnMember,  const FString& BaseCategory, int32 const ActionGrouping, FGraphActionListBuilderBase& ListBuilder, bool bCalledForEach = false, bool bPlaceOnTop = false);
+	static void AddSpawnInfoForFunction(const UFunction* Function, bool bIsParentContext, struct FFunctionTargetInfo const& TargetInfo, const struct FMemberReference& CallOnMember, const FString& BaseCategory, int32 const ActionGrouping, FGraphActionListBuilderBase& ListBuilder, bool bCalledForEach = false, bool bPlaceOnTop = false);
 
 	// This variable allows potentially unsafe (given the current known context) commands to appear in the context menu
 	// It defaults to off, but can be enabled via a console command when creating macros for internal use, where safety
 	// is guaranteed by higher level code.
-	BLUEPRINTGRAPH_API static int32 bAllowUnsafeCommands;
+	static int32 bAllowUnsafeCommands;
 
 public:
 	FK2ActionMenuBuilder(UEdGraphSchema_K2 const* Schema);
diff --git a/Engine/Source/Editor/UnrealEd/Classes/Commandlets/DumpBlueprintsInfoCommandlet.h b/Engine/Source/Editor/UnrealEd/Classes/Commandlets/DumpBlueprintsInfoCommandlet.h
new file mode 100644
--- /dev/null
+++ b/Engine/Source/Editor/UnrealEd/Classes/Commandlets/DumpBlueprintsInfoCommandlet.h
@@ -0,0 +1,15 @@
+// Copyright 1998-2014 Epic Games, Inc. All Rights Reserved.
+
+#pragma once
+#include "DumpBlueprintsInfoCommandlet.generated.h"
+
+UCLASS()
+class UDumpBlueprintsInfoCommandlet : public UCommandlet
+{
+	GENERATED_UCLASS_BODY()
+
+public:		
+	// Begin UCommandlet Interface
+	virtual int32 Main(FString const& Params) OVERRIDE;
+	// End UCommandlet Interface
+};
diff --git a/Engine/Source/Editor/UnrealEd/Private/Commandlets/DumpBlueprintsInfoCommandlet.cpp b/Engine/Source/Editor/UnrealEd/Private/Commandlets/DumpBlueprintsInfoCommandlet.cpp
new file mode 100644
--- /dev/null
+++ b/Engine/Source/Editor/UnrealEd/Private/Commandlets/DumpBlueprintsInfoCommandlet.cpp
@@ -0,0 +1,121 @@
+// Copyright 1998-2014 Epic Games, Inc. All Rights Reserved.
+
+#include "UnrealEd.h"
+#include "ObjectEditorUtils.h"
+#include "Kismet2/KismetEditorUtilities.h"
+#include "K2ActionMenuBuilder.h"
+#include "EdGraphSchema_K2_Actions.h"
+#include "K2Node.h"
+
+//------------------------------------------------------------------------------
+UDumpBlueprintsInfoCommandlet::UDumpBlueprintsInfoCommandlet(class FPostConstructInitializeProperties const& PCIP)
+	: Super(PCIP)
+{
+}
+
+//------------------------------------------------------------------------------
+int32 UDumpBlueprintsInfoCommandlet::Main(FString const& Params)
+{
+	TArray<FString> Tokens, Switches;
+	ParseCommandLine(*Params, Tokens, Switches);
+
+	FString CommandletSaveDir = FPaths::GameSavedDir() + TEXT("Commandlets/");
+	IFileManager::Get().MakeDirectory(*CommandletSaveDir);
+ 
+	FString Filename = FString::Printf(TEXT("BlueprintsInfoDump_%s.json"), FPlatformTime::StrTimestamp());
+	Filename = Filename.Replace(TEXT(" "), TEXT("_"));
+	Filename = Filename.Replace(TEXT("/"), TEXT("-"));
+	Filename = Filename.Replace(TEXT(":"), TEXT("."));
+	Filename = FPaths::GetCleanFilename(Filename);
+ 
+	FString FilePath = CommandletSaveDir / *Filename;
+	FArchive* FileOut = IFileManager::Get().CreateFileWriter(*FilePath);
+	if (FileOut != nullptr)
+	{
+		UClass* ClassFilter = nullptr;
+		for (FString& Switch : Switches)
+		{
+			if (Switch.StartsWith("class="))
+			{
+				FString ClassSwitch, ClassName;
+				Switch.Split(TEXT("="), &ClassSwitch, &ClassName);
+				ClassFilter = FindObject<UClass>(ANY_PACKAGE, *ClassName);
+			}
+		}
+
+		FString LeadingFields("{\n\t\"ClassContext\" : \"");
+		if (ClassFilter != nullptr)
+		{
+			LeadingFields += ClassFilter->GetName();
+		}
+		else
+		{
+			LeadingFields += "<NO_CONTEXT>";
+		}
+		LeadingFields += "\",\n\t\"Palette\" : {";
+		FileOut->Serialize(TCHAR_TO_ANSI(*LeadingFields), LeadingFields.Len());
+
+		UPackage* TransientPackage = GetTransientPackage();
+		FName TempBpName = MakeUniqueObjectName(TransientPackage, UBlueprint::StaticClass(), FName(TEXT("COMMANDLET_TEMP_Blueprint")));
+		UBlueprint* TempBlueprint = FKismetEditorUtilities::CreateBlueprint(AActor::StaticClass(), TransientPackage, TempBpName, BPTYPE_Normal, UBlueprint::StaticClass(), UBlueprintGeneratedClass::StaticClass());
+
+		FBlueprintPaletteListBuilder PaletteBuilder(TempBlueprint);
+		UEdGraphSchema_K2 const* K2Schema = GetDefault<UEdGraphSchema_K2>();
+		FK2ActionMenuBuilder(K2Schema).GetPaletteActions(PaletteBuilder, ClassFilter);
+
+		for (int32 ActionIndex = 0; ActionIndex < PaletteBuilder.GetNumActions(); ++ActionIndex)
+		{
+			FGraphActionListBuilderBase::ActionGroup& Action = PaletteBuilder.GetAction(ActionIndex);
+			if (Action.Actions.Num() <= 0)
+			{
+				continue;
+			}
+
+			TSharedPtr<FEdGraphSchemaAction> PrimeAction = Action.Actions[0];
+			FString ActionName = PrimeAction->MenuDescription.ToString();
+
+			FString ActionEntry("\n\t\t\"");
+			ActionEntry += ActionName + TEXT("\" : {");
+			
+			
+			TArray<FString> MenuHierarchy;
+			Action.GetCategoryChain(MenuHierarchy);
+			bool bHasCategory = (MenuHierarchy.Num() > 0);
+
+			ActionEntry += TEXT("\n\t\t\t\"Category\" : \"");
+			if (bHasCategory)
+			{
+				for (FString const& SubCategory : MenuHierarchy)
+				{
+					ActionEntry += SubCategory + TEXT("|");
+				}
+				ActionEntry.RemoveAt(ActionEntry.Len() - 1); // remove the last '|'
+			}
+			ActionEntry += TEXT("\"");
+
+
+			UK2Node const* NodeTemplate = FK2SchemaActionUtils::ExtractNodeTemplateFromAction(PrimeAction);
+			if (NodeTemplate != nullptr)
+			{
+
+				ActionEntry += TEXT(",\n\t\t\t\"Node\" : \"");
+				ActionEntry += NodeTemplate->GetClass()->GetPathName();
+				ActionEntry += TEXT("\"");
+			}
+			ActionEntry += TEXT("\n\t\t}");
+
+			bool bIsLastEntry = (ActionIndex == PaletteBuilder.GetNumActions()-1);
+			if (!bIsLastEntry)
+			{
+				ActionEntry += TEXT(",");
+			}
+			FileOut->Serialize(TCHAR_TO_ANSI(*ActionEntry), ActionEntry.Len());
+		}
+		FileOut->Serialize(TCHAR_TO_ANSI(TEXT("\n\t}\n}")), 5);
+ 
+		FileOut->Close();
+		delete FileOut;
+		FileOut = nullptr;
+	}
+	return 0;
+}
\ No newline at end of file
