BlueprintFunctionLibrary wont be listed in "Create Blueprint Based On Class" hierarchy.

diff --git a/Engine/Source/Editor/UnrealEd/Private/Factories/EditorFactories.cpp b/Engine/Source/Editor/UnrealEd/Private/Factories/EditorFactories.cpp
--- a/Engine/Source/Editor/UnrealEd/Private/Factories/EditorFactories.cpp
+++ b/Engine/Source/Editor/UnrealEd/Private/Factories/EditorFactories.cpp
@@ -6297,7 +6297,7 @@ UObject* UBlueprintFunctionLibraryFactory::FactoryCreateNew(UClass* Class, UObje
 	// Make sure we are trying to factory a blueprint, then create and init one
 	check(Class->IsChildOf(UBlueprint::StaticClass()));
 
-	if ((ParentClass == NULL) || !FKismetEditorUtilities::CanCreateBlueprintOfClass(ParentClass))
+	if (ParentClass != UBlueprintFunctionLibrary::StaticClass())
 	{
 		FFormatNamedArguments Args;
 		Args.Add(TEXT("ClassName"), (ParentClass != NULL) ? FText::FromString(ParentClass->GetName()) : LOCTEXT("Null", "(null)"));
diff --git a/Engine/Source/Editor/UnrealEd/Private/Kismet2/Kismet2.cpp b/Engine/Source/Editor/UnrealEd/Private/Kismet2/Kismet2.cpp
--- a/Engine/Source/Editor/UnrealEd/Private/Kismet2/Kismet2.cpp
+++ b/Engine/Source/Editor/UnrealEd/Private/Kismet2/Kismet2.cpp
@@ -485,9 +485,8 @@ bool FKismetEditorUtilities::CanCreateBlueprintOfClass(const UClass* Class)
 		&& (!Class->ClassGeneratedBy || (bAllowDerivedBlueprints && !IsClassABlueprintSkeleton(Class)));
 
 	const UEdGraphSchema_K2* K2Schema = GetDefault<UEdGraphSchema_K2>();
-	const bool bIsValidClass = Class->GetBoolMetaDataHierarchical(FBlueprintMetadata::MD_IsBlueprintBase) 
-		|| (Class == UObject::StaticClass()) 
-		|| (Class == UBlueprintFunctionLibrary::StaticClass());
+	const bool bIsValidClass = Class->GetBoolMetaDataHierarchical(FBlueprintMetadata::MD_IsBlueprintBase)
+		|| (Class == UObject::StaticClass());
 
 	return bCanCreateBlueprint && bIsValidClass;
 }
