No Crash after removing an event form an interface.
https://answers.unrealengine.com/questions/31043/macro-and-interface-library-killed-4-weeks-of-work.html

diff --git a/Engine/Source/Editor/BlueprintGraph/Private/EdGraphSchema_K2.cpp b/Engine/Source/Editor/BlueprintGraph/Private/EdGraphSchema_K2.cpp
--- a/Engine/Source/Editor/BlueprintGraph/Private/EdGraphSchema_K2.cpp
+++ b/Engine/Source/Editor/BlueprintGraph/Private/EdGraphSchema_K2.cpp
@@ -342,12 +342,14 @@ UFunction* UEdGraphSchema_K2::GetCallableParentFunction(UFunction* Function) con
 
 bool UEdGraphSchema_K2::CanUserKismetCallFunction(const UFunction* Function)
 {
-	return Function->HasAllFunctionFlags(FUNC_BlueprintCallable) && !Function->HasAllFunctionFlags(FUNC_Delegate) && !Function->GetBoolMetaData(FBlueprintMetadata::MD_BlueprintInternalUseOnly) && !Function->HasMetaData(FBlueprintMetadata::MD_DeprecatedFunction);
+	return Function && 
+		(Function->HasAllFunctionFlags(FUNC_BlueprintCallable) && !Function->HasAllFunctionFlags(FUNC_Delegate) && !Function->GetBoolMetaData(FBlueprintMetadata::MD_BlueprintInternalUseOnly) && !Function->HasMetaData(FBlueprintMetadata::MD_DeprecatedFunction));
 }
 
 bool UEdGraphSchema_K2::CanKismetOverrideFunction(const UFunction* Function)
 {
-	return Function->HasAllFunctionFlags(FUNC_BlueprintEvent) && !Function->HasAllFunctionFlags(FUNC_Delegate) && !Function->GetBoolMetaData(FBlueprintMetadata::MD_BlueprintInternalUseOnly) && !Function->HasMetaData(FBlueprintMetadata::MD_DeprecatedFunction);
+	return  Function && 
+		(Function->HasAllFunctionFlags(FUNC_BlueprintEvent) && !Function->HasAllFunctionFlags(FUNC_Delegate) && !Function->GetBoolMetaData(FBlueprintMetadata::MD_BlueprintInternalUseOnly) && !Function->HasMetaData(FBlueprintMetadata::MD_DeprecatedFunction));
 }
 
 struct FNoOutputParametersHelper 
@@ -371,7 +373,7 @@ struct FNoOutputParametersHelper
 bool UEdGraphSchema_K2::FunctionCanBePlacedAsEvent(const UFunction* InFunction)
 {
 	// First check we are override-able
-	if (!CanKismetOverrideFunction(InFunction))
+	if (!InFunction || !CanKismetOverrideFunction(InFunction))
 	{
 		return false;
 	}
diff --git a/Engine/Source/Editor/UnrealEd/Private/Kismet2/BlueprintEditorUtils.cpp b/Engine/Source/Editor/UnrealEd/Private/Kismet2/BlueprintEditorUtils.cpp
--- a/Engine/Source/Editor/UnrealEd/Private/Kismet2/BlueprintEditorUtils.cpp
+++ b/Engine/Source/Editor/UnrealEd/Private/Kismet2/BlueprintEditorUtils.cpp
@@ -1958,6 +1958,7 @@ UK2Node_Event* FBlueprintEditorUtils::FindOverrideForFunction(const UBlueprint*
 	return NULL;
 }
 
+// returns if Blueprint depends on TestBlueprint
 bool FBlueprintEditorUtils::IsBlueprintDependentOn(UBlueprint const* Blueprint, UBlueprint const* TestBlueprint)
 {
 	if ( Blueprint && (Blueprint != TestBlueprint) )
@@ -4018,13 +4019,19 @@ static void ConformInterfaceByName(UBlueprint* Blueprint, FBPInterfaceDescriptio
 			// destroy the old event node (this will also break all pin links and remove it from the graph)
 			EventNode->DestroyNode();
 
-			// have to rename so it doesn't conflict with the graph we're about to add
-			CustomEventNode->RenameCustomEventCloseToName();
+			if (InterfaceFunction)
+			{
+				// have to rename so it doesn't conflict with the graph we're about to add
+				CustomEventNode->RenameCustomEventCloseToName();
+			}
 			EventGraph->Nodes.Add(CustomEventNode);
 
 			// warn the user that their old functionality won't work (it's now connected 
 			// to a custom node that isn't triggered anywhere)
-			FText WarningMessageText = LOCTEXT("InterfaceEventNodeReplaced_Warn", "'%s' was promoted from an event to a function - it has been replaced by a custom event, which won't trigger unless you call it manually.");
+			FText WarningMessageText = InterfaceFunction ?
+				LOCTEXT("InterfaceEventNodeReplaced_Warn", "'%s' was promoted from an event to a function - it has been replaced by a custom event, which won't trigger unless you call it manually.") :
+				LOCTEXT("InterfaceEventRemovedNodeReplaced_Warn", "'%s' was removed from its interface - it has been replaced by a custom event, which won't trigger unless you call it manually.");
+
 			Blueprint->Message_Warn(FString::Printf(*WarningMessageText.ToString(), *FunctionName.ToString()));
 		}
 
