When converting actors, if the name isnt actually custom, allow the conversion process to rename the actor.  Avoids classes like converting PointLight->SpotLight still being called PointLight after conversion

diff --git a/Engine/Source/Editor/UnrealEd/Private/Editor.cpp b/Engine/Source/Editor/UnrealEd/Private/Editor.cpp
--- a/Engine/Source/Editor/UnrealEd/Private/Editor.cpp
+++ b/Engine/Source/Editor/UnrealEd/Private/Editor.cpp
@@ -5676,8 +5676,11 @@ void UEditorEngine::DoConvertActors( const TArray<AActor*>& ActorsToConvert, UCl
 		for( int32 ActorIdx = 0; ActorIdx < ActorsToConvert.Num(); ++ActorIdx )
 		{
 			AActor* ActorToConvert = ActorsToConvert[ ActorIdx ];
+			// Source actor display label
 			FString ActorLabel = ActorToConvert->GetActorLabel();
-
+			// Low level source actor object name
+			FName ActorObjectName = ActorToConvert->GetFName();
+	
 			// The class of the actor we are about to replace
 			UClass* ClassToReplace = ActorToConvert->GetClass();
 
@@ -5798,7 +5801,16 @@ void UEditorEngine::DoConvertActors( const TArray<AActor*>& ActorsToConvert, UCl
 
 			if (NewActor)
 			{
-				NewActor->SetActorLabel(ActorLabel);
+				// If the actor label isn't actually anything custom allow the name to be changed
+				// to avoid cases like converting PointLight->SpotLight still being called PointLight after conversion
+				FString ClassName = ClassToReplace->GetName();
+				
+				// Remove any number off the end of the label
+				int32 Number = 0;
+				if( !ActorLabel.StartsWith( ClassName ) || !FParse::Value(*ActorLabel, *ClassName, Number)  )
+				{
+					NewActor->SetActorLabel(ActorLabel);
+				}
 
 				ConvertedActors.Add(NewActor);
 
