Fix crash when curve editor range delegate is not set

diff --git a/Engine/Source/Editor/UnrealEd/Private/SCurveEditor.cpp b/Engine/Source/Editor/UnrealEd/Private/SCurveEditor.cpp
--- a/Engine/Source/Editor/UnrealEd/Private/SCurveEditor.cpp
+++ b/Engine/Source/Editor/UnrealEd/Private/SCurveEditor.cpp
@@ -1138,7 +1138,11 @@ FReply SCurveEditor::OnMouseMove( const FGeometry& InMyGeometry, const FPointerE
 
 				// clamp to data input if set
 				ClampViewRangeToDataIfBound(NewMinInput, NewMaxInput, DataMinInput, DataMaxInput, ScaleInfo.ViewInputRange);
-				SetInputViewRangeHandler.Execute(NewMinInput, NewMaxInput);
+				if (SetInputViewRangeHandler.IsBound())
+				{
+					SetInputViewRangeHandler.Execute(NewMinInput, NewMaxInput);
+				}
+				
 			}
 
 			return FReply::Handled();
@@ -1169,7 +1173,10 @@ FReply SCurveEditor::OnMouseWheel( const FGeometry& MyGeometry, const FPointerEv
 		const float NewMaxInput = ViewMaxInput.Get() + (InputChange * 0.5f);
 
 
-		SetInputViewRangeHandler.Execute(NewMinInput, NewMaxInput);
+		if (SetInputViewRangeHandler.IsBound())
+		{
+			SetInputViewRangeHandler.Execute(NewMinInput, NewMaxInput);
+		}
 	}
 
 	return FReply::Handled();
@@ -1509,7 +1516,10 @@ FReply SCurveEditor::ZoomToFitHorizontal()
 			InMax = (CONST_DefaultZoomRange + CONST_FitMargin)*2.0;
 		}
 
-		SetInputViewRangeHandler.Execute(InMin, InMax);
+		if (SetInputViewRangeHandler.IsBound())
+		{
+			SetInputViewRangeHandler.Execute(InMin, InMax);
+		}
 	}
 	return FReply::Handled();
 }
