Add non delegate fallback for SCurveEditor

diff --git a/Engine/Source/Editor/UnrealEd/Private/SCurveEditor.cpp b/Engine/Source/Editor/UnrealEd/Private/SCurveEditor.cpp
--- a/Engine/Source/Editor/UnrealEd/Private/SCurveEditor.cpp
+++ b/Engine/Source/Editor/UnrealEd/Private/SCurveEditor.cpp
@@ -1138,11 +1138,7 @@ FReply SCurveEditor::OnMouseMove( const FGeometry& InMyGeometry, const FPointerE
 
 				// clamp to data input if set
 				ClampViewRangeToDataIfBound(NewMinInput, NewMaxInput, DataMinInput, DataMaxInput, ScaleInfo.ViewInputRange);
-				if (SetInputViewRangeHandler.IsBound())
-				{
-					SetInputViewRangeHandler.Execute(NewMinInput, NewMaxInput);
-				}
-				
+				SetInputMinMax(NewMinInput, NewMaxInput);
 			}
 
 			return FReply::Handled();
@@ -1172,11 +1168,7 @@ FReply SCurveEditor::OnMouseWheel( const FGeometry& MyGeometry, const FPointerEv
 		const float NewMinInput = ViewMinInput.Get() - (InputChange * 0.5f);
 		const float NewMaxInput = ViewMaxInput.Get() + (InputChange * 0.5f);
 
-
-		if (SetInputViewRangeHandler.IsBound())
-		{
-			SetInputViewRangeHandler.Execute(NewMinInput, NewMaxInput);
-		}
+		SetInputMinMax(NewMinInput, NewMaxInput);
 	}
 
 	return FReply::Handled();
@@ -1516,10 +1508,7 @@ FReply SCurveEditor::ZoomToFitHorizontal()
 			InMax = (CONST_DefaultZoomRange + CONST_FitMargin)*2.0;
 		}
 
-		if (SetInputViewRangeHandler.IsBound())
-		{
-			SetInputViewRangeHandler.Execute(InMin, InMax);
-		}
+		SetInputMinMax(InMin, InMax);
 	}
 	return FReply::Handled();
 }
@@ -1721,6 +1710,27 @@ FText SCurveEditor::GetCurveName( FRichCurve* Curve ) const
 	return FText::GetEmpty();
 }
 
+void SCurveEditor::SetInputMinMax(float NewMin, float NewMax)
+{
+	if (SetInputViewRangeHandler.IsBound())
+	{
+		SetInputViewRangeHandler.Execute(NewMin, NewMax);
+	}
+	else
+	{
+		//if no delegate and view min input isn't using a delegate just set value directly
+		if (ViewMinInput.IsBound() == false)
+		{
+			ViewMinInput.Set(NewMin);
+		}
+
+		if (ViewMaxInput.IsBound() == false)
+		{
+			ViewMaxInput.Set(NewMax);
+		}
+	}
+}
+
 bool SCurveEditor::HitTestCurves(  const FGeometry& InMyGeometry, const FPointerEvent& InMouseEvent )
 {
 	bool bHit = false;
diff --git a/Engine/Source/Editor/UnrealEd/Public/SCurveEditor.h b/Engine/Source/Editor/UnrealEd/Public/SCurveEditor.h
--- a/Engine/Source/Editor/UnrealEd/Public/SCurveEditor.h
+++ b/Engine/Source/Editor/UnrealEd/Public/SCurveEditor.h
@@ -467,6 +467,8 @@ private:
 
 	/** Gradient editor */
 	TSharedPtr<class SColorGradientEditor> GradientViewer;
+
+	void SetInputMinMax(float NewMin, float NewMax);
 protected:
 	/** Minimum input of view range  */
 	TAttribute<float>	ViewMinInput;
