- Fixed issue with importing doesn't apply cluster transform causing skin to be positioned with wrong transform.
- Fixed calculating vert size to not apply ref inverse transform. Vert position is imported on ref pose, so if you inverse it, you'll get scale 1, 1, 1 transform, which isn't true. Just using normal vert size should work fine.

diff --git a/Engine/Source/Developer/MeshUtilities/Private/SkeletalMeshTools.cpp b/Engine/Source/Developer/MeshUtilities/Private/SkeletalMeshTools.cpp
--- a/Engine/Source/Developer/MeshUtilities/Private/SkeletalMeshTools.cpp
+++ b/Engine/Source/Developer/MeshUtilities/Private/SkeletalMeshTools.cpp
@@ -309,10 +309,10 @@ namespace SkeletalMeshTools
 				FRigidSkinVertex* RigidVert = &Chunk.RigidVertices[i];
 				int32 BoneIndex = Chunk.BoneMap[RigidVert->Bone];
 
-				FVector LocalPos = SkeletalMesh->RefBasesInvMatrix[BoneIndex].TransformPosition(RigidVert->Position);
+				FVector LocalPos = RigidVert->Position;
 				Infos[BoneIndex].Positions.Add(LocalPos);
 
-				FVector LocalNormal = SkeletalMesh->RefBasesInvMatrix[BoneIndex].TransformVector(RigidVert->TangentZ);
+				FVector LocalNormal = RigidVert->TangentZ;
 				Infos[BoneIndex].Normals.Add(LocalNormal);
 			}
 
@@ -324,10 +324,10 @@ namespace SkeletalMeshTools
 				{
 					int32 BoneIndex = Chunk.BoneMap[GetDominantBoneIndex(SoftVert)];
 
-					FVector LocalPos = SkeletalMesh->RefBasesInvMatrix[BoneIndex].TransformPosition(SoftVert->Position);
+					FVector LocalPos = SoftVert->Position;
 					Infos[BoneIndex].Positions.Add(LocalPos);
 
-					FVector LocalNormal = SkeletalMesh->RefBasesInvMatrix[BoneIndex].TransformVector(SoftVert->TangentZ);
+					FVector LocalNormal = SoftVert->TangentZ;
 					Infos[BoneIndex].Normals.Add(LocalNormal);
 				}
 				else
@@ -338,10 +338,10 @@ namespace SkeletalMeshTools
 						{
 							int32 BoneIndex = Chunk.BoneMap[SoftVert->InfluenceBones[j]];
 
-							FVector LocalPos = SkeletalMesh->RefBasesInvMatrix[BoneIndex].TransformPosition(SoftVert->Position);
+							FVector LocalPos = SoftVert->Position;
 							Infos[BoneIndex].Positions.Add(LocalPos);
 
-							FVector LocalNormal = SkeletalMesh->RefBasesInvMatrix[BoneIndex].TransformVector(SoftVert->TangentZ);
+							FVector LocalNormal = SoftVert->TangentZ;
 							Infos[BoneIndex].Normals.Add(LocalNormal);
 						}
 					}
diff --git a/Engine/Source/Editor/UnrealEd/Private/Fbx/FbxSkeletalMeshImport.cpp b/Engine/Source/Editor/UnrealEd/Private/Fbx/FbxSkeletalMeshImport.cpp
--- a/Engine/Source/Editor/UnrealEd/Private/Fbx/FbxSkeletalMeshImport.cpp
+++ b/Engine/Source/Editor/UnrealEd/Private/Fbx/FbxSkeletalMeshImport.cpp
@@ -806,9 +806,7 @@ bool UnFbx::FFbxImporter::ImportBone(TArray<FbxNode*>& NodeArray, FSkeletalMeshI
 	GlobalsPerLink[0].SetIdentity();
 	
 	bool GlobalLinkFoundFlag;
-	FbxVector4 LocalLinkT;
-	FbxQuaternion LocalLinkQ;
-	FbxVector4	LocalLinkS;
+	FTransform LocalTransform;
 
 	bool bAnyLinksNotInBindPose = false;
 	FString LinksWithoutBindPoses;
@@ -913,18 +911,13 @@ bool UnFbx::FFbxImporter::ImportBone(TArray<FbxNode*>& NodeArray, FSkeletalMeshI
 
 		if (LinkIndex)
 		{
-			FbxAMatrix	Matrix;
-			Matrix = GlobalsPerLink[ParentIndex].Inverse() * GlobalsPerLink[LinkIndex];
-			LocalLinkT = Matrix.GetT();
-			LocalLinkQ = Matrix.GetQ();
-			LocalLinkS = Matrix.GetS();
+			FTransform Parent = Converter.ConvertTransform(GlobalsPerLink[ParentIndex]);
+			FTransform Current = Converter.ConvertTransform(GlobalsPerLink[LinkIndex]);
+			LocalTransform = Current.GetRelativeTransform(Parent);
 		}
 		else	// skeleton root
 		{
-			// for root, this is global coordinate
-			LocalLinkT = GlobalsPerLink[LinkIndex].GetT();
-			LocalLinkQ = GlobalsPerLink[LinkIndex].GetQ();
-			LocalLinkS = GlobalsPerLink[LinkIndex].GetS();
+			LocalTransform = Converter.ConvertTransform(GlobalsPerLink[LinkIndex]);
 		}
 		
 		// set bone
@@ -964,9 +957,7 @@ bool UnFbx::FFbxImporter::ImportBone(TArray<FbxNode*>& NodeArray, FSkeletalMeshI
 			}
 		}
 
-		JointMatrix.Transform.SetTranslation(Converter.ConvertPos(LocalLinkT));
-		JointMatrix.Transform.SetRotation(Converter.ConvertRotToQuat(LocalLinkQ));
-		JointMatrix.Transform.SetScale3D(Converter.ConvertScale(LocalLinkS));
+		JointMatrix.Transform = LocalTransform;
 	}
 	
 	if(bAnyLinksNotInBindPose)
@@ -1884,29 +1875,11 @@ bool UnFbx::FFbxImporter::FillSkelMeshImporterFromFbx( FSkeletalMeshImportData&
 	int32 ExistPointNum = ImportData.Points.Num();
 	ImportData.Points.AddUninitialized(ControlPointsCount);
 
+	SkinControlPointsToPose(ImportData, Node->GetMesh(), FbxShape, false);
 	// Construct the matrices for the conversion from right handed to left handed system
+
 	FbxAMatrix TotalMatrix;
 	FbxAMatrix TotalMatrixForNormal;
-	TotalMatrix = ComputeTotalMatrix(Node);
-	TotalMatrixForNormal = TotalMatrix.Inverse();
-	TotalMatrixForNormal = TotalMatrixForNormal.Transpose();
-
-	int32 ControlPointsIndex;
-	for( ControlPointsIndex = 0 ; ControlPointsIndex < ControlPointsCount ;ControlPointsIndex++ )
-	{
-		FbxVector4 Position;
-		if (FbxShape)
-		{
-			Position = FbxShape->GetControlPoints()[ControlPointsIndex];
-		}
-		else
-		{
-			Position = Mesh->GetControlPoints()[ControlPointsIndex];
-		}																	 
-		FbxVector4 FinalPosition;
-		FinalPosition = TotalMatrix.MultT(Position);
-		ImportData.Points[ControlPointsIndex+ExistPointNum] = Converter.ConvertPos(FinalPosition);
-	}
 	
 	bool OddNegativeScale = IsOddNegativeScale(TotalMatrix);
 	
