reverted CL# 2009553 - disabling viewport real-time mode when starting standalone PIE ( TTP# 328925)

diff --git a/Engine/Source/Editor/UnrealEd/Private/PlayLevel.cpp b/Engine/Source/Editor/UnrealEd/Private/PlayLevel.cpp
--- a/Engine/Source/Editor/UnrealEd/Private/PlayLevel.cpp
+++ b/Engine/Source/Editor/UnrealEd/Private/PlayLevel.cpp
@@ -960,8 +960,6 @@ void UEditorEngine::PlayOnLocalPc(FString MapNameOverride, FString URLParms, FSt
 		return;
 	}
 
-	DisableRealtimeViewports();
-
 	FString GameNameOrProjectFile;
 	FString AdditionalParameters(TEXT(""));
 	if (FPaths::IsProjectFilePathSet())
