change sphyl/sphere to two different options

diff --git a/Engine/Source/Editor/PhAT/Private/SPhATNewAssetDlg.cpp b/Engine/Source/Editor/PhAT/Private/SPhATNewAssetDlg.cpp
--- a/Engine/Source/Editor/PhAT/Private/SPhATNewAssetDlg.cpp
+++ b/Engine/Source/Editor/PhAT/Private/SPhATNewAssetDlg.cpp
@@ -21,7 +21,8 @@ void SPhATNewAssetDlg::Construct(const FArguments& InArgs)
 	check(NewBodyData && NewBodyResponse);
 
 	// Initialize combobox options
-	CollisionGeometryOptions.Add(TSharedPtr<FString>(new FString(TEXT("Sphyl/Sphere"))));
+	CollisionGeometryOptions.Add(TSharedPtr<FString>(new FString(TEXT("Sphyl"))));
+	CollisionGeometryOptions.Add(TSharedPtr<FString>(new FString(TEXT("Sphere"))));
 	CollisionGeometryOptions.Add(TSharedPtr<FString>(new FString(TEXT("Box"))));
 
 	// Add in the ability to create a convex hull from the source geometry for skeletal meshes
@@ -262,6 +263,10 @@ void SPhATNewAssetDlg::OnCollisionGeometrySelectionChanged(TSharedPtr<FString> N
 	{
 		NewBodyData->GeomType = EFG_Box;
 	}
+	else if (*NewSelection == TEXT("Sphere"))
+	{
+		NewBodyData->GeomType = EFG_Sphere;
+	}
 	// Check to see if the user has picked that they want to create a convex hull
 	else if ( *NewSelection == TEXT("Single Convex Hull") )
 	{
@@ -273,7 +278,7 @@ void SPhATNewAssetDlg::OnCollisionGeometrySelectionChanged(TSharedPtr<FString> N
 	}
 	else
 	{
-		NewBodyData->GeomType = EFG_SphylSphere;
+		NewBodyData->GeomType = EFG_Sphyl;
 	}
 }
 
diff --git a/Engine/Source/Editor/UnrealEd/Private/PhysicsAssetUtils.cpp b/Engine/Source/Editor/UnrealEd/Private/PhysicsAssetUtils.cpp
--- a/Engine/Source/Editor/UnrealEd/Private/PhysicsAssetUtils.cpp
+++ b/Engine/Source/Editor/UnrealEd/Private/PhysicsAssetUtils.cpp
@@ -9,7 +9,7 @@
 void FPhysAssetCreateParams::Initialize()
 {
 	MinBoneSize = 5.0f;
-	GeomType = EFG_SphylSphere;
+	GeomType = EFG_Sphyl;
 	VertWeight = EVW_DominantWeight;
 	bAlignDownBone = true;
 	bCreateJoints = true;
@@ -247,7 +247,6 @@ void CreateCollisionFromBone( UBodySetup* bs, USkeletalMesh* skelMesh, int32 Bon
 
 	// Calculate orientation of to use for collision primitive.
 	FMatrix ElemTM;
-	bool bSphyl;
 
 	if(Params.bAlignDownBone)
 	{
@@ -269,26 +268,21 @@ void CreateCollisionFromBone( UBodySetup* bs, USkeletalMesh* skelMesh, int32 Bon
 				FVector XAxis, YAxis;
 				ZAxis.FindBestAxisVectors( YAxis, XAxis );
 
-				ElemTM = FMatrix( XAxis, YAxis, ZAxis, FVector(0) );
-
-				bSphyl = true;			
+				ElemTM = FMatrix( XAxis, YAxis, ZAxis, FVector(0) );	
 			}
 			else
 			{
 				ElemTM = FMatrix::Identity;
-				bSphyl = false;
 			}
 		}
 		else
 		{
 			ElemTM = FMatrix::Identity;
-			bSphyl = false;
 		}
 	}
 	else
 	{
 		ElemTM = FMatrix::Identity;
-		bSphyl = false;
 	}
 
 	// Get the (Unreal scale) bounding box for this bone using the rotation.
@@ -328,6 +322,14 @@ void CreateCollisionFromBone( UBodySetup* bs, USkeletalMesh* skelMesh, int32 Bon
 		be->Y = BoxExtent.Y * 2.0f * 1.01f;
 		be->Z = BoxExtent.Z * 2.0f * 1.01f;	
 	}
+	else if (Params.GeomType == EFG_Sphere)
+	{
+		int32 sx = bs->AggGeom.SphereElems.AddZeroed();
+		FKSphereElem* se = &bs->AggGeom.SphereElems[sx];
+
+		se->Center = ElemTM.GetOrigin();
+		se->Radius = BoxExtent.GetMax() * 1.01f;
+	}
 	// Deal with creating a single convex hull
 	else if (Params.GeomType == EFG_SingleConvexHull)
 	{
@@ -406,24 +408,13 @@ void CreateCollisionFromBone( UBodySetup* bs, USkeletalMesh* skelMesh, int32 Bon
 	}
 	else
 	{
-		if(bSphyl)
-		{
-			int32 sx = bs->AggGeom.SphylElems.AddZeroed();
-			FKSphylElem* se = &bs->AggGeom.SphylElems[sx];
+		int32 sx = bs->AggGeom.SphylElems.AddZeroed();
+		FKSphylElem* se = &bs->AggGeom.SphylElems[sx];
 
-			se->SetTransform( FTransform( ElemTM ) );
+		se->SetTransform( FTransform( ElemTM ) );
 
-			se->Radius = FMath::Max(BoxExtent.X, BoxExtent.Y) * 1.01f;
-			se->Length = BoxExtent.Z * 1.01f;
-		}
-		else
-		{
-			int32 sx = bs->AggGeom.SphereElems.AddZeroed();
-			FKSphereElem* se = &bs->AggGeom.SphereElems[sx];
-
-			se->Center = ElemTM.GetOrigin();
-			se->Radius = BoxExtent.GetMax() * 1.01f;
-		}
+		se->Radius = FMath::Max(BoxExtent.X, BoxExtent.Y) * 1.01f;
+		se->Length = BoxExtent.Z * 1.01f;
 	}
 }
 
diff --git a/Engine/Source/Editor/UnrealEd/Public/PhysicsAssetUtils.h b/Engine/Source/Editor/UnrealEd/Public/PhysicsAssetUtils.h
--- a/Engine/Source/Editor/UnrealEd/Public/PhysicsAssetUtils.h
+++ b/Engine/Source/Editor/UnrealEd/Public/PhysicsAssetUtils.h
@@ -6,7 +6,8 @@
 enum EPhysAssetFitGeomType
 {
 	EFG_Box,
-	EFG_SphylSphere,
+	EFG_Sphyl,
+	EFG_Sphere,
 	EFG_SingleConvexHull,
 	EFG_MultiConvexHull
 };
