Added commandlet which lists all UStaticMesh assets that were imported from speed trees

#ttp 334235 - UE4: RENDERING: Build list of UStaticMesh assets that were imported from speedtrees across all games / samples in main branch p4

diff --git a/Engine/Source/Editor/UnrealEd/Classes/Commandlets/ListStaticMeshesImportedFromSpeedTreesCommandlet.h b/Engine/Source/Editor/UnrealEd/Classes/Commandlets/ListStaticMeshesImportedFromSpeedTreesCommandlet.h
new file mode 100644
--- /dev/null
+++ b/Engine/Source/Editor/UnrealEd/Classes/Commandlets/ListStaticMeshesImportedFromSpeedTreesCommandlet.h
@@ -0,0 +1,14 @@
+// Copyright 1998-2014 Epic Games, Inc. All Rights Reserved.
+
+#pragma once
+#include "ListStaticMeshesImportedFromSpeedTreesCommandlet.generated.h"
+
+UCLASS()
+class UListStaticMeshesImportedFromSpeedTreesCommandlet : public UCommandlet
+{
+	GENERATED_UCLASS_BODY()
+
+	// Begin UCommandlet Interface
+	virtual int32 Main(const FString& Params) OVERRIDE;
+	// End UCommandlet Interface
+};
\ No newline at end of file
diff --git a/Engine/Source/Editor/UnrealEd/Private/Commandlets/ContentCommandlets.cpp b/Engine/Source/Editor/UnrealEd/Private/Commandlets/ContentCommandlets.cpp
--- a/Engine/Source/Editor/UnrealEd/Private/Commandlets/ContentCommandlets.cpp
+++ b/Engine/Source/Editor/UnrealEd/Private/Commandlets/ContentCommandlets.cpp
@@ -1704,3 +1704,92 @@ int32 UListMaterialsUsedWithMeshEmittersCommandlet::Main( const FString& Params
 	return 0;
 }
 
+
+/* ==========================================================================================================
+	UListStaticMeshesImportedFromSpeedTreesCommandlet
+========================================================================================================== */
+
+UListStaticMeshesImportedFromSpeedTreesCommandlet::UListStaticMeshesImportedFromSpeedTreesCommandlet(const class FPostConstructInitializeProperties& PCIP)
+: Super(PCIP)
+{
+}
+
+int32 UListStaticMeshesImportedFromSpeedTreesCommandlet::Main(const FString& Params)
+{
+
+	TArray<FString> FilesInPath;
+	FEditorFileUtils::FindAllPackageFiles(FilesInPath);
+
+	if (FilesInPath.Num() == 0)
+	{
+		UE_LOG(LogContentCommandlet, Warning, TEXT("No packages found"));
+		return 1;
+	}
+
+	TArray<FString> StaticMeshList;
+	int32 GCIndex = 0;
+	int32 TotalPackagesChecked = 0;
+
+	// Load the asset registry module
+	FAssetRegistryModule& AssetRegistryModule = FModuleManager::LoadModuleChecked<FAssetRegistryModule>(TEXT("AssetRegistry"));
+
+	// Update Registry Module
+	UE_LOG(LogContentCommandlet, Display, TEXT("Searching Asset Registry for static mesh "));
+	AssetRegistryModule.Get().SearchAllAssets(true);
+
+	// Retrieve list of all assets, used to find unreferenced ones.
+	TArray<FAssetData> AssetList;
+	AssetRegistryModule.Get().GetAssetsByClass(UStaticMesh::StaticClass()->GetFName(), AssetList, true);
+
+	for (int32 AssetIdx = 0; AssetIdx < AssetList.Num(); ++AssetIdx)
+	{
+		const FString Filename = AssetList[AssetIdx].ObjectPath.ToString();
+
+		UE_LOG(LogContentCommandlet, Display, TEXT("Processing static mesh (%i/%i):  %s "), AssetIdx, AssetList.Num(), *Filename);
+
+		UPackage* Package = LoadPackage(NULL, *Filename, LOAD_Quiet);
+		if (Package == NULL)
+		{
+			UE_LOG(LogContentCommandlet, Error, TEXT("Error loading %s!"), *Filename);
+			continue;
+		}
+
+		TotalPackagesChecked++;
+		for (TObjectIterator<UStaticMesh> It; It; ++It)
+		{
+			UStaticMesh* StaticMesh = *It;
+			if (StaticMesh->IsIn(Package) && !StaticMesh->IsTemplate())
+			{
+				// If the mesh was imported from a speedtree, we append the static mesh name to the list.
+				if (StaticMesh->SpeedTreeWind.IsValid())
+				{
+					StaticMeshList.Add(StaticMesh->GetPathName());
+				}
+			}
+		}
+
+		// Collect garbage every 10 packages instead of every package makes the commandlet run much faster
+		if ((++GCIndex % 10) == 0)
+		{
+			CollectGarbage(RF_Native);
+		}
+	}
+
+	if (StaticMeshList.Num() > 0)
+	{
+		// Now, dump out the list of materials that require updating.
+		UE_LOG(LogContentCommandlet, Display, TEXT("-------------------------------------------------------------------"));
+		UE_LOG(LogContentCommandlet, Display, TEXT("The following static meshes were imported from SpeedTrees:"));
+		for (int32 Index = 0; Index < StaticMeshList.Num(); ++Index)
+		{
+			UE_LOG(LogContentCommandlet, Error, TEXT("%s"), *StaticMeshList[Index]);
+		}
+		UE_LOG(LogContentCommandlet, Display, TEXT("-------------------------------------------------------------------"));
+	}
+	else
+	{
+		UE_LOG(LogContentCommandlet, Display, TEXT("No static meshes were imported from speedtrees in this project."));
+	}
+	return 0;
+}
+
