Fixed being able to move actors to a locked level

This addresses TTP 334274 - EDITOR: Scene Outliner: Moving an actor to a locked level through the scene outliner deletes the actor

diff --git a/Engine/Source/Editor/UnrealEd/Private/EditorServer.cpp b/Engine/Source/Editor/UnrealEd/Private/EditorServer.cpp
--- a/Engine/Source/Editor/UnrealEd/Private/EditorServer.cpp
+++ b/Engine/Source/Editor/UnrealEd/Private/EditorServer.cpp
@@ -2481,6 +2481,15 @@ void UEditorEngine::MoveSelectedActorsToLevel( ULevel* InDestLevel )
 
 void UEditorEngine::DoMoveSelectedActorsToLevel( ULevel* InDestLevel )
 {
+	// Can't move into a locked level
+	if (FLevelUtils::IsLevelLocked(InDestLevel))
+	{
+		FNotificationInfo Info(NSLOCTEXT("UnrealEd", "CannotMoveIntoLockedLevel", "Cannot move the selected actors into a locked level"));
+		Info.bUseThrobber = false;
+		FSlateNotificationManager::Get().AddNotification(Info)->SetCompletionState(SNotificationItem::CS_Fail);
+		return;
+	}
+
 	// Start the transaction
 	GEditor->Trans->Begin( NULL, NSLOCTEXT("UnrealEd", "MoveSelectedActorsToSelectedLevel", "Move Actors To Level"));
 	
@@ -2500,7 +2509,7 @@ void UEditorEngine::DoMoveSelectedActorsToLevel( ULevel* InDestLevel )
 		FirstSelectedActorLocation = FirstActor->GetActorLocation();
 	}
 
-	// Copt the actors we have selected to the clipboard
+	// Copy the actors we have selected to the clipboard
 	CopySelectedActorsToClipboard( World, true );
 
 	// Set the new level
