Removed STAT url until the page is working again

#ttp 330006 - EDITOR: "Learn More" link  found in the toast when switching on Show Stats option, does not do anything

#branch UE4

diff --git a/Engine/Source/Editor/UnrealEd/Private/SEditorViewport.cpp b/Engine/Source/Editor/UnrealEd/Private/SEditorViewport.cpp
--- a/Engine/Source/Editor/UnrealEd/Private/SEditorViewport.cpp
+++ b/Engine/Source/Editor/UnrealEd/Private/SEditorViewport.cpp
@@ -326,9 +326,10 @@ void SEditorViewport::OnToggleStats()
 		 // let the user know how they can enable stats via the console
 		 FNotificationInfo Info(LOCTEXT("StatsEnableHint", "Stats display can be toggled via the STAT [type] console command"));
 		 Info.ExpireDuration = 3.0f;
+		 /* Temporarily remove the link until the page is updated
 		 Info.HyperlinkText = LOCTEXT("StatsEnableHyperlink", "Learn more");
 		 Info.Hyperlink = FSimpleDelegate::CreateStatic([](){ IDocumentation::Get()->Open(TEXT("Engine/Basics/ConsoleCommands#statisticscommands")); });
-
+		 */
 		 FSlateNotificationManager::Get().AddNotification(Info);
 	}
 }
