#ttp 304473 - UE4: RENDERING: Visualizetexture looks broken when using r.ScreenPercentage / FullscreenAndUpsample
#proj Engine
#branch UE4

#summary Reposition texture visualization display when camera safe zone black bars are enabled

diff --git a/Engine/Source/Editor/UnrealEd/Private/EditorViewportClient.cpp b/Engine/Source/Editor/UnrealEd/Private/EditorViewportClient.cpp
--- a/Engine/Source/Editor/UnrealEd/Private/EditorViewportClient.cpp
+++ b/Engine/Source/Editor/UnrealEd/Private/EditorViewportClient.cpp
@@ -1295,7 +1295,7 @@ void FEditorViewportClient::UpdateLightingShowFlags( FEngineShowFlags& InOutShow
 	}
 }
 
-bool FEditorViewportClient::CalculateSafeFrameRect(FSlateRect& OutSafeFrameRect, FViewport* InViewport)
+bool FEditorViewportClient::CalculateEditorConstrainedViewRect(FSlateRect& OutSafeFrameRect, FViewport* InViewport)
 {
 	const int32 SizeX = InViewport->GetSizeXY().X;
 	const int32 SizeY = InViewport->GetSizeXY().Y;
@@ -1342,7 +1342,7 @@ void FEditorViewportClient::DrawSafeFrames(FViewport& InViewport, FSceneView& Vi
 	if (EngineShowFlags.CameraAspectRatioBars || EngineShowFlags.CameraSafeFrames)
 	{
 		FSlateRect SafeRect;
-		if (CalculateSafeFrameRect(SafeRect, &InViewport))
+		if (CalculateEditorConstrainedViewRect(SafeRect, &InViewport))
 		{
 			if (EngineShowFlags.CameraSafeFrames)
 			{
@@ -2494,6 +2494,13 @@ void FEditorViewportClient::Draw(FViewport* InViewport, FCanvas* Canvas)
 
 	SetupViewForRendering(ViewFamily,*View);
 
+	FSlateRect SafeFrame;
+	View->CameraConstrainedViewRect = View->UnscaledViewRect;
+	if (CalculateEditorConstrainedViewRect(SafeFrame, Viewport))
+	{
+		View->CameraConstrainedViewRect = FIntRect(SafeFrame.Left, SafeFrame.Top, SafeFrame.Right, SafeFrame.Bottom);
+	}
+
 	if (IsAspectRatioConstrained())
 	{
 		// Clear the background to black if the aspect ratio is constrained, as the scene view won't write to all pixels.
diff --git a/Engine/Source/Runtime/Engine/Private/GameViewportClient.cpp b/Engine/Source/Runtime/Engine/Private/GameViewportClient.cpp
--- a/Engine/Source/Runtime/Engine/Private/GameViewportClient.cpp
+++ b/Engine/Source/Runtime/Engine/Private/GameViewportClient.cpp
@@ -738,6 +738,8 @@ void UGameViewportClient::Draw(FViewport* InViewport, FCanvas* SceneCanvas)
 
 						View->CurrentBufferVisualizationMode = CurrentBufferVisualizationMode;
 
+						View->CameraConstrainedViewRect = View->UnscaledViewRect;
+
 						// If this is the primary drawing pass, update things that depend on the view location
 						if( i == 0 )
 						{
diff --git a/Engine/Source/Runtime/Engine/Public/SceneManagement.h b/Engine/Source/Runtime/Engine/Public/SceneManagement.h
--- a/Engine/Source/Runtime/Engine/Public/SceneManagement.h
+++ b/Engine/Source/Runtime/Engine/Public/SceneManagement.h
@@ -2653,6 +2653,12 @@ public:
 	/** 0 if valid (we are rendering a screen postprocess pass )*/
 	struct FRenderingCompositePassContext* RenderingCompositePassContext; 
 
+	/** Aspect ratio constrained view rect. In the editor, when attached to a camera actor and the camera black bar showflag is enabled, the normal viewrect 
+	  * remains as the full viewport, and the black bars are just simulated by drawing black bars. This member stores the effective constrained area within the
+	  * bars.
+	 **/
+	FIntRect CameraConstrainedViewRect;
+
 #if WITH_EDITOR
 	/** The set of (the first 64) groups' visibility info for this view */
 	uint64 EditorViewBitflag;
diff --git a/Engine/Source/Runtime/Renderer/Private/PostProcess/VisualizeTexture.cpp b/Engine/Source/Runtime/Renderer/Private/PostProcess/VisualizeTexture.cpp
--- a/Engine/Source/Runtime/Renderer/Private/PostProcess/VisualizeTexture.cpp
+++ b/Engine/Source/Runtime/Renderer/Private/PostProcess/VisualizeTexture.cpp
@@ -285,16 +285,23 @@ FIntRect FVisualizeTexture::ComputeVisualizeTextureRect(FIntPoint InputTextureSi
 			FIntPoint HalfMax = InputTextureSize - HalfMin;
 
 			ret = FIntRect(Center - HalfMin, Center + HalfMax);
+			break;
 		}
-		break;
 
 		// whole texture in PIP
 		case 3:
-			ret = FIntRect(80, ViewExtent.Y - ViewExtent.Y / 3 - 10, ViewExtent.X / 3 + 10, ViewExtent.Y - 10) + ViewRect.Min;
+		{
+			int32 LeftOffset = AspectRatioConstrainedViewRect.Min.X;
+			int32 BottomOffset = AspectRatioConstrainedViewRect.Max.Y - ViewRect.Max.Y;
+
+			ret = FIntRect(LeftOffset + 80, ViewExtent.Y - ViewExtent.Y / 3 - 10 + BottomOffset, ViewExtent.X / 3 + 10, ViewExtent.Y - 10 + BottomOffset) + ViewRect.Min;
 			break;
+		}
 
 		default:
+		{
 			break;
+		}
 	}
 
 	return ret;
@@ -867,6 +874,7 @@ IPooledRenderTarget* FVisualizeTexture::GetObservedElement() const
 void FVisualizeTexture::OnStartFrame(const FSceneView& View)
 {
 	ViewRect = View.UnscaledViewRect;
+	AspectRatioConstrainedViewRect = View.Family->EngineShowFlags.CameraAspectRatioBars ? View.CameraConstrainedViewRect : ViewRect;
 
 #if !(UE_BUILD_SHIPPING || UE_BUILD_TEST)
 	// VisualizeTexture observed render target is set each frame
diff --git a/Engine/Source/Runtime/Renderer/Private/PostProcess/VisualizeTexture.h b/Engine/Source/Runtime/Renderer/Private/PostProcess/VisualizeTexture.h
--- a/Engine/Source/Runtime/Renderer/Private/PostProcess/VisualizeTexture.h
+++ b/Engine/Source/Runtime/Renderer/Private/PostProcess/VisualizeTexture.h
@@ -75,8 +75,14 @@ private:
 	TRefCountPtr<FRHIShaderResourceView> StencilSRV;
 	FTextureRHIRef StencilSRVSrc;
 
+	// The view rectangle that we are drawing to
 	FIntRect ViewRect;
 
+	// View rectange, constrained to the camera aspect ratio (if required). In game modes, the view rectangle is set to the correct aspect ratio constrained rectangle, but
+	// in the editor it is set to the full viewport window, and separate black bars are drawn to simulate the contrained area. We need to know about that so we can keep
+	// the texture visualization image inside this area
+	FIntRect AspectRatioConstrainedViewRect;
+
 	// is called by FPooledRenderTarget
 
 	void GenerateContent(const FSceneRenderTargetItem& RenderTargetItem, const FPooledRenderTargetDesc& Desc);
