Temporary fix for crash of when shutting down Persona.

diff --git a/Engine/Source/Editor/UnrealEd/Private/Animation/DebugSkelMeshComponent.cpp b/Engine/Source/Editor/UnrealEd/Private/Animation/DebugSkelMeshComponent.cpp
--- a/Engine/Source/Editor/UnrealEd/Private/Animation/DebugSkelMeshComponent.cpp
+++ b/Engine/Source/Editor/UnrealEd/Private/Animation/DebugSkelMeshComponent.cpp
@@ -479,33 +479,34 @@ void UDebugSkelMeshComponent::ShowOnlyClothSections(bool bShow, int32 LODIndex)
 
 void UDebugSkelMeshComponent::RestoreClothSectionsVisibility()
 {
-	PreEditChange(NULL);
-
 	FSkeletalMeshResource* SkelMeshResource = GetSkeletalMeshResource();
-	check(SkelMeshResource);
-
-	for (int32 LODIndex = 0; LODIndex < SkelMeshResource->LODModels.Num(); LODIndex++)
+	if (ensure (SkelMeshResource))
 	{
-		FStaticLODModel& LODModel = SkelMeshResource->LODModels[LODIndex];
+		PreEditChange(NULL);
 
-		// enables all sections first
-		for (int32 SecIdx = 0; SecIdx < LODModel.Sections.Num(); SecIdx++)
+		for(int32 LODIndex = 0; LODIndex < SkelMeshResource->LODModels.Num(); LODIndex++)
 		{
-			LODModel.Sections[SecIdx].bDisabled = false;
-		}
+			FStaticLODModel& LODModel = SkelMeshResource->LODModels[LODIndex];
 
-		// disables corresponding original section to enable the cloth section instead
-		for (int32 SecIdx = 0; SecIdx < LODModel.Sections.Num(); SecIdx++)
-		{
-			FSkelMeshSection& Section = LODModel.Sections[SecIdx];
+			// enables all sections first
+			for(int32 SecIdx = 0; SecIdx < LODModel.Sections.Num(); SecIdx++)
+			{
+				LODModel.Sections[SecIdx].bDisabled = false;
+			}
 
-			if (LODModel.Chunks[Section.ChunkIndex].HasApexClothData())
+			// disables corresponding original section to enable the cloth section instead
+			for(int32 SecIdx = 0; SecIdx < LODModel.Sections.Num(); SecIdx++)
 			{
-				LODModel.Sections[Section.CorrespondClothSectionIndex].bDisabled = true;
+				FSkelMeshSection& Section = LODModel.Sections[SecIdx];
+
+				if(LODModel.Chunks[Section.ChunkIndex].HasApexClothData())
+				{
+					LODModel.Sections[Section.CorrespondClothSectionIndex].bDisabled = true;
+				}
 			}
 		}
-	}
 
-	PostEditChange();
+		PostEditChange();
+	}
 }
 #endif // #if WITH_APEX_CLOTHING
\ No newline at end of file
