#ue4 Build fix for previous checkin

diff --git a/Engine/Source/Editor/UnrealEd/Private/LevelEditorViewport.cpp b/Engine/Source/Editor/UnrealEd/Private/LevelEditorViewport.cpp
--- a/Engine/Source/Editor/UnrealEd/Private/LevelEditorViewport.cpp
+++ b/Engine/Source/Editor/UnrealEd/Private/LevelEditorViewport.cpp
@@ -4475,7 +4475,7 @@ void FLevelEditorViewportClient::LostFocus(FViewport* Viewport)
 bool FLevelEditorViewportClient::OverrideHighResScreenshotCaptureRegion(FIntRect& OutCaptureRegion)
 {
 	FSlateRect Rect;
-	if (CalculateSafeFrameRect(Rect, Viewport))
+	if (CalculateEditorConstrainedViewRect(Rect, Viewport))
 	{
 		FSlateRect InnerRect = Rect.InsetBy(FMargin(0.5f * SafePadding * Rect.GetSize().Size()));
 		OutCaptureRegion = FIntRect((int32)InnerRect.Left, (int32)InnerRect.Top, (int32)(InnerRect.Left + InnerRect.GetSize().X), (int32)(InnerRect.Top + InnerRect.GetSize().Y));
diff --git a/Engine/Source/Editor/UnrealEd/Public/EditorViewportClient.h b/Engine/Source/Editor/UnrealEd/Public/EditorViewportClient.h
--- a/Engine/Source/Editor/UnrealEd/Public/EditorViewportClient.h
+++ b/Engine/Source/Editor/UnrealEd/Public/EditorViewportClient.h
@@ -1006,7 +1006,7 @@ protected:
 	virtual bool GetActiveSafeFrame(float& OutAspectRatio) const { return false; }
 
 	/** Helper function to calculate the safe frame rectangle on the current viewport */
-	bool CalculateSafeFrameRect(FSlateRect& OutSafeFrameRect, FViewport* InViewport);
+	bool CalculateEditorConstrainedViewRect(FSlateRect& OutSafeFrameRect, FViewport* InViewport);
 
 private:
 	/** @return Whether or not the camera should be panned or dollied */
