Fixes to the Point Light Placement automation test

diff --git a/Engine/Source/Editor/UnrealEd/Private/Tests/EditorAutomationTests.cpp b/Engine/Source/Editor/UnrealEd/Private/Tests/EditorAutomationTests.cpp
--- a/Engine/Source/Editor/UnrealEd/Private/Tests/EditorAutomationTests.cpp
+++ b/Engine/Source/Editor/UnrealEd/Private/Tests/EditorAutomationTests.cpp
@@ -1391,27 +1391,36 @@ IMPLEMENT_SIMPLE_AUTOMATION_TEST(FLightPlacement, "QA.Point Light Placement", EA
 
 bool FLightPlacement::RunTest(const FString& Parameters)
 {
-	FString MapName = TEXT("/Engine/Maps/Templates/Template_Default");
-	FEditorAutomationTestUtilities::LoadMap(MapName);
+	//Open a new empty map
+	UWorld* World = GEditor->NewMap();
 
-	UWorld *World = NULL;
+	//Move the perspective viewport view to show the test.
+	for (int32 i = 0; i < GEditor->LevelViewportClients.Num(); i++)
 	{
-		//Gather assets
-		UObject* EditorCylinderMesh = (UStaticMesh*)StaticLoadObject(UStaticMesh::StaticClass(),NULL,TEXT("/Engine/EditorMeshes/EditorCylinder.EditorCylinder"),NULL,LOAD_None,NULL);
-		// Add cylinder to world
-		AStaticMeshActor* StaticMesh = Cast<AStaticMeshActor>(FActorFactoryAssetProxy::AddActorForAsset(EditorCylinderMesh));
-		World = StaticMesh->GetWorld();
-		//delete static mesh now that we have the world.  Not ideal. :/
-		GEditor->edactDeleteSelected(World);
+		FLevelEditorViewportClient* ViewportClient = GEditor->LevelViewportClients[i];
+		if (!ViewportClient->IsOrtho())
+		{
+			ViewportClient->SetViewLocation(FVector(890, 70, 280));
+			ViewportClient->SetViewRotation(FRotator(0, 180, 0));
+		}
 	}
 
+	{
+		//Gather assets
+		UObject* Astroid = (UStaticMesh*)StaticLoadObject(UStaticMesh::StaticClass(), NULL, TEXT("/Engine/Content/EditorAutomation/Astroid.Astroid"), NULL, LOAD_None, NULL);
+		//Add Astroid mesh to the world
+		AStaticMeshActor* StaticMesh = Cast<AStaticMeshActor>(FActorFactoryAssetProxy::AddActorForAsset(Astroid));
+		StaticMesh->TeleportTo(FVector(0.0f, 0.0f, 0.0f), FRotator(0, 0, 0));
+		StaticMesh->SetActorRelativeScale3D(FVector(1.0f, 1.0f, 1.0f));
+	}
 
 	//Point Light
-	APointLight* PointLight = Cast<APointLight>(GEditor->AddActor(World->GetCurrentLevel(), APointLight::StaticClass(), FVector(-80.0f, 400.0f, 544.0f)));
+	APointLight* PointLight = Cast<APointLight>(GEditor->AddActor(World->GetCurrentLevel(), APointLight::StaticClass(), FVector(0.0f, 0.0f, 400.0f)));
 	PointLight->SetMobility(EComponentMobility::Movable);
-	PointLight->SetBrightness(116.0f);
-	PointLight->SetLightColor(FColor(200, 100, 200));
+	PointLight->SetBrightness(5000.0f);
+	PointLight->SetLightColor(FColor(255, 0, 0));
 
+	//Duplicate the point light
 	{
 		FScopedTransaction DuplicateLightScope( NSLOCTEXT("UnrealEd.Test", "DuplicateLightScope", "Duplicate Light Scope") );
 
@@ -1421,14 +1430,18 @@ bool FLightPlacement::RunTest(const FString& Parameters)
 		for ( FSelectionIterator It( GEditor->GetSelectedActorIterator() ) ; It ; ++It )
 		{
 			AActor* Actor = static_cast<AActor*>( *It );
-			Actor->TeleportTo(FVector(192.0f,400.0f,544.0f), FRotator(0, 0, 0));
+			Actor->TeleportTo(FVector(10.0f,10.0f,400.0f), FRotator(0, 0, 0));
 		}
 	}
 
+	//Undo the duplication then redo it
 	GEditor->UndoTransaction();
+	GEditor->RedoTransaction();
 
-	PointLight->TeleportTo(FVector(-464.0f,400.0f, 544.0f), FRotator(0, 164, 0));
-	PointLight->SetRadius( 5000.0f );
+	//Update the duplicated asset
+	PointLight->TeleportTo(FVector(500.0f,300.0f, 300.0f), FRotator(0, 0, 0));
+	PointLight->SetRadius( 500.0f );
+	PointLight->SetLightColor(FColor(255, 255, 255));
 
 	return true;
 }
