Add messaging for when a user attempts to launch or package a iOS project with code on PC
#ue4
#ios

diff --git a/Engine/Source/Editor/MainFrame/Private/Frame/MainFrameActions.cpp b/Engine/Source/Editor/MainFrame/Private/Frame/MainFrameActions.cpp
--- a/Engine/Source/Editor/MainFrame/Private/Frame/MainFrameActions.cpp
+++ b/Engine/Source/Editor/MainFrame/Private/Frame/MainFrameActions.cpp
@@ -486,6 +486,13 @@ void FMainFrameActionCallbacks::PackageProject( const FString InPlatformName, co
 		}
 	}
 
+#if PLATFORM_WINDOWS
+	if (bProjectHasCode && FRocketSupport::IsRocket() && PlatformName == TEXT("IOS"))
+	{
+		FMessageDialog::Open(EAppMsgType::Ok, LOCTEXT("PackagingNotWorkingMessage", "Sorry, packaging is currently not supported for code-based iOS projects. This feature will be available in a future release.") );
+		return;
+	}
+#endif
 
 	UProjectPackagingSettings* PackagingSettings = Cast<UProjectPackagingSettings>(UProjectPackagingSettings::StaticClass()->GetDefaultObject());
 
diff --git a/Engine/Source/Editor/UnrealEd/Private/PlayLevel.cpp b/Engine/Source/Editor/UnrealEd/Private/PlayLevel.cpp
--- a/Engine/Source/Editor/UnrealEd/Private/PlayLevel.cpp
+++ b/Engine/Source/Editor/UnrealEd/Private/PlayLevel.cpp
@@ -1133,8 +1133,6 @@ static void HandleHyperlinkNavigate()
 
 void UEditorEngine::PlayUsingLauncher()
 {
-	//	FMessageDialog::Open(EAppMsgType::Ok, LOCTEXT("LaunchOnDeviceNotWorkingMessage", "Sorry, launching on a device is currently not supported. This feature will be available again soon.") );
-
 	if (!PlayUsingLauncherDeviceId.IsEmpty())
 	{
 		ILauncherServicesModule& LauncherServicesModule = FModuleManager::LoadModuleChecked<ILauncherServicesModule>(TEXT("LauncherServices"));
@@ -1148,6 +1146,15 @@ void UEditorEngine::PlayUsingLauncher()
 		FGameProjectGenerationModule& GameProjectModule = FModuleManager::LoadModuleChecked<FGameProjectGenerationModule>(TEXT("GameProjectGeneration"));
 		bool bHasCode = GameProjectModule.Get().GetProjectCodeFileCount() > 0;
 
+#if PLATFORM_WINDOWS
+		if (bHasCode && FRocketSupport::IsRocket() && PlayUsingLauncherDeviceId.Left(PlayUsingLauncherDeviceId.Find(TEXT("@"))) == TEXT("IOS"))
+		{
+			FMessageDialog::Open(EAppMsgType::Ok, LOCTEXT("LaunchOnDeviceNotWorkingMessage", "Sorry, launching on a device is currently not supported for code-based iOS projects. This feature will be available in a future release.") );
+			bPlayUsingLauncher = false;
+			return;
+		}
+#endif
+
 		ILauncherProfileRef LauncherProfile = LauncherServicesModule.CreateProfile(TEXT("Play On Device"));
 		LauncherProfile->SetBuildGame((bHasCode && FSourceCodeNavigation::IsCompilerAvailable()) || !FRocketSupport::IsRocket());
 		LauncherProfile->SetCookMode(ELauncherProfileCookModes::ByTheBook);
@@ -1199,7 +1206,7 @@ void UEditorEngine::PlayUsingLauncher()
 
 		for (const FString& MapName : MapNames)
 		{
-		LauncherProfile->AddCookedMap(MapName);
+			LauncherProfile->AddCookedMap(MapName);
 		}
 
 		ILauncherPtr Launcher = LauncherServicesModule.CreateLauncher();
