LevelEditor: fixed view port settings not being remembered after Editor restart

diff --git a/Engine/Source/Editor/UnrealEd/Classes/Settings/LevelEditorViewportSettings.h b/Engine/Source/Editor/UnrealEd/Classes/Settings/LevelEditorViewportSettings.h
--- a/Engine/Source/Editor/UnrealEd/Classes/Settings/LevelEditorViewportSettings.h
+++ b/Engine/Source/Editor/UnrealEd/Classes/Settings/LevelEditorViewportSettings.h
@@ -30,81 +30,80 @@ struct UNREALED_API FLevelEditorViewportInstanceSettings
 		, bIsRealtime(false)
 		, bShowFPS_DEPRECATED(false)
 		, bShowStats(false)
-	{
-	}
+	{ }
 
 	/**
 	 * The viewport type
 	 */
-	UPROPERTY()
+	UPROPERTY(config)
 	TEnumAsByte<ELevelViewportType> ViewportType;
 
 	/* 
 	 * View mode to set when this viewport is of type LVT_Perspective 
 	 */
-	UPROPERTY()
+	UPROPERTY(config)
 	TEnumAsByte<EViewModeIndex> PerspViewModeIndex;
 
 	/* 
 	 * View mode to set when this viewport is not of type LVT_Perspective 
 	 */
-	UPROPERTY()
+	UPROPERTY(config)
 	TEnumAsByte<EViewModeIndex> OrthoViewModeIndex;
 
 	/**
 	 * A set of flags that determines visibility for various scene elements (FEngineShowFlags), converted to string form
 	 * These have to be saved as strings since FEngineShowFlags is too complex for UHT to parse correctly
 	 */
-	UPROPERTY()
+	UPROPERTY(config)
 	FString EditorShowFlagsString;
 
 	/**
 	 * A set of flags that determines visibility for various scene elements (FEngineShowFlags), converted to string form
 	 * These have to be saved as strings since FEngineShowFlags is too complex for UHT to parse correctly
 	 */
-	UPROPERTY()
+	UPROPERTY(config)
 	FString GameShowFlagsString;
 
 	/**
 	 * The buffer visualization mode for the viewport
 	 */
-	UPROPERTY()
+	UPROPERTY(config)
 	FName BufferVisualizationMode;
 
 	/**
 	 * Setting to allow designers to override the automatic expose
 	 */
-	UPROPERTY()
+	UPROPERTY(config)
 	FExposureSettings ExposureSettings;
 
 	/*
 	 * Field of view angle for the viewport
 	 */
-	UPROPERTY()
+	UPROPERTY(config)
 	float FOVAngle;
 
 	/*
 	 * Is this viewport updating in real-time?
 	 */
-	UPROPERTY()
+	UPROPERTY(config)
 	bool bIsRealtime;
 
 	/*
 	 * Should this viewport show an FPS count?
 	 */
-	UPROPERTY()
+	UPROPERTY(config)
 	bool bShowFPS_DEPRECATED;
 
 	/*
 	 * Should this viewport show statistics?
 	 */
-	UPROPERTY()
+	UPROPERTY(config)
 	bool bShowStats;
 
 	/*
 	 * Should this viewport have any stats enabled by default?
 	 */
-	UPROPERTY()
+	UPROPERTY(config)
 	TArray<FString> EnabledStats;
 };
 
@@ -120,13 +119,13 @@ struct UNREALED_API FLevelEditorViewportInstanceSettingsKeyValuePair
 	/* 
 	 * Name identifying this config
 	 */
-	UPROPERTY()
+	UPROPERTY(config)
 	FString ConfigName;
 
 	/* 
 	 * Settings for this config
 	 */
-	UPROPERTY()
+	UPROPERTY(config)
 	FLevelEditorViewportInstanceSettings ConfigSettings;
 };
 
@@ -140,8 +139,6 @@ class UNREALED_API ULevelEditorViewportSettings
 {
 	GENERATED_UCLASS_BODY()
 
-public:
-
 	/**
 	 * Enable the use of flight camera controls under various circumstances.
 	 */
@@ -257,6 +254,7 @@ public:
 	/** If enabled the when dragging new objects out of the content browser, it will snap the objects Z coordinate to the floor below it (if any) instead of the Z grid snapping location */
 	UPROPERTY(EditAnywhere, config, Category=GridSnapping)
 	bool bSnapNewObjectsToFloor;
+
 private:
 
 	/** If enabled, use the old-style multiplicative/percentage scaling method instead of the new additive/fraction method */
@@ -372,6 +370,12 @@ public:
 	UPROPERTY(EditAnywhere, config, Category = LookAndFeel)
 	uint32 bSaveEngineStats : 1;
 
+private:
+
+	// Per-instance viewport settings.
+	UPROPERTY(config)
+	TArray<FLevelEditorViewportInstanceSettingsKeyValuePair> PerInstanceSettings;
+
 public:
 
 	/**
@@ -421,9 +425,11 @@ public:
 		PostEditChange();
 	}
 
-public:
-
-	/** @return True if percentage based scaling is enabled */
+	/**
+	 * Checks whether percentage based scaling should be used for view ports.
+	 *
+	 * @return true if percentage based scaling is enabled, false otherwise.
+	 */
 	bool UsePercentageBasedScaling( ) const
 	{
 		return bUsePercentageBasedScaling;
@@ -451,10 +457,4 @@ private:
 
 	// Holds an event delegate that is executed when a setting has changed.
 	FSettingChangedEvent SettingChangedEvent;
-
-	/* 
-	 * Per-instance viewport settings
-	 */
-	UPROPERTY(config)
-	TArray<FLevelEditorViewportInstanceSettingsKeyValuePair> PerInstanceSettings;
 };
