Adjust the way escape sequences are converted when exporting and importing Portable Object(.po) files.

diff --git a/Engine/Source/Editor/UnrealEd/Private/Commandlets/InternationalizationExportCommandlet.cpp b/Engine/Source/Editor/UnrealEd/Private/Commandlets/InternationalizationExportCommandlet.cpp
--- a/Engine/Source/Editor/UnrealEd/Private/Commandlets/InternationalizationExportCommandlet.cpp
+++ b/Engine/Source/Editor/UnrealEd/Private/Commandlets/InternationalizationExportCommandlet.cpp
@@ -179,59 +179,29 @@ static const TMap< FString, FString > POCulturePluralForms = TMapBuilder< FStrin
 /**
 *	Helper Functions
 */
-FString CustomReplaceCharWithEscapedChar( const FString& InStr )
+FString ConditionArchiveStrForPo( const FString& InStr )
 {
-	FString Result;
-	for (const TCHAR* Char = *InStr; *Char != TCHAR('\0'); ++Char)
-	{
-		switch (*Char)
-		{
-		case TCHAR('\\'): Result += TEXT("\\\\"); break;
-		case TCHAR('\n'): Result += TEXT("\\n"); break;
-		case TCHAR('\t'): Result += TEXT("\\t"); break;
-		case TCHAR('\b'): Result += TEXT("\\b"); break;
-		case TCHAR('\f'): Result += TEXT("\\f"); break;
-		case TCHAR('\r'): Result += TEXT("\\r"); break;
-		case TCHAR('\"'): Result += TEXT("\\\""); break;
-		default: Result += *Char;
-		}
-	}
+	FString Result = InStr;
+	Result.ReplaceInline(TEXT("\\'"), TEXT("\\\\'"));
+	Result.ReplaceInline(TEXT("\\x"), TEXT("\\\\x"));
+	Result.ReplaceInline(TEXT("\\u"), TEXT("\\\\u"));
+	Result.ReplaceInline(TEXT("\r"), TEXT("\\r"));
+	Result.ReplaceInline(TEXT("\n"), TEXT("\\n"));
+
+
 	return Result;
 }
 
-FString CustomReplaceEscapedCharWithChar( const FString& InStr )
+FString ConditionPoStringForArchive( const FString& InStr )
 {
-	FString Result;
-	for (const TCHAR* Char = *InStr; *Char != TCHAR('\0'); ++Char)
-	{
-		if (*Char == TCHAR('\\'))
-		{
-			++Char;
-			if( *Char == TCHAR('\0') )
-			{
-				Result += TCHAR('\\');
-				break;
-			}
-			switch (*Char)
-			{
-			case TCHAR('\"'): case TCHAR('\\'): case TCHAR('/'): Result += *Char; break;
-			case TCHAR('f'): Result += TCHAR('\f'); break;
-			case TCHAR('r'): Result += TCHAR('\r'); break;
-			case TCHAR('n'): Result += TCHAR('\n'); break;
-			case TCHAR('b'): Result += TCHAR('\b'); break;
-			case TCHAR('t'): Result += TCHAR('\t'); break;
-			default: Result += *Char;
-			}
-		}
-		else
-		{
-			Result += *Char;
-		}
-
-	}
+	FString Result = InStr;
+	Result.ReplaceInline(TEXT("\\\\u"), TEXT("\\u"));
+	Result.ReplaceInline(TEXT("\\\\x"), TEXT("\\x"));
+	Result.ReplaceInline(TEXT("\\\\'"), TEXT("\\'"));
 	return Result;
 }
 
+
 FString ConvertSrcLocationToPORef( const FString& InSrcLocation )
 {
 	// Source location format: /Path1/Path2/file.cpp - line 123
@@ -401,7 +371,7 @@ FString FPortableObjectHeader::ToString() const
 	for( auto Entry : HeaderEntries )
 	{
 		const FString& Key = Entry.Key;
-		const FString& Value = CustomReplaceCharWithEscapedChar( Entry.Value );
+		const FString& Value = Entry.Value;
 		Result += FString::Printf( TEXT("\"%s: %s\\n\"%s"), *Key, *Value, NewLineDelimiter );
 	}
 
@@ -420,7 +390,7 @@ bool FPortableObjectHeader::FromLocPOEntry( const TSharedRef<const FPortableObje
 
 	// The POEntry would store header info inside the MsgStr[0]
 	TArray<FString> HeaderLinesToProcess;
-	LocEntry->MsgStr[0].ParseIntoArray( &HeaderLinesToProcess, NewLineDelimiter, true );
+	LocEntry->MsgStr[0].ReplaceEscapedCharWithChar().ParseIntoArray( &HeaderLinesToProcess, NewLineDelimiter, true );
 
 	for( const FString& PotentialHeaderEntry : HeaderLinesToProcess )
 	{
@@ -429,7 +399,7 @@ bool FPortableObjectHeader::FromLocPOEntry( const TSharedRef<const FPortableObje
 		{
 			// Looks like a header entry so we add it
 			const FString& Key = PotentialHeaderEntry.LeftChop( PotentialHeaderEntry.Len() - SplitIndex ).Trim().TrimTrailing();
-			FString Value = CustomReplaceEscapedCharWithChar( PotentialHeaderEntry.RightChop( SplitIndex+1 ).Trim().TrimTrailing() );
+			FString Value = PotentialHeaderEntry.RightChop( SplitIndex+1 ).Trim().TrimTrailing();
 
 			HeaderEntries.Add( TPairInitializer<FString, FString>( Key, Value ) );
 		}
@@ -631,7 +601,7 @@ bool FPortableObjectFormatDOM::FromString( const FString& InStr )
 				NextLineIdx++;
 			}
 
-			ProcessedEntry->MsgCtxt = CustomReplaceEscapedCharWithChar( RawMsgCtxt );
+			ProcessedEntry->MsgCtxt = RawMsgCtxt;
 		}
 		else if( Line.StartsWith( TEXT("msgid") ) )
 		{
@@ -657,7 +627,7 @@ bool FPortableObjectFormatDOM::FromString( const FString& InStr )
 				LineIdx = NextLineIdx;	
 				NextLineIdx++;
 			}
-			ProcessedEntry->MsgId = CustomReplaceEscapedCharWithChar( RawMsgId );
+			ProcessedEntry->MsgId = RawMsgId;
 			bHasMsgId = true;
 		}
 		else if( Line.StartsWith( TEXT("msgid_plural") ) )
@@ -684,7 +654,7 @@ bool FPortableObjectFormatDOM::FromString( const FString& InStr )
 				LineIdx = NextLineIdx;	
 				NextLineIdx++;
 			}
-			ProcessedEntry->MsgIdPlural = CustomReplaceEscapedCharWithChar( RawMsgIdPlural );
+			ProcessedEntry->MsgIdPlural = RawMsgIdPlural;
 		}
 		else if( Line.StartsWith( TEXT("msgstr[") ) )
 		{
@@ -724,11 +694,11 @@ bool FPortableObjectFormatDOM::FromString( const FString& InStr )
 
 			if( ProcessedEntry->MsgStr.Num() > Index )
 			{
-				ProcessedEntry->MsgStr[Index] = CustomReplaceEscapedCharWithChar( RawMsgStr );
+				ProcessedEntry->MsgStr[Index] = RawMsgStr;
 			}
 			else
 			{
-				ProcessedEntry->MsgStr.Insert( CustomReplaceEscapedCharWithChar( RawMsgStr ), Index );
+				ProcessedEntry->MsgStr.Insert( RawMsgStr, Index );
 			}
 		}
 		else if( Line.StartsWith( TEXT("msgstr") ) )
@@ -756,7 +726,7 @@ bool FPortableObjectFormatDOM::FromString( const FString& InStr )
 				NextLineIdx++;
 			}
 
-			FString MsgStr = CustomReplaceEscapedCharWithChar( RawMsgStr );
+			FString MsgStr = RawMsgStr;
 			if( ProcessedEntry->MsgStr.Num() > 0 )
 			{
 				ProcessedEntry->MsgStr[0] = MsgStr;
@@ -880,8 +850,11 @@ TSharedPtr<FPortableObjectEntry> FPortableObjectFormatDOM::FindEntry( const FStr
 
 void FPortableObjectEntry::AddReference( const FString& InReference )
 {
-	ReferenceComments.AddUnique(InReference);
-
+	if(!InReference.IsEmpty())
+	{
+		ReferenceComments.AddUnique(InReference);
+	}
+	
 	//// Reference comments can contain multiple references in a single line so we parse those out.
 	//TArray<FString> ReferencesToProcess;
 	//InReference.ParseIntoArray( &ReferencesToProcess, TEXT(" "), true );
@@ -929,10 +902,10 @@ FString FPortableObjectEntry::ToString() const
 
 	if( !MsgCtxt.IsEmpty() )
 	{
-		Result += FString::Printf(TEXT("msgctxt \"%s\"%s"), *CustomReplaceCharWithEscapedChar( MsgCtxt ), NewLineDelimiter);
+		Result += FString::Printf(TEXT("msgctxt \"%s\"%s"), *MsgCtxt, NewLineDelimiter);
 	}
 
-	Result += FString::Printf(TEXT("msgid \"%s\"%s"), *CustomReplaceCharWithEscapedChar( MsgId ), NewLineDelimiter);
+	Result += FString::Printf(TEXT("msgid \"%s\"%s"), *MsgId, NewLineDelimiter);
 
 	if( MsgStr.Num() == 0)
 	{
@@ -940,13 +913,13 @@ FString FPortableObjectEntry::ToString() const
 	}
 	else if( MsgStr.Num() == 1 )
 	{
-		Result += FString::Printf(TEXT("msgstr \"%s\"%s"), *CustomReplaceCharWithEscapedChar( MsgStr[0] ), NewLineDelimiter);
+		Result += FString::Printf(TEXT("msgstr \"%s\"%s"), *MsgStr[0], NewLineDelimiter);
 	}
 	else
 	{
 		for( int32 Idx = 0; Idx < MsgStr.Num(); ++Idx )
 		{
-			Result += FString::Printf(TEXT("msgstr[%d] \"%s\"%s"), Idx, *CustomReplaceCharWithEscapedChar( MsgStr[Idx] ), NewLineDelimiter);
+			Result += FString::Printf(TEXT("msgstr[%d] \"%s\"%s"), Idx, *MsgStr[Idx], NewLineDelimiter);
 		}
 	}
 	return Result;
@@ -1049,12 +1022,15 @@ bool UInternationalizationExportCommandlet::DoExport( const FString& SourcePath,
 							TSharedPtr<FArchiveEntry> ArchiveEntry = InternationalizationArchive->FindEntryBySource( Namespace, Source, ContextIter->KeyMetadataObj );
 							if( ArchiveEntry.IsValid() )
 							{
+								const FString ConditionedArchiveSource = ConditionArchiveStrForPo(ArchiveEntry->Source.Text);
+								const FString ConditionedArchiveTranslation = ConditionArchiveStrForPo(ArchiveEntry->Translation.Text);
+
 								TSharedRef<FPortableObjectEntry> PoEntry = MakeShareable( new FPortableObjectEntry );
 								//@TODO: We support additional metadata entries that can be translated.  How do those fit in the PO file format?  Ex: isMature
-								PoEntry->MsgId = ArchiveEntry->Source.Text;
+								PoEntry->MsgId = ConditionedArchiveSource;
 								//@TODO: Take into account optional entries and entries that differ by keymetadata.  Ex. Each optional entry needs a unique msgCtxt
 								PoEntry->MsgCtxt = Namespace;
-								PoEntry->MsgStr.Add( ArchiveEntry->Translation.Text );
+								PoEntry->MsgStr.Add( ConditionedArchiveTranslation );
 
 								PoEntry->AddReference( ConvertSrcLocationToPORef( ContextIter->SourceLocation ) );
 								PortableObj.AddEntry( PoEntry );
@@ -1166,8 +1142,8 @@ bool UInternationalizationExportCommandlet::DoImport(const FString& SourcePath,
 				}
 				
 				const FString& Namespace = POEntry->MsgCtxt;
-				const FString& SourceText = POEntry->MsgId;
-				const FString& Translation = POEntry->MsgStr[0];
+				const FString& SourceText = ConditionPoStringForArchive(POEntry->MsgId);
+				const FString& Translation = ConditionPoStringForArchive(POEntry->MsgStr[0]);
 
 				//@TODO: Take into account optional entries and entries that differ by keymetadata.  Ex. Each optional entry needs a unique msgCtxt
 				TSharedPtr< FArchiveEntry > FoundEntry = InternationalizationArchive->FindEntryBySource( Namespace, SourceText, NULL );
