Duplicate now functions with hidden actors.
#ttp 334309 - EDITOR: Scene Outliner: Using copy/paste on a hidden actor results in a copied actor but using duplicate does nothing
#branch UE4
#proj Editor.UnrealEd
#change In FDuplicateJob::DuplicateActorsToLevel, changed the call to GEditor->SelectActor() to pass bSelectEvenIfHidden set to true.

diff --git a/Engine/Source/Editor/UnrealEd/Private/EditorActor.cpp b/Engine/Source/Editor/UnrealEd/Private/EditorActor.cpp
--- a/Engine/Source/Editor/UnrealEd/Private/EditorActor.cpp
+++ b/Engine/Source/Editor/UnrealEd/Private/EditorActor.cpp
@@ -328,7 +328,7 @@ public:
 		for ( int32 ActorIndex = 0 ; ActorIndex < Actors.Num() ; ++ActorIndex )
 		{
 			AActor* Actor = Actors[ ActorIndex ];
-			GEditor->SelectActor( Actor, true, false );
+			GEditor->SelectActor( Actor, true, false, true );
 		}
 
 		FString ScratchData;
