Cooker - Fixed a problem with checking for the max path

diff --git a/Engine/Source/Editor/UnrealEd/Private/Commandlets/CookCommandlet.cpp b/Engine/Source/Editor/UnrealEd/Private/Commandlets/CookCommandlet.cpp
--- a/Engine/Source/Editor/UnrealEd/Private/Commandlets/CookCommandlet.cpp
+++ b/Engine/Source/Editor/UnrealEd/Private/Commandlets/CookCommandlet.cpp
@@ -442,7 +442,8 @@ bool UCookCommandlet::SaveCookedPackage( UPackage* Package, uint32 SaveFlags, bo
 					}
 				}
 
-				if( PlatFilename.Len() >= PLATFORM_MAX_FILEPATH_LENGTH )
+				const FString FullFilename = FPaths::ConvertRelativePathToFull( PlatFilename );
+				if( FullFilename.Len() >= PLATFORM_MAX_FILEPATH_LENGTH )
 				{
 					UE_LOG( LogCookCommandlet, Error, TEXT( "Couldn't save package, filename is too long :%s" ), *PlatFilename );
 					bSavedCorrectly = false;
diff --git a/Engine/Source/Editor/UnrealEd/Private/CookOnTheFlyServer.cpp b/Engine/Source/Editor/UnrealEd/Private/CookOnTheFlyServer.cpp
--- a/Engine/Source/Editor/UnrealEd/Private/CookOnTheFlyServer.cpp
+++ b/Engine/Source/Editor/UnrealEd/Private/CookOnTheFlyServer.cpp
@@ -638,7 +638,8 @@ bool UCookOnTheFlyServer::SaveCookedPackage( UPackage* Package, uint32 SaveFlags
 					World->PersistentLevel->OwningWorld = World;
 				}
 
-				if( PlatFilename.Len() >= PLATFORM_MAX_FILEPATH_LENGTH )
+				const FString FullFilename = FPaths::ConvertRelativePathToFull( PlatFilename );
+				if( FullFilename.Len() >= PLATFORM_MAX_FILEPATH_LENGTH )
 				{
 					UE_LOG( LogCookOnTheFly, Error, TEXT( "Couldn't save package, filename is too long :%s" ), *PlatFilename );
 					bSavedCorrectly = false;
