Fixed logspam when WorldSettings panel is open after creating a new gamemode class.

diff --git a/Engine/Source/Editor/UnrealEd/Public/GameModeInfoCustomizer.h b/Engine/Source/Editor/UnrealEd/Public/GameModeInfoCustomizer.h
--- a/Engine/Source/Editor/UnrealEd/Public/GameModeInfoCustomizer.h
+++ b/Engine/Source/Editor/UnrealEd/Public/GameModeInfoCustomizer.h
@@ -174,7 +174,10 @@ public:
 			return NULL;
 		}
 
-		const UClass* GameModeClass = FindObject<UClass>(ANY_PACKAGE, *ClassName);
+		FString StrippedClassName = ClassName;
+		ConstructorHelpers::StripObjectClass(StrippedClassName);
+
+		const UClass* GameModeClass = FindObject<UClass>(ANY_PACKAGE, *StrippedClassName);
 		if (!GameModeClass)
 		{
 			GameModeClass = LoadObject<UClass>(NULL, *ClassName);
