[330502] Temp fix for a crash when attempting to place a migrated Blueprint asset that derives from a native C++ class.

- Modified AttemptDropObjAsActors() to not allow placement for now if the Blueprint GeneratedClass is invalid.

diff --git a/Engine/Source/Editor/UnrealEd/Private/LevelEditorViewport.cpp b/Engine/Source/Editor/UnrealEd/Private/LevelEditorViewport.cpp
--- a/Engine/Source/Editor/UnrealEd/Private/LevelEditorViewport.cpp
+++ b/Engine/Source/Editor/UnrealEd/Private/LevelEditorViewport.cpp
@@ -235,10 +235,14 @@ static TArray<AActor*> AttemptDropObjAsActors( ULevel* InLevel, UObject* ObjToUs
 		if (ObjectClass->IsChildOf(UBlueprint::StaticClass()))
 		{
 			UBlueprint* BlueprintObj = StaticCast<UBlueprint*>(ObjToUse);
-			check(BlueprintObj->ParentClass == BlueprintObj->GeneratedClass->GetSuperClass());
-			if (BlueprintObj->GeneratedClass->HasAnyClassFlags(CLASS_NotPlaceable | CLASS_Abstract))
+			bPlace = BlueprintObj->GeneratedClass != NULL;
+			if(bPlace)
 			{
-				bPlace = false;
+				check(BlueprintObj->ParentClass == BlueprintObj->GeneratedClass->GetSuperClass());
+				if (BlueprintObj->GeneratedClass->HasAnyClassFlags(CLASS_NotPlaceable | CLASS_Abstract))
+				{
+					bPlace = false;
+				}
 			}
 		}
 
