* added/fixed mouse leave event

diff --git a/Engine/Source/Editor/UnrealEd/Private/EditorViewportClient.cpp b/Engine/Source/Editor/UnrealEd/Private/EditorViewportClient.cpp
--- a/Engine/Source/Editor/UnrealEd/Private/EditorViewportClient.cpp
+++ b/Engine/Source/Editor/UnrealEd/Private/EditorViewportClient.cpp
@@ -3673,6 +3673,12 @@ void FEditorViewportClient::OnJoystickPlugged(const uint32 InControllerID, const
 	}
 }
 
+
+void FEditorViewportClient::MouseEnter(FViewport* Viewport,int32 x, int32 y)
+{
+	MouseMove(Viewport, x, y);
+}
+
 void FEditorViewportClient::MouseMove(FViewport* InViewport,int32 x, int32 y)
 {
 	check(IsInGameThread());
diff --git a/Engine/Source/Editor/UnrealEd/Private/LevelEditorViewport.cpp b/Engine/Source/Editor/UnrealEd/Private/LevelEditorViewport.cpp
--- a/Engine/Source/Editor/UnrealEd/Private/LevelEditorViewport.cpp
+++ b/Engine/Source/Editor/UnrealEd/Private/LevelEditorViewport.cpp
@@ -3715,6 +3715,8 @@ void FLevelEditorViewportClient::MouseEnter( FViewport* Viewport,int32 x, int32
 void FLevelEditorViewportClient::MouseLeave( FViewport* Viewport ) 
 {
 	GEditorModeTools().MouseLeave(this, Viewport);
+
+	FEditorViewportClient::MouseLeave(Viewport);
 }
 
 void FLevelEditorViewportClient::MouseMove(FViewport* Viewport,int32 x, int32 y)
diff --git a/Engine/Source/Editor/UnrealEd/Public/EditorViewportClient.h b/Engine/Source/Editor/UnrealEd/Public/EditorViewportClient.h
--- a/Engine/Source/Editor/UnrealEd/Public/EditorViewportClient.h
+++ b/Engine/Source/Editor/UnrealEd/Public/EditorViewportClient.h
@@ -407,6 +407,7 @@ public:
 	virtual bool InputGesture(FViewport* Viewport, EGestureEvent::Type GestureType, const FVector2D& GestureDelta) OVERRIDE;
 	virtual void ReceivedFocus(FViewport* Viewport) OVERRIDE;
 	virtual void OnJoystickPlugged(const uint32 InControllerID, const uint32 InType, const uint32 bInConnected)  OVERRIDE;
+	virtual void MouseEnter(FViewport* Viewport,int32 x, int32 y) OVERRIDE;
 	virtual void MouseMove(FViewport* Viewport,int32 x, int32 y) OVERRIDE;
 	virtual void MouseLeave( FViewport* Viewport ) OVERRIDE;
 	virtual EMouseCursor::Type GetCursor(FViewport* Viewport,int32 X,int32 Y) OVERRIDE;
