Fixing non-unity build.

diff --git a/Engine/Source/Editor/UnrealEd/Public/Kismet2/BlueprintEditorUtils.h b/Engine/Source/Editor/UnrealEd/Public/Kismet2/BlueprintEditorUtils.h
--- a/Engine/Source/Editor/UnrealEd/Public/Kismet2/BlueprintEditorUtils.h
+++ b/Engine/Source/Editor/UnrealEd/Public/Kismet2/BlueprintEditorUtils.h
@@ -4,6 +4,8 @@
 
 #include "Editor/ClassViewer/Public/ClassViewerModule.h"
 
+#include "EdGraphSchema_K2.h"
+
 /** 
   * Flags describing how to handle graph removal
   */
