Adding missing file change to fix the build.

diff --git a/Engine/Source/Editor/UnrealEd/UnrealEd.Build.cs b/Engine/Source/Editor/UnrealEd/UnrealEd.Build.cs
--- a/Engine/Source/Editor/UnrealEd/UnrealEd.Build.cs
+++ b/Engine/Source/Editor/UnrealEd/UnrealEd.Build.cs
@@ -63,6 +63,7 @@ public class UnrealEd : ModuleRules
                 "EditorStyle",
 				"SourceControl",
 				"UnrealEdMessages",
+				"BlueprintGraph",
 			}
 		);
 
@@ -185,6 +186,7 @@ public class UnrealEd : ModuleRules
 			new string[] {
 				"UserFeedback",
              	"CollectionManager",
+				"BlueprintGraph",
 			}
 			);
 
