Fix "create blueprint based on selected actor" to work with actor types with out an asset associated with them.

- Move check for null asset reference to the Post Create Blueprint function.
- Fix a crash when using this button to create an Particle System blueprint. UActorFactoryEmitter::GetAssetFromActorInstance() was returning a ParticleSystemComponent when PostCreateBlueprint() (and PostCreateActor for that matter) were expecting a ParticleSystem reference.

diff --git a/Engine/Source/Editor/UnrealEd/Private/Factories/ActorFactory.cpp b/Engine/Source/Editor/UnrealEd/Private/Factories/ActorFactory.cpp
--- a/Engine/Source/Editor/UnrealEd/Private/Factories/ActorFactory.cpp
+++ b/Engine/Source/Editor/UnrealEd/Private/Factories/ActorFactory.cpp
@@ -184,12 +184,15 @@ UObject* UActorFactoryStaticMesh::GetAssetFromActorInstance(AActor* Instance)
 
 void UActorFactoryStaticMesh::PostCreateBlueprint( UObject* Asset, AActor* CDO )
 {
-	UStaticMesh* StaticMesh = CastChecked<UStaticMesh>( Asset );
-	AStaticMeshActor* StaticMeshActor = CastChecked<AStaticMeshActor>( CDO );
-	UStaticMeshComponent* StaticMeshComponent = StaticMeshActor->StaticMeshComponent;
+	if (Asset != NULL && CDO != NULL)
+	{
+		UStaticMesh* StaticMesh = CastChecked<UStaticMesh>(Asset);
+		AStaticMeshActor* StaticMeshActor = CastChecked<AStaticMeshActor>(CDO);
+		UStaticMeshComponent* StaticMeshComponent = StaticMeshActor->StaticMeshComponent;
 
-	StaticMeshComponent->StaticMesh = StaticMesh;
-	StaticMeshComponent->StaticMeshDerivedDataKey = StaticMesh->RenderData->DerivedDataKey;
+		StaticMeshComponent->StaticMesh = StaticMesh;
+		StaticMeshComponent->StaticMeshDerivedDataKey = StaticMesh->RenderData->DerivedDataKey;
+	}
 }
 
 /*-----------------------------------------------------------------------------
@@ -305,21 +308,24 @@ void UActorFactoryDeferredDecal::PostSpawnActor(UObject* Asset, AActor* NewActor
 
 void UActorFactoryDeferredDecal::PostCreateBlueprint( UObject* Asset, AActor* CDO )
 {
-	UMaterialInterface* Material = GetMaterial( Asset );
-
-	if ( Material != NULL )
+	if (Asset != NULL && CDO != NULL)
 	{
-		TArray<UDecalComponent*> DecalComponents;
-		CDO->GetComponents(DecalComponents);
+		UMaterialInterface* Material = GetMaterial(Asset);
 
-		UDecalComponent* DecalComponent = NULL;
-		for (int32 Idx = 0; Idx < DecalComponents.Num() && DecalComponent == NULL; Idx++)
+		if (Material != NULL)
 		{
-			DecalComponent = Cast<UDecalComponent>(DecalComponents[Idx]);
-		}
+			TArray<UDecalComponent*> DecalComponents;
+			CDO->GetComponents(DecalComponents);
 
-		check(DecalComponent);
-		DecalComponent->DecalMaterial = Material;
+			UDecalComponent* DecalComponent = NULL;
+			for (int32 Idx = 0; Idx < DecalComponents.Num() && DecalComponent == NULL; Idx++)
+			{
+				DecalComponent = Cast<UDecalComponent>(DecalComponents[Idx]);
+			}
+
+			check(DecalComponent);
+			DecalComponent->DecalMaterial = Material;
+		}
 	}
 }
 
@@ -399,14 +405,24 @@ UObject* UActorFactoryEmitter::GetAssetFromActorInstance(AActor* Instance)
 {
 	check(Instance->IsA(NewActorClass));
 	AEmitter* Emitter = CastChecked<AEmitter>(Instance);
-	return Emitter->ParticleSystemComponent;
+	if (Emitter->ParticleSystemComponent.IsValid())
+	{
+		return Emitter->ParticleSystemComponent->Template;
+	}
+	else
+	{
+		return NULL;
+	}
 }
 
 void UActorFactoryEmitter::PostCreateBlueprint( UObject* Asset, AActor* CDO )
 {
-	UParticleSystem* ParticleSystem = CastChecked<UParticleSystem>(Asset);
-	AEmitter* Emitter = CastChecked<AEmitter>( CDO );
-	Emitter->SetTemplate( ParticleSystem );
+	if (Asset != NULL && CDO != NULL)
+	{
+		UParticleSystem* ParticleSystem = CastChecked<UParticleSystem>(Asset);
+		AEmitter* Emitter = CastChecked<AEmitter>(CDO);
+		Emitter->SetTemplate(ParticleSystem);
+	}
 }
 
 
@@ -509,22 +525,29 @@ void UActorFactoryPhysicsAsset::PostSpawnActor(UObject* Asset, AActor* NewActor)
 
 void UActorFactoryPhysicsAsset::PostCreateBlueprint( UObject* Asset, AActor* CDO )
 {
-	UPhysicsAsset* PhysicsAsset = CastChecked<UPhysicsAsset>(Asset);
-	ASkeletalMeshActor* SkeletalPhysicsActor = CastChecked<ASkeletalMeshActor>( CDO );
+	if (CDO != NULL)
+	{
+		ASkeletalMeshActor* SkeletalPhysicsActor = CastChecked<ASkeletalMeshActor>(CDO);
 
-	USkeletalMesh* UseSkelMesh = PhysicsAsset->PreviewSkeletalMesh.Get();
+		if (Asset != NULL)
+		{
+			UPhysicsAsset* PhysicsAsset = CastChecked<UPhysicsAsset>(Asset);
 
-	SkeletalPhysicsActor->SkeletalMeshComponent->SkeletalMesh = UseSkelMesh;
-	SkeletalPhysicsActor->SkeletalMeshComponent->PhysicsAssetOverride = PhysicsAsset;
+			USkeletalMesh* UseSkelMesh = PhysicsAsset->PreviewSkeletalMesh.Get();
 
-	// set physics setup
-	SkeletalPhysicsActor->SkeletalMeshComponent->KinematicBonesUpdateType = EKinematicBonesUpdateToPhysics::SkipAllBones;
-	SkeletalPhysicsActor->SkeletalMeshComponent->BodyInstance.bSimulatePhysics = true;
-	SkeletalPhysicsActor->SkeletalMeshComponent->bBlendPhysics = true;
+			SkeletalPhysicsActor->SkeletalMeshComponent->SkeletalMesh = UseSkelMesh;
+			SkeletalPhysicsActor->SkeletalMeshComponent->PhysicsAssetOverride = PhysicsAsset;
+		}
 
-	SkeletalPhysicsActor->bAlwaysRelevant = true;
-	SkeletalPhysicsActor->bReplicateMovement = true;
-	SkeletalPhysicsActor->SetReplicates(true);
+		// set physics setup
+		SkeletalPhysicsActor->SkeletalMeshComponent->KinematicBonesUpdateType = EKinematicBonesUpdateToPhysics::SkipAllBones;
+		SkeletalPhysicsActor->SkeletalMeshComponent->BodyInstance.bSimulatePhysics = true;
+		SkeletalPhysicsActor->SkeletalMeshComponent->bBlendPhysics = true;
+
+		SkeletalPhysicsActor->bAlwaysRelevant = true;
+		SkeletalPhysicsActor->bReplicateMovement = true;
+		SkeletalPhysicsActor->SetReplicates(true);
+	}
 }
 
 
@@ -680,20 +703,24 @@ void UActorFactoryAnimationAsset::PostSpawnActor( UObject* Asset, AActor* NewAct
 void UActorFactoryAnimationAsset::PostCreateBlueprint( UObject* Asset,  AActor* CDO )
 {
 	Super::PostCreateBlueprint( Asset, CDO );
-	UAnimationAsset* AnimationAsset = Cast<UAnimationAsset>(Asset);
-	UVertexAnimation* VertexAnimation = Cast<UVertexAnimation>(Asset);
-
-	ASkeletalMeshActor* SkeletalMeshActor = CastChecked<ASkeletalMeshActor>( CDO );
-	USkeletalMeshComponent* SkeletalComponent = ( SkeletalMeshActor->SkeletalMeshComponent );
-	if (AnimationAsset)
-	{
-		SkeletalComponent->SetAnimationMode(EAnimationMode::Type::AnimationSingleNode);
-		SkeletalComponent->SetAnimation(AnimationAsset);
-	}
-	else if (VertexAnimation)
+	
+	if (Asset != NULL && CDO != NULL)
 	{
-		SkeletalComponent->SetAnimationMode(EAnimationMode::Type::AnimationSingleNode);
-		SkeletalComponent->SetVertexAnimation(VertexAnimation);
+		UAnimationAsset* AnimationAsset = Cast<UAnimationAsset>(Asset);
+		UVertexAnimation* VertexAnimation = Cast<UVertexAnimation>(Asset);
+
+		ASkeletalMeshActor* SkeletalMeshActor = CastChecked<ASkeletalMeshActor>(CDO);
+		USkeletalMeshComponent* SkeletalComponent = (SkeletalMeshActor->SkeletalMeshComponent);
+		if (AnimationAsset)
+		{
+			SkeletalComponent->SetAnimationMode(EAnimationMode::Type::AnimationSingleNode);
+			SkeletalComponent->SetAnimation(AnimationAsset);
+		}
+		else if (VertexAnimation)
+		{
+			SkeletalComponent->SetAnimationMode(EAnimationMode::Type::AnimationSingleNode);
+			SkeletalComponent->SetVertexAnimation(VertexAnimation);
+		}
 	}
 }
 
@@ -854,12 +881,15 @@ void UActorFactorySkeletalMesh::PostSpawnActor( UObject* Asset, AActor* NewActor
 
 void UActorFactorySkeletalMesh::PostCreateBlueprint( UObject* Asset, AActor* CDO )
 {
-	USkeletalMesh* SkeletalMesh = GetSkeletalMeshFromAsset( Asset );
-	UAnimBlueprint* AnimBlueprint = Cast<UAnimBlueprint>( Asset );
+	if (Asset != NULL && CDO != NULL)
+	{
+		USkeletalMesh* SkeletalMesh = GetSkeletalMeshFromAsset(Asset);
+		UAnimBlueprint* AnimBlueprint = Cast<UAnimBlueprint>(Asset);
 
-	ASkeletalMeshActor* SkeletalMeshActor = CastChecked<ASkeletalMeshActor>( CDO );
-	SkeletalMeshActor->SkeletalMeshComponent->SkeletalMesh = SkeletalMesh;
-	SkeletalMeshActor->SkeletalMeshComponent->AnimBlueprintGeneratedClass = AnimBlueprint ? Cast<UAnimBlueprintGeneratedClass>(AnimBlueprint->GeneratedClass) : NULL;
+		ASkeletalMeshActor* SkeletalMeshActor = CastChecked<ASkeletalMeshActor>(CDO);
+		SkeletalMeshActor->SkeletalMeshComponent->SkeletalMesh = SkeletalMesh;
+		SkeletalMeshActor->SkeletalMeshComponent->AnimBlueprintGeneratedClass = AnimBlueprint ? Cast<UAnimBlueprintGeneratedClass>(AnimBlueprint->GeneratedClass) : NULL;
+	}
 }
 
 
@@ -923,12 +953,15 @@ UObject* UActorFactoryAmbientSound::GetAssetFromActorInstance(AActor* Instance)
 
 void UActorFactoryAmbientSound::PostCreateBlueprint( UObject* Asset, AActor* CDO )
 {
-	USoundBase* AmbientSound = Cast<USoundBase>( Asset );
-
-	if ( AmbientSound != NULL )
+	if (Asset != NULL && CDO != NULL)
 	{
-		AAmbientSound* NewSound = CastChecked<AAmbientSound>( CDO );
-		NewSound->AudioComponent->SetSound(AmbientSound);
+		USoundBase* AmbientSound = Cast<USoundBase>(Asset);
+
+		if (AmbientSound != NULL)
+		{
+			AAmbientSound* NewSound = CastChecked<AAmbientSound>(CDO);
+			NewSound->AudioComponent->SetSound(AmbientSound);
+		}
 	}
 }
 
@@ -1154,13 +1187,16 @@ void UActorFactoryMatineeActor::PostSpawnActor( UObject* Asset, AActor* NewActor
 
 void UActorFactoryMatineeActor::PostCreateBlueprint( UObject* Asset, AActor* CDO )
 {
-	UInterpData* MatineeData = Cast<UInterpData>( Asset );
-	AMatineeActor* MatineeActor = CastChecked<AMatineeActor>( CDO );
+	if (Asset != NULL && CDO != NULL)
+	{
+		UInterpData* MatineeData = Cast<UInterpData>(Asset);
+		AMatineeActor* MatineeActor = CastChecked<AMatineeActor>(CDO);
 
-	// @todo sequencer: Don't ever need or want InterpData for Sequencer.  We will probably get rid of this after old Matinee goes away.
-	// also note the PostSpawnActor() code above creates an UInterpData and puts it in the actor's outermost package.  Definitely do not
-	// want that for Sequencer.
-	MatineeActor->MatineeData = MatineeData;
+		// @todo sequencer: Don't ever need or want InterpData for Sequencer.  We will probably get rid of this after old Matinee goes away.
+		// also note the PostSpawnActor() code above creates an UInterpData and puts it in the actor's outermost package.  Definitely do not
+		// want that for Sequencer.
+		MatineeActor->MatineeData = MatineeData;
+	}
 }
 
 
@@ -1346,10 +1382,13 @@ UObject* UActorFactoryDestructible::GetAssetFromActorInstance(AActor* Instance)
 
 void UActorFactoryDestructible::PostCreateBlueprint( UObject* Asset, AActor* CDO )
 {
-	UDestructibleMesh* DestructibleMesh = CastChecked<UDestructibleMesh>( Asset );
-	ADestructibleActor* DestructibleActor = CastChecked<ADestructibleActor>( CDO );
+	if (Asset != NULL && CDO != NULL)
+	{
+		UDestructibleMesh* DestructibleMesh = CastChecked<UDestructibleMesh>(Asset);
+		ADestructibleActor* DestructibleActor = CastChecked<ADestructibleActor>(CDO);
 
-	DestructibleActor->DestructibleComponent->SetSkeletalMesh( DestructibleMesh );
+		DestructibleActor->DestructibleComponent->SetSkeletalMesh(DestructibleMesh);
+	}
 }
 
 
diff --git a/Engine/Source/Editor/UnrealEd/Private/Kismet2/Kismet2.cpp b/Engine/Source/Editor/UnrealEd/Private/Kismet2/Kismet2.cpp
--- a/Engine/Source/Editor/UnrealEd/Private/Kismet2/Kismet2.cpp
+++ b/Engine/Source/Editor/UnrealEd/Private/Kismet2/Kismet2.cpp
@@ -522,10 +522,8 @@ UBlueprint* FKismetEditorUtilities::CreateBlueprintFromActor(const FString& Path
 			{
 				// Create the blueprint
 				UObject* Asset = FactoryToUse->GetAssetFromActorInstance(Actor);
-				if (Asset)
-				{
-					NewBlueprint = FactoryToUse->CreateBlueprint( Asset, Package, BlueprintName, FName("CreateFromActor") );
-				}
+				// For Actors that don't have an asset associated with them, Asset will be null
+				NewBlueprint = FactoryToUse->CreateBlueprint( Asset, Package, BlueprintName, FName("CreateFromActor") );
 			}
 			else 
 			{
