When replacing actors with blueprints via the right click "replace with actors" menu, if the blueprint is a child of that actor, try to copy over properties.

diff --git a/Engine/Source/Editor/UnrealEd/Private/Editor.cpp b/Engine/Source/Editor/UnrealEd/Private/Editor.cpp
--- a/Engine/Source/Editor/UnrealEd/Private/Editor.cpp
+++ b/Engine/Source/Editor/UnrealEd/Private/Editor.cpp
@@ -4995,6 +4995,19 @@ void UEditorEngine::ReplaceSelectedActors(UActorFactory* Factory, const FAssetDa
 		if (Factory != NULL)
 		{
 			NewActor = Factory->CreateActor( Asset, Level, OldLocation, &OldRotation);
+			// For blueprints, try to copy over properties
+			if (Factory->IsA(UActorFactoryBlueprint::StaticClass()))
+			{
+				UBlueprint* Blueprint = CastChecked<UBlueprint>(Asset);
+				// Only try to copy properties if this blueprint is based on the actor
+				UClass* OldActorClass = OldActor->GetClass();
+				if (Blueprint->GeneratedClass->IsChildOf(OldActorClass))
+				{
+					NewActor->UnregisterAllComponents();
+					UEditorEngine::CopyPropertiesForUnrelatedObjects(OldActor, NewActor);
+					NewActor->RegisterAllComponents();
+				}
+			}
 		}
 		else
 		{
