Move BodyFindConstraints to be a member of the PhysicsAsset so engine code can use it as well as editor. Fixes compile error.

diff --git a/Engine/Source/Editor/PhAT/Private/PhATSharedData.cpp b/Engine/Source/Editor/PhAT/Private/PhATSharedData.cpp
--- a/Engine/Source/Editor/PhAT/Private/PhATSharedData.cpp
+++ b/Engine/Source/Editor/PhAT/Private/PhATSharedData.cpp
@@ -708,7 +708,7 @@ bool FPhATSharedData::WeldSelectedBodies(bool bWeld /* = true */)
 
 		// .. and any constraints of the 'child' body..
 		TArray<int32>	Constraints;
-		FPhysicsAssetUtils::BodyFindConstraints(PhysicsAsset, ChildBodyIndex, Constraints);
+		PhysicsAsset->BodyFindConstraints(ChildBodyIndex, Constraints);
 
 		for (int32 i = 0; i <Constraints.Num(); ++i)
 		{
@@ -1057,7 +1057,7 @@ void FPhATSharedData::DeleteBody(int32 DelBodyIndex, bool bRefreshComponent)
 
 	// .. and any constraints to the body.
 	TArray<int32>	Constraints;
-	FPhysicsAssetUtils::BodyFindConstraints(PhysicsAsset, DelBodyIndex, Constraints);
+	PhysicsAsset->BodyFindConstraints(DelBodyIndex, Constraints);
 
 	for (int32 i = 0; i <Constraints.Num(); ++i)
 	{
diff --git a/Engine/Source/Editor/UnrealEd/Private/PhysicsAssetUtils.cpp b/Engine/Source/Editor/UnrealEd/Private/PhysicsAssetUtils.cpp
--- a/Engine/Source/Editor/UnrealEd/Private/PhysicsAssetUtils.cpp
+++ b/Engine/Source/Editor/UnrealEd/Private/PhysicsAssetUtils.cpp
@@ -515,7 +515,7 @@ void WeldBodies(UPhysicsAsset* PhysAsset, int32 BaseBodyIndex, int32 AddBodyInde
 	// Then deal with any constraints.
 
 	TArray<int32>	Body2Constraints;
-	BodyFindConstraints(PhysAsset, AddBodyIndex, Body2Constraints);
+	PhysAsset->BodyFindConstraints(AddBodyIndex, Body2Constraints);
 
 	while( Body2Constraints.Num() > 0 )
 	{
@@ -552,27 +552,13 @@ void WeldBodies(UPhysicsAsset* PhysAsset, int32 BaseBodyIndex, int32 AddBodyInde
 		}
 
 		// See if we have any more constraints to body2.
-		BodyFindConstraints(PhysAsset, AddBodyIndex, Body2Constraints);
+		PhysAsset->BodyFindConstraints(AddBodyIndex, Body2Constraints);
 	}
 
 	// Finally remove the body
 	DestroyBody(PhysAsset, AddBodyIndex);
 }
 
-void BodyFindConstraints(UPhysicsAsset* PhysAsset, int32 BodyIndex, TArray<int32>& Constraints)
-{
-	Constraints.Empty();
-	FName bodyName = PhysAsset->BodySetup[BodyIndex]->BoneName;
-
-	for(int32 i=0; i<PhysAsset->ConstraintSetup.Num(); i++)
-	{
-		if( PhysAsset->ConstraintSetup[i]->DefaultInstance.ConstraintBone1 == bodyName || PhysAsset->ConstraintSetup[i]->DefaultInstance.ConstraintBone2 == bodyName )
-		{
-			Constraints.Add(i);
-		}
-	}
-}
-
 int32 CreateNewConstraint(UPhysicsAsset* PhysAsset, FName InConstraintName, UPhysicsConstraintTemplate* InConstraintSetup)
 {
 	// constraintClass must be a subclass of UPhysicsConstraintTemplate
@@ -662,12 +648,12 @@ void DestroyBody(UPhysicsAsset* PhysAsset, int32 bodyIndex)
 	// Now remove any constraints that were attached to this body.
 	// This is a bit yuck and slow...
 	TArray<int32> Constraints;
-	BodyFindConstraints(PhysAsset, bodyIndex, Constraints);
+	PhysAsset->BodyFindConstraints(bodyIndex, Constraints);
 
 	while(Constraints.Num() > 0)
 	{
 		DestroyConstraint( PhysAsset, Constraints[0] );
-		BodyFindConstraints(PhysAsset, bodyIndex, Constraints);
+		PhysAsset->BodyFindConstraints(bodyIndex, Constraints);
 	}
 
 	// Remove pointer from array. Actual objects will be garbage collected.
diff --git a/Engine/Source/Editor/UnrealEd/Public/PhysicsAssetUtils.h b/Engine/Source/Editor/UnrealEd/Public/PhysicsAssetUtils.h
--- a/Engine/Source/Editor/UnrealEd/Public/PhysicsAssetUtils.h
+++ b/Engine/Source/Editor/UnrealEd/Public/PhysicsAssetUtils.h
@@ -66,14 +66,6 @@ namespace FPhysicsAssetUtils
 	 */
 	UNREALED_API void WeldBodies(UPhysicsAsset* PhysAsset, int32 BaseBodyIndex, int32 AddBodyIndex, USkeletalMeshComponent* SkelComp);
 
-	/** Find all the constraints that are connected to a particular body.
-	 * 
-	 * @param	PhysAsset		The PhysicsAsset to find the constraints for
-	 * @param	BodyIndex		The index of the body to find the constraints for
-	 * @param	Constraints		Returns the found constraints
-	 **/
-	UNREALED_API void BodyFindConstraints(UPhysicsAsset* PhysAsset, int32 BodyIndex, TArray<int32>& Constraints);
-
 	/**
 	 * Creates a new constraint
 	 *
diff --git a/Engine/Source/Runtime/Engine/Classes/PhysicsEngine/PhysicsAsset.h b/Engine/Source/Runtime/Engine/Classes/PhysicsEngine/PhysicsAsset.h
--- a/Engine/Source/Runtime/Engine/Classes/PhysicsEngine/PhysicsAsset.h
+++ b/Engine/Source/Runtime/Engine/Classes/PhysicsEngine/PhysicsAsset.h
@@ -165,5 +165,12 @@ public:
 
 	// @todo document
 	ENGINE_API int32 FindBodyIndex(FName BodyName);
+
+	/** Find all the constraints that are connected to a particular body.
+	 * 
+	 * @param	BodyIndex		The index of the body to find the constraints for
+	 * @param	Constraints		Returns the found constraints
+	 **/
+	ENGINE_API void BodyFindConstraints(int32 BodyIndex, TArray<int32>& Constraints);
 };
 
diff --git a/Engine/Source/Runtime/Engine/Private/PhysicsEngine/PhysicsAsset.cpp b/Engine/Source/Runtime/Engine/Private/PhysicsEngine/PhysicsAsset.cpp
--- a/Engine/Source/Runtime/Engine/Private/PhysicsEngine/PhysicsAsset.cpp
+++ b/Engine/Source/Runtime/Engine/Private/PhysicsEngine/PhysicsAsset.cpp
@@ -308,3 +308,18 @@ void UPhysicsAsset::GetAssetRegistryTags(TArray<FAssetRegistryTag>& OutTags) con
 
 	Super::GetAssetRegistryTags(OutTags);
 }
+
+
+void UPhysicsAsset::BodyFindConstraints(int32 BodyIndex, TArray<int32>& Constraints)
+{
+	Constraints.Empty();
+	FName BodyName = BodySetup[BodyIndex]->BoneName;
+
+	for(int32 ConIdx=0; ConIdx<ConstraintSetup.Num(); ConIdx++)
+	{
+		if( ConstraintSetup[ConIdx]->DefaultInstance.ConstraintBone1 == BodyName || ConstraintSetup[ConIdx]->DefaultInstance.ConstraintBone2 == BodyName )
+		{
+			Constraints.Add(ConIdx);
+		}
+	}
+}
\ No newline at end of file
diff --git a/Engine/Source/Runtime/Engine/Private/Vehicles/WheeledVehicleMovementComponent.cpp b/Engine/Source/Runtime/Engine/Private/Vehicles/WheeledVehicleMovementComponent.cpp
--- a/Engine/Source/Runtime/Engine/Private/Vehicles/WheeledVehicleMovementComponent.cpp
+++ b/Engine/Source/Runtime/Engine/Private/Vehicles/WheeledVehicleMovementComponent.cpp
@@ -3,7 +3,6 @@
 #include "EnginePrivate.h"
 #include "Net/UnrealNetwork.h"
 #include "MessageLog.h"
-#include "PhysicsAssetUtils.h"
 
 #if WITH_PHYSX
 #include "../PhysicsEngine/PhysXSupport.h"
@@ -1230,7 +1229,7 @@ void UWheeledVehicleMovementComponent::FixupSkeletalMesh()
 
 							//we also need to make sure the wheel body has no constraints with the parent body
 							TArray<int32> WheelConstraints;
-							FPhysicsAssetUtils::BodyFindConstraints(PhysicsAsset, BodySetupIdx, WheelConstraints);
+							PhysicsAsset->BodyFindConstraints(BodySetupIdx, WheelConstraints);
 							for (int32 ConstraintIdx = 0; ConstraintIdx < WheelConstraints.Num(); ++ConstraintIdx)
 							{
 								UPhysicsConstraintTemplate * ConstraintTemplate = PhysicsAsset->ConstraintSetup[WheelConstraints[ConstraintIdx]];
