Remove #pragma warning(disable : 4996) (functions is deprecated) from compiler setup.
#proj core
#branch UE4
#summary Remove #pragma warning(disable : 4996) (functions is deprecated) from compiler setup. Disabled warning only where absolutely necessary, fixed other cases. Introduced DEPRECATED(VERSION, MESSAGE) macro.

diff --git a/Engine/Plugins/Experimental/Phya/Source/Phya/Private/PhyaLib/src/Resonator/paModalData.cpp b/Engine/Plugins/Experimental/Phya/Source/Phya/Private/PhyaLib/src/Resonator/paModalData.cpp
--- a/Engine/Plugins/Experimental/Phya/Source/Phya/Private/PhyaLib/src/Resonator/paModalData.cpp
+++ b/Engine/Plugins/Experimental/Phya/Source/Phya/Private/PhyaLib/src/Resonator/paModalData.cpp
@@ -37,6 +37,9 @@ paModalData :: ~paModalData() {
 	paFree(m_amp);
 }
 
+#if PLATFORM_WINDOWS
+#pragma warning (disable : 4996) // scanf is deprecated, but can't use the safe version as may not be supported on all platforms.
+#endif
 
 int
 paModalData :: read( const char* data ) {
@@ -80,4 +83,8 @@ paModalData :: read( const char* data ) {
 //		fprintf(stderr, "Total mode limit of %d exceeded: Some modes lost.\n\n", paModal::nMaxModes);
 
 	return(0);
-}
\ No newline at end of file
+}
+
+#if PLATFORM_WINDOWS
+#pragma warning (default : 4996) // scanf is deprecated.
+#endif
diff --git a/Engine/Source/Editor/UnrealEd/Private/Fbx/FbxMainExport.cpp b/Engine/Source/Editor/UnrealEd/Private/Fbx/FbxMainExport.cpp
--- a/Engine/Source/Editor/UnrealEd/Private/Fbx/FbxMainExport.cpp
+++ b/Engine/Source/Editor/UnrealEd/Private/Fbx/FbxMainExport.cpp
@@ -1143,7 +1143,7 @@ void FFbxExporter::ExportMatineeTrackMove(FbxNode* FbxActor, UInterpTrackInstMov
 
 	if (MoveTrack != NULL)
 	{
-		FbxAnimLayer* BaseLayer = (FbxAnimLayer*)AnimStack->GetMember(FBX_TYPE(FbxAnimLayer), 0);
+		FbxAnimLayer* BaseLayer = AnimStack->GetMember<FbxAnimLayer>(0);
 		FbxAnimCurve* Curve;
 
 		bool bPosCurve = true;
diff --git a/Engine/Source/Editor/UnrealEd/Private/Fbx/FbxMainImport.cpp b/Engine/Source/Editor/UnrealEd/Private/Fbx/FbxMainImport.cpp
--- a/Engine/Source/Editor/UnrealEd/Private/Fbx/FbxMainImport.cpp
+++ b/Engine/Source/Editor/UnrealEd/Private/Fbx/FbxMainImport.cpp
@@ -507,9 +507,9 @@ bool FFbxImporter::GetSceneInfo(FString Filename, FbxSceneInfo& SceneInfo)
 		
 		// TODO: display multiple anim stack
 		SceneInfo.TakeName = NULL;
-		for( int32 AnimStackIndex = 0; AnimStackIndex < Scene->GetSrcObjectCount(FbxAnimStack::ClassId); AnimStackIndex++ )
+		for (int32 AnimStackIndex = 0; AnimStackIndex < Scene->GetSrcObjectCount<FbxAnimStack>(); AnimStackIndex++)
 		{
-			FbxAnimStack* CurAnimStack = FbxCast<FbxAnimStack>(Scene->GetSrcObject(FbxAnimStack::ClassId, 0));
+			FbxAnimStack* CurAnimStack = Scene->GetSrcObject<FbxAnimStack>(0);
 			// TODO: skip empty anim stack
 			const char* AnimStackName = CurAnimStack->GetName();
 			SceneInfo.TakeName = new char[FCStringAnsi::Strlen(AnimStackName) + 1];
diff --git a/Engine/Source/Editor/UnrealEd/Private/Fbx/FbxMaterialImport.cpp b/Engine/Source/Editor/UnrealEd/Private/Fbx/FbxMaterialImport.cpp
--- a/Engine/Source/Editor/UnrealEd/Private/Fbx/FbxMaterialImport.cpp
+++ b/Engine/Source/Editor/UnrealEd/Private/Fbx/FbxMaterialImport.cpp
@@ -172,17 +172,17 @@ void UnFbx::FFbxImporter::ImportTexturesFromNode(FbxNode* Node)
 					FbxTexture * lTexture= NULL;
 
 					//Here we have to check if it's layered textures, or just textures:
-					int32 LayeredTextureCount = Property.GetSrcObjectCount(FbxLayeredTexture::ClassId);
+					int32 LayeredTextureCount = Property.GetSrcObjectCount<FbxLayeredTexture>();
 					FbxString PropertyName = Property.GetName();
 					if(LayeredTextureCount > 0)
 					{
 						for(int32 LayerIndex=0; LayerIndex<LayeredTextureCount; ++LayerIndex)
 						{
-							FbxLayeredTexture *lLayeredTexture = FbxCast <FbxLayeredTexture>(Property.GetSrcObject(FbxLayeredTexture::ClassId, LayerIndex));
-							int32 NbTextures = lLayeredTexture->GetSrcObjectCount(FbxTexture::ClassId);
+							FbxLayeredTexture *lLayeredTexture = Property.GetSrcObject<FbxLayeredTexture>(LayerIndex);
+							int32 NbTextures = lLayeredTexture->GetSrcObjectCount<FbxTexture>();
 							for(int32 TexIndex =0; TexIndex<NbTextures; ++TexIndex)
 							{
-								FbxFileTexture* Texture = FbxCast <FbxFileTexture> (lLayeredTexture->GetSrcObject(FbxFileTexture::ClassId,TexIndex));
+								FbxFileTexture* Texture = lLayeredTexture->GetSrcObject<FbxFileTexture>(TexIndex);
 								if(Texture)
 								{
 									ImportTexture(Texture, PropertyName == FbxSurfaceMaterial::sNormalMap || PropertyName == FbxSurfaceMaterial::sBump);
@@ -193,11 +193,11 @@ void UnFbx::FFbxImporter::ImportTexturesFromNode(FbxNode* Node)
 					else
 					{
 						//no layered texture simply get on the property
-						int32 NbTextures = Property.GetSrcObjectCount(FbxTexture::ClassId);
+						int32 NbTextures = Property.GetSrcObjectCount<FbxTexture>();
 						for(int32 TexIndex =0; TexIndex<NbTextures; ++TexIndex)
 						{
 
-							FbxFileTexture* Texture = FbxCast <FbxFileTexture> (Property.GetSrcObject(FbxFileTexture::ClassId,TexIndex));
+							FbxFileTexture* Texture = Property.GetSrcObject<FbxFileTexture>(TexIndex);
 							if(Texture)
 							{
 								ImportTexture(Texture, PropertyName == FbxSurfaceMaterial::sNormalMap || PropertyName == FbxSurfaceMaterial::sBump);
@@ -228,19 +228,19 @@ bool UnFbx::FFbxImporter::CreateAndLinkExpressionForMaterialProperty(
 	FbxProperty FbxProperty = FbxMaterial.FindProperty( MaterialProperty );
 	if( FbxProperty.IsValid() )
 	{
-		int32 LayeredTextureCount = FbxProperty.GetSrcObjectCount(FbxLayeredTexture::ClassId);
+		int32 LayeredTextureCount = FbxProperty.GetSrcObjectCount<FbxLayeredTexture>();
 		if (LayeredTextureCount>0)
 		{
 			UE_LOG(LogFbxMaterialImport, Warning,TEXT("Layered Textures are not supported (material %s)"),ANSI_TO_TCHAR(FbxMaterial.GetName()));
 		}
 		else
 		{
-			int32 TextureCount = FbxProperty.GetSrcObjectCount(FbxTexture::ClassId);
+			int32 TextureCount = FbxProperty.GetSrcObjectCount<FbxTexture>();
 			if (TextureCount>0)
 			{
 				for(int32 TextureIndex =0; TextureIndex<TextureCount; ++TextureIndex)
 				{
-					FbxFileTexture* FbxTexture = FbxProperty.GetSrcObject(FBX_TYPE(FbxFileTexture), TextureIndex);
+					FbxFileTexture* FbxTexture = FbxProperty.GetSrcObject<FbxFileTexture>(TextureIndex);
 
 					// create an unreal texture asset
 					UTexture* UnrealTexture = ImportTexture(FbxTexture, bSetupAsNormalMap);
diff --git a/Engine/Source/Editor/UnrealEd/Private/Fbx/FbxMatineeImport.cpp b/Engine/Source/Editor/UnrealEd/Private/Fbx/FbxMatineeImport.cpp
--- a/Engine/Source/Editor/UnrealEd/Private/Fbx/FbxMatineeImport.cpp
+++ b/Engine/Source/Editor/UnrealEd/Private/Fbx/FbxMatineeImport.cpp
@@ -113,10 +113,10 @@ bool FFbxImporter::IsNodeAnimated(FbxNode* Node, FbxAnimLayer* AnimLayer)
 {
 	if (!AnimLayer)
 	{
-		FbxAnimStack* AnimStack = Scene->GetMember(FBX_TYPE(FbxAnimStack), 0);
+		FbxAnimStack* AnimStack = Scene->GetMember<FbxAnimStack>(0);
 		if (!AnimStack) return false;
 
-		AnimLayer = AnimStack->GetMember(FBX_TYPE(FbxAnimLayer), 0);
+		AnimLayer = AnimStack->GetMember<FbxAnimLayer>(0);
 		if (AnimLayer == NULL) return false;
 	}
 	
@@ -126,9 +126,9 @@ bool FFbxImporter::IsNodeAnimated(FbxNode* Node, FbxAnimLayer* AnimLayer)
 
 	// translation animation
 	FbxProperty TransProp = Node->LclTranslation;
-	for (int32 i = 0; i < TransProp.GetSrcObjectCount(FBX_TYPE(FbxAnimCurveNode)); i++)
+	for (int32 i = 0; i < TransProp.GetSrcObjectCount<FbxAnimCurveNode>(); i++)
 	{
-		FbxAnimCurveNode* CurveNode = FbxCast<FbxAnimCurveNode>(TransProp.GetSrcObject(FbxAnimCurveNode::ClassId, i));
+		FbxAnimCurveNode* CurveNode = TransProp.GetSrcObject<FbxAnimCurveNode>(i);
 		if (CurveNode && AnimLayer->IsConnectedSrcObject(CurveNode))
 		{
 			bIsAnimated |= (bool)CurveNode->GetAnimationInterval(AnimTimeSpan);
@@ -137,9 +137,9 @@ bool FFbxImporter::IsNodeAnimated(FbxNode* Node, FbxAnimLayer* AnimLayer)
 	}
 	// rotation animation
 	FbxProperty RotProp = Node->LclRotation;
-	for (int32 i = 0; bIsAnimated == false && i < RotProp.GetSrcObjectCount(FBX_TYPE(FbxAnimCurveNode)); i++)
+	for (int32 i = 0; bIsAnimated == false && i < RotProp.GetSrcObjectCount<FbxAnimCurveNode>(); i++)
 	{
-		FbxAnimCurveNode* CurveNode = FbxCast<FbxAnimCurveNode>(RotProp.GetSrcObject(FbxAnimCurveNode::ClassId, i));
+		FbxAnimCurveNode* CurveNode = RotProp.GetSrcObject<FbxAnimCurveNode>(i);
 		if (CurveNode && AnimLayer->IsConnectedSrcObject(CurveNode))
 		{
 			bIsAnimated |= (bool)CurveNode->GetAnimationInterval(AnimTimeSpan);
@@ -174,12 +174,12 @@ void FFbxImporter::ImportMatineeSequence(AMatineeActor* InMatineeActor)
 	if (Scene == NULL || InMatineeActor == NULL) return;
 	
 	// merge animation layer at first
-	FbxAnimStack* AnimStack = Scene->GetMember(FBX_TYPE(FbxAnimStack), 0);
+	FbxAnimStack* AnimStack = Scene->GetMember<FbxAnimStack>(0);
 	if (!AnimStack) return;
 		
 	MergeAllLayerAnimation(AnimStack, FbxTime::GetFrameRate(Scene->GetGlobalSettings().GetTimeMode()));
 
-	FbxAnimLayer* AnimLayer = AnimStack->GetMember(FBX_TYPE(FbxAnimLayer), 0);
+	FbxAnimLayer* AnimLayer = AnimStack->GetMember<FbxAnimLayer>(0);
 	if (AnimLayer == NULL) return;
 
 	// If the Matinee editor is not open, we need to initialize the sequence.
@@ -552,12 +552,12 @@ float FFbxImporter::ImportMatineeActor(FbxNode* Node, UInterpGroupInst* MatineeG
 	// Bucket the transforms at the same time.
 	// The Matinee Movement track can take in a Translation vector
 	// and three animated Euler rotation angles.
-	FbxAnimStack* AnimStack = Scene->GetMember(FBX_TYPE(FbxAnimStack), 0);
+	FbxAnimStack* AnimStack = Scene->GetMember<FbxAnimStack>(0);
 	if (!AnimStack) return -1.0f;
 
 	MergeAllLayerAnimation(AnimStack, FbxTime::GetFrameRate(Scene->GetGlobalSettings().GetTimeMode()));
 
-	FbxAnimLayer* AnimLayer = AnimStack->GetMember(FBX_TYPE(FbxAnimLayer), 0);
+	FbxAnimLayer* AnimLayer = AnimStack->GetMember<FbxAnimLayer>(0);
 	if (AnimLayer == NULL) return -1.0f;
 	
 	bool bNodeAnimated = IsNodeAnimated(Node, AnimLayer);
diff --git a/Engine/Source/Editor/UnrealEd/Private/Fbx/FbxSkeletalMeshImport.cpp b/Engine/Source/Editor/UnrealEd/Private/Fbx/FbxSkeletalMeshImport.cpp
--- a/Engine/Source/Editor/UnrealEd/Private/Fbx/FbxSkeletalMeshImport.cpp
+++ b/Engine/Source/Editor/UnrealEd/Private/Fbx/FbxSkeletalMeshImport.cpp
@@ -802,7 +802,7 @@ bool UnFbx::FFbxImporter::ImportBone(TArray<FbxNode*>& NodeArray, FSkeletalMeshI
 	}
 
 	FbxArray<FbxAMatrix> GlobalsPerLink;
-	GlobalsPerLink.AddMultiple(SortedLinks.Num());
+	GlobalsPerLink.Grow(SortedLinks.Num());
 	GlobalsPerLink[0].SetIdentity();
 	
 	bool GlobalLinkFoundFlag;
@@ -1348,7 +1348,7 @@ void UnFbx::FFbxImporter::SetupAnimationDataFromMesh(USkeletalMesh* SkeletalMesh
 {
 	USkeleton * Skeleton = SkeletalMesh->Skeleton;
 
-	if ( Scene->GetSrcObjectCount(FbxAnimStack::ClassId) > 0 )
+	if (Scene->GetSrcObjectCount<FbxAnimStack>() > 0)
 	{
 		if ( ensure ( Skeleton != NULL ) )
 		{
diff --git a/Engine/Source/Editor/UnrealEd/Private/SkeletalMeshEdit.cpp b/Engine/Source/Editor/UnrealEd/Private/SkeletalMeshEdit.cpp
--- a/Engine/Source/Editor/UnrealEd/Private/SkeletalMeshEdit.cpp
+++ b/Engine/Source/Editor/UnrealEd/Private/SkeletalMeshEdit.cpp
@@ -170,7 +170,7 @@ bool UEditorEngine::ReimportFbxAnimation( USkeleton * Skeleton, UAnimSequence *
 		else
 		{
 			// pick first one
-			FbxAnimStack* CurAnimStack = FbxCast<FbxAnimStack>(FbxImporter->Scene->GetSrcObject(FbxAnimStack::ClassId, 0));
+			FbxAnimStack* CurAnimStack = FbxImporter->Scene->GetSrcObject<FbxAnimStack>(0);
 			if (CurAnimStack)
 			{
 				// set current anim stack
@@ -245,7 +245,7 @@ void UnFbx::FFbxImporter::MergeAllLayerAnimation(FbxAnimStack* AnimStack, int32
 	// always apply unroll filter
 	FbxAnimCurveFilterUnroll UnrollFilter;
 
-	FbxAnimLayer* lLayer = static_cast<FbxAnimLayer*>(AnimStack->GetMember(FBX_TYPE(FbxAnimLayer),0));
+	FbxAnimLayer* lLayer = AnimStack->GetMember<FbxAnimLayer>(0);
 	UnrollFilter.Reset();
 	ApplyUnroll(Scene->GetRootNode(), lLayer, &UnrollFilter);
 }
@@ -260,12 +260,12 @@ bool UnFbx::FFbxImporter::IsValidAnimationData(TArray<FbxNode*>& SortedLinks, TA
 
 	ValidTakeCount = 0;
 
-	int32 AnimStackCount = Scene->GetSrcObjectCount( FbxAnimStack::ClassId );
+	int32 AnimStackCount = Scene->GetSrcObjectCount<FbxAnimStack>();
 
 	int32 AnimStackIndex;
 	for (AnimStackIndex = 0; AnimStackIndex < AnimStackCount; AnimStackIndex++ )
 	{
-		FbxAnimStack* CurAnimStack = FbxCast<FbxAnimStack>(Scene->GetSrcObject(FbxAnimStack::ClassId, AnimStackIndex));
+		FbxAnimStack* CurAnimStack = Scene->GetSrcObject<FbxAnimStack>(AnimStackIndex);
 		// set current anim stack
 		Scene->SetCurrentAnimationStack(CurAnimStack);
 
@@ -475,10 +475,10 @@ UAnimSequence * UnFbx::FFbxImporter::ImportAnimations(USkeleton * Skeleton, UObj
 
 	}
 
-	int32 AnimStackCount = Scene->GetSrcObjectCount(FbxAnimStack::ClassId);
+	int32 AnimStackCount = Scene->GetSrcObjectCount<FbxAnimStack>();
 	for( int32 AnimStackIndex = 0; AnimStackIndex < AnimStackCount; AnimStackIndex++ )
 	{
-		FbxAnimStack* CurAnimStack = FbxCast<FbxAnimStack>(Scene->GetSrcObject(FbxAnimStack::ClassId, AnimStackIndex));
+		FbxAnimStack* CurAnimStack = Scene->GetSrcObject<FbxAnimStack>(AnimStackIndex);
 
 		FbxTimeSpan AnimTimeSpan = GetAnimationTimeSpan(SortedLinks[0], CurAnimStack);
 		bool bValidAnimStack = ValidateAnimStack(SortedLinks, NodeArray, CurAnimStack, ResampleRate, ImportOptions->bImportMorph, AnimTimeSpan);
@@ -542,10 +542,10 @@ int32 UnFbx::FFbxImporter::GetMaxSampleRate(TArray<FbxNode*>& SortedLinks, TArra
 {
 	int32 MaxStackResampleRate = 0;
 
-	int32 AnimStackCount = Scene->GetSrcObjectCount(FbxAnimStack::ClassId);
+	int32 AnimStackCount = Scene->GetSrcObjectCount<FbxAnimStack>();
 	for( int32 AnimStackIndex = 0; AnimStackIndex < AnimStackCount; AnimStackIndex++)
 	{
-		FbxAnimStack* CurAnimStack = FbxCast<FbxAnimStack>(Scene->GetSrcObject(FbxAnimStack::ClassId, AnimStackIndex));
+		FbxAnimStack* CurAnimStack = Scene->GetSrcObject<FbxAnimStack>(AnimStackIndex);
 
 		FbxTimeSpan AnimStackTimeSpan = GetAnimationTimeSpan(SortedLinks[0], CurAnimStack);
 
diff --git a/Engine/Source/Editor/UnrealEd/Private/SpeedTreeImportFactory.cpp b/Engine/Source/Editor/UnrealEd/Private/SpeedTreeImportFactory.cpp
--- a/Engine/Source/Editor/UnrealEd/Private/SpeedTreeImportFactory.cpp
+++ b/Engine/Source/Editor/UnrealEd/Private/SpeedTreeImportFactory.cpp
@@ -13,9 +13,9 @@
 #include "RawMesh.h"
 
 #if WITH_SPEEDTREE
-
+#pragma warning(disable : 4996) // 'function' was was declared deprecated  (needed for the secure string functions)
 #include "Core/Core.h"
-
+#pragma warning(default : 4996)
 #endif
 
 #define LOCTEXT_NAMESPACE "SpeedTreeImportFactory"
diff --git a/Engine/Source/Runtime/Core/Private/Windows/WindowsPlatformMisc.cpp b/Engine/Source/Runtime/Core/Private/Windows/WindowsPlatformMisc.cpp
--- a/Engine/Source/Runtime/Core/Private/Windows/WindowsPlatformMisc.cpp
+++ b/Engine/Source/Runtime/Core/Private/Windows/WindowsPlatformMisc.cpp
@@ -142,11 +142,6 @@ void FWindowsPlatformMisc::PlatformPreInit()
 	// Use our own handler for pure virtuals being called.
 	DefaultPureCallHandler = _set_purecall_handler( PureCallHandler );
 
-	// Check Windows version.
-	OSVERSIONINFOEX OsVersionInfo = { 0 };
-	OsVersionInfo.dwOSVersionInfoSize = sizeof( OSVERSIONINFOEX );
-	GetVersionEx( ( LPOSVERSIONINFO )&OsVersionInfo );
-
 	const int32 MinResolution[] = {640,480};
 	if ( ::GetSystemMetrics(SM_CXSCREEN) < MinResolution[0] || ::GetSystemMetrics(SM_CYSCREEN) < MinResolution[1] )
 	{
diff --git a/Engine/Source/Runtime/Core/Private/Windows/WindowsPlatformSurvey.cpp b/Engine/Source/Runtime/Core/Private/Windows/WindowsPlatformSurvey.cpp
--- a/Engine/Source/Runtime/Core/Private/Windows/WindowsPlatformSurvey.cpp
+++ b/Engine/Source/Runtime/Core/Private/Windows/WindowsPlatformSurvey.cpp
@@ -685,7 +685,10 @@ void FWindowsPlatformSurvey::GetOSVersionLabels(const SYSTEM_INFO& SystemInfo, F
 	OsVersionInfo.dwOSVersionInfoSize = sizeof( OSVERSIONINFOEX );
 	FString OSVersionLabel = TEXT("Windows (unknown version)");
 	FString OSSubVersionLabel;
+#pragma warning(disable : 4996) // 'function' was declared deprecated
 	if (GetVersionEx( ( LPOSVERSIONINFO )&OsVersionInfo ))
+#pragma warning(default : 4996)
+
 	{
 		bool bIsInvalidVersion = false;
 
diff --git a/Engine/Source/Runtime/Core/Private/Windows/WindowsTextInputMethodSystem.cpp b/Engine/Source/Runtime/Core/Private/Windows/WindowsTextInputMethodSystem.cpp
--- a/Engine/Source/Runtime/Core/Private/Windows/WindowsTextInputMethodSystem.cpp
+++ b/Engine/Source/Runtime/Core/Private/Windows/WindowsTextInputMethodSystem.cpp
@@ -324,7 +324,10 @@ bool FWindowsTextInputMethodSystem::InitializeTSF()
 
 		TSFActivationProxy = new FTSFActivationProxy(this);
 
+#pragma warning(disable : 4996) // 'function' was was declared deprecated
 		const DWORD WindowsVersion = ::GetVersion();
+#pragma warning(default : 4996)
+
 		const DWORD WindowsMajorVersion = LOBYTE(LOWORD(WindowsVersion));
 		const DWORD WindowsMinorVersion = HIBYTE(LOWORD(WindowsVersion));
 
diff --git a/Engine/Source/Runtime/Core/Public/GenericPlatform/GenericPlatformCompilerSetup.h b/Engine/Source/Runtime/Core/Public/GenericPlatform/GenericPlatformCompilerSetup.h
new file mode 100644
--- /dev/null
+++ b/Engine/Source/Runtime/Core/Public/GenericPlatform/GenericPlatformCompilerSetup.h
@@ -0,0 +1,17 @@
+// Copyright 1998-2014 Epic Games, Inc. All Rights Reserved.
+
+/*=============================================================================================
+	GenericPlatformCompilerSetup.h: pragmas, version checks and other things for generic compiler
+==============================================================================================*/
+
+#pragma once
+
+#ifndef DEPRECATED
+
+	#ifndef DISABLE_DEPRECATION
+		#define DEPRECATED(VERSION, MESSAGE)
+	#else
+		#define DEPRECATED(VERSION, MESSAGE)
+	#endif // DISABLE_DEPRECATION
+
+#endif // DEPRECATED
\ No newline at end of file
diff --git a/Engine/Source/Runtime/Core/Public/GenericPlatform/GenericPlatformMath.h b/Engine/Source/Runtime/Core/Public/GenericPlatform/GenericPlatformMath.h
--- a/Engine/Source/Runtime/Core/Public/GenericPlatform/GenericPlatformMath.h
+++ b/Engine/Source/Runtime/Core/Public/GenericPlatform/GenericPlatformMath.h
@@ -23,6 +23,12 @@ struct FGenericPlatformMath
 		return (int32)F;
 	}
 
+	DEPRECATED(4.2, "FPlatformMath::Trunc is deprecated, use FPlatformMath::TruncToInt instead")
+	static CONSTEXPR FORCEINLINE int32 Trunc(float F)
+	{
+		return TruncToInt(F);
+	}
+
 	/**
 	 * Converts a float to an integer value with truncation towards zero.
 	 * @param F		Floating point value to convert
@@ -33,6 +39,12 @@ struct FGenericPlatformMath
 		return (float)TruncToInt(F);
 	}
 
+	DEPRECATED(4.2, "FPlatformMath::TruncFloat is deprecated, use FPlatformMath::TruncToFloat instead")
+	static CONSTEXPR FORCEINLINE float TruncFloat(float F)
+	{
+		return TruncToFloat(F);
+	}
+
 	/**
 	 * Converts a float to a nearest less or equal integer.
 	 * @param F		Floating point value to convert
diff --git a/Engine/Source/Runtime/Core/Public/GenericPlatform/MicrosoftPlatformString.h b/Engine/Source/Runtime/Core/Public/GenericPlatform/MicrosoftPlatformString.h
--- a/Engine/Source/Runtime/Core/Public/GenericPlatform/MicrosoftPlatformString.h
+++ b/Engine/Source/Runtime/Core/Public/GenericPlatform/MicrosoftPlatformString.h
@@ -8,6 +8,10 @@
 * Microsoft specific implementation 
 **/
 
+#if !USE_SECURE_CRT
+#pragma warning(disable : 4996) // 'function' was declared deprecated  (needed for the secure string functions)
+#endif
+
 struct FMicrosoftPlatformString : public FGenericPlatformString
 {
 	/** 
@@ -286,4 +290,8 @@ struct FMicrosoftPlatformString : public FGenericPlatformString
 	{
 		return _tcslen( (const WIDECHAR*)String );
 	}
-};
\ No newline at end of file
+};
+
+#if !USE_SECURE_CRT
+#pragma warning(default : 4996) // 'function' was was declared deprecated  (needed for the secure string functions)
+#endif
\ No newline at end of file
diff --git a/Engine/Source/Runtime/Core/Public/HAL/PlatformIncludes.h b/Engine/Source/Runtime/Core/Public/HAL/PlatformIncludes.h
--- a/Engine/Source/Runtime/Core/Public/HAL/PlatformIncludes.h
+++ b/Engine/Source/Runtime/Core/Public/HAL/PlatformIncludes.h
@@ -15,7 +15,13 @@
 #include <float.h>
 #include <string.h>
 
+// Platform specific compiler pre-setup, must be included before everything else
+#if PLATFORM_WINDOWS
+#include "Windows/WindowsPlatformCompilerPreSetup.h"
+#endif
+
 // Generic versions of the platform abstractions
+#include "GenericPlatform/GenericPlatformCompilerSetup.h"
 #include "GenericPlatform/GenericPlatformMemory.h"
 #include "GenericPlatform/GenericPlatformString.h"
 #include "GenericPlatform/GenericPlatformMisc.h"
diff --git a/Engine/Source/Runtime/Core/Public/HTML5/HTML5PlatformCompilerSetup.h b/Engine/Source/Runtime/Core/Public/HTML5/HTML5PlatformCompilerSetup.h
--- a/Engine/Source/Runtime/Core/Public/HTML5/HTML5PlatformCompilerSetup.h
+++ b/Engine/Source/Runtime/Core/Public/HTML5/HTML5PlatformCompilerSetup.h
@@ -89,12 +89,6 @@
 #pragma warning(disable : 4121) // 'symbol' : alignment of a member was sensitive to packing
 #pragma warning(disable : 4345) // behavior change: an object of POD type constructed with an initializer of the form () will be default-initialized
 
-#if !USE_SECURE_CRT
-#pragma warning(disable : 4996) // 'function' was was declared deprecated  (needed for the secure string functions)
-#else
-#pragma warning(default : 4996)	// enable deprecation warnings
-#endif
-
 #if UE_BUILD_DEBUG
 // xstring.h causes this warning in debug builds
 #pragma warning(disable : 4548) // expression before comma has no effect; expected expression with side-effect
@@ -121,4 +115,4 @@
 #ifdef _CHAR_UNSIGNED
 #error "Bad VC++ option: Characters must be signed" //lint !e309 suppress as lint doesn't have this defined
 #endif
-#endif
\ No newline at end of file
+#endif
diff --git a/Engine/Source/Runtime/Core/Public/WinRT/WinRTCompilerSetup.h b/Engine/Source/Runtime/Core/Public/WinRT/WinRTCompilerSetup.h
--- a/Engine/Source/Runtime/Core/Public/WinRT/WinRTCompilerSetup.h
+++ b/Engine/Source/Runtime/Core/Public/WinRT/WinRTCompilerSetup.h
@@ -87,12 +87,6 @@
 #pragma warning(disable : 4180) // qualifier applied to function type has no meaning; ignored
 #pragma warning( disable : 4121 )
 
-#if !USE_SECURE_CRT
-	#pragma warning(disable : 4996) // 'function' was was declared deprecated  (needed for the secure string functions)
-#else
-	#pragma warning(default : 4996)	// enable deprecation warnings
-#endif
-
 #if UE_BUILD_DEBUG
 	// xstring.h causes this warning in debug builds
 	#pragma warning(disable : 4548) // expression before comma has no effect; expected expression with side-effect
diff --git a/Engine/Source/Runtime/Core/Public/Windows/WindowsPlatformCompilerPreSetup.h b/Engine/Source/Runtime/Core/Public/Windows/WindowsPlatformCompilerPreSetup.h
new file mode 100644
--- /dev/null
+++ b/Engine/Source/Runtime/Core/Public/Windows/WindowsPlatformCompilerPreSetup.h
@@ -0,0 +1,13 @@
+// Copyright 1998-2014 Epic Games, Inc. All Rights Reserved.
+
+/*==========================================================================================================
+	WindowsPlatformCompilerPreSetup.h: setup for the compiler that needs to be defined before everything else
+===========================================================================================================*/
+
+#pragma once
+
+#ifndef DISABLE_DEPRECATION
+	#define DEPRECATED(VERSION, MESSAGE) __declspec(deprecated(MESSAGE))
+#else
+	#define DEPRECATED(VERSION, MESSAGE)
+#endif // DISABLE_DEPRECATION
diff --git a/Engine/Source/Runtime/Core/Public/Windows/WindowsPlatformCompilerSetup.h b/Engine/Source/Runtime/Core/Public/Windows/WindowsPlatformCompilerSetup.h
--- a/Engine/Source/Runtime/Core/Public/Windows/WindowsPlatformCompilerSetup.h
+++ b/Engine/Source/Runtime/Core/Public/Windows/WindowsPlatformCompilerSetup.h
@@ -101,12 +101,6 @@ checkAtCompileTime(_MSC_VER >= 1700, Visual_Studio_2012_Required);
 #pragma warning(disable : 4121) // 'symbol' : alignment of a member was sensitive to packing
 #pragma warning(disable : 4345) // behavior change: an object of POD type constructed with an initializer of the form () will be default-initialized
 
-#if !USE_SECURE_CRT
-#pragma warning(disable : 4996) // 'function' was was declared deprecated  (needed for the secure string functions)
-#else
-#pragma warning(default : 4996)	// enable deprecation warnings
-#endif
-
 #if UE_BUILD_DEBUG
 // xstring.h causes this warning in debug builds
 #pragma warning(disable : 4548) // expression before comma has no effect; expected expression with side-effect
diff --git a/Engine/Source/Runtime/OpenGLDrv/Private/OpenGLShaders.cpp b/Engine/Source/Runtime/OpenGLDrv/Private/OpenGLShaders.cpp
--- a/Engine/Source/Runtime/OpenGLDrv/Private/OpenGLShaders.cpp
+++ b/Engine/Source/Runtime/OpenGLDrv/Private/OpenGLShaders.cpp
@@ -197,9 +197,10 @@ static void ReplaceShaderSubstring(ANSICHAR* ShaderCode, const ANSICHAR* OldStri
 		Curr = Substr;
 		if (Substr)
 		{
-			strcpy(Substr, NewString);
-			Curr += strlen(NewString);
-			for (size_t index = 0; index < strlen(OldString)-strlen(NewString); ++index)
+			size_t NewStringLength = strlen(NewString);
+			FCStringAnsi::Strcpy(Substr, strlen(Substr), NewString);
+			Curr += NewStringLength;
+			for (size_t index = 0; index < strlen(OldString) - NewStringLength; ++index)
 			{
 				*Curr = ' ';
 				++Curr;
