Changed the UI for selecting an asset path, so it's no longer possible to enter an invalid path (uses the PathPicker widget from the content browser).
#ttp 332335 - EDITOR: BSP: CRASH: When creating a static mesh from selected brush, if you enter in an incorrect path name you will crash
#branch UE4
#proj Editor.UnrealEd, Editor.KismetWidgets, Editor.Kismet, Editor.DetailCustomizations
#add Created SCreateAssetFromActor, based on SCreateBlueprintFromActorDialog, as a generic base widget which can be used for choosing an asset path.
#remove Removed the old popup text box for entering an asset path
#add Made the 'Create Static Mesh' button on a brush open an SCreateAssetFromActor dialog so that a valid package name can be chosen.
#add Put a check on the resulting package object being non-null.
#change Changed SCreateBlueprintFromActorDialog to instance SCreateAssetFromActor, with appropriate properties.  Renamed to FCreateBlueprintFromActorDialog, as it is no longer a Slate class.  Renamed the files accordingly.
#change Made the chain of methods involved with converting a brush to a static mesh take an FString as the asset name instead of FText.

diff --git a/Engine/Source/Editor/DetailCustomizations/Private/ActorDetails.cpp b/Engine/Source/Editor/DetailCustomizations/Private/ActorDetails.cpp
--- a/Engine/Source/Editor/DetailCustomizations/Private/ActorDetails.cpp
+++ b/Engine/Source/Editor/DetailCustomizations/Private/ActorDetails.cpp
@@ -17,7 +17,7 @@
 #include "IContentBrowserSingleton.h"
 #include "ContentBrowserModule.h"
 #include "ScopedTransaction.h"
-#include "SCreateBlueprintFromActorDialog.h"
+#include "CreateBlueprintFromActorDialog.h"
 #include "ActorEditorUtils.h"
 #include "BlueprintEditorUtils.h"
 #include "ComponentTransformDetails.h"
@@ -899,7 +899,7 @@ END_SLATE_FUNCTION_BUILD_OPTIMIZATION
 
 FReply FActorDetails::OnPickBlueprintPathClicked(bool bHarvest )
 {
-	SCreateBlueprintFromActorDialog::OpenDialog(bHarvest);
+	FCreateBlueprintFromActorDialog::OpenDialog(bHarvest);
 
 	return FReply::Handled();
 }
diff --git a/Engine/Source/Editor/Kismet/Private/BlueprintEditorViewportContextMenuExtender.cpp b/Engine/Source/Editor/Kismet/Private/BlueprintEditorViewportContextMenuExtender.cpp
--- a/Engine/Source/Editor/Kismet/Private/BlueprintEditorViewportContextMenuExtender.cpp
+++ b/Engine/Source/Editor/Kismet/Private/BlueprintEditorViewportContextMenuExtender.cpp
@@ -4,7 +4,7 @@
 #include "BlueprintEditorPrivatePCH.h"
 #include "BlueprintEditorModule.h"
 #include "Editor/LevelEditor/Public/LevelEditor.h"
-#include "SCreateBlueprintFromActorDialog.h"
+#include "CreateBlueprintFromActorDialog.h"
 
 DEFINE_LOG_CATEGORY_STATIC(LogViewportBlueprintMenu, Log, All);
 
@@ -187,7 +187,7 @@ void FillBlueprintOptions(FMenuBuilder& MenuBuilder, TArray<AActor*> SelectedAct
 
 		if(bCanHarvestComponentsForBlueprint)
 		{
-			FUIAction CreateBlueprintAction( FExecuteAction::CreateStatic( &SCreateBlueprintFromActorDialog::OpenDialog, true ) );
+			FUIAction CreateBlueprintAction( FExecuteAction::CreateStatic( &FCreateBlueprintFromActorDialog::OpenDialog, true ) );
 			MenuBuilder.AddMenuEntry(LOCTEXT("CreateBlueprint", "Create Blueprint..."), LOCTEXT("CreateBlueprint_Tooltip", "Harvest Components from Selected Actors and create Blueprint"), FSlateIcon(FEditorStyle::GetStyleSetName(), "Kismet.HarvestBlueprintFromActors"), CreateBlueprintAction);
 		}
 	}
diff --git a/Engine/Source/Editor/KismetWidgets/Private/CreateBlueprintFromActorDialog.cpp b/Engine/Source/Editor/KismetWidgets/Private/CreateBlueprintFromActorDialog.cpp
new file mode 100644
--- /dev/null
+++ b/Engine/Source/Editor/KismetWidgets/Private/CreateBlueprintFromActorDialog.cpp
@@ -0,0 +1,92 @@
+// Copyright 1998-2014 Epic Games, Inc. All Rights Reserved.
+
+#include "KismetWidgetsPrivatePCH.h"
+#include "CreateBlueprintFromActorDialog.h"
+#include "SCreateAssetFromActor.h"
+#include "UnrealEd.h"
+#include "Editor/UnrealEd/Public/Kismet2/KismetEditorUtilities.h"
+
+#define LOCTEXT_NAMESPACE "CreateBlueprintFromActorDialog"
+
+void FCreateBlueprintFromActorDialog::OpenDialog(bool bInHarvest)
+{
+	TSharedPtr<SWindow> PickBlueprintPathWidget;
+	SAssignNew(PickBlueprintPathWidget, SWindow)
+		.Title(LOCTEXT("SelectPath", "Select Path"))
+		.ToolTipText(LOCTEXT("SelectPathTooltip", "Select the path where the Blueprint will be created at"))
+		.ClientSize(FVector2D(400, 400));
+
+	TSharedPtr<SCreateAssetFromActor> CreateBlueprintFromActorDialog;
+	PickBlueprintPathWidget->SetContent
+		(
+		SAssignNew(CreateBlueprintFromActorDialog, SCreateAssetFromActor, PickBlueprintPathWidget)
+		.AssetFilenameSuffix(TEXT("Blueprint"))
+		.HeadingText(LOCTEXT("CreateBlueprintFromActor_Heading", "Blueprint Name:"))
+		.CreateButtonText(LOCTEXT("CreateBlueprintFromActor_ButtonLabel", "Create Blueprint"))
+		.OnCreateAssetAction(FOnPathChosen::CreateStatic(FCreateBlueprintFromActorDialog::OnCreateBlueprint, bInHarvest))
+		);
+
+	TSharedPtr<SWindow> RootWindow = FGlobalTabmanager::Get()->GetRootWindow();
+	FSlateApplication::Get().AddWindowAsNativeChild(PickBlueprintPathWidget.ToSharedRef(), RootWindow.ToSharedRef() );
+}
+
+void FCreateBlueprintFromActorDialog::OnCreateBlueprint(const FString& InAssetPath, bool bInHarvest)
+{
+	UBlueprint* Blueprint = NULL;
+
+	if(bInHarvest)
+	{
+		TArray<AActor*> Actors;
+
+		USelection* SelectedActors = GEditor->GetSelectedActors();
+		for(FSelectionIterator Iter(*SelectedActors); Iter; ++Iter)
+		{
+			// We only care about actors that are referenced in the world for literals, and also in the same level as this blueprint
+			AActor* Actor = Cast<AActor>(*Iter);
+			if(Actor)
+			{
+				Actors.Add(Actor);
+			}
+		}
+		Blueprint = FKismetEditorUtilities::HarvestBlueprintFromActors(InAssetPath, Actors, true);
+	}
+	else
+	{
+		TArray< UObject* > SelectedActors;
+		GEditor->GetSelectedActors()->GetSelectedObjects(AActor::StaticClass(), SelectedActors);
+		check(SelectedActors.Num());
+		Blueprint = FKismetEditorUtilities::CreateBlueprintFromActor(InAssetPath, (AActor*)SelectedActors[0], true);
+	}
+
+	if(Blueprint)
+	{
+		// Rename new instance based on the original actor label rather than the asset name
+		USelection* SelectedActors = GEditor->GetSelectedActors();
+		if( SelectedActors && SelectedActors->Num() == 1 )
+		{
+			AActor* Actor = Cast<AActor>(SelectedActors->GetSelectedObject(0));
+			if(Actor)
+			{
+				GEditor->SetActorLabelUnique(Actor, FPackageName::GetShortName(InAssetPath));
+			}
+		}
+
+		TArray<UObject*> Objects;
+		Objects.Add(Blueprint);
+		GEditor->SyncBrowserToObjects( Objects );
+	}
+	else
+	{
+		FNotificationInfo Info( LOCTEXT("CreateBlueprintFromActorFailed", "Unable to create a blueprint from actor.") );
+		Info.ExpireDuration = 3.0f;
+		Info.bUseLargeFont = false;
+		TSharedPtr<SNotificationItem> Notification = FSlateNotificationManager::Get().AddNotification(Info);
+		if ( Notification.IsValid() )
+		{
+			Notification->SetCompletionState( SNotificationItem::CS_Fail );
+		}
+	}
+}
+
+
+#undef LOCTEXT_NAMESPACE
diff --git a/Engine/Source/Editor/KismetWidgets/Private/SCreateBlueprintFromActorDialog.cpp b/Engine/Source/Editor/KismetWidgets/Private/SCreateBlueprintFromActorDialog.cpp
deleted file mode 100644
diff --git a/Engine/Source/Editor/KismetWidgets/Public/CreateBlueprintFromActorDialog.h b/Engine/Source/Editor/KismetWidgets/Public/CreateBlueprintFromActorDialog.h
new file mode 100644
--- /dev/null
+++ b/Engine/Source/Editor/KismetWidgets/Public/CreateBlueprintFromActorDialog.h
@@ -0,0 +1,26 @@
+// Copyright 1998-2014 Epic Games, Inc. All Rights Reserved.
+
+#pragma once
+
+//////////////////////////////////////////////////////////////////////////
+// FCreateBlueprintFromActorDialog
+
+class FCreateBlueprintFromActorDialog
+{
+public:
+	/** 
+	 * Static function to access constructing this window.
+	 *
+	 * @param bInHarvest		true if the components of the selected actors should be harvested for the blueprint.
+	 */
+	static KISMETWIDGETS_API void OpenDialog(bool bInHarvest);
+private:
+
+	/** 
+	 * Will create the blueprint
+	 *
+	 * @param InAssetPath		Path of the asset to create
+	 * @param bInHarvest		true if the components of the selected actors should be harvested for the blueprint.
+	 */
+	static void OnCreateBlueprint(const FString& InAssetPath, bool bInHarvest);
+};
diff --git a/Engine/Source/Editor/KismetWidgets/Public/SCreateBlueprintFromActorDialog.h b/Engine/Source/Editor/KismetWidgets/Public/SCreateBlueprintFromActorDialog.h
deleted file mode 100644
diff --git a/Engine/Source/Editor/UnrealEd/Classes/Editor/EditorEngine.h b/Engine/Source/Editor/UnrealEd/Classes/Editor/EditorEngine.h
--- a/Engine/Source/Editor/UnrealEd/Classes/Editor/EditorEngine.h
+++ b/Engine/Source/Editor/UnrealEd/Classes/Editor/EditorEngine.h
@@ -1089,7 +1089,7 @@ public:
 	 *
 	 * @return							Returns the newly created actor with the newly created static mesh.
 	 */
-	AActor* ConvertBrushesToStaticMesh(const FText& InStaticMeshPackageName, TArray<ABrush*>& InBrushesToConvert, const FVector& InPivotLocation);
+	AActor* ConvertBrushesToStaticMesh(const FString& InStaticMeshPackageName, TArray<ABrush*>& InBrushesToConvert, const FVector& InPivotLocation);
 
 	/**
 	 * Converts passed in light actors into new actors of another type.
@@ -1117,7 +1117,7 @@ public:
 	 * @param	bUseSpecialCases			If true, looks for classes that can be handled by hardcoded conversions
 	 * @param	InStaticMeshPackageName		The name to save the brushes to.
 	 */
-	void DoConvertActors( const TArray<AActor*>& ActorsToConvert, UClass* ConvertToClass, const TSet<FString>& ComponentsToConsider, bool bUseSpecialCases, const FText& InStaticMeshPackageName );
+	void DoConvertActors( const TArray<AActor*>& ActorsToConvert, UClass* ConvertToClass, const TSet<FString>& ComponentsToConsider, bool bUseSpecialCases, const FString& InStaticMeshPackageName );
 
 	/**
 	 * Sets up for a potentially deferred ConvertActors call, based on if any brushes are being converted to a static mesh. If one (or more)
diff --git a/Engine/Source/Editor/UnrealEd/Private/Editor.cpp b/Engine/Source/Editor/UnrealEd/Private/Editor.cpp
--- a/Engine/Source/Editor/UnrealEd/Private/Editor.cpp
+++ b/Engine/Source/Editor/UnrealEd/Private/Editor.cpp
@@ -36,6 +36,7 @@
 #include "AnimationUtils.h"
 #include "AudioDecompress.h"
 #include "LevelEditor.h"
+#include "SCreateAssetFromActor.h"
 
 #include "Developer/DirectoryWatcher/Public/DirectoryWatcherModule.h"
 
@@ -5474,16 +5475,15 @@ void CopyActorComponentProperties( const AActor* SourceActor, AActor* DestActor,
 	}
 }
 
-AActor* UEditorEngine::ConvertBrushesToStaticMesh(const FText& InStaticMeshPackageName, TArray<ABrush*>& InBrushesToConvert, const FVector& InPivotLocation)
+AActor* UEditorEngine::ConvertBrushesToStaticMesh(const FString& InStaticMeshPackageName, TArray<ABrush*>& InBrushesToConvert, const FVector& InPivotLocation)
 {
 	AActor* NewActor(NULL);
 
-	FString PackageName = InStaticMeshPackageName.ToString();
+	FName ObjName = *FPackageName::GetLongPackageAssetName(InStaticMeshPackageName);
 
-	FName ObjName = *FPackageName::GetLongPackageAssetName(PackageName);
 
-
-	UPackage* Pkg = CreatePackage(NULL, *PackageName);
+	UPackage* Pkg = CreatePackage(NULL, *InStaticMeshPackageName);
+	check(Pkg != nullptr);
 
 	FVector Location(0.0f, 0.0f, 0.0f);
 	FRotator Rotation(0.0f, 0.0f, 0.0f);
@@ -5553,14 +5553,9 @@ struct TConvertData
 
 namespace ConvertHelpers
 {
-	void OnBrushToStaticMeshNameCommitted(const FText& InSettingsPackageName, ETextCommit::Type CommitInfo, TConvertData InConvertData)
+	void OnBrushToStaticMeshNameCommitted(const FString& InSettingsPackageName, TConvertData InConvertData)
 	{
-		if (CommitInfo == ETextCommit::OnEnter)
-		{
-			GEditor->DoConvertActors(InConvertData.ActorsToConvert, InConvertData.ConvertToClass, InConvertData.ComponentsToConsider, InConvertData.bUseSpecialCases, InSettingsPackageName);
-		}
-
-		GEditor->CloseEntryPopupWindow();
+		GEditor->DoConvertActors(InConvertData.ActorsToConvert, InConvertData.ConvertToClass, InConvertData.ComponentsToConsider, InConvertData.bUseSpecialCases, InSettingsPackageName);
 	}
 
 	void GetBrushList(const TArray<AActor*>& InActorsToConvert, UClass* InConvertToClass, TArray<ABrush*>& OutBrushList, int32& OutBrushIndexForReattachment)
@@ -5597,34 +5592,34 @@ void UEditorEngine::ConvertActors( const TArray<AActor*>& ActorsToConvert, UClas
 
 	if( BrushList.Num() )
 	{
-		const FString PackageName = TEXT("/Game/Unsorted/MyMesh");
-		IMainFrameModule& MainFrameModule = FModuleManager::LoadModuleChecked<IMainFrameModule>(TEXT("MainFrame"));
-
 		TConvertData ConvertData(ActorsToConvert, ConvertToClass, ComponentsToConsider, bUseSpecialCases);
 
-		TSharedRef<STextEntryPopup> TextEntry = 
-			SNew(STextEntryPopup)
-			.Label( LOCTEXT("ConvertActorsPopUpMsg", "New Actor Path (Convert Brush to Static Mesh)") )
-			.DefaultText( FText::FromString(*PackageName) )
-			.OnTextCommitted(FOnTextCommitted::CreateStatic(ConvertHelpers::OnBrushToStaticMeshNameCommitted, ConvertData))
-			.ClearKeyboardFocusOnCommit( false );
-
-		PopupWindow = FSlateApplication::Get().PushMenu(
-			MainFrameModule.GetParentWindow().ToSharedRef(),
-			TextEntry,
-			FSlateApplication::Get().GetCursorPos(),
-			FPopupTransitionEffect( FPopupTransitionEffect::TypeInPopup )
+		TSharedPtr<SWindow> CreateAssetFromActorWindow =
+			SNew(SWindow)
+			.Title(LOCTEXT("SelectPath", "Select Path"))
+			.ToolTipText(LOCTEXT("SelectPathTooltip", "Select the path where the static mesh will be created"))
+			.ClientSize(FVector2D(400, 400));
+
+		TSharedPtr<SCreateAssetFromActor> CreateAssetFromActorWidget;
+		CreateAssetFromActorWindow->SetContent
+			(
+			SAssignNew(CreateAssetFromActorWidget, SCreateAssetFromActor, CreateAssetFromActorWindow)
+			.AssetFilenameSuffix(TEXT("StaticMesh"))
+			.HeadingText(LOCTEXT("ConvertBrushesToStaticMesh_Heading", "Static Mesh Name:"))
+			.CreateButtonText(LOCTEXT("ConvertBrushesToStaticMesh_ButtonLabel", "Create Static Mesh"))
+			.OnCreateAssetAction(FOnPathChosen::CreateStatic(ConvertHelpers::OnBrushToStaticMeshNameCommitted, ConvertData))
 			);
 
-		TextEntry->FocusDefaultWidget();
+		TSharedPtr<SWindow> RootWindow = FGlobalTabmanager::Get()->GetRootWindow();
+		FSlateApplication::Get().AddWindowAsNativeChild(CreateAssetFromActorWindow.ToSharedRef(), RootWindow.ToSharedRef());
 	}
 	else
 	{
-		DoConvertActors(ActorsToConvert, ConvertToClass, ComponentsToConsider, bUseSpecialCases, FText::FromString(TEXT("")));
+		DoConvertActors(ActorsToConvert, ConvertToClass, ComponentsToConsider, bUseSpecialCases, TEXT(""));
 	}
 }
 
-void UEditorEngine::DoConvertActors( const TArray<AActor*>& ActorsToConvert, UClass* ConvertToClass, const TSet<FString>& ComponentsToConsider, bool bUseSpecialCases, const FText& InStaticMeshPackageName )
+void UEditorEngine::DoConvertActors( const TArray<AActor*>& ActorsToConvert, UClass* ConvertToClass, const TSet<FString>& ComponentsToConsider, bool bUseSpecialCases, const FString& InStaticMeshPackageName )
 {
 	// Early out if actor deletion is currently forbidden
 	if (GEditor->ShouldAbortActorDeletion())
diff --git a/Engine/Source/Editor/UnrealEd/Private/SCreateAssetFromActor.cpp b/Engine/Source/Editor/UnrealEd/Private/SCreateAssetFromActor.cpp
new file mode 100644
--- /dev/null
+++ b/Engine/Source/Editor/UnrealEd/Private/SCreateAssetFromActor.cpp
@@ -0,0 +1,181 @@
+// Copyright 1998-2014 Epic Games, Inc. All Rights Reserved.
+
+#include "UnrealEd.h"
+#include "SCreateAssetFromActor.h"
+#include "ContentBrowserModule.h"
+#include "AssetToolsModule.h"
+#include "AssetRegistryModule.h"
+#include "PackageTools.h"
+
+#define LOCTEXT_NAMESPACE "SCreateAssetFromActor"
+
+void SCreateAssetFromActor::Construct(const FArguments& InArgs, TSharedPtr<SWindow> InParentWindow)
+{
+	AssetFilenameSuffix = InArgs._AssetFilenameSuffix;
+	HeadingText = InArgs._HeadingText;
+	CreateButtonText = InArgs._CreateButtonText;
+	OnCreateAssetAction = InArgs._OnCreateAssetAction;
+
+	bIsReportingError = false;
+	AssetPath = FString("/Game");
+	FPathPickerConfig PathPickerConfig;
+	PathPickerConfig.DefaultPath = AssetPath;
+	PathPickerConfig.OnPathSelected = FOnPathSelected::CreateRaw(this, &SCreateAssetFromActor::OnSelectAssetPath);
+
+	USelection::SelectionChangedEvent.AddRaw(this, &SCreateAssetFromActor::OnLevelSelectionChanged);
+
+	// Set up PathPickerConfig.
+	FContentBrowserModule& ContentBrowserModule = FModuleManager::LoadModuleChecked<FContentBrowserModule>("ContentBrowser");
+
+	ParentWindow = InParentWindow;
+
+	FString PackageName;
+	ActorInstanceLabel.Empty();
+
+	USelection* SelectedActors = GEditor->GetSelectedActors();
+	for (FSelectionIterator Iter(*SelectedActors); Iter; ++Iter)
+	{
+		AActor* Actor = Cast<AActor>(*Iter);
+		if (Actor)
+		{
+			ActorInstanceLabel += Actor->GetActorLabel();
+			ActorInstanceLabel += TEXT("_");
+			break;
+		}
+	}
+
+	ActorInstanceLabel = PackageTools::SanitizePackageName(ActorInstanceLabel + AssetFilenameSuffix);
+
+	FString AssetName = ActorInstanceLabel;
+	FString BasePath = AssetPath / AssetName;
+	FAssetToolsModule& AssetToolsModule = FModuleManager::LoadModuleChecked<FAssetToolsModule>("AssetTools");
+	AssetToolsModule.Get().CreateUniqueAssetName(BasePath, TEXT(""), PackageName, AssetName);
+
+	ChildSlot
+	[
+		SNew(SVerticalBox)
+		+ SVerticalBox::Slot()
+		.FillHeight(1.0f)
+		[
+			ContentBrowserModule.Get().CreatePathPicker(PathPickerConfig)
+		]
+		+ SVerticalBox::Slot()
+		.AutoHeight()
+		[
+			SNew(SHorizontalBox)
+			+ SHorizontalBox::Slot()
+			.Padding(0.0f, 0.0f, 4.0f, 0.0f)
+			.AutoWidth()
+			[
+				SNew(STextBlock)
+				.Text(LOCTEXT("BlueprintNameLabel", "Blueprint Name:"))
+			]
+			+ SHorizontalBox::Slot()
+			.FillWidth(1.0f)
+			[
+				SAssignNew(FileNameWidget, SEditableTextBox)
+				.Text(FText::FromString(AssetName))
+				.OnTextChanged(this, &SCreateAssetFromActor::OnFilenameChanged)
+			]
+		]
+		+ SVerticalBox::Slot()
+		.HAlign(HAlign_Right)
+		.AutoHeight()
+		[
+			SNew(SHorizontalBox)
+			+ SHorizontalBox::Slot()
+			.Padding(0, 2, 6, 0)
+			.AutoWidth()
+			[
+				SNew(SButton)
+				.VAlign(VAlign_Center)
+				.HAlign(HAlign_Center)
+				.OnClicked(this, &SCreateAssetFromActor::OnCreateAssetFromActorClicked)
+				.IsEnabled(this, &SCreateAssetFromActor::IsCreateAssetFromActorEnabled)
+				[
+					SNew(STextBlock)
+					.Text(LOCTEXT("CreateBlueprintFromActorButton", "Create Blueprint"))
+				]
+			]
+			+ SHorizontalBox::Slot()
+			.Padding(0, 2, 0, 0)
+			.AutoWidth()
+			[
+				SNew(SButton)
+				.VAlign(VAlign_Center)
+				.HAlign(HAlign_Center)
+				.ContentPadding(FMargin(8, 2, 8, 2))
+				.OnClicked(this, &SCreateAssetFromActor::OnCancelCreateAssetFromActor)
+				[
+					SNew(STextBlock)
+					.Text(LOCTEXT("CancelBlueprintFromActorButton", "Cancel"))
+				]
+			]
+		]
+	];
+
+	OnFilenameChanged(FText::FromString(AssetName));
+}
+
+FReply SCreateAssetFromActor::OnCreateAssetFromActorClicked()
+{
+	ParentWindow->RequestDestroyWindow();
+	OnCreateAssetAction.ExecuteIfBound(AssetPath / FileNameWidget->GetText().ToString());
+	return FReply::Handled();
+}
+
+FReply SCreateAssetFromActor::OnCancelCreateAssetFromActor()
+{
+	ParentWindow->RequestDestroyWindow();
+	return FReply::Handled();
+}
+
+void SCreateAssetFromActor::OnSelectAssetPath(const FString& Path)
+{
+	AssetPath = Path;
+	OnFilenameChanged(FileNameWidget->GetText());
+}
+
+void SCreateAssetFromActor::OnLevelSelectionChanged(UObject* InObjectSelected)
+{
+	// When actor selection changes, this window should be destroyed.
+	ParentWindow->RequestDestroyWindow();
+}
+
+void SCreateAssetFromActor::OnFilenameChanged(const FText& InNewName)
+{
+	TArray<FAssetData> AssetData;
+	FAssetRegistryModule& AssetRegistryModule = FModuleManager::LoadModuleChecked<FAssetRegistryModule>("AssetRegistry");
+	AssetRegistryModule.Get().GetAssetsByPath(FName(*AssetPath), AssetData);
+
+	FText ErrorText;
+	if (!FEditorFileUtils::IsFilenameValidForSaving(InNewName.ToString(), ErrorText) || !FName(*InNewName.ToString()).IsValidObjectName(ErrorText))
+	{
+		FileNameWidget->SetError(ErrorText);
+		bIsReportingError = true;
+		return;
+	}
+	else
+	{
+		// Check to see if the name conflicts
+		for (auto Iter = AssetData.CreateConstIterator(); Iter; ++Iter)
+		{
+			if (Iter->AssetName.ToString() == InNewName.ToString())
+			{
+				FileNameWidget->SetError(LOCTEXT("AssetInUseError", "Asset name already in use!"));
+				bIsReportingError = true;
+				return;
+			}
+		}
+	}
+
+	FileNameWidget->SetError(FText::FromString(TEXT("")));
+	bIsReportingError = false;
+}
+
+bool SCreateAssetFromActor::IsCreateAssetFromActorEnabled() const
+{
+	return !bIsReportingError;
+}
+
+#undef LOCTEXT_NAMESPACE
diff --git a/Engine/Source/Editor/UnrealEd/Private/Tests/EditorAutomationTests.cpp b/Engine/Source/Editor/UnrealEd/Private/Tests/EditorAutomationTests.cpp
--- a/Engine/Source/Editor/UnrealEd/Private/Tests/EditorAutomationTests.cpp
+++ b/Engine/Source/Editor/UnrealEd/Private/Tests/EditorAutomationTests.cpp
@@ -1218,7 +1218,7 @@ bool FConvertToValidation::RunTest(const FString& Parameters)
 		//generate static mesh package name
 		FString PackageName = FPackageName::FilenameToLongPackageName(FPaths::AutomationTransientDir() + TEXT("ConvertToBSPToStaticMesh"));
 		//Convert brush to specific package name
-		GEditor->DoConvertActors(ToStaticMeshActors, AStaticMeshActor::StaticClass(), TSet<FString>(), true, FText::FromString(PackageName));
+		GEditor->DoConvertActors(ToStaticMeshActors, AStaticMeshActor::StaticClass(), TSet<FString>(), true, PackageName);
 
 		//find the package
 		UPackage* NewPackage = FindPackage(NULL, *PackageName);
diff --git a/Engine/Source/Editor/UnrealEd/Public/SCreateAssetFromActor.h b/Engine/Source/Editor/UnrealEd/Public/SCreateAssetFromActor.h
new file mode 100644
--- /dev/null
+++ b/Engine/Source/Editor/UnrealEd/Public/SCreateAssetFromActor.h
@@ -0,0 +1,83 @@
+// Copyright 1998-2014 Epic Games, Inc. All Rights Reserved.
+
+#pragma once
+
+//////////////////////////////////////////////////////////////////////////
+// SCreateAssetFromActor
+
+DECLARE_DELEGATE_OneParam(FOnPathChosen, const FString&);
+
+class SCreateAssetFromActor : public SCompoundWidget
+{
+public:
+	SLATE_BEGIN_ARGS(SCreateAssetFromActor)
+		: _AssetFilenameSuffix()
+		, _HeadingText()
+		, _CreateButtonText()
+	{}
+
+	/** The default suffix to use for the asset filename */
+	SLATE_ARGUMENT(FString, AssetFilenameSuffix)
+
+	/** The text to display at the top of the dialog */
+	SLATE_ARGUMENT(FText, HeadingText)
+
+	/** The label for the create button */
+	SLATE_ARGUMENT(FText, CreateButtonText)
+
+	/** Action to perform when create clicked */
+	SLATE_EVENT(FOnPathChosen, OnCreateAssetAction)
+
+	SLATE_END_ARGS()
+
+	/** Constructs this widget with InArgs */
+	UNREALED_API void Construct(const FArguments& InArgs, TSharedPtr<SWindow> InParentWindow);
+
+private:
+
+	/** Callback when the "create asset" button is clicked. */
+	FReply OnCreateAssetFromActorClicked();
+
+	/** Callback when the selected asset path has changed. */
+	void OnSelectAssetPath(const FString& Path);
+
+	/** Destroys the window when the operation is cancelled. */
+	FReply OnCancelCreateAssetFromActor();
+
+	/** Callback when level selection has changed, will destroy the window. */
+	void OnLevelSelectionChanged(UObject* InObjectSelected);
+
+	/** Callback when the user changes the filename for the Blueprint */
+	void OnFilenameChanged(const FText& InNewName);
+
+	/** Callback to see if creating an asset is enabled */
+	bool IsCreateAssetFromActorEnabled() const;
+
+private:
+	/** The window this widget is nested in */
+	TSharedPtr<SWindow> ParentWindow;
+
+	/** The selected path to create the asset */
+	FString AssetPath;
+
+	/** The resultant actor instance label, based on the original actor labels */
+	FString ActorInstanceLabel;
+
+	/** The default suffix to use for the asset filename */
+	FString AssetFilenameSuffix;
+
+	/** The text to display as a heading for the dialog */
+	FText HeadingText;
+
+	/** The label to be displayed on the create button */
+	FText CreateButtonText;
+
+	/** Filename textbox widget */
+	TSharedPtr<SEditableTextBox> FileNameWidget;
+
+	/** True if an error is currently being reported */
+	bool bIsReportingError;
+
+	/** Called when the create button is clicked */
+	FOnPathChosen OnCreateAssetAction;
+};
