Don't report an Actor reference error when deleting an Actor if the referencing Actor is spawned from a Child Actor Component

diff --git a/Engine/Source/Editor/UnrealEd/Private/EditorActor.cpp b/Engine/Source/Editor/UnrealEd/Private/EditorActor.cpp
--- a/Engine/Source/Editor/UnrealEd/Private/EditorActor.cpp
+++ b/Engine/Source/Editor/UnrealEd/Private/EditorActor.cpp
@@ -566,7 +566,16 @@ bool UUnrealEdEngine::edactDeleteSelected( UWorld* InWorld, bool bVerifyDeletion
 		FBlueprintEditorUtils::FindActorsThatReferenceActor( Actor, ClassTypesToIgnore, ReferencingActors );
 
 		bool bReferencedByLevelScript = (NULL != LSB && FBlueprintEditorUtils::FindNumReferencesToActorFromLevelScript(LSB, Actor) > 0);
-		bool bReferencedByActor = ( ReferencingActors.Num() > 0 );
+		bool bReferencedByActor = false;
+		for (AActor* ReferencingActor : ReferencingActors)
+		{
+			if (ReferencingActor->ParentComponentActor.Get() != Actor)
+			{
+				bReferencedByActor = true;
+				break;
+			}
+		}
+		(ReferencingActors.Num() > 0);
 
 		if ( bReferencedByLevelScript || bReferencedByActor )
 		{
