fixed inconsistent line spacing; added missing documentation

diff --git a/Engine/Source/Editor/UnrealEd/Private/Settings/GameMapsSettingsCustomization.h b/Engine/Source/Editor/UnrealEd/Private/Settings/GameMapsSettingsCustomization.h
--- a/Engine/Source/Editor/UnrealEd/Private/Settings/GameMapsSettingsCustomization.h
+++ b/Engine/Source/Editor/UnrealEd/Private/Settings/GameMapsSettingsCustomization.h
@@ -6,7 +6,6 @@
 
 #pragma once
 
-
 #include "PropertyEditing.h"
 #include "PropertyCustomizationHelpers.h"
 #include "GameModeInfoCustomizer.h"
@@ -14,8 +13,6 @@
 #define LOCTEXT_NAMESPACE "FLevelEditorPlaySettingsCustomization"
 
 
-
-
 /**
  * Implements a details view customization for UGameMapsSettingsCustomization objects.
  */
@@ -56,7 +53,6 @@ public:
 
 	// End IDetailCustomization interface
 
-
 public:
 
 	/**
@@ -69,7 +65,6 @@ public:
 		return MakeShareable(new FGameMapsSettingsCustomization());
 	}
 
-
 protected:
 
 	/**
@@ -124,7 +119,6 @@ protected:
 			];
 	}
 
-
 	/**
 	 * Checks whether the specified map name is valid.
 	 *
@@ -142,7 +136,6 @@ protected:
 		return (FPaths::FileExists(FPackageName::LongPackageNameToFilename(MapName, FPackageName::GetMapPackageExtension())));
 	}
 
-
 	/**
 	 * Creates a widget for the map picker.
 	 *
@@ -196,7 +189,6 @@ protected:
 		return MenuBuilder.MakeWidget();
 	}
 
-
 private:
 
 	// Handles selecting a map from a map picker.
@@ -256,9 +248,9 @@ private:
 		}
 	}
 
-
 private:
-	/** Helper class to customizer GameMode property */
+
+	// Helper class to customizer GameMode property.
 	TSharedPtr<FGameModeInfoCustomizer>	GameInfoModeCustomizer;
 };
 
diff --git a/Engine/Source/Editor/UnrealEd/Private/Settings/LevelEditorPlaySettingsCustomization.h b/Engine/Source/Editor/UnrealEd/Private/Settings/LevelEditorPlaySettingsCustomization.h
--- a/Engine/Source/Editor/UnrealEd/Private/Settings/LevelEditorPlaySettingsCustomization.h
+++ b/Engine/Source/Editor/UnrealEd/Private/Settings/LevelEditorPlaySettingsCustomization.h
@@ -6,7 +6,6 @@
 
 #pragma once
 
-
 #include "PropertyEditing.h"
 #include "PropertyCustomizationHelpers.h"
 
@@ -216,6 +215,13 @@ protected:
 	{
 	}
 
+	/**
+	 * Adds a section to the screen resolution menu.
+	 *
+	 * @param MenuBuilder The menu builder to add the section to.
+	 * @param Resolutions The collection of screen resolutions to add.
+	 * @param SectionName The name of the section to add.
+	 */
 	void AddScreenResolutionSection( FMenuBuilder& MenuBuilder, const TArray<FPlayScreenResolution>& Resolutions, const FText& SectionName )
 	{
 		MenuBuilder.BeginSection(NAME_None, SectionName);
@@ -424,14 +430,12 @@ public:
 				.IsEnabled(TAttribute<bool>::Create(TAttribute<bool>::FGetter::CreateSP(this, &FLevelEditorPlaySettingsCustomization::HandleRerouteInputToSecondWindowEnabled)))
 				.Visibility(TAttribute<EVisibility>::Create(TAttribute<EVisibility>::FGetter::CreateSP(this, &FLevelEditorPlaySettingsCustomization::HandleRerouteInputToSecondWindowVisibility)));
 			
-
 			// Run under one instance
 			if (GEditor && GEditor->bAllowMultiplePIEWorlds)
 			{
-
 				NetworkCategory.AddProperty("RunUnderOneProcess")
-					.DisplayName( TEXT("Use Single Process") )
-					.IsEnabled( TAttribute<bool>::Create(TAttribute<bool>::FGetter::CreateSP( this, &FLevelEditorPlaySettingsCustomization::HandleRunUnderOneProcessIsEnabled ) ) );
+					.DisplayName(TEXT("Use Single Process"))
+					.IsEnabled(TAttribute<bool>::Create(TAttribute<bool>::FGetter::CreateSP( this, &FLevelEditorPlaySettingsCustomization::HandleRunUnderOneProcessIsEnabled)));
 			}
 			else
 			{
@@ -502,6 +506,7 @@ private:
 		return (NumClients >= 2);
 	}
 
+	// Callback for getting the visibility of the ClientWindowSize property.
 	EVisibility HandleClientWindowSizePropertyVisibility() const
 	{
 		uint8 NetMode;
@@ -569,11 +574,13 @@ private:
 		return (NetMode != PIE_Standalone) || RunUnderOneProcess;
 	}
 
+	// Callback for getting the enabled state of the RunUnderOneProcess property.
 	bool HandleRunUnderOneProcessIsEnabled( ) const
 	{
 		return true;
 	}
 
+	// Callback for getting the enabled state of the RerouteInputToSecondWindow property.
 	bool HandleRerouteInputToSecondWindowEnabled( ) const
 	{
 		int32 NumClients;
@@ -582,6 +589,7 @@ private:
 		return NumClients > 1;	
 	}
 	
+	// Callback for getting the visibility of the RerouteInputToSecondWindow property.
 	EVisibility HandleRerouteInputToSecondWindowVisibility( ) const
 	{
 		bool RunUnderOneProcess;
@@ -589,6 +597,7 @@ private:
 		return (RunUnderOneProcess ? EVisibility::Visible : EVisibility::Hidden);
 	}
 
+	// Callback for getting the visibility of the PlayNetMode property.
 	EVisibility HandlePlayNetModeVisibility( ) const
 	{
 		bool RunUnderOneProcess;
@@ -596,6 +605,7 @@ private:
 		return (RunUnderOneProcess ? EVisibility::Hidden : EVisibility::Visible);
 	}
 
+	// Callback for getting the visibility of the CmdLine property.
 	EVisibility HandleCmdLineVisibility( ) const
 	{
 		bool RunUnderOneProcess;
