ProjectSettings: Cleaned up the build config visibility hack from CL# 2021006 (TTP# 328795)

diff --git a/Engine/Source/Editor/ProjectSettingsViewer/Private/ProjectSettingsViewerModule.cpp b/Engine/Source/Editor/ProjectSettingsViewer/Private/ProjectSettingsViewerModule.cpp
--- a/Engine/Source/Editor/ProjectSettingsViewer/Private/ProjectSettingsViewerModule.cpp
+++ b/Engine/Source/Editor/ProjectSettingsViewer/Private/ProjectSettingsViewerModule.cpp
@@ -185,9 +185,6 @@ protected:
 			LOCTEXT("MovieSettingsDescription", "Movie player settings"),
 			GetMutableDefault<UMoviePlayerSettings>()
 		);
-
-		// update the visible packaging enum values
-		UProjectPackagingSettings::UpdateBuildConfigurationVisibility();
 /*
 		// game session
 		SettingsModule.RegisterSettings("Project", "Game", "GameSession",
diff --git a/Engine/Source/Editor/UnrealEd/Classes/Settings/ProjectPackagingSettings.h b/Engine/Source/Editor/UnrealEd/Classes/Settings/ProjectPackagingSettings.h
--- a/Engine/Source/Editor/UnrealEd/Classes/Settings/ProjectPackagingSettings.h
+++ b/Engine/Source/Editor/UnrealEd/Classes/Settings/ProjectPackagingSettings.h
@@ -104,9 +104,4 @@ public:
 	 */
 	UPROPERTY(config, EditAnywhere, Category=Packaging)
 	bool UsePakFile;
-
-	/*
-	 * Updates visibility metadata for values of the EProjectPackagingBuildConfigurations, as appropriate for the current project.
-	 */
-	static void UpdateBuildConfigurationVisibility();
 };
diff --git a/Engine/Source/Editor/UnrealEd/Private/Settings/ProjectPackagingSettingsCustomization.h b/Engine/Source/Editor/UnrealEd/Private/Settings/ProjectPackagingSettingsCustomization.h
new file mode 100644
--- /dev/null
+++ b/Engine/Source/Editor/UnrealEd/Private/Settings/ProjectPackagingSettingsCustomization.h
@@ -0,0 +1,74 @@
+// Copyright 1998-2014 Epic Games, Inc. All Rights Reserved.
+
+/*=============================================================================
+	ProjectPackagingSettingsCustomization.h: Declares the FProjectPackagingSettingsCustomization class.
+=============================================================================*/
+
+#pragma once
+
+#include "PropertyEditing.h"
+#include "PropertyCustomizationHelpers.h"
+#include "PropertyRestriction.h"
+
+
+#define LOCTEXT_NAMESPACE "FProjectPackagingSettingsCustomization"
+
+
+/**
+ * Implements a details view customization for UProjectPackagingSettingsCustomization objects.
+ */
+class FProjectPackagingSettingsCustomization
+	: public IDetailCustomization
+{
+public:
+
+	// Begin IDetailCustomization interface
+
+	virtual void CustomizeDetails( IDetailLayoutBuilder& LayoutBuilder ) OVERRIDE
+	{
+		CustomizeProjectCategory(LayoutBuilder);
+	}
+
+	// End IDetailCustomization interface
+
+public:
+
+	/**
+	 * Creates a new instance.
+	 *
+	 * @return A new struct customization for play-in settings.
+	 */
+	static TSharedRef<IDetailCustomization> MakeInstance( )
+	{
+		return MakeShareable(new FProjectPackagingSettingsCustomization());
+	}
+
+protected:
+
+	/**
+	 * Customizes the Project property category.
+	 *
+	 * @param LayoutBuilder - The layout builder.
+	 */
+	void CustomizeProjectCategory( IDetailLayoutBuilder& LayoutBuilder )
+	{
+		// hide the DebugGame configuration for content-only games
+		TArray<FString> TargetFileNames;
+		IFileManager::Get().FindFiles(TargetFileNames, *(FPaths::GameSourceDir() / TEXT("*.target.cs")), true, false);
+
+		if (TargetFileNames.Num() == 0)
+		{
+			IDetailCategoryBuilder& ProjectCategory = LayoutBuilder.EditCategory("Project");
+			{
+				TSharedRef<FPropertyRestriction> BuildConfigurationRestruction = MakeShareable(new FPropertyRestriction(LOCTEXT("DebugGameRestrictionReason", "The DebugGame build configuration is not available in content-only projects.")));
+				BuildConfigurationRestruction->AddValue("DebugGame");
+
+				TSharedRef<IPropertyHandle> BuildConfigurationHandle = LayoutBuilder.GetProperty("BuildConfiguration");
+				BuildConfigurationHandle->AddRestriction(BuildConfigurationRestruction);
+			}
+		}
+	}
+};
+
+
+#undef LOCTEXT_NAMESPACE
diff --git a/Engine/Source/Editor/UnrealEd/Private/Settings/SettingsClasses.cpp b/Engine/Source/Editor/UnrealEd/Private/Settings/SettingsClasses.cpp
--- a/Engine/Source/Editor/UnrealEd/Private/Settings/SettingsClasses.cpp
+++ b/Engine/Source/Editor/UnrealEd/Private/Settings/SettingsClasses.cpp
@@ -303,24 +303,3 @@ void UProjectPackagingSettings::PostEditChangeProperty( FPropertyChangedEvent& P
 		StagingDirectory.Path = Path;
 	}
 }
-
-void UProjectPackagingSettings::UpdateBuildConfigurationVisibility()
-{
-	// Find if there are any .target.cs files in the project source directory.
-	TArray<FString> TargetFileNames;
-	IFileManager::Get().FindFiles(TargetFileNames, *(FPaths::GameSourceDir() / TEXT("*.target.cs")), true, false);
-
-	// Get the EProjectPackagingBuildConfigurations enum
-	extern UEnum *Z_Construct_UEnum_UProjectPackagingSettings_EProjectPackagingBuildConfigurations();
-	UEnum *ProjectPackagingConfigEnum = Z_Construct_UEnum_UProjectPackagingSettings_EProjectPackagingBuildConfigurations();
-
-	// Mark the DebugGame value as hidden if it's a content only project.
-	if (TargetFileNames.Num() == 0)
-	{
-		ProjectPackagingConfigEnum->SetMetaData(TEXT("Hidden"), TEXT("1"), PPBC_DebugGame);
-	}
-	else
-	{
-		ProjectPackagingConfigEnum->RemoveMetaData(TEXT("Hidden"), PPBC_DebugGame);
-	}
-}
\ No newline at end of file
diff --git a/Engine/Source/Editor/UnrealEd/Private/UnrealEdEngine.cpp b/Engine/Source/Editor/UnrealEd/Private/UnrealEdEngine.cpp
--- a/Engine/Source/Editor/UnrealEd/Private/UnrealEdEngine.cpp
+++ b/Engine/Source/Editor/UnrealEd/Private/UnrealEdEngine.cpp
@@ -13,6 +13,7 @@
 #include "EditorLoadingSavingSettingsCustomization.h"
 #include "GameMapsSettingsCustomization.h"
 #include "LevelEditorPlaySettingsCustomization.h"
+#include "ProjectPackagingSettingsCustomization.h"
 #include "ISourceControlModule.h"
 #include "Editor/StatsViewer/Public/StatsViewerModule.h"
 #include "SnappingUtils.h"
@@ -91,6 +92,7 @@ void UUnrealEdEngine::Init(IEngineLoop* InEngineLoop)
 		PropertyModule.RegisterCustomPropertyLayout("EditorLoadingSavingSettings", FOnGetDetailCustomizationInstance::CreateStatic(&FEditorLoadingSavingSettingsCustomization::MakeInstance));
 		PropertyModule.RegisterCustomPropertyLayout("GameMapsSettings", FOnGetDetailCustomizationInstance::CreateStatic(&FGameMapsSettingsCustomization::MakeInstance));
 		PropertyModule.RegisterCustomPropertyLayout("LevelEditorPlaySettings", FOnGetDetailCustomizationInstance::CreateStatic(&FLevelEditorPlaySettingsCustomization::MakeInstance));
+		PropertyModule.RegisterCustomPropertyLayout("ProjectPackagingSettings", FOnGetDetailCustomizationInstance::CreateStatic(&FProjectPackagingSettingsCustomization::MakeInstance));
 	}
 
 	bool bCookOnTheSide = FParse::Param( FCommandLine::Get(),TEXT("COOKONTHESIDE"));
