Added support for moving camera forward and backward with magnify trackpad gesture

diff --git a/Engine/Source/Editor/UnrealEd/Private/EditorViewportClient.cpp b/Engine/Source/Editor/UnrealEd/Private/EditorViewportClient.cpp
--- a/Engine/Source/Editor/UnrealEd/Private/EditorViewportClient.cpp
+++ b/Engine/Source/Editor/UnrealEd/Private/EditorViewportClient.cpp
@@ -210,6 +210,7 @@ FEditorViewportClient::FEditorViewportClient(FPreviewScene* InPreviewScene)
 	, bDraggingByHandle( false )
 	, CurrentGestureDragDelta(FVector::ZeroVector)
 	, CurrentGestureRotDelta(FRotator::ZeroRotator)
+	, GestureMoveForwardBackwardImpulse(0.0f)
 	, bDisableInput( false )
 	, bDrawAxes(true)
 	, bForceAudioRealtime(false)
@@ -1077,8 +1078,8 @@ void FEditorViewportClient::UpdateCameraMovement( float DeltaTime )
 			( GetDefault<ULevelEditorViewportSettings>()->FlightCameraControlType == WASD_RMBOnly && (Viewport->KeyState(EKeys::RightMouseButton ) ||Viewport->KeyState(EKeys::MiddleMouseButton) || Viewport->KeyState(EKeys::LeftMouseButton) || bIsUsingTrackpad ) ) ) ) &&
 			!MouseDeltaTracker->UsingDragTool();
 
-		//reset impulses if we're using WASD keys
-		CameraUserImpulseData->MoveForwardBackwardImpulse = 0.0f;
+		// Apply impulse from magnify gesture and reset impulses if we're using WASD keys
+		CameraUserImpulseData->MoveForwardBackwardImpulse = GestureMoveForwardBackwardImpulse;
 		CameraUserImpulseData->MoveRightLeftImpulse = 0.0f;
 		CameraUserImpulseData->MoveUpDownImpulse = 0.0f;
 		CameraUserImpulseData->ZoomOutInImpulse = 0.0f;
@@ -1086,6 +1087,7 @@ void FEditorViewportClient::UpdateCameraMovement( float DeltaTime )
 		CameraUserImpulseData->RotatePitchImpulse = 0.0f;
 		CameraUserImpulseData->RotateRollImpulse = 0.0f;
 
+		GestureMoveForwardBackwardImpulse = 0.0f;
 
 		// Forward/back
 		if( ( bRemapWASDKeys && Viewport->KeyState( FViewportNavigationCommands::Get().Forward->GetActiveGesture()->Key ) ) ||
@@ -2875,6 +2877,10 @@ bool FEditorViewportClient::InputGesture(FViewport* InViewport, EGestureEvent::T
 
 				FEditorViewportStats::Used(FEditorViewportStats::CAT_ORTHOGRAPHIC_GESTURE_SCROLL);
 			}
+			else if (GestureType == EGestureEvent::Magnify)
+			{
+				GestureMoveForwardBackwardImpulse = GestureDelta.X * 4.0f;
+			}
 		}
 		break;
 
diff --git a/Engine/Source/Editor/UnrealEd/Public/EditorViewportClient.h b/Engine/Source/Editor/UnrealEd/Public/EditorViewportClient.h
--- a/Engine/Source/Editor/UnrealEd/Public/EditorViewportClient.h
+++ b/Engine/Source/Editor/UnrealEd/Public/EditorViewportClient.h
@@ -1185,7 +1185,9 @@ protected:
 	/** Cumulative camera drag and rotation deltas from trackpad gesture in current Tick */
 	FVector CurrentGestureDragDelta;
 	FRotator CurrentGestureRotDelta;
-	
+
+	float GestureMoveForwardBackwardImpulse;
+
 	/** If true, force this viewport to use real time audio regardless of other settings */
 	bool bForceAudioRealtime;
 
