- Fixed bind pose to discard invalid ones and better tolerant when invalid
- If failed, it will default to using time 0 as bind pose

diff --git a/Engine/Source/Editor/UnrealEd/Private/Fbx/FbxOptionWindow.cpp b/Engine/Source/Editor/UnrealEd/Private/Fbx/FbxOptionWindow.cpp
--- a/Engine/Source/Editor/UnrealEd/Private/Fbx/FbxOptionWindow.cpp
+++ b/Engine/Source/Editor/UnrealEd/Private/Fbx/FbxOptionWindow.cpp
@@ -818,7 +818,7 @@ TSharedRef<SWidget> SFbxOptionWindow::ConstructSkeletalMeshAdvanced()
 		.OnCheckStateChanged(this, &SFbxOptionWindow::SetSkeletalMesh_UseT0AsRefPose)
 		[
 			SNew(STextBlock)
-			.Text(LOCTEXT("FbxOptionWindow_UseRefPose", "Use T0 as Ref Pose"))
+			.Text(LOCTEXT("FbxOptionWindow_UseRefPose", "Use Time 0 Pose for Ref Pose"))
 		]
 	];
 
diff --git a/Engine/Source/Editor/UnrealEd/Private/Fbx/FbxSkeletalMeshImport.cpp b/Engine/Source/Editor/UnrealEd/Private/Fbx/FbxSkeletalMeshImport.cpp
--- a/Engine/Source/Editor/UnrealEd/Private/Fbx/FbxSkeletalMeshImport.cpp
+++ b/Engine/Source/Editor/UnrealEd/Private/Fbx/FbxSkeletalMeshImport.cpp
@@ -648,7 +648,7 @@ void UnFbx::FFbxImporter::BuildSkeletonSystem(TArray<FbxCluster*>& ClusterArray,
 	}
 }
 
-bool UnFbx::FFbxImporter::ImportBone(TArray<FbxNode*>& NodeArray, FSkeletalMeshImportData &ImportData, TArray<FbxNode*> &SortedLinks, bool& bOutDiffPose, bool bDisableMissingBindPoseWarning)
+bool UnFbx::FFbxImporter::ImportBone(TArray<FbxNode*>& NodeArray, FSkeletalMeshImportData &ImportData, TArray<FbxNode*> &SortedLinks, bool& bOutDiffPose, bool bDisableMissingBindPoseWarning, bool & bUseTime0AsRefPose)
 {
 	bOutDiffPose = false;
 	int32 SkelType = 0; // 0 for skeletal mesh, 1 for rigid mesh
@@ -670,7 +670,7 @@ bool UnFbx::FFbxImporter::ImportBone(TArray<FbxNode*>& NodeArray, FSkeletalMeshI
 		// we have built the ClusterInformation.
 		int32 Default_NbPoses = SdkManager->GetBindPoseCount(Scene);
 		// If there are no BindPoses, the following will generate them.
-		SdkManager->CreateMissingBindPoses(Scene);
+		//SdkManager->CreateMissingBindPoses(Scene);
 
 		//if we created missing bind poses, update the number of bind poses
 		int32 NbPoses = SdkManager->GetBindPoseCount(Scene);
@@ -704,8 +704,6 @@ bool UnFbx::FFbxImporter::ImportBone(TArray<FbxNode*>& NodeArray, FSkeletalMeshI
 			return false;
 		}
 
-#define CHECK_VALID_BIND_POSE 0
-
 		// get bind pose
 		const int32 PoseCount = Scene->GetPoseCount();
 		for (int32 PoseIndex = 0; PoseIndex < PoseCount; PoseIndex++)
@@ -715,9 +713,6 @@ bool UnFbx::FFbxImporter::ImportBone(TArray<FbxNode*>& NodeArray, FSkeletalMeshI
 			// current pose is bind pose, 
 			if (CurrentPose && CurrentPose->IsBindPose())
 			{
-#if !CHECK_VALID_BIND_POSE
-				PoseArray.Add(CurrentPose);
-#else
 				// IsValidBindPose doesn't work reliably
 				// It checks all the parent chain(regardless root given), and if the parent doesn't have correct bind pose, it fails
 				// It causes more false positive issues than the real issue we have to worry about
@@ -731,42 +726,42 @@ bool UnFbx::FFbxImporter::ImportBone(TArray<FbxNode*>& NodeArray, FSkeletalMeshI
 
 				// it does not make any difference of checking with different node
 				// it is possible pose 0 -> node array 2, but isValidBindPose function returns true even with node array 0
-				FbxNode * Current = NodeArray[0];
-				if (CurrentPose->IsValidBindPoseVerbose(Current, &Notification, 0.0001, &Status))
-				{
-					PoseArray.Add(CurrentPose);
-					UE_LOG(LogFbx, Warning, TEXT("Valid bind pose for Pose (%s) - %s"), *PoseName, *FString(Current->GetName()));
-				}
-				else
+				FbxNode * ValidPoseNode = NULL;
+				for (auto Current : NodeArray)
 				{
-					FString ErrorString = Status.GetErrorString();
-
-					UE_LOG(LogFbx, Warning, TEXT("Not valid bind pose for Pose (%s) - Node %s : %s"), *PoseName, *FString(Current->GetName()), *ErrorString );
-					for(int32 NotificationIndex=0; NotificationIndex<Notification.GetNbEntries(); ++NotificationIndex)
+					FString CurrentName = Current->GetName();
+					if(CurrentPose->IsValidBindPoseVerbose(Current, &Notification, 0.0001, &Status))
 					{
-						const FbxAccumulatorEntry* Entry = Notification.GetEntryAt(NotificationIndex);
-						if(Entry)
-						{
-							FString Name = (const char*)(Entry->GetName());
-							FString Description = (const char*)(Entry->GetDescription());
-							FString Detail;
+						PoseArray.Add(CurrentPose);
+						UE_LOG(LogFbx, Warning, TEXT("Valid bind pose for Pose (%s) - %s"), *PoseName, *FString(Current->GetName()));
+						break;
+					}
+					else
+					{
+						FString ErrorString = Status.GetErrorString();
 
-							for(int32 DetailCount =0; DetailCount< Entry->GetDetailsCount(); ++DetailCount)
+						UE_LOG(LogFbx, Warning, TEXT("Not valid bind pose for Pose (%s) - Node %s : %s"), *PoseName, *FString(Current->GetName()), *ErrorString);
+						for(int32 NotificationIndex=0; NotificationIndex<Notification.GetNbEntries(); ++NotificationIndex)
+						{
+							const FbxAccumulatorEntry* Entry = Notification.GetEntryAt(NotificationIndex);
+							if(Entry)
 							{
-								Detail += FString((const char*)(*Entry->GetDetail(DetailCount)));
-							}
+								FString Name = (const char*)(Entry->GetName());
+								FString Description = (const char*)(Entry->GetDescription());
+								FString Detail;
 
-							UE_LOG(LogFbx, Warning, TEXT("\t(%d) %s : %s (%s)"), NotificationIndex+1, *Name, *Description, *Detail );
-						}
-					}
+								for(int32 DetailCount =0; DetailCount< Entry->GetDetailsCount(); ++DetailCount)
+								{
+									Detail += FString((const char*)(*Entry->GetDetail(DetailCount)));
+								}
 
-					Notification.ClearAccumulator();
+								UE_LOG(LogFbx, Warning, TEXT("\t(%d) %s : %s (%s)"), NotificationIndex+1, *Name, *Description, *Detail);
+							}
+						}
 
-					// add to tokenized error message
-					AddTokenizedErrorMessage(FTokenizedMessage::Create(EMessageSeverity::Warning, FText::Format(LOCTEXT("FbxSkeletaLMeshimport_InvalidBindPose", "FBX bind pose '{0}' is invalid.  It will be ignored"), FText::FromString(Scene->GetPose(PoseIndex)->GetName()))));
+						Notification.ClearAccumulator();
+					}
 				}
-#endif // CHECK_VALID_BIND_POSE
-
 			}
 		}
 
@@ -774,12 +769,22 @@ bool UnFbx::FFbxImporter::ImportBone(TArray<FbxNode*>& NodeArray, FSkeletalMeshI
 		BuildSkeletonSystem(ClusterArray, SortedLinks);
 	}
 
-	if( SortedLinks.Num() == 0 )
+	// error check
+	// if no bond is found
+	if(SortedLinks.Num() == 0)
 	{
 		AddTokenizedErrorMessage(FTokenizedMessage::Create(EMessageSeverity::Warning, FText::Format(LOCTEXT("FbxSkeletaLMeshimport_NoBone", "'{0}' has no bones"), FText::FromString(NodeArray[0]->GetName()))));
 		return false;
 	}
 
+	// if no bind pose is found
+	if (!bUseTime0AsRefPose && PoseArray.GetCount() == 0)
+	{
+		// add to tokenized error message
+		AddTokenizedErrorMessage(FTokenizedMessage::Create(EMessageSeverity::Warning, LOCTEXT("FbxSkeletaLMeshimport_InvalidBindPose", "Could not find the bind pose.  It will use time 0 as bind pose.")));
+		bUseTime0AsRefPose = true;
+	}
+
 	int32 LinkIndex;
 
 	// Check for duplicate bone names and issue a warning if found
@@ -869,7 +874,9 @@ bool UnFbx::FFbxImporter::ImportBone(TArray<FbxNode*>& NodeArray, FSkeletalMeshI
 
 			if (!GlobalLinkFoundFlag)
 			{
-				if(!ImportOptions->bUseT0AsRefPose && !bDisableMissingBindPoseWarning)
+				// since now we set use time 0 as ref pose this won't unlikely happen
+				// but leaving it just in case it still has case where it's missing partial bind pose
+				if(!bUseTime0AsRefPose && !bDisableMissingBindPoseWarning)
 				{
 					bAnyLinksNotInBindPose = true;
 					LinksWithoutBindPoses +=  ANSI_TO_TCHAR(Link->GetName());
@@ -897,7 +904,7 @@ bool UnFbx::FFbxImporter::ImportBone(TArray<FbxNode*>& NodeArray, FSkeletalMeshI
 			GlobalsPerLink[LinkIndex] = Link->EvaluateGlobalTransform();
 		}
 		
-		if (ImportOptions->bUseT0AsRefPose)
+		if (bUseTime0AsRefPose)
 		{
 			FbxAMatrix& T0Matrix = Scene->GetAnimationEvaluator()->GetNodeGlobalTransform(Link, 0);
 			if (GlobalsPerLink[LinkIndex] != T0Matrix)
@@ -1039,9 +1046,10 @@ USkeletalMesh* UnFbx::FFbxImporter::ImportSkeletalMesh(UObject* InParent, TArray
 	TArray<FbxNode*> SortedLinkArray;
 	FbxArray<FbxAMatrix> GlobalsPerLink;
 
+	bool bUseTime0AsRefPose = ImportOptions->bUseT0AsRefPose;
 	// Note: importing morph data causes additional passes through this function, so disable the warning dialogs
 	// from popping up again on each additional pass.  
-	if ( !ImportBone(NodeArray, *SkelMeshImportDataPtr, SortedLinkArray,  bDiffPose, (FbxShapeArray != NULL)))
+	if ( !ImportBone(NodeArray, *SkelMeshImportDataPtr, SortedLinkArray,  bDiffPose, (FbxShapeArray != NULL), bUseTime0AsRefPose))
 	{
 		AddTokenizedErrorMessage(FTokenizedMessage::Create(EMessageSeverity::Error, LOCTEXT("FbxSkeletaLMeshimport_MultipleRootFound", "Multiple roots found")));
 		// I can't delete object here since this is middle of import
@@ -1093,7 +1101,7 @@ USkeletalMesh* UnFbx::FFbxImporter::ImportSkeletalMesh(UObject* InParent, TArray
 			return NULL;
 		}
 		
-		if (ImportOptions->bUseT0AsRefPose && bDiffPose)
+		if (bUseTime0AsRefPose && bDiffPose)
 		{
 			// deform skin vertex to the frame 0 from bind pose
 			SkinControlPointsToPose(*SkelMeshImportDataPtr, FbxMesh, FbxShape, true);
