Do not try to allocate zero storage, dereference null vertices array or build normals without there being vertices.
#ttp 334052 LIVE: EDITOR: CRASH: Importing Skeletal Mesh FBX as rigid mesh causes engine crash
reviewer Andrew.Brown

diff --git a/Engine/Source/Editor/UnrealEd/Private/Fbx/FbxSkeletalMeshImport.cpp b/Engine/Source/Editor/UnrealEd/Private/Fbx/FbxSkeletalMeshImport.cpp
--- a/Engine/Source/Editor/UnrealEd/Private/Fbx/FbxSkeletalMeshImport.cpp
+++ b/Engine/Source/Editor/UnrealEd/Private/Fbx/FbxSkeletalMeshImport.cpp
@@ -1680,28 +1680,31 @@ bool UnFbx::FFbxImporter::FillSkelMeshImporterFromFbx( FSkeletalMeshImportData&
 		}
 	}
 
-	if( ImportOptions->bPreserveSmoothingGroups )
+	if( LayerCount > 0 && ImportOptions->bPreserveSmoothingGroups )
 	{
 		// Check and see if the smooothing data is valid.  If not generate it from the normals
 		BaseLayer = Mesh->GetLayer(0);
-		const FbxLayerElementSmoothing* SmoothingLayer = BaseLayer->GetSmoothing();
-
-		if( SmoothingLayer )
+		if( BaseLayer )
 		{
-			bool bValidSmoothingData = false;
-			FbxLayerElementArrayTemplate<int32>& Array = SmoothingLayer->GetDirectArray();
-			for( int32 SmoothingIndex = 0; SmoothingIndex < Array.GetCount(); ++SmoothingIndex )
+			const FbxLayerElementSmoothing* SmoothingLayer = BaseLayer->GetSmoothing();
+
+			if( SmoothingLayer )
 			{
-				if( Array[SmoothingIndex] != 0 )
+				bool bValidSmoothingData = false;
+				FbxLayerElementArrayTemplate<int32>& Array = SmoothingLayer->GetDirectArray();
+				for( int32 SmoothingIndex = 0; SmoothingIndex < Array.GetCount(); ++SmoothingIndex )
 				{
-					bValidSmoothingData = true;
-					break;
+					if( Array[SmoothingIndex] != 0 )
+					{
+						bValidSmoothingData = true;
+						break;
+					}
 				}
-			}
 
-			if( !bValidSmoothingData )
-			{
-				GeometryConverter->ComputeEdgeSmoothingFromNormals(Mesh);
+				if( !bValidSmoothingData && Mesh->GetPolygonVertexCount() > 0 )
+				{
+					GeometryConverter->ComputeEdgeSmoothingFromNormals(Mesh);
+				}
 			}
 		}
 	}
diff --git a/Engine/Source/Runtime/Engine/Private/SkeletalMesh.cpp b/Engine/Source/Runtime/Engine/Private/SkeletalMesh.cpp
--- a/Engine/Source/Runtime/Engine/Private/SkeletalMesh.cpp
+++ b/Engine/Source/Runtime/Engine/Private/SkeletalMesh.cpp
@@ -255,9 +255,12 @@ void FSkeletalMeshVertexBuffer::Init(const TArray<FSoftSkinVertex>& InVertices)
 	
 	VertexData->ResizeBuffer(InVertices.Num());
 	
-	Data = VertexData->GetDataPointer();
-	Stride = VertexData->GetStride();
-	NumVertices = VertexData->GetNumVertices();
+	if (InVertices.Num() > 0)
+	{
+		Data = VertexData->GetDataPointer();
+		Stride = VertexData->GetStride();
+		NumVertices = VertexData->GetNumVertices();
+	}
 	
 	if (bExtraBoneInfluences)
 	{
