Missing file

diff --git a/Engine/Source/Editor/UnrealEd/Private/FbxImporter.h b/Engine/Source/Editor/UnrealEd/Private/FbxImporter.h
--- a/Engine/Source/Editor/UnrealEd/Private/FbxImporter.h
+++ b/Engine/Source/Editor/UnrealEd/Private/FbxImporter.h
@@ -180,7 +180,7 @@ struct FbxSceneInfo
 	TArray<FbxMeshInfo> MeshInfo;
 	
 	// only one take supported currently
-	char* TakeName;
+	FString TakeName;
 	double FrameRate;
 	double TotalTime;
 
@@ -710,8 +710,11 @@ protected:
 	 * @param NodeArray Fbx Nodes to import, they are bound to the same skeleton system
 	 * @param ImportData object to store skeletal mesh data
 	 * @param OutSortedLinks return all skeletons sorted by depth traversal
+	 * @param bOutDiffPose
+	 * @param bDisableMissingBindPoseWarning
+	 * @param bUseTime0AsRefPose	in/out - Use Time 0 as Ref Pose 
 	 */
-	bool ImportBone(TArray<FbxNode*>& NodeArray, FSkeletalMeshImportData &ImportData, TArray<FbxNode*> &OutSortedLinks, bool& bOutDiffPose, bool bDisableMissingBindPoseWarning);
+	bool ImportBone(TArray<FbxNode*>& NodeArray, FSkeletalMeshImportData &ImportData, TArray<FbxNode*> &OutSortedLinks, bool& bOutDiffPose, bool bDisableMissingBindPoseWarning, bool & bUseTime0AsRefPose);
 	
 	/**
 	 * Skins the control points of the given mesh or shape using either the default pose for skinning or the first frame of the
