Avoid dereferencing nullptr when AssetData has no assoicated asset or class, switch to deleting the uasset package directly instead.
#ttp 334907 CRASH: Trying to delete an asset of a type that no longer exists crashes
reviewer Thomas.Sarkanen

diff --git a/Engine/Source/Editor/UnrealEd/Private/ObjectTools.cpp b/Engine/Source/Editor/UnrealEd/Private/ObjectTools.cpp
--- a/Engine/Source/Editor/UnrealEd/Private/ObjectTools.cpp
+++ b/Engine/Source/Editor/UnrealEd/Private/ObjectTools.cpp
@@ -1456,13 +1456,55 @@ namespace ObjectTools
 
 	int32 DeleteAssets( const TArray<FAssetData>& AssetsToDelete, bool bShowConfirmation )
 	{
+		TArray<UPackage*> PackageFilesToDelete;
 		TArray<UObject*> ObjectsToDelete;
 		for ( int i = 0; i < AssetsToDelete.Num(); i++ )
 		{
-			ObjectsToDelete.Add( AssetsToDelete[i].GetAsset() );
+			const FAssetData& AssetData = AssetsToDelete[i];
+			UObject *ObjectToDelete = AssetData.GetAsset();
+			// Assets can be loaded even when their underlying type/class no longer exists...
+			if ( ObjectToDelete!=nullptr )
+			{
+				ObjectsToDelete.Add( ObjectToDelete );
+			}
+			else if ( AssetData.IsUAsset() )
+			{
+				// ... In this cases there is no underlying asset or type so remove the package itself directly after confirming it's valid to do so.
+				FString PackageFilename;
+				if( !FPackageName::DoesPackageExist( AssetData.PackageName.ToString(), NULL, &PackageFilename ) )
+				{
+					// Could not determine filename for package so we can not delete
+					continue;
+				}
+
+				if ( FPaths::GetExtension(PackageFilename, /*bIncludeDot=*/true).ToLower() != FPackageName::GetAssetPackageExtension() )
+				{
+					// Only delete UAsset packages because that is what we checked for in ShowDeleteConfirmationDialog()
+					continue;
+				}
+
+				UPackage* Package = FindPackage(nullptr, *AssetData.PackageName.ToString());
+				if ( Package )
+				{
+					PackageFilesToDelete.Add(Package);
+				}
+			}
+		}
+
+		const int32 NumPackagesToDelete = PackageFilesToDelete.Num();
+		if ( NumPackagesToDelete > 0 )
+		{
+			const bool bPerformReferenceCheck = true;
+			CleanupAfterSuccessfulDelete(PackageFilesToDelete, true);
+		}
+
+		int32 NumObjectsToDelete = ObjectsToDelete.Num();
+		if ( NumObjectsToDelete > 0 )
+		{
+			NumObjectsToDelete = DeleteObjects( ObjectsToDelete, bShowConfirmation );
 		}
 
-		return DeleteObjects( ObjectsToDelete, bShowConfirmation );
+		return NumPackagesToDelete + NumObjectsToDelete;
 	}
 
 	int32 DeleteObjects( const TArray< UObject* >& ObjectsToDelete, bool bShowConfirmation )
