fixing CIS (bad string conversion)

diff --git a/Engine/Source/Editor/UnrealEd/Private/Fbx/FbxMainImport.cpp b/Engine/Source/Editor/UnrealEd/Private/Fbx/FbxMainImport.cpp
--- a/Engine/Source/Editor/UnrealEd/Private/Fbx/FbxMainImport.cpp
+++ b/Engine/Source/Editor/UnrealEd/Private/Fbx/FbxMainImport.cpp
@@ -514,7 +514,7 @@ bool FFbxImporter::GetSceneInfo(FString Filename, FbxSceneInfo& SceneInfo)
 			const char* AnimStackName = CurAnimStack->GetName();
 			SceneInfo.TakeName = new char[FCStringAnsi::Strlen(AnimStackName) + 1];
 
-			FCStringAnsi::Strcpy(SceneInfo.TakeName, FCStringAnsi::Strlen(AnimStackName) + 1, AnimStackName);
+			FCStringAnsi::Strcpy(TCHAR_TO_ANSI(*SceneInfo.TakeName), FCStringAnsi::Strlen(AnimStackName) + 1, AnimStackName);
 		}
 		SceneInfo.FrameRate = FbxTime::GetFrameRate(Scene->GetGlobalSettings().GetTimeMode());
 		
