Delete accidentally left-in no-op line.

diff --git a/Engine/Source/Editor/UnrealEd/Private/EditorActor.cpp b/Engine/Source/Editor/UnrealEd/Private/EditorActor.cpp
--- a/Engine/Source/Editor/UnrealEd/Private/EditorActor.cpp
+++ b/Engine/Source/Editor/UnrealEd/Private/EditorActor.cpp
@@ -575,7 +575,6 @@ bool UUnrealEdEngine::edactDeleteSelected( UWorld* InWorld, bool bVerifyDeletion
 				break;
 			}
 		}
-		(ReferencingActors.Num() > 0);
 
 		if ( bReferencedByLevelScript || bReferencedByActor )
 		{
