Backing out the unintentional change

diff --git a/Engine/Source/Editor/UnrealEd/Private/Fbx/FbxMainImport.cpp b/Engine/Source/Editor/UnrealEd/Private/Fbx/FbxMainImport.cpp
--- a/Engine/Source/Editor/UnrealEd/Private/Fbx/FbxMainImport.cpp
+++ b/Engine/Source/Editor/UnrealEd/Private/Fbx/FbxMainImport.cpp
@@ -514,7 +514,7 @@ bool FFbxImporter::GetSceneInfo(FString Filename, FbxSceneInfo& SceneInfo)
 			const char* AnimStackName = CurAnimStack->GetName();
 			SceneInfo.TakeName = new char[FCStringAnsi::Strlen(AnimStackName) + 1];
 
-			FCStringAnsi::Strcpy(TCHAR_TO_ANSI(*SceneInfo.TakeName), FCStringAnsi::Strlen(AnimStackName) + 1, AnimStackName);
+			FCStringAnsi::Strcpy(SceneInfo.TakeName, FCStringAnsi::Strlen(AnimStackName) + 1, AnimStackName);
 		}
 		SceneInfo.FrameRate = FbxTime::GetFrameRate(Scene->GetGlobalSettings().GetTimeMode());
 		
diff --git a/Engine/Source/Editor/UnrealEd/Private/FbxImporter.h b/Engine/Source/Editor/UnrealEd/Private/FbxImporter.h
--- a/Engine/Source/Editor/UnrealEd/Private/FbxImporter.h
+++ b/Engine/Source/Editor/UnrealEd/Private/FbxImporter.h
@@ -180,7 +180,7 @@ struct FbxSceneInfo
 	TArray<FbxMeshInfo> MeshInfo;
 	
 	// only one take supported currently
-	FString TakeName;
+	char* TakeName;
 	double FrameRate;
 	double TotalTime;
 
