fix for not cooking the joystick assets if a different asset is specified
TTP335076
#ue4

diff --git a/Engine/Source/Editor/UnrealEd/Private/Commandlets/CookCommandlet.cpp b/Engine/Source/Editor/UnrealEd/Private/Commandlets/CookCommandlet.cpp
--- a/Engine/Source/Editor/UnrealEd/Private/Commandlets/CookCommandlet.cpp
+++ b/Engine/Source/Editor/UnrealEd/Private/Commandlets/CookCommandlet.cpp
@@ -911,6 +911,16 @@ void UCookCommandlet::CollectFilesToCook(TArray<FString>& FilesInPath)
 		}
 	}
 
+	// make sure we cook any extra assets for the default touch interface
+	// @todo need a better approach to cooking assets which are dynamically loaded by engine code based on settings
+	FConfigFile InputIni;
+	FString InterfaceFile;
+	FConfigCacheIni::LoadLocalIniFile(InputIni, TEXT("Input"), true);
+	if (InputIni.GetString(TEXT("/Script/Engine.InputSettings"), TEXT("DefaultTouchInterface"), InterfaceFile))
+	{
+		FilesInPath.AddUnique(InterfaceFile);
+	}
+
 	//@todo SLATE: This is a hack to ensure all slate referenced assets get cooked.
 	// Slate needs to be refactored to properly identify required assets at cook time.
 	// Simply jamming everything in a given directory into the cook list is error-prone
