fix for when there isn't any asset specified
#ue4

diff --git a/Engine/Source/Editor/UnrealEd/Private/Commandlets/CookCommandlet.cpp b/Engine/Source/Editor/UnrealEd/Private/Commandlets/CookCommandlet.cpp
--- a/Engine/Source/Editor/UnrealEd/Private/Commandlets/CookCommandlet.cpp
+++ b/Engine/Source/Editor/UnrealEd/Private/Commandlets/CookCommandlet.cpp
@@ -918,7 +918,10 @@ void UCookCommandlet::CollectFilesToCook(TArray<FString>& FilesInPath)
 	FConfigCacheIni::LoadLocalIniFile(InputIni, TEXT("Input"), true);
 	if (InputIni.GetString(TEXT("/Script/Engine.InputSettings"), TEXT("DefaultTouchInterface"), InterfaceFile))
 	{
-		FilesInPath.AddUnique(InterfaceFile);
+		if (InterfaceFile != TEXT("None"))
+		{
+			FilesInPath.AddUnique(InterfaceFile);
+		}
 	}
 
 	//@todo SLATE: This is a hack to ensure all slate referenced assets get cooked.
