- Fixed the problems on clothing sections display mode in Persona show menu
- Resolved compile errors when turned off #WITH_APEX_CLOTHING definition

diff --git a/Engine/Source/Editor/Persona/Private/Persona.cpp b/Engine/Source/Editor/Persona/Private/Persona.cpp
--- a/Engine/Source/Editor/Persona/Private/Persona.cpp
+++ b/Engine/Source/Editor/Persona/Private/Persona.cpp
@@ -460,9 +460,6 @@ FPersona::~FPersona()
 	if(PreviewComponent)
 	{
 		PreviewComponent->RemoveFromRoot();
-#if WITH_APEX_CLOTHING
-		PreviewComponent->RestoreClothSectionsVisibility();
-#endif //#if WITH_APEX_CLOTHING
 	}
 	
 	// NOTE: Any tabs that we still have hanging out when destroyed will be cleaned up by FBaseToolkit's destructor
diff --git a/Engine/Source/Editor/Persona/Private/SAnimViewportToolBar.cpp b/Engine/Source/Editor/Persona/Private/SAnimViewportToolBar.cpp
--- a/Engine/Source/Editor/Persona/Private/SAnimViewportToolBar.cpp
+++ b/Engine/Source/Editor/Persona/Private/SAnimViewportToolBar.cpp
@@ -585,6 +585,7 @@ void SAnimViewportToolBar::FillShowAdvancedMenu(FMenuBuilder& MenuBuilder) const
 
 void SAnimViewportToolBar::FillShowClothingMenu(FMenuBuilder& MenuBuilder) const
 {
+#if WITH_APEX_CLOTHING
 	const FAnimViewportShowCommands& Actions = FAnimViewportShowCommands::Get();
 
 	MenuBuilder.BeginSection("ClothPreview", LOCTEXT("ClothPreview_Label", "Preview"));
@@ -622,6 +623,7 @@ void SAnimViewportToolBar::FillShowClothingMenu(FMenuBuilder& MenuBuilder) const
 		MenuBuilder.AddMenuEntry(Actions.ShowOnlyClothSections);
 		MenuBuilder.AddMenuEntry(Actions.HideOnlyClothSections);
 	}
+#endif // #if WITH_APEX_CLOTHING
 }
 
 FText SAnimViewportToolBar::GetLODMenuLabel() const
diff --git a/Engine/Source/Editor/Persona/Private/SAnimationEditorViewport.cpp b/Engine/Source/Editor/Persona/Private/SAnimationEditorViewport.cpp
--- a/Engine/Source/Editor/Persona/Private/SAnimationEditorViewport.cpp
+++ b/Engine/Source/Editor/Persona/Private/SAnimationEditorViewport.cpp
@@ -583,25 +583,23 @@ void SAnimationEditorViewportTabBody::BindCommands()
 		FCanExecuteAction(),
 		FIsActionChecked::CreateSP( this, &SAnimationEditorViewportTabBody::IsShowingClothFixedVertices ) );
 
-	SectionsDisplayMode = ESectionDisplayMode::ShowAll;
-
 	CommandList.MapAction(
 		ViewportShowMenuCommands.ShowAllSections,
-		FExecuteAction::CreateSP(this, &SAnimationEditorViewportTabBody::OnSetSectionsDisplayMode, (int32)ESectionDisplayMode::ShowAll),
+		FExecuteAction::CreateSP(this, &SAnimationEditorViewportTabBody::OnSetSectionsDisplayMode, (int32)UDebugSkelMeshComponent::ESectionDisplayMode::ShowAll),
 		FCanExecuteAction(),
-		FIsActionChecked::CreateSP(this, &SAnimationEditorViewportTabBody::IsSectionsDisplayMode, (int32)ESectionDisplayMode::ShowAll));
+		FIsActionChecked::CreateSP(this, &SAnimationEditorViewportTabBody::IsSectionsDisplayMode, (int32)UDebugSkelMeshComponent::ESectionDisplayMode::ShowAll));
 
 	CommandList.MapAction(
 		ViewportShowMenuCommands.ShowOnlyClothSections,
-		FExecuteAction::CreateSP(this, &SAnimationEditorViewportTabBody::OnSetSectionsDisplayMode, (int32)ESectionDisplayMode::ShowOnlyClothSections),
+		FExecuteAction::CreateSP(this, &SAnimationEditorViewportTabBody::OnSetSectionsDisplayMode, (int32)UDebugSkelMeshComponent::ESectionDisplayMode::ShowOnlyClothSections),
 		FCanExecuteAction(),
-		FIsActionChecked::CreateSP(this, &SAnimationEditorViewportTabBody::IsSectionsDisplayMode, (int32)ESectionDisplayMode::ShowOnlyClothSections));
+		FIsActionChecked::CreateSP(this, &SAnimationEditorViewportTabBody::IsSectionsDisplayMode, (int32)UDebugSkelMeshComponent::ESectionDisplayMode::ShowOnlyClothSections));
 
 	CommandList.MapAction(
 		ViewportShowMenuCommands.HideOnlyClothSections,
-		FExecuteAction::CreateSP(this, &SAnimationEditorViewportTabBody::OnSetSectionsDisplayMode, (int32)ESectionDisplayMode::HideOnlyClothSections),
+		FExecuteAction::CreateSP(this, &SAnimationEditorViewportTabBody::OnSetSectionsDisplayMode, (int32)UDebugSkelMeshComponent::ESectionDisplayMode::HideOnlyClothSections),
 		FCanExecuteAction(),
-		FIsActionChecked::CreateSP(this, &SAnimationEditorViewportTabBody::IsSectionsDisplayMode, (int32)ESectionDisplayMode::HideOnlyClothSections));
+		FIsActionChecked::CreateSP(this, &SAnimationEditorViewportTabBody::IsSectionsDisplayMode, (int32)UDebugSkelMeshComponent::ESectionDisplayMode::HideOnlyClothSections));
 
 #endif// #if WITH_APEX_CLOTHING		
 
@@ -1676,26 +1674,26 @@ void SAnimationEditorViewportTabBody::OnSetSectionsDisplayMode(int32 DisplayMode
 {
 	UDebugSkelMeshComponent* PreviewComponent = PersonaPtr.Pin()->PreviewComponent;
 
-	SectionsDisplayMode = DisplayMode;
-
 	if (!PreviewComponent)
 	{
 		return;
 	}
 
+	PreviewComponent->SectionsDisplayMode = DisplayMode;
+
 	switch (DisplayMode)
 	{
-	case ESectionDisplayMode::ShowAll:
+	case UDebugSkelMeshComponent::ESectionDisplayMode::ShowAll:
 		// restore to the original states
 		PreviewComponent->RestoreClothSectionsVisibility();
 		break;
-	case ESectionDisplayMode::ShowOnlyClothSections:
-		// disable all except clothing sections
-		PreviewComponent->ShowOnlyClothSections(true, PreviewComponent->PredictedLODLevel);
+	case UDebugSkelMeshComponent::ESectionDisplayMode::ShowOnlyClothSections:
+		// disable all except clothing sections and shows only cloth sections
+		PreviewComponent->ToggleClothSectionsVisibility(true);
 		break;
-	case ESectionDisplayMode::HideOnlyClothSections:
+	case UDebugSkelMeshComponent::ESectionDisplayMode::HideOnlyClothSections:
 		// disables only clothing sections
-		PreviewComponent->ShowOnlyClothSections(false, PreviewComponent->PredictedLODLevel);
+		PreviewComponent->ToggleClothSectionsVisibility(false);
 		break;
 	}
 
@@ -1704,8 +1702,16 @@ void SAnimationEditorViewportTabBody::OnSetSectionsDisplayMode(int32 DisplayMode
 
 bool SAnimationEditorViewportTabBody::IsSectionsDisplayMode(int32 DisplayMode) const
 {
-	return (SectionsDisplayMode == DisplayMode);
+	UDebugSkelMeshComponent* PreviewComponent = PersonaPtr.Pin()->PreviewComponent;
+
+	if (ensure(PreviewComponent))
+	{
+		return (PreviewComponent->SectionsDisplayMode == DisplayMode);
+	}
+
+	return false;
 }
+#endif // #if WITH_APEX_CLOTHING
 
 EVisibility SAnimationEditorViewportTabBody::GetViewportCornerTextVisibility() const
 {
@@ -1762,5 +1768,4 @@ FReply SAnimationEditorViewportTabBody::ClickedOnViewportCornerText()
 	return FReply::Handled();
 }
 
-#endif // #if WITH_APEX_CLOTHING
 #undef LOCTEXT_NAMESPACE
diff --git a/Engine/Source/Editor/Persona/Private/SAnimationEditorViewport.h b/Engine/Source/Editor/Persona/Private/SAnimationEditorViewport.h
--- a/Engine/Source/Editor/Persona/Private/SAnimationEditorViewport.h
+++ b/Engine/Source/Editor/Persona/Private/SAnimationEditorViewport.h
@@ -328,14 +328,6 @@ private:
 	bool IsAudioMuted();
 
 #if WITH_APEX_CLOTHING
-
-	enum ESectionDisplayMode
-	{
-		ShowAll,
-		ShowOnlyClothSections,
-		HideOnlyClothSections,
-		NumSectionDisplayMode
-	};
 	/** 
 	 * clothing show options 
 	*/
@@ -381,9 +373,6 @@ private:
 	void OnSetSectionsDisplayMode(int32 DisplayMode);
 	bool IsSectionsDisplayMode(int32 DisplayMode) const;
 
-	/** Draw All/ Draw only clothing sections/ Hide only clothing sections */
-	uint32 SectionsDisplayMode;
-
 #endif // #if WITH_APEX_CLOTHING
 
 private:
diff --git a/Engine/Source/Editor/UnrealEd/Classes/Animation/DebugSkelMeshComponent.h b/Engine/Source/Editor/UnrealEd/Classes/Animation/DebugSkelMeshComponent.h
--- a/Engine/Source/Editor/UnrealEd/Classes/Animation/DebugSkelMeshComponent.h
+++ b/Engine/Source/Editor/UnrealEd/Classes/Animation/DebugSkelMeshComponent.h
@@ -212,10 +212,29 @@ class UDebugSkelMeshComponent : public USkeletalMeshComponent
 #endif
 
 #if WITH_APEX_CLOTHING
-	/** toggle visibility between cloth sections and non-cloth sections */
-	UNREALED_API void ShowOnlyClothSections(bool bShow, int32 LODIndex);
+
+	enum ESectionDisplayMode
+	{
+		None = -1,
+		ShowAll,
+		ShowOnlyClothSections,
+		HideOnlyClothSections,
+		NumSectionDisplayMode
+	};
+	/** Draw All/ Draw only clothing sections/ Hide only clothing sections */
+	int32 SectionsDisplayMode;
+
+	/** 
+	 * toggle visibility between cloth sections and non-cloth sections for all LODs
+	 * if bShowOnlyClothSections is true, shows only cloth sections. On the other hand, 
+	 * if bShowOnlyClothSections is false, hides only cloth sections.
+	 */
+	UNREALED_API void ToggleClothSectionsVisibility(bool bShowOnlyClothSections);
 	/** Restore all section visibilities to original states for all LODs */
 	UNREALED_API void RestoreClothSectionsVisibility();
+
+	int32 FindCurrentSectionDisplayMode();
+
 #endif //#if WITH_APEX_CLOTHING
 };
 
diff --git a/Engine/Source/Editor/UnrealEd/Private/Animation/DebugSkelMeshComponent.cpp b/Engine/Source/Editor/UnrealEd/Private/Animation/DebugSkelMeshComponent.cpp
--- a/Engine/Source/Editor/UnrealEd/Private/Animation/DebugSkelMeshComponent.cpp
+++ b/Engine/Source/Editor/UnrealEd/Private/Animation/DebugSkelMeshComponent.cpp
@@ -88,6 +88,10 @@ UDebugSkelMeshComponent::UDebugSkelMeshComponent(const class FPostConstructIniti
 
 	bMeshSocketsVisible = true;
 	bSkeletonSocketsVisible = true;
+
+#if WITH_APEX_CLOTHING
+	SectionsDisplayMode = ESectionDisplayMode::None;
+#endif //#if WITH_APEX_CLOTHING
 }
 
 FBoxSphereBounds UDebugSkelMeshComponent::CalcBounds(const FTransform & LocalToWorld) const
@@ -158,6 +162,14 @@ FPrimitiveSceneProxy* UDebugSkelMeshComponent::CreateSceneProxy()
 		Result = ::new FDebugSkelMeshSceneProxy(this, SkelMeshResource, WireframeMeshOverlayColor);
 	}
 
+#if WITH_APEX_CLOTHING
+	if (SectionsDisplayMode == ESectionDisplayMode::None)
+	{
+		SectionsDisplayMode = FindCurrentSectionDisplayMode();
+	}
+
+#endif //#if WITH_APEX_CLOTHING
+
 	return Result;
 }
 
@@ -434,53 +446,62 @@ void UDebugSkelMeshComponent::ClearAnimNotifyErrors(UObject* InSourceNotify)
 
 #if WITH_APEX_CLOTHING
 
-void UDebugSkelMeshComponent::ShowOnlyClothSections(bool bShow, int32 LODIndex)
+void UDebugSkelMeshComponent::ToggleClothSectionsVisibility(bool bShowOnlyClothSections)
 {
-	PreEditChange(NULL);
-
 	FSkeletalMeshResource* SkelMeshResource = GetSkeletalMeshResource();
-	check(SkelMeshResource);
-	check(LODIndex >= 0 && LODIndex < SkelMeshResource->LODModels.Num());
-	FStaticLODModel& LODModel = SkelMeshResource->LODModels[LODIndex];
-
-	for (int32 SecIdx = 0; SecIdx < LODModel.Sections.Num(); SecIdx++)
+	if (SkelMeshResource)
 	{
-		FSkelMeshSection& Section = LODModel.Sections[SecIdx];
+		PreEditChange(NULL);
 
-		// toggle visibility between cloth sections and non-cloth sections
-		if (bShow)
+		for (int32 LODIndex = 0; LODIndex < SkelMeshResource->LODModels.Num(); LODIndex++)
 		{
-			// enables only cloth sections
-			if (LODModel.Chunks[Section.ChunkIndex].HasApexClothData())
-			{
-				Section.bDisabled = false;
-			}
-			else
-			{
-				Section.bDisabled = true;
-			}
-		}
-		else
-		{   // disables cloth sections and also corresponding original sections
-			if (LODModel.Chunks[Section.ChunkIndex].HasApexClothData())
-			{
-				Section.bDisabled = true;
-				LODModel.Sections[Section.CorrespondClothSectionIndex].bDisabled = true;
-			}
-			else
+			FStaticLODModel& LODModel = SkelMeshResource->LODModels[LODIndex];
+
+			for (int32 SecIdx = 0; SecIdx < LODModel.Sections.Num(); SecIdx++)
 			{
-				Section.bDisabled = false;
+				FSkelMeshSection& Section = LODModel.Sections[SecIdx];
+
+				// toggle visibility between cloth sections and non-cloth sections
+				if (bShowOnlyClothSections)
+				{
+					// enables only cloth sections
+					if (LODModel.Chunks[Section.ChunkIndex].HasApexClothData())
+					{
+						Section.bDisabled = false;
+					}
+					else
+					{
+						Section.bDisabled = true;
+					}
+				}
+				else
+				{   // disables cloth sections and also corresponding original sections
+					if (LODModel.Chunks[Section.ChunkIndex].HasApexClothData())
+					{
+						Section.bDisabled = true;
+						LODModel.Sections[Section.CorrespondClothSectionIndex].bDisabled = true;
+					}
+					else
+					{
+						Section.bDisabled = false;
+					}
+				}
 			}
 		}
+		PostEditChange();
 	}
-
-	PostEditChange();
 }
 
 void UDebugSkelMeshComponent::RestoreClothSectionsVisibility()
 {
+	// if this skeletal mesh doesn't have any clothing assets, just return
+	if (!SkeletalMesh || SkeletalMesh->ClothingAssets.Num() == 0)
+	{
+		return;
+	}
+
 	FSkeletalMeshResource* SkelMeshResource = GetSkeletalMeshResource();
-	if (ensure (SkelMeshResource))
+	if (SkelMeshResource)
 	{
 		PreEditChange(NULL);
 
@@ -509,4 +530,84 @@ void UDebugSkelMeshComponent::RestoreClothSectionsVisibility()
 		PostEditChange();
 	}
 }
+
+int32 UDebugSkelMeshComponent::FindCurrentSectionDisplayMode()
+{
+	ESectionDisplayMode DisplayMode = ESectionDisplayMode::None;
+
+	FSkeletalMeshResource* SkelMeshResource = GetSkeletalMeshResource();
+	// if this skeletal mesh doesn't have any clothing asset, returns "None"
+	if (!SkelMeshResource || !SkeletalMesh || SkeletalMesh->ClothingAssets.Num() == 0)
+	{
+		return ESectionDisplayMode::None;
+	}
+	else
+	{
+		int32 LODIndex;
+		int32 NumLODs = SkelMeshResource->LODModels.Num();
+		for (LODIndex = 0; LODIndex < NumLODs; LODIndex)
+		{
+			// if find any LOD model which has cloth data, then break
+			if (SkelMeshResource->LODModels[LODIndex].HasApexClothData())
+			{
+				break;
+			}
+		}
+
+		// couldn't find 
+		if (LODIndex == NumLODs)
+		{
+			return ESectionDisplayMode::None;
+		}
+
+		FStaticLODModel& LODModel = SkelMeshResource->LODModels[LODIndex];
+
+		// firstly, find cloth sections
+		for (int32 SecIdx = 0; SecIdx < LODModel.Sections.Num(); SecIdx++)
+		{
+			FSkelMeshSection& Section = LODModel.Sections[SecIdx];
+
+			if (LODModel.Chunks[Section.ChunkIndex].HasApexClothData())
+			{
+				// Normal state if the cloth section is visible and the corresponding section is disabled
+				if (Section.bDisabled == false &&
+					LODModel.Sections[Section.CorrespondClothSectionIndex].bDisabled == true)
+				{
+					DisplayMode = ESectionDisplayMode::ShowOnlyClothSections;
+					break;
+				}
+			}
+		}
+
+		// secondly, find non-cloth sections except cloth-corresponding sections
+		bool bFoundNonClothSection = false;
+
+		for (int32 SecIdx = 0; SecIdx < LODModel.Sections.Num(); SecIdx++)
+		{
+			FSkelMeshSection& Section = LODModel.Sections[SecIdx];
+
+			// not related to cloth sections
+			if (!LODModel.Chunks[Section.ChunkIndex].HasApexClothData() &&
+				Section.CorrespondClothSectionIndex < 0)
+			{
+				bFoundNonClothSection = true;
+				if (!Section.bDisabled)
+				{
+					if (DisplayMode == ESectionDisplayMode::ShowOnlyClothSections)
+					{
+						DisplayMode = ESectionDisplayMode::ShowAll;
+					}
+					else
+					{
+						DisplayMode = ESectionDisplayMode::HideOnlyClothSections;
+					}
+				}
+				break;
+			}
+		}
+	}
+
+	return DisplayMode;
+}
+
 #endif // #if WITH_APEX_CLOTHING
\ No newline at end of file
