Fix normals not exporting correctly from the editor using FBX.  This is due to an issue in the FBX sdk when one tries to import an FBX 2014 file into tools that use FBX 2013.  We now export a 2013 file to work around this.

diff --git a/Engine/Source/Editor/UnrealEd/Private/Fbx/FbxMainExport.cpp b/Engine/Source/Editor/UnrealEd/Private/Fbx/FbxMainExport.cpp
--- a/Engine/Source/Editor/UnrealEd/Private/Fbx/FbxMainExport.cpp
+++ b/Engine/Source/Editor/UnrealEd/Private/Fbx/FbxMainExport.cpp
@@ -92,9 +92,9 @@ void FFbxExporter::CreateDocument()
 	
 	// create scene info
 	FbxDocumentInfo* SceneInfo = FbxDocumentInfo::Create(SdkManager,"SceneInfo");
-	SceneInfo->mTitle = "Unreal Matinee Sequence";
-	SceneInfo->mSubject = "Export Unreal Matinee";
-	SceneInfo->mComment = "no particular comments required.";
+	SceneInfo->mTitle = "Unreal FBX Exporter";
+	SceneInfo->mSubject = "Export FBX meshes from Unreal";
+
 
 	Scene->SetSceneInfo(SceneInfo);
 	
@@ -108,9 +108,9 @@ void FFbxExporter::CreateDocument()
 	//FbxScene->GetGlobalSettings().SetOriginalSystemUnit( KFbxSystemUnit::m );
 	
 	// setup anim stack
-	AnimStack = FbxAnimStack::Create(Scene, "Unreal Matinee Take");
+	AnimStack = FbxAnimStack::Create(Scene, "Unreal Take");
 	//KFbxSet<KTime>(AnimStack->LocalStart, KTIME_ONE_SECOND);
-	AnimStack->Description.Set("Animation Take for Unreal Matinee.");
+	AnimStack->Description.Set("Animation Take for Unreal.");
 
 	// this take contains one base layer. In fact having at least one layer is mandatory.
 	AnimLayer = FbxAnimLayer::Create(Scene, "Base Layer");
@@ -152,6 +152,10 @@ void FFbxExporter::WriteToFile(const TCHAR* Filename)
 	IOS_REF.SetBoolProp(EXP_FBX_ANIMATION,       true);
 	IOS_REF.SetBoolProp(EXP_FBX_GLOBAL_SETTINGS, true);
 
+	// We export using FBX 2013 format because many users are still on that version and FBX 2014 files has compatibility issues with
+	// normals when importing to an earlier version of the plugin
+	Exporter->SetFileExportVersion(FBX_FILE_VERSION_7300, FbxSceneRenamer::eNone );
+
 	// Initialize the exporter by providing a filename.
 	if( !Exporter->Initialize(TCHAR_TO_UTF8(Filename), FileFormat, SdkManager->GetIOSettings()) )
 	{
