fix for cook errors in Strategy game
#ue4

diff --git a/Engine/Source/Editor/UnrealEd/Private/Commandlets/CookCommandlet.cpp b/Engine/Source/Editor/UnrealEd/Private/Commandlets/CookCommandlet.cpp
--- a/Engine/Source/Editor/UnrealEd/Private/Commandlets/CookCommandlet.cpp
+++ b/Engine/Source/Editor/UnrealEd/Private/Commandlets/CookCommandlet.cpp
@@ -892,7 +892,7 @@ void UCookCommandlet::CollectFilesToCook(TArray<FString>& FilesInPath)
 	FConfigCacheIni::LoadLocalIniFile(InputIni, TEXT("Input"), true);
 	if (InputIni.GetString(TEXT("/Script/Engine.InputSettings"), TEXT("DefaultTouchInterface"), InterfaceFile))
 	{
-		if (InterfaceFile != TEXT("None"))
+		if (InterfaceFile != TEXT("None") && InterfaceFile != TEXT(""))
 		{
 			FilesInPath.AddUnique(InterfaceFile);
 		}
