Editor: Fixed realtime mode always resetting to enabled after PIE/SIE

diff --git a/Engine/Source/Editor/UnrealEd/Private/Editor.cpp b/Engine/Source/Editor/UnrealEd/Private/Editor.cpp
--- a/Engine/Source/Editor/UnrealEd/Private/Editor.cpp
+++ b/Engine/Source/Editor/UnrealEd/Private/Editor.cpp
@@ -5892,7 +5892,7 @@ void UEditorEngine::RestoreRealtimeViewports()
 		FEditorViewportClient* VC = AllViewportClients[x];
 		if( VC )
 		{
-			VC->RestoreRealtime();
+			VC->RestoreRealtime(true);
 		}
 	}
 
diff --git a/Engine/Source/Editor/UnrealEd/Private/Kismet2/DebuggerCommands.cpp b/Engine/Source/Editor/UnrealEd/Private/Kismet2/DebuggerCommands.cpp
--- a/Engine/Source/Editor/UnrealEd/Private/Kismet2/DebuggerCommands.cpp
+++ b/Engine/Source/Editor/UnrealEd/Private/Kismet2/DebuggerCommands.cpp
@@ -830,12 +830,7 @@ void FInternalPlayWorldCommandCallbacks::Simulate_Clicked()
 		// Start a new simulation session!
 		if( !HasPlayWorld() )
 		{
-			// Make sure the view port is in real-time mode
-			ActiveLevelViewport->GetLevelViewportClient().SetRealtime( true );
-
-			// Start simulating!
-			const bool bSimulateInEditor = true;
-			GUnrealEd->RequestPlaySession(false, ActiveLevelViewport, bSimulateInEditor, NULL, NULL, -1, false );
+			GUnrealEd->RequestPlaySession(false, ActiveLevelViewport, true/*bSimulateInEditor*/, NULL, NULL, -1, false );
 		}
 		else
 		{
