Undo then redo on a duplicated scaled brush now preserves scaling as expected.
#ttp 332995 - Using Undo and then Redo after duplicating edited geometry will return the duplicate to default size and shape
#branch UE4
#proj Editor.UnrealEd
#add Added a check in UEditorBrushBuilder::PostEditChangeProperty so that the brush shouldn't be rebuilt on a property change if it is performing a transaction.

diff --git a/Engine/Source/Editor/UnrealEd/Private/EditorBrushBuilder.cpp b/Engine/Source/Editor/UnrealEd/Private/EditorBrushBuilder.cpp
--- a/Engine/Source/Editor/UnrealEd/Private/EditorBrushBuilder.cpp
+++ b/Engine/Source/Editor/UnrealEd/Private/EditorBrushBuilder.cpp
@@ -225,11 +225,14 @@ bool UEditorBrushBuilder::Build( UWorld* InWorld, ABrush* InBrush )
 
 void UEditorBrushBuilder::PostEditChangeProperty(struct FPropertyChangedEvent& PropertyChangedEvent)
 {
-	// Rebuild brush on property change
-	ABrush* Brush = Cast<ABrush>( GetOuter() );
-	if( Brush )
+	if (!GIsTransacting)
 	{
-		Build(Brush->GetWorld(), Brush);
+		// Rebuild brush on property change
+		ABrush* Brush = Cast<ABrush>(GetOuter());
+		if (Brush)
+		{
+			Build(Brush->GetWorld(), Brush);
+		}
 	}
 }
 
