PIE: Temp fix for PIE not remembering last executed play mode

diff --git a/Engine/Source/Editor/LevelEditor/Private/LevelEditorActions.cpp b/Engine/Source/Editor/LevelEditor/Private/LevelEditorActions.cpp
--- a/Engine/Source/Editor/LevelEditor/Private/LevelEditorActions.cpp
+++ b/Engine/Source/Editor/LevelEditor/Private/LevelEditorActions.cpp
@@ -1935,6 +1935,7 @@ void FLevelEditorActionCallbacks::OnToggleOnlyLoadVisibleInPIE()
 	ULevelEditorPlaySettings* PlaySettings = GetMutableDefault<ULevelEditorPlaySettings>();
 	PlaySettings->bOnlyLoadVisibleLevelsInPIE = !PlaySettings->bOnlyLoadVisibleLevelsInPIE;
 	PlaySettings->PostEditChange();
+	PlaySettings->SaveConfig();
 }
 
 bool FLevelEditorActionCallbacks::OnIsOnlyLoadVisibleInPIEEnabled()
diff --git a/Engine/Source/Editor/UnrealEd/Private/Kismet2/DebuggerCommands.cpp b/Engine/Source/Editor/UnrealEd/Private/Kismet2/DebuggerCommands.cpp
--- a/Engine/Source/Editor/UnrealEd/Private/Kismet2/DebuggerCommands.cpp
+++ b/Engine/Source/Editor/UnrealEd/Private/Kismet2/DebuggerCommands.cpp
@@ -178,6 +178,7 @@ FPlayWorldCommands::FPlayWorldCommands()
 				if (PlayDevice.IsValid())
 				{
 					PlaySettings->LastExecutedLaunchDevice = PlayDevice->GetId().ToString();
+					PlaySettings->SaveConfig();
 				}
 			}
 		}
@@ -887,6 +888,7 @@ void SetLastExecutedPlayMode( EPlayModeType PlayMode )
 {
 	ULevelEditorPlaySettings* PlaySettings = GetMutableDefault<ULevelEditorPlaySettings>();
 	PlaySettings->LastExecutedPlayModeType = PlayMode;
+	PlaySettings->SaveConfig();
 }
 
 
@@ -949,6 +951,7 @@ void SetLastExecutedLaunchMode( ELaunchModeType LaunchMode )
 {
 	ULevelEditorPlaySettings* PlaySettings = GetMutableDefault<ULevelEditorPlaySettings>();
 	PlaySettings->LastExecutedLaunchModeType = LaunchMode;
+	PlaySettings->SaveConfig();
 }
 
 
@@ -1161,6 +1164,7 @@ void FInternalPlayWorldCommandCallbacks::PlayInLocation_Clicked( EPlayModeLocati
 {
 	ULevelEditorPlaySettings* PlaySettings = GetMutableDefault<ULevelEditorPlaySettings>();
 	PlaySettings->LastExecutedPlayModeLocation = Location;
+	PlaySettings->SaveConfig();
 }
 
 
@@ -1296,6 +1300,7 @@ void FInternalPlayWorldCommandCallbacks::HandleLaunchOnDeviceActionExecute( FStr
 	PlaySettings->LastExecutedLaunchModeType = LaunchMode_OnDevice;
 	PlaySettings->LastExecutedLaunchDevice = DeviceId;
     PlaySettings->LastExecutedLaunchName = DeviceName;
+	PlaySettings->SaveConfig();
 
 	LaunchOnDevice(DeviceId, DeviceName);
 }
diff --git a/Engine/Source/Editor/UnrealEd/Private/PlayLevel.cpp b/Engine/Source/Editor/UnrealEd/Private/PlayLevel.cpp
--- a/Engine/Source/Editor/UnrealEd/Private/PlayLevel.cpp
+++ b/Engine/Source/Editor/UnrealEd/Private/PlayLevel.cpp
@@ -709,9 +709,6 @@ void UEditorEngine::StartQueuedPlayMapRequest()
 {
 	const bool bWantSimulateInEditor = bIsSimulateInEditorQueued;
 
-	//ULevelEditorPlaySettings* PlaySettings = GetMutableDefault<ULevelEditorPlaySettings>();
-	//PlaySettings->LastExecutedPlayModeType = PlayMode;
-
 	EndPlayOnLocalPc();
 
 	const ULevelEditorPlaySettings* PlayInSettings = GetDefault<ULevelEditorPlaySettings>();
