FBX: Bind pose issue : When look for bind pose, look for parent node if not valid, and see if it's valid there.

diff --git a/Engine/Source/Editor/UnrealEd/Private/Fbx/FbxSkeletalMeshImport.cpp b/Engine/Source/Editor/UnrealEd/Private/Fbx/FbxSkeletalMeshImport.cpp
--- a/Engine/Source/Editor/UnrealEd/Private/Fbx/FbxSkeletalMeshImport.cpp
+++ b/Engine/Source/Editor/UnrealEd/Private/Fbx/FbxSkeletalMeshImport.cpp
@@ -726,7 +726,6 @@ bool UnFbx::FFbxImporter::ImportBone(TArray<FbxNode*>& NodeArray, FSkeletalMeshI
 
 				// it does not make any difference of checking with different node
 				// it is possible pose 0 -> node array 2, but isValidBindPose function returns true even with node array 0
-				FbxNode * ValidPoseNode = NULL;
 				for (auto Current : NodeArray)
 				{
 					FString CurrentName = Current->GetName();
@@ -738,28 +737,52 @@ bool UnFbx::FFbxImporter::ImportBone(TArray<FbxNode*>& NodeArray, FSkeletalMeshI
 					}
 					else
 					{
-						FString ErrorString = Status.GetErrorString();
-
-						UE_LOG(LogFbx, Warning, TEXT("Not valid bind pose for Pose (%s) - Node %s : %s"), *PoseName, *FString(Current->GetName()), *ErrorString);
-						for(int32 NotificationIndex=0; NotificationIndex<Notification.GetNbEntries(); ++NotificationIndex)
+						// first try to find parent who is null group and see if you can try test it again
+						FbxNode * Parent = Current->GetParent();
+						while (Parent)
 						{
-							const FbxAccumulatorEntry* Entry = Notification.GetEntryAt(NotificationIndex);
-							if(Entry)
+							FbxNodeAttribute* Attr = Parent->GetNodeAttribute();
+							if ( Attr && Attr->GetAttributeType() == FbxNodeAttribute::eNull )
 							{
-								FString Name = (const char*)(Entry->GetName());
-								FString Description = (const char*)(Entry->GetDescription());
-								FString Detail;
+								// found it 
+								break;
+							}
+
+							// find next parent
+							Parent = Parent->GetParent();
+						}
+
+						if ( Parent && CurrentPose->IsValidBindPose(Parent) )
+						{
+							PoseArray.Add(CurrentPose);
+							UE_LOG(LogFbx, Warning, TEXT("Valid bind pose for Pose (%s) - %s"), *PoseName, *FString(Current->GetName()));
+							break;
+						}
+						else
+						{
+							FString ErrorString = Status.GetErrorString();
 
-								for(int32 DetailCount =0; DetailCount< Entry->GetDetailsCount(); ++DetailCount)
+							UE_LOG(LogFbx, Warning, TEXT("Not valid bind pose for Pose (%s) - Node %s : %s"), *PoseName, *FString(Current->GetName()), *ErrorString);
+							for(int32 NotificationIndex=0; NotificationIndex<Notification.GetNbEntries(); ++NotificationIndex)
+							{
+								const FbxAccumulatorEntry* Entry = Notification.GetEntryAt(NotificationIndex);
+								if(Entry)
 								{
-									Detail += FString((const char*)(*Entry->GetDetail(DetailCount)));
-								}
+									FString Name = (const char*)(Entry->GetName());
+									FString Description = (const char*)(Entry->GetDescription());
+									FString Detail;
 
-								UE_LOG(LogFbx, Warning, TEXT("\t(%d) %s : %s (%s)"), NotificationIndex+1, *Name, *Description, *Detail);
+									for(int32 DetailCount =0; DetailCount< Entry->GetDetailsCount(); ++DetailCount)
+									{
+										Detail += FString((const char*)(*Entry->GetDetail(DetailCount)));
+									}
+
+									UE_LOG(LogFbx, Warning, TEXT("\t(%d) %s : %s (%s)"), NotificationIndex+1, *Name, *Description, *Detail);
+								}
 							}
-						}
 
-						Notification.ClearAccumulator();
+							Notification.ClearAccumulator();
+						}
 					}
 				}
 			}
