Editor: Generate better default name for newly created unusual Blueprint assets (interfaces, function libraries, macro libraries)
- Add GetDefaultNewAssetName method to UFactory interface
- Add GetDefaultNewAssetName overrides for various blueprint factory subclasses

diff --git a/Engine/Source/Editor/ContentBrowser/Private/SContentBrowser.cpp b/Engine/Source/Editor/ContentBrowser/Private/SContentBrowser.cpp
--- a/Engine/Source/Editor/ContentBrowser/Private/SContentBrowser.cpp
+++ b/Engine/Source/Editor/ContentBrowser/Private/SContentBrowser.cpp
@@ -1046,7 +1046,7 @@ void SContentBrowser::NewAssetRequested(const FString& SelectedPath, TWeakObject
 
 			static FName AssetToolsModuleName = FName("AssetTools");
 			FAssetToolsModule& AssetToolsModule = FModuleManager::GetModuleChecked<FAssetToolsModule>(AssetToolsModuleName);
-			AssetToolsModule.Get().CreateUniqueAssetName(SelectedPath + TEXT("/New") + NewFactory->GetSupportedClass()->GetName(), TEXT(""), PackageNameToUse, DefaultAssetName);
+			AssetToolsModule.Get().CreateUniqueAssetName(SelectedPath + TEXT("/") + NewFactory->GetDefaultNewAssetName(), TEXT(""), PackageNameToUse, DefaultAssetName);
 			CreateNewAsset(DefaultAssetName, SelectedPath, NewFactory->GetSupportedClass(), NewFactory);
 		}
 	}
diff --git a/Engine/Source/Editor/UnrealEd/Classes/Factories/BlueprintFactory.h b/Engine/Source/Editor/UnrealEd/Classes/Factories/BlueprintFactory.h
--- a/Engine/Source/Editor/UnrealEd/Classes/Factories/BlueprintFactory.h
+++ b/Engine/Source/Editor/UnrealEd/Classes/Factories/BlueprintFactory.h
@@ -16,11 +16,11 @@ class UNREALED_API UBlueprintFactory : public UFactory
 	UPROPERTY(EditAnywhere, Category=BlueprintFactory, meta=(AllowAbstract = "", BlueprintBaseOnly = ""))
 	TSubclassOf<class UObject> ParentClass;
 
-	// Begin UFactory Interface
+	// UFactory Interface
 	virtual bool ConfigureProperties() OVERRIDE;
 	virtual UObject* FactoryCreateNew(UClass* Class, UObject* InParent, FName Name, EObjectFlags Flags, UObject* Context, FFeedbackContext* Warn, FName CallingContext) OVERRIDE;
-	virtual UObject* FactoryCreateNew(UClass* Class,UObject* InParent,FName Name,EObjectFlags Flags,UObject* Context,FFeedbackContext* Warn) OVERRIDE;
-	// Begin UFactory Interface
+	virtual UObject* FactoryCreateNew(UClass* Class, UObject* InParent, FName Name, EObjectFlags Flags, UObject* Context, FFeedbackContext* Warn) OVERRIDE;
+	// End of UFactory Interface
 
 protected:
 	// Hook for derived classes to indicate if they want to skip some filtering that is unnecessary for macro libraries
diff --git a/Engine/Source/Editor/UnrealEd/Classes/Factories/BlueprintFunctionLibraryFactory.h b/Engine/Source/Editor/UnrealEd/Classes/Factories/BlueprintFunctionLibraryFactory.h
--- a/Engine/Source/Editor/UnrealEd/Classes/Factories/BlueprintFunctionLibraryFactory.h
+++ b/Engine/Source/Editor/UnrealEd/Classes/Factories/BlueprintFunctionLibraryFactory.h
@@ -12,13 +12,14 @@ class UBlueprintFunctionLibraryFactory : public UBlueprintFactory
 {
 	GENERATED_UCLASS_BODY()
 
-	// Begin UFactory Interface
+	// UFactory interface
 	virtual FText GetDisplayName() const OVERRIDE;
 	virtual FName GetNewAssetThumbnailOverride() const OVERRIDE;
 	virtual uint32 GetMenuCategories() const OVERRIDE;
 	virtual UObject* FactoryCreateNew(UClass* Class, UObject* InParent, FName Name, EObjectFlags Flags, UObject* Context, FFeedbackContext* Warn, FName CallingContext) OVERRIDE;
 	virtual bool ConfigureProperties() OVERRIDE;
-	// Begin UFactory Interface
+	virtual FString GetDefaultNewAssetName() const OVERRIDE;
+	// End of UFactory interface
 };
 
 
diff --git a/Engine/Source/Editor/UnrealEd/Classes/Factories/BlueprintInterfaceFactory.h b/Engine/Source/Editor/UnrealEd/Classes/Factories/BlueprintInterfaceFactory.h
--- a/Engine/Source/Editor/UnrealEd/Classes/Factories/BlueprintInterfaceFactory.h
+++ b/Engine/Source/Editor/UnrealEd/Classes/Factories/BlueprintInterfaceFactory.h
@@ -12,13 +12,14 @@ class UBlueprintInterfaceFactory : public UFactory
 {
 	GENERATED_UCLASS_BODY()
 
-	// Begin UFactory Interface
+	// UFactory interface
 	virtual FText GetDisplayName() const OVERRIDE;
 	virtual FName GetNewAssetThumbnailOverride() const OVERRIDE;
 	virtual uint32 GetMenuCategories() const OVERRIDE;
 	virtual UObject* FactoryCreateNew(UClass* Class, UObject* InParent, FName Name, EObjectFlags Flags, UObject* Context, FFeedbackContext* Warn, FName CallingContext) OVERRIDE;
 	virtual UObject* FactoryCreateNew(UClass* Class, UObject* InParent, FName Name, EObjectFlags Flags, UObject* Context, FFeedbackContext* Warn) OVERRIDE;
-	// Begin UFactory Interface
+	virtual FString GetDefaultNewAssetName() const OVERRIDE;
+	// End of UFactory interface
 };
 
 
diff --git a/Engine/Source/Editor/UnrealEd/Classes/Factories/BlueprintMacroFactory.h b/Engine/Source/Editor/UnrealEd/Classes/Factories/BlueprintMacroFactory.h
--- a/Engine/Source/Editor/UnrealEd/Classes/Factories/BlueprintMacroFactory.h
+++ b/Engine/Source/Editor/UnrealEd/Classes/Factories/BlueprintMacroFactory.h
@@ -12,12 +12,13 @@ class UBlueprintMacroFactory : public UBlueprintFactory
 {
 	GENERATED_UCLASS_BODY()
 
-	// Begin UFactory Interface
+	// UFactory interface
 	virtual FText GetDisplayName() const OVERRIDE;
 	virtual FName GetNewAssetThumbnailOverride() const OVERRIDE;
 	virtual uint32 GetMenuCategories() const OVERRIDE;
 	virtual UObject* FactoryCreateNew(UClass* Class, UObject* InParent, FName Name, EObjectFlags Flags, UObject* Context, FFeedbackContext* Warn, FName CallingContext) OVERRIDE;
-	// Begin UFactory Interface
+	virtual FString GetDefaultNewAssetName() const OVERRIDE;
+	// End of UFactory interface
 
 protected:
 	// UBlueprintFactory interface
diff --git a/Engine/Source/Editor/UnrealEd/Classes/Factories/Factory.h b/Engine/Source/Editor/UnrealEd/Classes/Factories/Factory.h
--- a/Engine/Source/Editor/UnrealEd/Classes/Factories/Factory.h
+++ b/Engine/Source/Editor/UnrealEd/Classes/Factories/Factory.h
@@ -9,8 +9,8 @@
 #pragma once
 #include "Factory.generated.h"
 
-UCLASS(abstract,MinimalAPI)
-class UFactory : public UObject
+UCLASS(abstract)
+class UNREALED_API UFactory : public UObject
 {
 	GENERATED_UCLASS_BODY()
 
@@ -50,57 +50,58 @@ class UFactory : public UObject
 	UPROPERTY()
 	TArray<FString> ValidGameNames;
 
-	UNREALED_API static FString GetCurrentFilename() { return CurrentFilename; }
-	UNREALED_API static FString	CurrentFilename;
+	static FString GetCurrentFilename() { return CurrentFilename; }
+
+	static FString CurrentFilename;
 
 	/** For interactive object imports, this value indicates whether the user wants objects to be automatically
 		overwritten (See EAppReturnType), or -1 if the user should be prompted. */
-	UNREALED_API static int32 OverwriteYesOrNoToAllState;
+	static int32 OverwriteYesOrNoToAllState;
 
 	/** If this value is true, warning messages will be shown once for all objects being imported at the same time.  
 		This value will be reset to false each time a new import operation is started. */
 	static bool bAllowOneTimeWarningMessages;
 
 	// Begin UObject interface.
-	UNREALED_API static void AddReferencedObjects(UObject* InThis, FReferenceCollector& Collector);
+	static void AddReferencedObjects(UObject* InThis, FReferenceCollector& Collector);
 	// End UObject interface.
 
 	/** Returns true if this factory should be shown in the New Asset menu. */
-	UNREALED_API virtual bool ShouldShowInNewMenu() const;
+	virtual bool ShouldShowInNewMenu() const;
 
 	/** Returns an optional override brush name for the new asset menu. If this is not specified, the thumbnail for the supported class will be used. */
-	UNREALED_API virtual FName GetNewAssetThumbnailOverride() const { return NAME_None; }
+	virtual FName GetNewAssetThumbnailOverride() const { return NAME_None; }
 
 	/** Returns the name of the factory for menus */
-	UNREALED_API virtual FText GetDisplayName() const;
+	virtual FText GetDisplayName() const;
 
 	/** When shown in menus, this is the category containing this factory. Return type is a BitFlag mask using EAssetTypeCategories. */
-	UNREALED_API virtual uint32 GetMenuCategories() const;
+	virtual uint32 GetMenuCategories() const;
 
 	/** Returns the tooltip text description of this factory */
-	UNREALED_API virtual FText GetToolTip() const;
+	virtual FText GetToolTip() const;
 
 	/**
 	 * @return		The object class supported by this factory.
 	 */
-	UNREALED_API UClass* GetSupportedClass() const;
+	UClass* GetSupportedClass() const;
 
 	/**
 	 * @return true if it supports this class 
 	 */
-	UNREALED_API virtual bool DoesSupportClass(UClass * Class);
+	virtual bool DoesSupportClass(UClass * Class);
 
 	/**
 	 * Resolves SupportedClass for factories which support multiple classes.
 	 * Such factories will have a NULL SupportedClass member.
 	 */
-	UNREALED_API virtual UClass* ResolveSupportedClass();
+	virtual UClass* ResolveSupportedClass();
 
 	/**
 	 * Resets the saved state of this factory.  The states are used to suppress messages during multiple object import. 
 	 * It needs to be reset each time a new import is started
 	 */
-	UNREALED_API static void ResetState();
+	static void ResetState();
 
 	/** Opens a dialog to configure the factory properties. Return false if user opted out of configuring properties */
 	virtual bool ConfigureProperties() { return true; }
@@ -124,16 +125,16 @@ class UFactory : public UObject
 	/**
 	 * @return	true if this factory can deal with the file sent in.
 	 */
-	UNREALED_API virtual bool FactoryCanImport( const FString& Filename );
+	virtual bool FactoryCanImport( const FString& Filename );
 
 
 	// @todo document
-	UNREALED_API virtual bool ImportUntypedBulkDataFromText(const TCHAR*& Buffer, FUntypedBulkData& BulkData);
+	virtual bool ImportUntypedBulkDataFromText(const TCHAR*& Buffer, FUntypedBulkData& BulkData);
 
 
 	// @todo document
-	UNREALED_API static UObject* StaticImportObject( UClass* Class, UObject* InOuter, FName Name, EObjectFlags Flags, const TCHAR* Filename=TEXT(""), UObject* Context=NULL, UFactory* Factory=NULL, const TCHAR* Parms=NULL, FFeedbackContext* Warn=GWarn, int32 MaxImportFileSize = 0xC100000 );
-	UNREALED_API static UObject* StaticImportObject( UClass* Class, UObject* InOuter, FName Name, EObjectFlags Flags, bool& bOutOperationCanceled, const TCHAR* Filename=TEXT(""), UObject* Context=NULL, UFactory* Factory=NULL, const TCHAR* Parms=NULL, FFeedbackContext* Warn=GWarn, int32 MaxImportFileSize = 0xC100000 );
+	static UObject* StaticImportObject( UClass* Class, UObject* InOuter, FName Name, EObjectFlags Flags, const TCHAR* Filename=TEXT(""), UObject* Context=NULL, UFactory* Factory=NULL, const TCHAR* Parms=NULL, FFeedbackContext* Warn=GWarn, int32 MaxImportFileSize = 0xC100000 );
+	static UObject* StaticImportObject( UClass* Class, UObject* InOuter, FName Name, EObjectFlags Flags, bool& bOutOperationCanceled, const TCHAR* Filename=TEXT(""), UObject* Context=NULL, UFactory* Factory=NULL, const TCHAR* Parms=NULL, FFeedbackContext* Warn=GWarn, int32 MaxImportFileSize = 0xC100000 );
 
 	/**
 	 * UFactory::ValidForCurrentGame
@@ -144,15 +145,15 @@ class UFactory : public UObject
 	 *
 	 * author: superville
 	 */
-	UNREALED_API bool	ValidForCurrentGame();
+	bool ValidForCurrentGame();
 
 	/** Creates a list of file extensions supported by this factory */
-	UNREALED_API void GetSupportedFileExtensions( TArray<FString>& OutExtensions ) const;
+	void GetSupportedFileExtensions( TArray<FString>& OutExtensions ) const;
 
 	/** 
 	 * Do clean up after importing is done. Will be called once for multi batch import
 	 */
-	virtual void CleanUp() {};
+	virtual void CleanUp() {}
 
 	/**
 	 * Creates an asset if it doesn't exist. If it does exist then it overwrites it if possible. If it can not overwrite then it will delete and replace. If it can not delete, it will return NULL.
@@ -165,5 +166,8 @@ class UFactory : public UObject
 	 *						If NULL, the class default object is used instead.
 	 * @return	A pointer to a new asset of the specified type or null if the creation failed.
 	 */
-	UNREALED_API UObject* CreateOrOverwriteAsset(UClass* InClass, UObject* InParent, FName InName, EObjectFlags InFlags, UObject* InTemplate = NULL) const;
+	UObject* CreateOrOverwriteAsset(UClass* InClass, UObject* InParent, FName InName, EObjectFlags InFlags, UObject* InTemplate = NULL) const;
+
+	/** Returns a new starting point name for newly created assets in the content browser */
+	virtual FString GetDefaultNewAssetName() const;
 };
diff --git a/Engine/Source/Editor/UnrealEd/Private/Factories/EditorFactories.cpp b/Engine/Source/Editor/UnrealEd/Private/Factories/EditorFactories.cpp
--- a/Engine/Source/Editor/UnrealEd/Private/Factories/EditorFactories.cpp
+++ b/Engine/Source/Editor/UnrealEd/Private/Factories/EditorFactories.cpp
@@ -6264,9 +6264,15 @@ UObject* UBlueprintMacroFactory::FactoryCreateNew(UClass* Class, UObject* InPare
 	}
 }
 
+FString UBlueprintMacroFactory::GetDefaultNewAssetName() const
+{
+	return FString(TEXT("NewMacroLibrary"));
+}
+
 /*------------------------------------------------------------------------------
 BlueprintFunctionLibraryFactory implementation.
 ------------------------------------------------------------------------------*/
+
 UBlueprintFunctionLibraryFactory::UBlueprintFunctionLibraryFactory(const class FPostConstructInitializeProperties& PCIP)
 : Super(PCIP)
 {
@@ -6324,6 +6330,11 @@ bool UBlueprintFunctionLibraryFactory::ConfigureProperties()
 	return true;
 }
 
+FString UBlueprintFunctionLibraryFactory::GetDefaultNewAssetName() const
+{
+	return FString(TEXT("NewFunctionLibrary"));
+}
+
 /*------------------------------------------------------------------------------
 	UBlueprintInterfaceFactory implementation.
 ------------------------------------------------------------------------------*/
@@ -6378,6 +6389,11 @@ UObject* UBlueprintInterfaceFactory::FactoryCreateNew(UClass* Class, UObject* In
 	return FactoryCreateNew(Class, InParent, Name, Flags, Context, Warn, NAME_None);
 }
 
+FString UBlueprintInterfaceFactory::GetDefaultNewAssetName() const
+{
+	return FString(TEXT("NewInterface"));
+}
+
 /*------------------------------------------------------------------------------
 	UCurveFactory implementation.
 ------------------------------------------------------------------------------*/
diff --git a/Engine/Source/Editor/UnrealEd/Private/Factories/Factory.cpp b/Engine/Source/Editor/UnrealEd/Private/Factories/Factory.cpp
--- a/Engine/Source/Editor/UnrealEd/Private/Factories/Factory.cpp
+++ b/Engine/Source/Editor/UnrealEd/Private/Factories/Factory.cpp
@@ -493,3 +493,8 @@ UObject* UFactory::CreateOrOverwriteAsset(UClass* InClass, UObject* InParent, FN
 		return NULL;
 	}
 }
+
+FString UFactory::GetDefaultNewAssetName() const
+{
+	return FString(TEXT("New")) + GetSupportedClass()->GetName();
+}
\ No newline at end of file
