BuildScale is now applied when building Sphere collision

#ttp 334968 - EDITOR: Static Mesh: Build scale isn't taken into account when building sphere collision

#branch UE4

diff --git a/Engine/Source/Editor/UnrealEd/Private/GeomFitUtils.cpp b/Engine/Source/Editor/UnrealEd/Private/GeomFitUtils.cpp
--- a/Engine/Source/Editor/UnrealEd/Private/GeomFitUtils.cpp
+++ b/Engine/Source/Editor/UnrealEd/Private/GeomFitUtils.cpp
@@ -329,7 +329,7 @@ void GenerateSphereAsSimpleCollision(UStaticMesh* StaticMesh)
 	SrcModel.RawMeshBulkData->LoadRawMesh(RawMesh);
 
 	FSphere bSphere, bSphere2, bestSphere;
-	FVector unitVec = FVector(1,1,1);
+	FVector unitVec = StaticMesh->BodySetup ? StaticMesh->BodySetup->BuildScale3D : FVector(1.f, 1.f, 1.f);
 	CalcBoundingSphere(RawMesh, bSphere, unitVec);
 	CalcBoundingSphere2(RawMesh, bSphere2, unitVec);
 
