Fix TTP 332124 - EDITOR: LANDSCAPE: CRASH: Using the move to level tool in combination with simulate can lead to a crash
The crash was caused by various things getting confused between the editor world and the pie world when in pie/simulate. I cleaned up the code surrounding landscape undo/redo as well, a lot of it was quite old and going a bit stale.

diff --git a/Engine/Source/Editor/UnrealEd/Classes/Editor/UnrealEdEngine.h b/Engine/Source/Editor/UnrealEd/Classes/Editor/UnrealEdEngine.h
--- a/Engine/Source/Editor/UnrealEd/Classes/Editor/UnrealEdEngine.h
+++ b/Engine/Source/Editor/UnrealEd/Classes/Editor/UnrealEdEngine.h
@@ -687,7 +687,6 @@ class UNREALED_API UUnrealEdEngine : public UEditorEngine, public FNotifyHook
 	bool HandleDumpSelectionCommand( const TCHAR* Str, FOutputDevice& Ar );
 	bool HandleBuildLightingCommand( const TCHAR* Str, FOutputDevice& Ar, UWorld* InWorld );
 	bool HandleBuildPathsCommand( const TCHAR* Str, FOutputDevice& Ar, UWorld* InWorld );
-	bool HandleRestoreLandscapeLayerInfosCommand( const TCHAR* Str, FOutputDevice& Ar, UWorld* InWorld );
 	bool HandleUpdateLandscapeEditorDataCommand( const TCHAR* Str, FOutputDevice& Ar, UWorld* InWorld );
 	bool HandleUpdateLandscapeMICCommand( const TCHAR* Str, FOutputDevice& Ar, UWorld* InWorld );
 	bool HandleConvertMatineesCommand( const TCHAR* Str, FOutputDevice& Ar, UWorld* InWorld );
diff --git a/Engine/Source/Editor/UnrealEd/Private/EditorActorFolders.cpp b/Engine/Source/Editor/UnrealEd/Private/EditorActorFolders.cpp
--- a/Engine/Source/Editor/UnrealEd/Private/EditorActorFolders.cpp
+++ b/Engine/Source/Editor/UnrealEd/Private/EditorActorFolders.cpp
@@ -158,8 +158,12 @@ UEditorActorFolders& FActorFolders::GetFoldersForWorld_Internal(const UWorld& In
 		return **Folders;
 	}
 
-	UEditorActorFolders* Folders = ConstructObject<UEditorActorFolders>(UEditorActorFolders::StaticClass(), GetTransientPackage(), NAME_None, RF_Transactional);
-	return *TemporaryWorldFolders.Add(&InWorld, Folders);
+	// We intentionally don't pass RF_Transactional to ConstructObject so that we don't record the creation of the object into the undo buffer
+	// (to stop it getting deleted on undo as we manage its lifetime), but we still want it to be RF_Transactional so we can record any changes later
+	UEditorActorFolders* Folders = ConstructObject<UEditorActorFolders>(UEditorActorFolders::StaticClass(), GetTransientPackage(), NAME_None, RF_NoFlags);
+	Folders->SetFlags(RF_Transactional);
+	TemporaryWorldFolders.Add(&InWorld, Folders);
+	return *Folders;
 }
 
 TSet<FName> FActorFolders::ScanWorldForFolders(UWorld& InWorld)
diff --git a/Engine/Source/Editor/UnrealEd/Private/UnrealEdSrv.cpp b/Engine/Source/Editor/UnrealEd/Private/UnrealEdSrv.cpp
--- a/Engine/Source/Editor/UnrealEd/Private/UnrealEdSrv.cpp
+++ b/Engine/Source/Editor/UnrealEd/Private/UnrealEdSrv.cpp
@@ -481,42 +481,15 @@ bool UUnrealEdEngine::HandleBuildPathsCommand( const TCHAR* Str, FOutputDevice&
 	return FEditorBuildUtils::EditorBuild(InWorld, EBuildOptions::BuildAIPaths);
 }
 
-bool UUnrealEdEngine::HandleRestoreLandscapeLayerInfosCommand( const TCHAR* Str, FOutputDevice& Ar, UWorld* InWorld )
-{
-	if (!PlayWorld && InWorld && InWorld->GetWorldSettings())
-	{
-		// Search all layer info into LandscapeInfo
-		for (auto It = InWorld->LandscapeInfoMap.CreateIterator(); It; ++It)
-		{
-			ULandscapeInfo* LandscapeInfo = It.Value();
-			if (LandscapeInfo)
-			{
-				LandscapeInfo->UpdateLayerInfoMap();
-
-				// PostEditUndo for Components
-				for ( auto LandscapeComponentIt = LandscapeInfo->XYtoComponentMap.CreateIterator(); LandscapeComponentIt; ++LandscapeComponentIt )
-				{
-					ULandscapeComponent* Comp = LandscapeComponentIt.Value();
-					if (Comp && Comp->bNeedPostUndo)
-					{
-						Comp->PostEditUndo();
-					}
-				}
-			}
-		}
-	}
-	return true;
-}
-
 bool UUnrealEdEngine::HandleUpdateLandscapeEditorDataCommand( const TCHAR* Str, FOutputDevice& Ar, UWorld* InWorld )
 {
 	const bool bShowWarnings = (FString(Str) == TEXT("-warnings"));
 	
-	if (!PlayWorld && InWorld && InWorld->GetWorldSettings())
+	if (InWorld->GetWorldSettings())
 	{
 		ULandscapeInfo::RecreateLandscapeInfo(InWorld, bShowWarnings);
-				
-		// for removing 
+
+		// for removing
 		TMap<ULandscapeInfo*, ALandscapeGizmoActiveActor*> GizmoMap;
 		for (FActorIterator It(InWorld); It; ++It)
 		{
@@ -601,9 +574,11 @@ bool UUnrealEdEngine::HandleUpdateLandscapeEditorDataCommand( const TCHAR* Str,
 
 bool UUnrealEdEngine::HandleUpdateLandscapeMICCommand( const TCHAR* Str, FOutputDevice& Ar, UWorld* InWorld )
 {
-	if (!PlayWorld && InWorld && InWorld->GetWorldSettings())
+	UWorld* World = GetEditorWorldContext().World();
+
+	if (World && World->GetWorldSettings())
 	{
-		for (auto It = InWorld->LandscapeInfoMap.CreateIterator(); It; ++It)
+		for (auto It = World->LandscapeInfoMap.CreateIterator(); It; ++It)
 		{
 			ULandscapeInfo* Info = It.Value();
 			for (auto LandscapeComponentIt = Info->XYtoComponentMap.CreateIterator(); LandscapeComponentIt; ++LandscapeComponentIt )
@@ -811,17 +786,17 @@ bool UUnrealEdEngine::Exec( UWorld* InWorld, const TCHAR* Stream, FOutputDevice&
 		HandleBuildPathsCommand( Str, Ar, InWorld );
 	}
 #if WITH_EDITOR
-	else if (FParse::Command(&Str, TEXT("RestoreLandscapeLayerInfos")))
-	{
-		HandleRestoreLandscapeLayerInfosCommand( Str, Ar, InWorld );
-	}
 	else if (FParse::Command(&Str, TEXT("UpdateLandscapeEditorData")))
 	{
-		return HandleUpdateLandscapeEditorDataCommand( Str, Ar, InWorld );
+		// InWorld above is the PIE world if PIE is active, but this is specifically an editor command
+		UWorld* World = GetEditorWorldContext().World();
+		return HandleUpdateLandscapeEditorDataCommand(Str, Ar, World);
 	}
 	else if (FParse::Command(&Str, TEXT("UpdateLandscapeMIC")))
 	{
-		return HandleUpdateLandscapeMICCommand( Str, Ar, InWorld );
+		// InWorld above is the PIE world if PIE is active, but this is specifically an editor command
+		UWorld* World = GetEditorWorldContext().World();
+		return HandleUpdateLandscapeMICCommand(Str, Ar, World);
 	}
 #endif // WITH_EDITOR
 	else if( FParse::Command(&Str, TEXT("CONVERTMATINEES")) )
diff --git a/Engine/Source/Editor/UnrealEd/Public/EditorModes.h b/Engine/Source/Editor/UnrealEd/Public/EditorModes.h
--- a/Engine/Source/Editor/UnrealEd/Public/EditorModes.h
+++ b/Engine/Source/Editor/UnrealEd/Public/EditorModes.h
@@ -318,7 +318,7 @@ public:
 
 	UWorld* GetWorld() const
 	{
-		return GWorld;
+		return GEditor->GetEditorWorldContext().World();
 	}
 
 	// Property Widgets
diff --git a/Engine/Source/Runtime/Engine/Classes/Landscape/LandscapeComponent.h b/Engine/Source/Runtime/Engine/Classes/Landscape/LandscapeComponent.h
--- a/Engine/Source/Runtime/Engine/Classes/Landscape/LandscapeComponent.h
+++ b/Engine/Source/Runtime/Engine/Classes/Landscape/LandscapeComponent.h
@@ -198,11 +198,6 @@ public:
 	UPROPERTY(EditAnywhere, Category=LandscapeComponent)
 	float StaticLightingResolution;
 
-#if WITH_EDITORONLY_DATA
-	UPROPERTY(transient)
-	uint32 bNeedPostUndo:1;
-#endif // WITH_EDITORONLY_DATA
-
 	/** Forced LOD level to use when rendering */
 	UPROPERTY(EditAnywhere, BlueprintReadOnly, Category=LandscapeComponent)
 	int32 ForcedLOD;
diff --git a/Engine/Source/Runtime/Engine/Private/Landscape/Landscape.cpp b/Engine/Source/Runtime/Engine/Private/Landscape/Landscape.cpp
--- a/Engine/Source/Runtime/Engine/Private/Landscape/Landscape.cpp
+++ b/Engine/Source/Runtime/Engine/Private/Landscape/Landscape.cpp
@@ -53,11 +53,8 @@ ULandscapeComponent::ULandscapeComponent(const class FPostConstructInitializePro
 	bAllowCullDistanceVolume = false;
 	CollisionMipLevel = 0;
 	StaticLightingResolution = 0.f; // Default value 0 means no overriding
-#if WITH_EDITORONLY_DATA
-	bNeedPostUndo = false;
-#endif // WITH_EDITORONLY_DATA
-	HeightmapScaleBias = FVector4(0.0f, 0.0f, 0.0f, 1.0f);
 
+	HeightmapScaleBias = FVector4(0.0f, 0.0f, 0.0f, 1.0f);
 	WeightmapScaleBias = FVector4(0.0f, 0.0f, 0.0f, 1.0f);
 
 	bBoundsChangeTriggersStreamingDataRebuild = true;
@@ -869,32 +866,34 @@ void ALandscapeProxy::AddReferencedObjects(UObject* InThis, FReferenceCollector&
 void ALandscapeProxy::PreEditUndo()
 {
 	Super::PreEditUndo();
-	if (GIsEditor)
+
+	// PostEditUndo doesn't get called when undoing a create action
+	// and PreEditUndo doesn't get called when undoing a delete action
+	// so this code needs to be in both!
+	if (GIsEditor && GetWorld() && !GetWorld()->IsPlayInEditor())
 	{
-		// Remove all layer info for this Proxy
-		ULandscapeInfo* LandscapeInfo = GetLandscapeInfo(false);
-		if (LandscapeInfo)
-		{
-			LandscapeInfo->UpdateLayerInfoMap(this, true);
-		}
+		GEngine->DeferredCommands.AddUnique(TEXT("UpdateLandscapeEditorData"));
 	}
 }
 
 void ALandscapeProxy::PostEditUndo()
 {
 	Super::PostEditUndo();
+
 	if (GIsEditor && GetWorld() && !GetWorld()->IsPlayInEditor())
 	{
-		// defer LandscapeInfo setup
 		GEngine->DeferredCommands.AddUnique(TEXT("UpdateLandscapeEditorData"));
-		GEngine->DeferredCommands.AddUnique(TEXT("RestoreLandscapeLayerInfos"));
 	}
 }
 
 void ULandscapeInfo::PostEditUndo()
 {
 	Super::PostEditUndo();
-	ULandscapeInfo::RecreateLandscapeInfo(CastChecked<UWorld>(GetOuter()), true);
+
+	if (GIsEditor && GetWorld() && !GetWorld()->IsPlayInEditor())
+	{
+		GEngine->DeferredCommands.AddUnique(TEXT("UpdateLandscapeEditorData"));
+	}
 }
 
 FName FLandscapeInfoLayerSettings::GetLayerName() const
@@ -1220,11 +1219,6 @@ void ALandscapeProxy::PostLoad()
 void ALandscape::Destroyed()
 {
 	Super::Destroyed();
-
-	if (SplineComponent)
-	{
-		SplineComponent->ModifySplines();
-	}
 }
 
 void ALandscapeProxy::Destroyed()
@@ -1235,6 +1229,11 @@ void ALandscapeProxy::Destroyed()
 	{
 		ULandscapeInfo::RecreateLandscapeInfo(GetWorld(), false);
 	}
+
+	if (SplineComponent)
+	{
+		SplineComponent->ModifySplines();
+	}
 }
 #endif
 
diff --git a/Engine/Source/Runtime/Engine/Private/Landscape/LandscapeEdit.cpp b/Engine/Source/Runtime/Engine/Private/Landscape/LandscapeEdit.cpp
--- a/Engine/Source/Runtime/Engine/Private/Landscape/LandscapeEdit.cpp
+++ b/Engine/Source/Runtime/Engine/Private/Landscape/LandscapeEdit.cpp
@@ -280,22 +280,6 @@ void ULandscapeComponent::PostEditUndo()
 
 	UpdateMaterialInstances();
 
-	if (!bNeedPostUndo)
-	{
-		ULandscapeInfo* Info = GetLandscapeInfo(false);
-		if (Info && !Info->bIsValid)
-		{
-			bNeedPostUndo = true;
-			if (EditToolRenderData)
-			{
-				EditToolRenderData->DebugChannelR = INDEX_NONE;
-				EditToolRenderData->DebugChannelG = INDEX_NONE;
-				EditToolRenderData->DebugChannelB = INDEX_NONE;
-			}
-			return;
-		}
-	}
-
 	if (EditToolRenderData)
 	{
 		EditToolRenderData->UpdateDebugColorMaterial();
@@ -3763,7 +3747,7 @@ void ALandscape::PostEditChangeProperty(FPropertyChangedEvent& PropertyChangedEv
 				}
 			}
 		}
-				
+
 		if (ChangedMaterial)
 		{
 			if (GIsEditor && GetWorld() && !GetWorld()->IsPlayInEditor())
