Stuctures are parsed in FBlueprintEditorUtils::PropertyValueFromString.

https://udn.unrealengine.com/questions/207259/setarrayelem-not-respecting-entered-vector2.html

diff --git a/Engine/Source/Editor/KismetCompiler/Private/KismetCompilerVMBackend.cpp b/Engine/Source/Editor/KismetCompiler/Private/KismetCompilerVMBackend.cpp
--- a/Engine/Source/Editor/KismetCompiler/Private/KismetCompilerVMBackend.cpp
+++ b/Engine/Source/Editor/KismetCompiler/Private/KismetCompilerVMBackend.cpp
@@ -452,6 +452,7 @@ public:
 					int32 StructSize = Struct->GetStructureSize() * StructProperty->ArrayDim;
 					uint8* StructData = (uint8*)FMemory_Alloca(StructSize);
 					StructProperty->InitializeValue(StructData);
+					ensure(1 == StructProperty->ArrayDim);
 					if(!FStructureEditorUtils::Fill_MakeStructureDefaultValue(Cast<UUserDefinedStruct>(Struct), StructData))
 					{
 						UE_LOG(LogK2Compiler, Warning, TEXT("MakeStructureDefaultValue parsing error. Property: %s, Struct: %s"), *StructProperty->GetName(), *Struct->GetName());
diff --git a/Engine/Source/Editor/UnrealEd/Private/Kismet2/BlueprintEditorUtils.cpp b/Engine/Source/Editor/UnrealEd/Private/Kismet2/BlueprintEditorUtils.cpp
--- a/Engine/Source/Editor/UnrealEd/Private/Kismet2/BlueprintEditorUtils.cpp
+++ b/Engine/Source/Editor/UnrealEd/Private/Kismet2/BlueprintEditorUtils.cpp
@@ -18,6 +18,7 @@
 #include "Editor/UnrealEd/Public/Kismet2/KismetEditorUtilities.h"
 #include "Editor/UnrealEd/Public/Kismet2/KismetDebugUtilities.h"
 #include "Editor/UnrealEd/Public/Kismet2/BlueprintEditorUtils.h"
+#include "Editor/UnrealEd/Public/Kismet2/StructureEditorUtils.h"
 #include "ScopedTransaction.h"
 #include "Editor/ClassViewer/Public/ClassViewerModule.h"
 #include "Editor/ClassViewer/Public/ClassViewerFilter.h"
@@ -5540,12 +5541,16 @@ bool FBlueprintEditorUtils::PropertyValueFromString(const UProperty* Property, c
 		else if( Property->IsA(UArrayProperty::StaticClass()) )
 		{
 			const UArrayProperty* ArrayProp = CastChecked<UArrayProperty>(Property);
-			ArrayProp->ImportText(*Value, ArrayProp->ContainerPtrToValuePtr<uint8>(DefaultObject), 0, NULL);
+			FStringOutputDevice ImportError;
+			ArrayProp->ImportText(*Value, ArrayProp->ContainerPtrToValuePtr<uint8>(DefaultObject), 0, NULL, &ImportError);
+			bParseSuccedded = ImportError.IsEmpty();
 		}
 		else if( Property->IsA(UMulticastDelegateProperty::StaticClass()) )
 		{
 			const UMulticastDelegateProperty* MCDelegateProp = CastChecked<UMulticastDelegateProperty>(Property);
-			MCDelegateProp->ImportText(*Value, MCDelegateProp->ContainerPtrToValuePtr<uint8>(DefaultObject), 0, NULL);
+			FStringOutputDevice ImportError;
+			MCDelegateProp->ImportText(*Value, MCDelegateProp->ContainerPtrToValuePtr<uint8>(DefaultObject), 0, NULL, &ImportError);
+			bParseSuccedded = ImportError.IsEmpty();
 		}
 		else
 		{
@@ -5589,6 +5594,19 @@ bool FBlueprintEditorUtils::PropertyValueFromString(const UProperty* Property, c
 			bParseSuccedded = Color.InitFromString(Value);
 			Property->CopyCompleteValue( Property->ContainerPtrToValuePtr<uint8>(DefaultObject), &Color );
 		}
+		else if (StructProperty->Struct)
+		{
+			const UScriptStruct* Struct = StructProperty->Struct;
+			const int32 StructSize = Struct->GetStructureSize() * StructProperty->ArrayDim;
+			uint8* StructData = Property->ContainerPtrToValuePtr<uint8>(DefaultObject);
+			StructProperty->InitializeValue(StructData);
+			ensure(1 == StructProperty->ArrayDim);
+			bParseSuccedded = FStructureEditorUtils::Fill_MakeStructureDefaultValue(Cast<const UUserDefinedStruct>(Struct), StructData);
+
+			FStringOutputDevice ImportError;
+			StructProperty->ImportText(Value.IsEmpty() ? TEXT("()") : *Value, StructData, 0, NULL, &ImportError);
+			bParseSuccedded &= ImportError.IsEmpty();
+		}
 	}
 
 	return bParseSuccedded;
