Fix destructible editor resetting material when fracturing
TTP#335444

diff --git a/Engine/Source/Editor/UnrealEd/Private/Factories/ApexDestructibleAssetImport.cpp b/Engine/Source/Editor/UnrealEd/Private/Factories/ApexDestructibleAssetImport.cpp
--- a/Engine/Source/Editor/UnrealEd/Private/Factories/ApexDestructibleAssetImport.cpp
+++ b/Engine/Source/Editor/UnrealEd/Private/Factories/ApexDestructibleAssetImport.cpp
@@ -939,11 +939,26 @@ UNREALED_API bool BuildDestructibleMeshFromFractureSettings(UDestructibleMesh& D
 #if WITH_EDITORONLY_DATA
 	if (DestructibleMesh.FractureSettings != NULL)
 	{
+		TArray<UMaterialInterface*> OverrideMaterials;
+		OverrideMaterials.Init(DestructibleMesh.Materials.Num());	//save old materials
+		for (int32 MaterialIndex = 0; MaterialIndex < DestructibleMesh.Materials.Num(); ++MaterialIndex)
+		{
+			OverrideMaterials[MaterialIndex] = DestructibleMesh.Materials[MaterialIndex].MaterialInterface;
+		}
+
 		DestructibleMesh.Materials.Init(DestructibleMesh.FractureSettings->Materials.Num());
 
 		for (int32 MaterialIndex = 0; MaterialIndex < DestructibleMesh.Materials.Num(); ++MaterialIndex)
 		{
-			DestructibleMesh.Materials[MaterialIndex].MaterialInterface = DestructibleMesh.FractureSettings->Materials[MaterialIndex];
+			if (MaterialIndex < OverrideMaterials.Num())	//if user has overriden materials use it
+			{
+				DestructibleMesh.Materials[MaterialIndex].MaterialInterface = OverrideMaterials[MaterialIndex];
+			}
+			else
+			{
+				DestructibleMesh.Materials[MaterialIndex].MaterialInterface = DestructibleMesh.FractureSettings->Materials[MaterialIndex];
+			}
+
 		}
 
 		NxDestructibleAssetCookingDesc DestructibleAssetCookingDesc;
