Allow corrupted assets to be loaded without a crash. More robust FArchiveMoveSkeletalRefs.

diff --git a/Engine/Source/Editor/UnrealEd/Private/Kismet2/BlueprintEditorUtils.cpp b/Engine/Source/Editor/UnrealEd/Private/Kismet2/BlueprintEditorUtils.cpp
--- a/Engine/Source/Editor/UnrealEd/Private/Kismet2/BlueprintEditorUtils.cpp
+++ b/Engine/Source/Editor/UnrealEd/Private/Kismet2/BlueprintEditorUtils.cpp
@@ -514,10 +514,14 @@ protected:
 		{
 			if( UClass* RefClass = Cast<UClass>(Object) )
 			{
-				UClass* AuthClass = RefClass->GetAuthoritativeClass();
-				if( RefClass != AuthClass )
+				const bool bIsValidBPGeneratedClass = RefClass->HasAnyClassFlags(CLASS_CompiledFromBlueprint) && RefClass->ClassGeneratedBy;
+				if (bIsValidBPGeneratedClass)
 				{
-					Object = AuthClass;
+					UClass* AuthClass = RefClass->GetAuthoritativeClass();
+					if (RefClass != AuthClass)
+					{
+						Object = AuthClass;
+					}
 				}
 			}
 		}
