Landscape : Added RecreateLandscapeCollision command to recreate all landscape collision component

diff --git a/Engine/Source/Editor/UnrealEd/Classes/Editor/UnrealEdEngine.h b/Engine/Source/Editor/UnrealEd/Classes/Editor/UnrealEdEngine.h
--- a/Engine/Source/Editor/UnrealEd/Classes/Editor/UnrealEdEngine.h
+++ b/Engine/Source/Editor/UnrealEd/Classes/Editor/UnrealEdEngine.h
@@ -689,6 +689,7 @@ class UNREALED_API UUnrealEdEngine : public UEditorEngine, public FNotifyHook
 	bool HandleBuildPathsCommand( const TCHAR* Str, FOutputDevice& Ar, UWorld* InWorld );
 	bool HandleUpdateLandscapeEditorDataCommand( const TCHAR* Str, FOutputDevice& Ar, UWorld* InWorld );
 	bool HandleUpdateLandscapeMICCommand( const TCHAR* Str, FOutputDevice& Ar, UWorld* InWorld );
+	bool HandleRecreateLandscapeCollisionCommand(const TCHAR* Str, FOutputDevice& Ar, UWorld* InWorld);
 	bool HandleConvertMatineesCommand( const TCHAR* Str, FOutputDevice& Ar, UWorld* InWorld );
 	bool HandleDisasmScriptCommand( const TCHAR* Str, FOutputDevice& Ar );	
 
diff --git a/Engine/Source/Editor/UnrealEd/Private/UnrealEdSrv.cpp b/Engine/Source/Editor/UnrealEd/Private/UnrealEdSrv.cpp
--- a/Engine/Source/Editor/UnrealEd/Private/UnrealEdSrv.cpp
+++ b/Engine/Source/Editor/UnrealEd/Private/UnrealEdSrv.cpp
@@ -488,8 +488,8 @@ bool UUnrealEdEngine::HandleUpdateLandscapeEditorDataCommand( const TCHAR* Str,
 	if (InWorld->GetWorldSettings())
 	{
 		ULandscapeInfo::RecreateLandscapeInfo(InWorld, bShowWarnings);
-
-		// for removing
+				
+		// for removing 
 		TMap<ULandscapeInfo*, ALandscapeGizmoActiveActor*> GizmoMap;
 		for (FActorIterator It(InWorld); It; ++It)
 		{
@@ -596,6 +596,19 @@ bool UUnrealEdEngine::HandleUpdateLandscapeMICCommand( const TCHAR* Str, FOutput
 	return true;
 }
 
+bool UUnrealEdEngine::HandleRecreateLandscapeCollisionCommand(const TCHAR* Str, FOutputDevice& Ar, UWorld* InWorld)
+{
+	if (!PlayWorld && InWorld && InWorld->GetWorldSettings())
+	{
+		for (auto It = InWorld->LandscapeInfoMap.CreateIterator(); It; ++It)
+		{
+			ULandscapeInfo* Info = It.Value();
+			Info->RecreateCollisionComponents();
+		}
+	}
+	return true;
+}
+
 bool UUnrealEdEngine::HandleConvertMatineesCommand( const TCHAR* Str, FOutputDevice& Ar, UWorld* InWorld )
 {
 	FVector StartLocation= FVector::ZeroVector;
@@ -798,6 +811,12 @@ bool UUnrealEdEngine::Exec( UWorld* InWorld, const TCHAR* Stream, FOutputDevice&
 		UWorld* World = GetEditorWorldContext().World();
 		return HandleUpdateLandscapeMICCommand(Str, Ar, World);
 	}
+	else if (FParse::Command(&Str, TEXT("RecreateLandscapeCollision")))
+	{
+		// InWorld above is the PIE world if PIE is active, but this is specifically an editor command
+		UWorld* World = GetEditorWorldContext().World();
+		return HandleRecreateLandscapeCollisionCommand(Str, Ar, World);
+	}
 #endif // WITH_EDITOR
 	else if( FParse::Command(&Str, TEXT("CONVERTMATINEES")) )
 	{
diff --git a/Engine/Source/Runtime/Engine/Classes/Landscape/LandscapeInfo.h b/Engine/Source/Runtime/Engine/Classes/Landscape/LandscapeInfo.h
--- a/Engine/Source/Runtime/Engine/Classes/Landscape/LandscapeInfo.h
+++ b/Engine/Source/Runtime/Engine/Classes/Landscape/LandscapeInfo.h
@@ -244,6 +244,7 @@ public:
 	 */
 	ENGINE_API void FixupProxiesWeightmaps();
 
+	ENGINE_API void RecreateCollisionComponents();
 #endif
 };
 
diff --git a/Engine/Source/Runtime/Engine/Private/Landscape/LandscapeEdit.cpp b/Engine/Source/Runtime/Engine/Private/Landscape/LandscapeEdit.cpp
--- a/Engine/Source/Runtime/Engine/Private/Landscape/LandscapeEdit.cpp
+++ b/Engine/Source/Runtime/Engine/Private/Landscape/LandscapeEdit.cpp
@@ -3420,6 +3420,20 @@ void ALandscapeProxy::RecreateCollisionComponents()
 	}
 }
 
+void ULandscapeInfo::RecreateCollisionComponents()
+{
+	if (LandscapeActor.IsValid())
+	{
+		LandscapeActor->RecreateCollisionComponents();
+	}
+
+	for (auto It = Proxies.CreateConstIterator(); It; ++It)
+	{
+		ALandscapeProxy* Proxy = (*It);
+		Proxy->RecreateCollisionComponents();
+	}
+}
+
 namespace
 {
 	inline float AdjustStaticLightingResolution(float StaticLightingResolution, int32 NumSubsections, int32 SubsectionSizeQuads, int32 ComponentSizeQuads)
@@ -3747,7 +3761,7 @@ void ALandscape::PostEditChangeProperty(FPropertyChangedEvent& PropertyChangedEv
 				}
 			}
 		}
-
+				
 		if (ChangedMaterial)
 		{
 			if (GIsEditor && GetWorld() && !GetWorld()->IsPlayInEditor())
