[333778] Fix a crash when removing a variable from a parent class and then attempting to recreate the variable in a Blueprint based on that class.

- modified FBlueprintEditorUtils::AddMemberVariable() to fall back to UObject for "object" variable types that no longer exist when the variable is also missing

diff --git a/Engine/Source/Editor/UnrealEd/Private/Kismet2/BlueprintEditorUtils.cpp b/Engine/Source/Editor/UnrealEd/Private/Kismet2/BlueprintEditorUtils.cpp
--- a/Engine/Source/Editor/UnrealEd/Private/Kismet2/BlueprintEditorUtils.cpp
+++ b/Engine/Source/Editor/UnrealEd/Private/Kismet2/BlueprintEditorUtils.cpp
@@ -2998,6 +2998,14 @@ bool FBlueprintEditorUtils::AddMemberVariable(UBlueprint* Blueprint, const FName
 			NewVar.VarType.PinSubCategory.Empty();
 			NewVar.VarType.PinSubCategoryObject = *Blueprint->GeneratedClass;
 		}
+		else if (!NewVar.VarType.PinSubCategoryObject.IsValid())
+		{
+			// Fall back to UObject if the given type is not valid. This can happen for example if a variable is removed from
+			// a Blueprint parent class along with the variable's type and the user then attempts to recreate the missing variable
+			// through a stale variable node's context menu in a child Blueprint graph.
+			NewVar.VarType.PinSubCategory.Empty();
+			NewVar.VarType.PinSubCategoryObject = UObject::StaticClass();
+		}
 
 		// if it's a PC_Object, then it should have an associated UClass object
 		check(NewVar.VarType.PinSubCategoryObject.IsValid());
