Force a full panel refresh of the World Settings tab when creating a new GameMode.
#ttp 331586 Editor: Selecting a new GameMode in world settings causes issues.
reviewer Thomas.Sarkanen

diff --git a/Engine/Source/Editor/UnrealEd/Public/GameModeInfoCustomizer.h b/Engine/Source/Editor/UnrealEd/Public/GameModeInfoCustomizer.h
--- a/Engine/Source/Editor/UnrealEd/Public/GameModeInfoCustomizer.h
+++ b/Engine/Source/Editor/UnrealEd/Public/GameModeInfoCustomizer.h
@@ -137,7 +137,7 @@ public:
 				.HAlign(HAlign_Center)
 				.VAlign(VAlign_Center)
 				.ContentPadding(FMargin(2.0f,0.0f))
-				.OnClicked(this, &FGameModeInfoCustomizer::OnClickNewGameMode)
+				.OnClicked(this, &FGameModeInfoCustomizer::OnClickNewGameMode, &LayoutBuilder)
 				.ToolTipText(LOCTEXT("NewGameMode_ToolTip", "Create a new Game Mode"))
 				[
 					SNew( STextBlock )
@@ -272,7 +272,7 @@ public:
 		}
 	}
 
-	FReply OnClickNewGameMode()
+	FReply OnClickNewGameMode(IDetailLayoutBuilder* DetailLayout)
 	{
 		// Create a new GameMode BP
 		UBlueprint* Blueprint = FKismetEditorUtilities::CreateBlueprintFromClass(LOCTEXT("CreateNewGameMode", "Create New GameMode"), AGameMode::StaticClass(), TEXT("NewGameMode"));
@@ -282,6 +282,11 @@ public:
 			DefaultGameModeClassHandle->SetValueFromFormattedString(Blueprint->GeneratedClass->GetPathName());
 		}
 
+		if (DetailLayout)
+		{
+			DetailLayout->ForceRefreshDetails();
+		}
+
 		return FReply::Handled();
 	}
 
