Renaming a UDEnum used in a UDStruct, won't cause a crash.

diff --git a/Engine/Source/Editor/UnrealEd/Private/Kismet2/EnumEditorUtils.cpp b/Engine/Source/Editor/UnrealEd/Private/Kismet2/EnumEditorUtils.cpp
--- a/Engine/Source/Editor/UnrealEd/Private/Kismet2/EnumEditorUtils.cpp
+++ b/Engine/Source/Editor/UnrealEd/Private/Kismet2/EnumEditorUtils.cpp
@@ -204,7 +204,7 @@ void FEnumEditorUtils::BroadcastChanges(const UUserDefinedEnum* Enum, const TArr
 			if (ByteProperty && (Enum == ByteProperty->GetIntPropertyEnum()))
 			{
 				UClass* OwnerClass = ByteProperty->GetOwnerClass();
-				if (ensure(OwnerClass))
+				if (OwnerClass)
 				{
 					ClassesToCheck.Add(OwnerClass);
 				}
