We now throttle cpu usage by ticking much slower when all windows are minimized. Added a setting (disabled by default) to also throttle when the application is not in the foreground (but could still be visible).

 Fixed Slate windows drawing when minimized

diff --git a/Engine/Config/BaseEditorUserSettings.ini b/Engine/Config/BaseEditorUserSettings.ini
--- a/Engine/Config/BaseEditorUserSettings.ini
+++ b/Engine/Config/BaseEditorUserSettings.ini
@@ -47,6 +47,9 @@ bDisplayPathsInLevelBrowser=False
 ; When enabled, the application frame rate, memory, Unreal object count, as well as a debug console will be displayed in the main editor UI
 bShowFrameRateAndMemory=False
 
+; Lowers CPU usage when the editor is in the background and not the active application
+bThrottleWhenNotForeground=False
+
 ; True if we should automatically reimport animset assets when a change to source content is detected
 bAutoReimportAnimSets=False
 
diff --git a/Engine/Source/Editor/UnrealEd/Classes/Editor/EditorEngine.h b/Engine/Source/Editor/UnrealEd/Classes/Editor/EditorEngine.h
--- a/Engine/Source/Editor/UnrealEd/Classes/Editor/EditorEngine.h
+++ b/Engine/Source/Editor/UnrealEd/Classes/Editor/EditorEngine.h
@@ -1583,6 +1583,13 @@ public:
 	 */
 	bool IsAnyViewportRealtime();
 
+	virtual bool ShouldThrottleCPUUsage() const OVERRIDE;
+
+	/**
+	 * @return true if all windows are hidden (including minimized)                                                         
+	 */
+	bool AreAllWindowsHidden() const;
+
 	/**
 	 *	Returns pointer to a temporary render target.
 	 *	If it has not already been created, does so here.
diff --git a/Engine/Source/Editor/UnrealEd/Classes/Editor/EditorUserSettings.h b/Engine/Source/Editor/UnrealEd/Classes/Editor/EditorUserSettings.h
--- a/Engine/Source/Editor/UnrealEd/Classes/Editor/EditorUserSettings.h
+++ b/Engine/Source/Editor/UnrealEd/Classes/Editor/EditorUserSettings.h
@@ -47,9 +47,13 @@ class UEditorUserSettings : public UObject
 	uint32 bDisplayActionListItemRefIds:1;
 	
 	/** When enabled, the application frame rate, memory and Unreal object count will be displayed in the main editor UI */
-	UPROPERTY(EditAnywhere, config, Category=DeveloperTools)
+	UPROPERTY(EditAnywhere, config, Category=Performance)
 	uint32 bShowFrameRateAndMemory:1;
 
+	/** Lowers CPU usage when the editor is in the background and not the active application */
+	UPROPERTY(EditAnywhere, config, Category=Performance, meta=(DisplayName="Use Less CPU when in Background") )
+	uint32 bThrottleWhenNotForeground:1;
+
 	/** Select to make Distributions use the curves, not the baked lookup tables. */
 	UPROPERTY(config)
 	uint32 bUseCurvesForDistributions:1; //(GDistributionType == 0)
diff --git a/Engine/Source/Editor/UnrealEd/Private/Editor.cpp b/Engine/Source/Editor/UnrealEd/Private/Editor.cpp
--- a/Engine/Source/Editor/UnrealEd/Private/Editor.cpp
+++ b/Engine/Source/Editor/UnrealEd/Private/Editor.cpp
@@ -1102,18 +1102,6 @@ void UEditorEngine::Tick( float DeltaSeconds, bool bIdleMode )
 				continue;
 			}
 
-			// HACK
-			/*
-			FSlatePlayInEditorInfo * SlateInfoPtr = SlatePlayInEditorMap.Find(PieContext.ContextHandle);
-			if ( SlateInfoPtr )
-			{
-				if (SlateInfoPtr->SlatePlayInEditorWindowViewport->IsForegroundWindow())
-				{
-					FSlateApplication::Get().SetKeyboardFocus( SlateInfoPtr->SlatePlayInEditorWindowViewport->GetViewportWidget().Pin(), EKeyboardFocusCause::SetDirectly );
-				}
-			}
-			*/
-
 			GPlayInEditorID = PieContext.PIEInstance;
 
 			PlayWorld = PieContext.World();
@@ -1299,43 +1287,48 @@ void UEditorEngine::Tick( float DeltaSeconds, bool bIdleMode )
 	/////////////////////////////
 	// Redraw viewports.
 
-	// Render view parents, then view children.
-	bool bEditorFrameNonRealtimeViewportDrawn = false;
-	if (GCurrentLevelEditingViewportClient && GCurrentLevelEditingViewportClient->IsVisible())
+	// Do not redraw if the application is hidden
+	bool bAllWindowsHidden = !bHasFocus && GEditor->AreAllWindowsHidden();
+	if( !bAllWindowsHidden )
 	{
-		bool bAllowNonRealtimeViewports = true;
-		bool bWasNonRealtimeViewportDraw = UpdateSingleViewportClient(GCurrentLevelEditingViewportClient, bAllowNonRealtimeViewports, bUpdateLinkedOrthoViewports );
-		if (GCurrentLevelEditingViewportClient->IsLevelEditorClient())
+		// Render view parents, then view children.
+		bool bEditorFrameNonRealtimeViewportDrawn = false;
+		if (GCurrentLevelEditingViewportClient && GCurrentLevelEditingViewportClient->IsVisible())
 		{
-			bEditorFrameNonRealtimeViewportDrawn |= bWasNonRealtimeViewportDraw;
+			bool bAllowNonRealtimeViewports = true;
+			bool bWasNonRealtimeViewportDraw = UpdateSingleViewportClient(GCurrentLevelEditingViewportClient, bAllowNonRealtimeViewports, bUpdateLinkedOrthoViewports);
+			if (GCurrentLevelEditingViewportClient->IsLevelEditorClient())
+			{
+				bEditorFrameNonRealtimeViewportDrawn |= bWasNonRealtimeViewportDraw;
+			}
 		}
-	}
-	for(int32 bRenderingChildren = 0;bRenderingChildren < 2;bRenderingChildren++)
-	{
-		for(int32 ViewportIndex=0; ViewportIndex<AllViewportClients.Num(); ViewportIndex++ )
+		for (int32 bRenderingChildren = 0; bRenderingChildren < 2; bRenderingChildren++)
 		{
-			FEditorViewportClient* ViewportClient = AllViewportClients[ViewportIndex];
-			if (ViewportClient == GCurrentLevelEditingViewportClient)
+			for (int32 ViewportIndex = 0; ViewportIndex < AllViewportClients.Num(); ViewportIndex++)
 			{
-				//already given this window a chance to update
-				continue;
-			}
+				FEditorViewportClient* ViewportClient = AllViewportClients[ViewportIndex];
+				if (ViewportClient == GCurrentLevelEditingViewportClient)
+				{
+					//already given this window a chance to update
+					continue;
+				}
 
-			if( !ViewportClient->IsLevelEditorClient() || ViewportClient->IsVisible() )
-			{
-				// Only update ortho viewports if that mode is turned on, the viewport client we are about to update is orthographic and the current editing viewport is orthographic and tracking mouse movement.
-				bUpdateLinkedOrthoViewports = GetDefault<ULevelEditorViewportSettings>()->bUseLinkedOrthographicViewports && ViewportClient->IsOrtho() && GCurrentLevelEditingViewportClient && GCurrentLevelEditingViewportClient->IsOrtho() && GCurrentLevelEditingViewportClient->IsTracking();
-				
-				const bool bIsViewParent = ViewportClient->ViewState.GetReference()->IsViewParent();
-				if(	(bRenderingChildren && !bIsViewParent) ||
-					(!bRenderingChildren && bIsViewParent) || bUpdateLinkedOrthoViewports )
+				if (!ViewportClient->IsLevelEditorClient() || ViewportClient->IsVisible())
 				{
-					//if we haven't drawn a non-realtime viewport OR not one of the main viewports
-					bool bAllowNonRealtimeViewports = (!bEditorFrameNonRealtimeViewportDrawn) || !(ViewportClient->IsLevelEditorClient());
-					bool bWasNonRealtimeViewportDrawn = UpdateSingleViewportClient(ViewportClient, bAllowNonRealtimeViewports, bUpdateLinkedOrthoViewports );
-					if (ViewportClient->IsLevelEditorClient())
+					// Only update ortho viewports if that mode is turned on, the viewport client we are about to update is orthographic and the current editing viewport is orthographic and tracking mouse movement.
+					bUpdateLinkedOrthoViewports = GetDefault<ULevelEditorViewportSettings>()->bUseLinkedOrthographicViewports && ViewportClient->IsOrtho() && GCurrentLevelEditingViewportClient && GCurrentLevelEditingViewportClient->IsOrtho() && GCurrentLevelEditingViewportClient->IsTracking();
+
+					const bool bIsViewParent = ViewportClient->ViewState.GetReference()->IsViewParent();
+					if ((bRenderingChildren && !bIsViewParent) ||
+						(!bRenderingChildren && bIsViewParent) || bUpdateLinkedOrthoViewports)
 					{
-						bEditorFrameNonRealtimeViewportDrawn |= bWasNonRealtimeViewportDrawn;
+						//if we haven't drawn a non-realtime viewport OR not one of the main viewports
+						bool bAllowNonRealtimeViewports = (!bEditorFrameNonRealtimeViewportDrawn) || !(ViewportClient->IsLevelEditorClient());
+						bool bWasNonRealtimeViewportDrawn = UpdateSingleViewportClient(ViewportClient, bAllowNonRealtimeViewports, bUpdateLinkedOrthoViewports);
+						if (ViewportClient->IsLevelEditorClient())
+						{
+							bEditorFrameNonRealtimeViewportDrawn |= bWasNonRealtimeViewportDrawn;
+						}
 					}
 				}
 			}
@@ -5918,7 +5911,43 @@ bool UEditorEngine::IsAnyViewportRealtime()
 	return false;
 }
 
+bool UEditorEngine::ShouldThrottleCPUUsage() const
+{
+	bool bShouldThrottle = false;
 
+	bool bIsForeground = FPlatformProcess::IsThisApplicationForeground();
+
+	if( !bIsForeground )
+	{
+		const UEditorUserSettings* Settings = GetDefault<UEditorUserSettings>();
+		bShouldThrottle = Settings->bThrottleWhenNotForeground;
+
+		// Check if we should throttle due to all windows being minimized
+		if ( !bShouldThrottle )
+		{
+			return bShouldThrottle = AreAllWindowsHidden();
+		}
+	}
+
+	return bShouldThrottle;
+}
+
+bool UEditorEngine::AreAllWindowsHidden() const
+{
+	const TArray< TSharedRef<SWindow> > AllWindows = FSlateApplication::Get().GetInteractiveTopLevelWindows();
+
+	bool bAllHidden = true;
+	for( const TSharedRef<SWindow>& Window : AllWindows )
+	{
+		if( !Window->IsWindowMinimized() && Window->IsVisible() )
+		{
+			bAllHidden = false;
+			break;
+		}
+	}
+
+	return bAllHidden;
+}
 
 AActor* UEditorEngine::AddActor(ULevel* InLevel, UClass* Class, const FVector& Location, bool bSilent, EObjectFlags ObjectFlags)
 {
diff --git a/Engine/Source/Editor/UnrealEd/Private/UnrealEdMisc.cpp b/Engine/Source/Editor/UnrealEd/Private/UnrealEdMisc.cpp
--- a/Engine/Source/Editor/UnrealEd/Private/UnrealEdMisc.cpp
+++ b/Engine/Source/Editor/UnrealEd/Private/UnrealEdMisc.cpp
@@ -1214,6 +1214,13 @@ void FUnrealEdMisc::TickPerformanceSurvey()
 		return;
 	}
 
+	// Don't run the test if we are throttling (due to minimized or not in foreground) as this will 
+	// greatly affect the framerate
+	if( GEditor->ShouldThrottleCPUUsage() )
+	{
+		return;
+	}
+
 	// Sample the frame rate until we have enough samples to take the average
 	if (FrameRateSamples.Num() < PerformanceSurveyDefs::NumFrameRateSamples)
 	{
diff --git a/Engine/Source/Runtime/Core/Private/GenericPlatform/GenericWindow.cpp b/Engine/Source/Runtime/Core/Private/GenericPlatform/GenericWindow.cpp
--- a/Engine/Source/Runtime/Core/Private/GenericPlatform/GenericWindow.cpp
+++ b/Engine/Source/Runtime/Core/Private/GenericPlatform/GenericWindow.cpp
@@ -89,6 +89,12 @@ bool FGenericWindow::IsMaximized() const
 	return true;
 }
 
+bool FGenericWindow::IsMinimized() const
+{
+	return false;
+}
+
+
 bool FGenericWindow::IsVisible() const
 {
 	// empty default functionality
diff --git a/Engine/Source/Runtime/Core/Private/Mac/MacWindow.cpp b/Engine/Source/Runtime/Core/Private/Mac/MacWindow.cpp
--- a/Engine/Source/Runtime/Core/Private/Mac/MacWindow.cpp
+++ b/Engine/Source/Runtime/Core/Private/Mac/MacWindow.cpp
@@ -1079,9 +1079,14 @@ bool FMacWindow::IsMaximized() const
 	return WindowHandle->bZoomed;
 }
 
+bool FMacWindow::IsMinimized() const
+{
+	return [WindowHandle isMiniaturized];
+}
+
 bool FMacWindow::IsVisible() const
 {
-	return bIsVisible;
+	return bIsVisible && [NSApp isHidden] == false;
 }
 
 bool FMacWindow::GetRestoredDimensions(int32& X, int32& Y, int32& Width, int32& Height)
diff --git a/Engine/Source/Runtime/Core/Private/Windows/WindowsWindow.cpp b/Engine/Source/Runtime/Core/Private/Windows/WindowsWindow.cpp
--- a/Engine/Source/Runtime/Core/Private/Windows/WindowsWindow.cpp
+++ b/Engine/Source/Runtime/Core/Private/Windows/WindowsWindow.cpp
@@ -523,6 +523,11 @@ bool FWindowsWindow::IsMaximized() const
 	return bIsMaximized;
 }
 
+bool FWindowsWindow::IsMinimized() const
+{
+	return !!::IsIconic(HWnd);
+}
+
 /** @return true if the native window is visible, false otherwise */
 bool FWindowsWindow::IsVisible() const
 {
diff --git a/Engine/Source/Runtime/Core/Public/GenericPlatform/GenericWindow.h b/Engine/Source/Runtime/Core/Public/GenericPlatform/GenericWindow.h
--- a/Engine/Source/Runtime/Core/Public/GenericPlatform/GenericWindow.h
+++ b/Engine/Source/Runtime/Core/Public/GenericPlatform/GenericWindow.h
@@ -89,6 +89,9 @@ public:
 	/** @return true if the native window is maximized, false otherwise */
 	virtual bool IsMaximized() const;
 
+	/** @return true if the native window is minimized, false otherwise */
+	virtual bool IsMinimized() const;
+
 	/** @return true if the native window is visible, false otherwise */
 	virtual bool IsVisible() const;
 
diff --git a/Engine/Source/Runtime/Core/Public/Mac/MacWindow.h b/Engine/Source/Runtime/Core/Public/Mac/MacWindow.h
--- a/Engine/Source/Runtime/Core/Public/Mac/MacWindow.h
+++ b/Engine/Source/Runtime/Core/Public/Mac/MacWindow.h
@@ -167,6 +167,8 @@ public:
 
 	virtual bool IsMaximized() const OVERRIDE;
 
+	virtual bool IsMinimized() const OVERRIDE;
+
 	virtual bool IsVisible() const OVERRIDE;
 
 	virtual bool GetRestoredDimensions(int32& X, int32& Y, int32& Width, int32& Height) OVERRIDE;
diff --git a/Engine/Source/Runtime/Core/Public/Windows/WindowsWindow.h b/Engine/Source/Runtime/Core/Public/Windows/WindowsWindow.h
--- a/Engine/Source/Runtime/Core/Public/Windows/WindowsWindow.h
+++ b/Engine/Source/Runtime/Core/Public/Windows/WindowsWindow.h
@@ -66,6 +66,8 @@ public:
 
 	virtual bool IsMaximized() const OVERRIDE;
 
+	virtual bool IsMinimized() const OVERRIDE;
+
 	virtual bool IsVisible() const OVERRIDE;
 
 	virtual bool GetRestoredDimensions(int32& X, int32& Y, int32& Width, int32& Height) OVERRIDE;
diff --git a/Engine/Source/Runtime/Engine/Classes/Engine/Engine.h b/Engine/Source/Runtime/Engine/Classes/Engine/Engine.h
--- a/Engine/Source/Runtime/Engine/Classes/Engine/Engine.h
+++ b/Engine/Source/Runtime/Engine/Classes/Engine/Engine.h
@@ -1734,6 +1734,10 @@ public:
 	 */
 	void SetAverageUnitTimes(float FrameTime, float RenderThreadTime, float GameThreadTime, float GPUFrameTime);
 
+	/**
+	 * @return true to throttle CPU usage based on current state (usually editor minimized or not in foreground)
+	 */
+	virtual bool ShouldThrottleCPUUsage() const;
 protected:
 	/** 
 	 * Determines whether a hardware survey should be run now.
diff --git a/Engine/Source/Runtime/Engine/Private/UnrealEngine.cpp b/Engine/Source/Runtime/Engine/Private/UnrealEngine.cpp
--- a/Engine/Source/Runtime/Engine/Private/UnrealEngine.cpp
+++ b/Engine/Source/Runtime/Engine/Private/UnrealEngine.cpp
@@ -6178,6 +6178,11 @@ void UEngine::SetAverageUnitTimes(float FrameTime, float RenderThreadTime, float
 	}
 }
 
+bool UEngine::ShouldThrottleCPUUsage() const
+{
+	return false;
+}
+
 /**
  *	Renders stats
  *
diff --git a/Engine/Source/Runtime/Launch/Private/LaunchEngineLoop.cpp b/Engine/Source/Runtime/Launch/Private/LaunchEngineLoop.cpp
--- a/Engine/Source/Runtime/Launch/Private/LaunchEngineLoop.cpp
+++ b/Engine/Source/Runtime/Launch/Private/LaunchEngineLoop.cpp
@@ -1971,10 +1971,10 @@ bool FEngineLoop::ShouldUseIdleMode() const
 	static const auto CVarIdleWhenNotForeground = IConsoleManager::Get().FindTConsoleVariableDataInt(TEXT("t.IdleWhenNotForeground"));
 	bool bIdleMode = false;
 
-	// Yield and prevent ticking and rendering when desired
-	if (FApp::IsGame() 
-		&& CVarIdleWhenNotForeground->GetValueOnGameThread()
+	// Yield cpu usage if desired
+	if (FApp::IsGame()
 		&& FPlatformProperties::SupportsWindowedMode()
+		&& CVarIdleWhenNotForeground->GetValueOnGameThread()
 		&& !FPlatformProcess::IsThisApplicationForeground())
 	{
 		bIdleMode = true;
@@ -2065,9 +2065,12 @@ void FEngineLoop::Tick()
 			FPlatformMisc::PumpMessages(true);
 		}
 
+		// Idle mode prevents ticking and rendering completely
 		const bool bIdleMode = ShouldUseIdleMode();
 
-		if (bIdleMode)
+		// Should we throttle CPU usage based on current state (usually in editor if we are minimized or not in foreground)
+		const bool bShouldThrottleCPU = GEngine->ShouldThrottleCPUUsage(); 
+		if (bIdleMode || bShouldThrottleCPU )
 		{
 			// Yield CPU time
 			FPlatformProcess::Sleep(.1f);
diff --git a/Engine/Source/Runtime/Slate/Private/Framework/Application/SlateApplication.cpp b/Engine/Source/Runtime/Slate/Private/Framework/Application/SlateApplication.cpp
--- a/Engine/Source/Runtime/Slate/Private/Framework/Application/SlateApplication.cpp
+++ b/Engine/Source/Runtime/Slate/Private/Framework/Application/SlateApplication.cpp
@@ -637,7 +637,7 @@ struct FDrawWindowArgs
 void FSlateApplication::DrawWindowAndChildren( const TSharedRef<SWindow>& WindowToDraw, FDrawWindowArgs& DrawWindowArgs )
 {
 	// Only draw visible windows
-	if( WindowToDraw->IsVisible() )
+	if( WindowToDraw->IsVisible() && !WindowToDraw->IsWindowMinimized() )
 	{
 		// Switch to the appropriate world for drawing
 		FScopedSwitchWorldHack SwitchWorld( WindowToDraw );
diff --git a/Engine/Source/Runtime/SlateCore/Private/Widgets/SWindow.cpp b/Engine/Source/Runtime/SlateCore/Private/Widgets/SWindow.cpp
--- a/Engine/Source/Runtime/SlateCore/Private/Widgets/SWindow.cpp
+++ b/Engine/Source/Runtime/SlateCore/Private/Widgets/SWindow.cpp
@@ -1039,6 +1039,12 @@ bool SWindow::IsWindowMaximized() const
 	return NativeWindow->IsMaximized();
 }
 
+bool SWindow::IsWindowMinimized() const
+{
+	return NativeWindow->IsMinimized();
+}
+
+
 /** Maximize the window if bInitiallyMaximized is set */
 void SWindow::InitialMaximize()
 {
diff --git a/Engine/Source/Runtime/SlateCore/Public/Widgets/SWindow.h b/Engine/Source/Runtime/SlateCore/Public/Widgets/SWindow.h
--- a/Engine/Source/Runtime/SlateCore/Public/Widgets/SWindow.h
+++ b/Engine/Source/Runtime/SlateCore/Public/Widgets/SWindow.h
@@ -425,6 +425,9 @@ public:
 	/** @return true if the window is maximized, false otherwise*/
 	bool IsWindowMaximized() const;
 
+	/** @return true of the window is minimized (iconic), false otherwise */
+	bool IsWindowMinimized() const;
+
 	/** Maximize the window if bInitiallyMaximized is set */
 	void InitialMaximize();
 
