Fix grammar in popup messages when launching games (was 'Building game for Launch on...', now 'Building game for launch...', for example. TTP 333743.

diff --git a/Engine/Source/Editor/UnrealEd/Private/PlayLevel.cpp b/Engine/Source/Editor/UnrealEd/Private/PlayLevel.cpp
--- a/Engine/Source/Editor/UnrealEd/Private/PlayLevel.cpp
+++ b/Engine/Source/Editor/UnrealEd/Private/PlayLevel.cpp
@@ -1087,7 +1087,16 @@ static void HandleStageStarted(const FString& InStage, FString DeviceName, TWeak
 	Arguments.Add(TEXT("StageName"), FText::FromString(InStage));
 	Arguments.Add(TEXT("DeviceName"), FText::FromString(DeviceName));
 
-	NotificationItemPtr.Pin()->SetText(FText::Format(LOCTEXT("LauncherTaskStageNotification", "{StageName} for Launch on {DeviceName}..."), Arguments));
+	FText NotificationText;
+	if(DeviceName.Len() == 0)
+	{
+		NotificationText = FText::Format(LOCTEXT("LauncherTaskStageNotificationNoDevice", "{StageName} for launch..."), Arguments);
+	}
+	else
+	{
+		NotificationText = FText::Format(LOCTEXT("LauncherTaskStageNotification", "{StageName} for launch on {DeviceName}..."), Arguments);
+	}
+	NotificationItemPtr.Pin()->SetText(NotificationText);
 }
 
 static void HandleLaunchCanceled(bool* bPlayUsingLauncher, TWeakPtr<SNotificationItem> NotificationItemPtr)
@@ -1207,9 +1216,14 @@ void UEditorEngine::PlayUsingLauncher()
 		ILauncherWorkerPtr LauncherWorker = Launcher->Launch(TargetDeviceServicesModule.GetDeviceProxyManager(), LauncherProfile);
 
 		// create notification item
-		FFormatNamedArguments Arguments;
-		Arguments.Add(TEXT("Device"), FText::FromString(PlayUsingLauncherDeviceName));
-		FNotificationInfo Info( FText::Format( LOCTEXT("LauncherTaskInProgressNotification", "Launching on {Device}..."), Arguments) );
+		FText LaunchingText = LOCTEXT("LauncherTaskInProgressNotificationNoDevice", "Launching...");
+		if(PlayUsingLauncherDeviceName.Len() > 0)
+		{
+			FFormatNamedArguments Arguments;
+			Arguments.Add(TEXT("Device"), FText::FromString(PlayUsingLauncherDeviceName));
+			LaunchingText = FText::Format( LOCTEXT("LauncherTaskInProgressNotification", "Launching on {Device}..."), Arguments);
+		}
+		FNotificationInfo Info(LaunchingText);
 
 		Info.Image = FEditorStyle::GetBrush(TEXT("MainFrame.CookContent"));
 		Info.bFireAndForget = false;
