Fix for TTP 334143, no longer crashes when importing a zero sized distance field shadowmap from Lightmass.

diff --git a/Engine/Source/Editor/UnrealEd/Private/Lightmass/Lightmass.cpp b/Engine/Source/Editor/UnrealEd/Private/Lightmass/Lightmass.cpp
--- a/Engine/Source/Editor/UnrealEd/Private/Lightmass/Lightmass.cpp
+++ b/Engine/Source/Editor/UnrealEd/Private/Lightmass/Lightmass.cpp
@@ -3815,7 +3815,7 @@ bool FLightmassProcessor::ImportSignedDistanceFieldShadowMapData2D(int32 Channel
 		FQuantizedShadowSignedDistanceFieldData2D* ShadowMapData = new FQuantizedShadowSignedDistanceFieldData2D(SMData.SizeX, SMData.SizeY);
 		check(ShadowMapData);
 
-		FQuantizedSignedDistanceFieldShadowSample* DataBuffer = &((*ShadowMapData)(0, 0));
+		FQuantizedSignedDistanceFieldShadowSample* DataBuffer = ShadowMapData->GetData();
 		uint32 DataBufferSize = SMData.SizeX * SMData.SizeY * sizeof(Lightmass::FQuantizedSignedDistanceFieldShadowSampleData);
 
 		uint32 CompressedSize = SMData.CompressedDataSize;
diff --git a/Engine/Source/Runtime/Engine/Public/ShadowMap.h b/Engine/Source/Runtime/Engine/Public/ShadowMap.h
--- a/Engine/Source/Runtime/Engine/Public/ShadowMap.h
+++ b/Engine/Source/Runtime/Engine/Public/ShadowMap.h
@@ -162,6 +162,8 @@ public:
 		Data.AddZeroed(SizeX * SizeY);
 	}
 
+	FQuantizedSignedDistanceFieldShadowSample* GetData() { return Data.GetData(); }
+
 	const FQuantizedSignedDistanceFieldShadowSample& operator()(uint32 X,uint32 Y) const { return Data[SizeX * Y + X]; }
 	FQuantizedSignedDistanceFieldShadowSample& operator()(uint32 X,uint32 Y) { return Data[SizeX * Y + X]; }
 
