Fixing a variable typo in Fbx material import
https://github.com/EpicGames/UnrealEngine/pull/165 ([REDACTED_PROFILE_URL])

diff --git a/Engine/Source/Editor/UnrealEd/Private/Fbx/FbxMaterialImport.cpp b/Engine/Source/Editor/UnrealEd/Private/Fbx/FbxMaterialImport.cpp
--- a/Engine/Source/Editor/UnrealEd/Private/Fbx/FbxMaterialImport.cpp
+++ b/Engine/Source/Editor/UnrealEd/Private/Fbx/FbxMaterialImport.cpp
@@ -58,7 +58,7 @@ UTexture* UnFbx::FFbxImporter::ImportTexture( FbxFileTexture* FbxTexture, bool b
 	BasePackageName = PackageTools::SanitizePackageName(BasePackageName);
 
 	UTexture* ExistingTexture = NULL;
-	UPackage* TexturePacakge = NULL;
+	UPackage* TexturePackage = NULL;
 	// First check if the asset already exists.
 	{
 		FString ObjectPath = BasePackageName + TEXT(".") + TextureName;
@@ -74,11 +74,11 @@ UTexture* UnFbx::FFbxImporter::ImportTexture( FbxFileTexture* FbxTexture, bool b
 		FString FinalPackageName;
 		AssetToolsModule.Get().CreateUniqueAssetName(BasePackageName, Suffix, FinalPackageName, TextureName);
 
-		TexturePacakge = CreatePackage(NULL, *FinalPackageName);
+		TexturePackage = CreatePackage(NULL, *FinalPackageName);
 	}
 	else
 	{
-		TexturePacakge = ExistingTexture->GetOutermost();
+		TexturePackage = ExistingTexture->GetOutermost();
 	}
 
 
@@ -130,7 +130,7 @@ UTexture* UnFbx::FFbxImporter::ImportTexture( FbxFileTexture* FbxTexture, bool b
 		}
 
 		UnrealTexture = (UTexture*)TextureFact->FactoryCreateBinary(
-			UTexture2D::StaticClass(), TexturePacakge, *TextureName, 
+			UTexture2D::StaticClass(), TexturePackage, *TextureName, 
 			RF_Standalone|RF_Public, NULL, TextureType, 
 			PtrTexture, PtrTexture+DataBinary.Num(), GWarn );
 
@@ -140,7 +140,7 @@ UTexture* UnFbx::FFbxImporter::ImportTexture( FbxFileTexture* FbxTexture, bool b
 			FAssetRegistryModule::AssetCreated(UnrealTexture);
 
 			// Set the dirty flag so this package will get saved later
-			TexturePacakge->SetDirtyFlag(true);
+			TexturePackage->SetDirtyFlag(true);
 		}
 		TextureFact->RemoveFromRoot();
 	}
