UE4: Fixed several issues related to loading an already-loaded world.
Loading the same world that is already loaded directly now early-outs.
Loading a sublevel of the currently edited world now keeps the sublevel in memory, patches the owning world, and properly GCs the old world.
Redirectors are now unloaded during editor cleanse (map transitions) to prevent references to worlds from preventing GC.
Already-loaded worlds are no longer unloaded when the editor switches to begin editing them. Switching away from that world, however, does unload it naturally.

diff --git a/Engine/Source/Developer/AssetTools/Private/AssetTypeActions/AssetTypeActions_World.cpp b/Engine/Source/Developer/AssetTools/Private/AssetTypeActions/AssetTypeActions_World.cpp
--- a/Engine/Source/Developer/AssetTools/Private/AssetTypeActions/AssetTypeActions_World.cpp
+++ b/Engine/Source/Developer/AssetTools/Private/AssetTypeActions/AssetTypeActions_World.cpp
@@ -9,7 +9,7 @@ void FAssetTypeActions_World::OpenAssetEditor( const TArray<UObject*>& InObjects
 	for (auto ObjIt = InObjects.CreateConstIterator(); ObjIt; ++ObjIt)
 	{
 		auto World = Cast<UWorld>(*ObjIt);
-		if (World != NULL)
+		if (World != NULL && World != GEditor->GetEditorWorldContext().World())
 		{
 			const FString FileToOpen = FPackageName::LongPackageNameToFilename(World->GetOutermost()->GetName(), FPackageName::GetMapPackageExtension());
 			const bool bLoadAsTemplate = false;
diff --git a/Engine/Source/Editor/UnrealEd/Classes/Editor/EditorEngine.h b/Engine/Source/Editor/UnrealEd/Classes/Editor/EditorEngine.h
--- a/Engine/Source/Editor/UnrealEd/Classes/Editor/EditorEngine.h
+++ b/Engine/Source/Editor/UnrealEd/Classes/Editor/EditorEngine.h
@@ -2405,8 +2405,10 @@ private:
 	 *
 	 * @param	Context		The world to destroy
 	 * @param	CleanseText	Reason for the destruction
+	 * @param	NewWorld	An optional new world to keep in memory after destroying the world referenced in the Context.
+	 *						This world and it's sublevels will remain in memory.
 	 */
-	void EditorDestroyWorld( FWorldContext & Context, const FText& CleanseText );
+	void EditorDestroyWorld( FWorldContext & Context, const FText& CleanseText, UWorld* NewWorld = nullptr );
 
 	/** Returns the GameViewport widget */
 	virtual TSharedPtr<SViewport> GetGameViewportWidget() const OVERRIDE;
diff --git a/Engine/Source/Editor/UnrealEd/Private/Editor.cpp b/Engine/Source/Editor/UnrealEd/Private/Editor.cpp
--- a/Engine/Source/Editor/UnrealEd/Private/Editor.cpp
+++ b/Engine/Source/Editor/UnrealEd/Private/Editor.cpp
@@ -1678,8 +1678,23 @@ void UEditorEngine::Cleanse( bool ClearSelection, bool Redraw, const FText& Tran
 			RedrawLevelEditingViewports();
 		}
 
+		// Attempt to unload any loaded redirectors. Redirectors should not be referenced in memory and are only used to forward references at load time
+		for (TObjectIterator<UObjectRedirector> RedirIt; RedirIt; ++RedirIt)
+		{
+			RedirIt->ClearFlags(RF_Standalone | RF_RootSet | RF_Transactional);
+		}
+
 		// Collect garbage.
 		CollectGarbage( GARBAGE_COLLECTION_KEEPFLAGS );
+
+		// Remaining redirectors are probably referenced by editor tools. Keep them in memory for now.
+		for (TObjectIterator<UObjectRedirector> RedirIt; RedirIt; ++RedirIt)
+		{
+			if ( RedirIt->IsAsset() )
+			{
+				RedirIt->SetFlags(RF_Standalone);
+			}
+		}
 	}
 }
 
diff --git a/Engine/Source/Editor/UnrealEd/Private/EditorServer.cpp b/Engine/Source/Editor/UnrealEd/Private/EditorServer.cpp
--- a/Engine/Source/Editor/UnrealEd/Private/EditorServer.cpp
+++ b/Engine/Source/Editor/UnrealEd/Private/EditorServer.cpp
@@ -1581,7 +1581,7 @@ void UEditorEngine::BSPIntersectionHelper(UWorld* InWorld, ECsgOper Operation)
 	bspBrushCSG( InWorld->GetBrush(), InWorld->GetModel(), 0, Brush_MAX, Operation, false, true, true );
 }
 
-void UEditorEngine::EditorDestroyWorld( FWorldContext & Context, const FText& CleanseText )
+void UEditorEngine::EditorDestroyWorld( FWorldContext & Context, const FText& CleanseText, UWorld* NewWorld )
 {
 	if (Context.World() == NULL )
 	{
@@ -1601,7 +1601,7 @@ void UEditorEngine::EditorDestroyWorld( FWorldContext & Context, const FText& Cl
 		Context.World()->ClearFlags(RF_Standalone | RF_RootSet | RF_Transactional);
 		Context.World()->SetFlags(RF_Transient);
 	}
-	
+
 	GUnrealEd->CurrentLODParentActor = NULL;
 	SelectNone( true, true );
 	EditorClearComponents();
@@ -1613,22 +1613,45 @@ void UEditorEngine::EditorDestroyWorld( FWorldContext & Context, const FText& Cl
 	CloseEditedWorldAssets(Context.World());
 
 	// Stop all audio and remove references 
-	if (GetAudioDevice() )
+	if ( GetAudioDevice() )
 	{
-		GetAudioDevice()->Flush( NULL );
+		GetAudioDevice()->Flush(NULL);
 	}
 
-	Context.World()->DestroyWorld( true );
+	UWorld* ContextWorld = Context.World();
+	Context.World()->DestroyWorld( true, NewWorld );
 	Context.SetCurrentWorld(NULL);
-	
+
+	// Add the new world to root if it wasn't already and keep track of it so we can remove it from root later if appropriate
+	bool bNewWorldAddedToRoot = false;
+	if ( NewWorld )
+	{
+		if ( !NewWorld->IsRooted() )
+		{
+			NewWorld->AddToRoot();
+			bNewWorldAddedToRoot = true;
+		}
+
+		// Reset the owning level to allow the old world to GC if it was a sublevel
+		NewWorld->PersistentLevel->OwningWorld = NewWorld;
+	}
+
 	// Cleanse which should remove the old world which we are going to verify.
 	GEditor->Cleanse( true, 0, CleanseText );
 
+	// If we added the world to the root set above, remove it now that the GC is complete.
+	if ( bNewWorldAddedToRoot )
+	{
+		NewWorld->RemoveFromRoot();
+	}
+
+	// Make sure the old world is completely gone, except if the new world was one of it's sublevels
 	for( TObjectIterator<UWorld> It; It; ++It )
 	{
 		UWorld* RemainingWorld = *It;
+		const bool bIsNewWorld = (RemainingWorld == NewWorld);
 		const bool bIsPersistantWorldType = (RemainingWorld->WorldType == EWorldType::Inactive) || (RemainingWorld->WorldType == EWorldType::Preview);
-		if (!bIsPersistantWorldType && !WorldHasValidContext(RemainingWorld))
+		if (!bIsNewWorld && !bIsPersistantWorldType && !WorldHasValidContext(RemainingWorld))
 		{
 			StaticExec(RemainingWorld, *FString::Printf(TEXT("OBJ REFS CLASS=WORLD NAME=%s"), *RemainingWorld->GetPathName()));
 
@@ -1644,7 +1667,8 @@ void UEditorEngine::EditorDestroyWorld( FWorldContext & Context, const FText& Cl
 		for( TObjectIterator<UPackage> It; It; ++It )
 		{
 			UPackage* RemainingPackage = *It;
-			if (RemainingPackage == WorldPackage)
+			const bool bIsNewWorldPackage = (RemainingPackage == NewWorld->GetOutermost());
+			if (!bIsNewWorldPackage && RemainingPackage == WorldPackage)
 			{
 				StaticExec(NULL, *FString::Printf(TEXT("OBJ REFS CLASS=PACKAGE NAME=%s"), *RemainingPackage->GetPathName()));
 
@@ -1897,6 +1921,12 @@ bool UEditorEngine::Map_Load(const TCHAR* Str, FOutputDevice& Ar)
 		FString LongTempFname;
 		if ( FPackageName::TryConvertFilenameToLongPackageName(TempFname, LongTempFname) )
 		{
+ 			if ( Context.World() && Context.World()->GetOutermost() && Context.World()->GetOutermost() == FindPackage(NULL, *LongTempFname) )
+ 			{
+ 				// This map is already loaded.
+ 				return true;
+ 			}
+
 			FString AlteredPath;
 			if ( FPackageName::DoesPackageExist(LongTempFname, NULL, &AlteredPath) )
 			{
@@ -1918,6 +1948,14 @@ bool UEditorEngine::Map_Load(const TCHAR* Str, FOutputDevice& Ar)
 				// Should we display progress while loading?
 				int32 bShowProgress = 1;
 				FParse::Value(Str, TEXT("SHOWPROGRESS="), bShowProgress);
+
+				// Is the new world already loaded?
+				UPackage* ExistingPackage = FindPackage(NULL, *LongTempFname);
+				UWorld* ExistingWorld = NULL;
+				if ( ExistingPackage )
+				{
+					ExistingWorld = UWorld::FindWorldInPackage(ExistingPackage);
+				}
 			
 				UObject* OldOuter = NULL;
 
@@ -1962,7 +2000,7 @@ bool UEditorEngine::Map_Load(const TCHAR* Str, FOutputDevice& Ar)
 						Arguments.Add(TEXT("MapFileName"), FText::FromString( MapFileName ));
 						FMessageLog("LoadErrors").NewPage( FText::Format( LOCTEXT("LoadMapLogPage", "Loading map: {MapFileName}"), Arguments ) );
 					}
-					EditorDestroyWorld( Context, LocalizedLoadingMap );
+					EditorDestroyWorld( Context, LocalizedLoadingMap, ExistingWorld );
 
 					GWarn->StatusUpdate( -1, -1, LocalizedLoadingMap );
 				}
@@ -2002,8 +2040,15 @@ bool UEditorEngine::Map_Load(const TCHAR* Str, FOutputDevice& Ar)
 				}
 				else
 				{
-					//Load the map normally into a new package
-					WorldPackage = LoadPackage( NULL, *LongTempFname, LoadFlags );
+					if ( ExistingPackage )
+					{
+						WorldPackage = ExistingPackage;
+					}
+					else
+					{
+						//Load the map normally into a new package
+						WorldPackage = LoadPackage( NULL, *LongTempFname, LoadFlags );
+					}
 				}
 
 				if (WorldPackage == NULL)
@@ -2034,10 +2079,14 @@ bool UEditorEngine::Map_Load(const TCHAR* Str, FOutputDevice& Ar)
 
 				Context.World()->WorldType = EWorldType::Editor;
 
-				Context.World()->InitWorld(UWorld::InitializationValues()
-											.ShouldSimulatePhysics(false)
-											.EnableTraceCollision(true)
-											);
+				if ( !Context.World()->bIsWorldInitialized )
+				{
+					Context.World()->InitWorld(UWorld::InitializationValues()
+												.ShouldSimulatePhysics(false)
+												.EnableTraceCollision(true)
+												);
+				}
+
 				{
 					FBSPOps::bspValidateBrush( Context.World()->GetBrush()->Brush, 0, 1 );
 
diff --git a/Engine/Source/Runtime/Engine/Classes/Engine/World.h b/Engine/Source/Runtime/Engine/Classes/Engine/World.h
--- a/Engine/Source/Runtime/Engine/Classes/Engine/World.h
+++ b/Engine/Source/Runtime/Engine/Classes/Engine/World.h
@@ -1650,8 +1650,9 @@ public:
 	/**
 	 * Cleans up components, streaming data and assorted other intermediate data.
 	 * @param bSessionEnded whether to notify the viewport that the game session has ended
+	 * @param NewWorld Optional new world that will be loaded after this world is cleaned up. Specify a new world to prevent it and it's sublevels from being GCed during map transitions.
 	 */
-	void CleanupWorld(bool bSessionEnded = true, bool bCleanupResources = true);
+	void CleanupWorld(bool bSessionEnded = true, bool bCleanupResources = true, UWorld* NewWorld = nullptr);
 	
 	/**
 	 * Invalidates the cached data used to render the levels' UModel.
@@ -1824,9 +1825,9 @@ public:
 	static UWorld* CreateWorld( const EWorldType::Type InWorldType, bool bInformEngineOfWorld, FName WorldName = NAME_None, UPackage* InWorldPackage = NULL, bool bAddToRoot = true );
 
 	/** 
-	 * Destroy this World instance 
+	 * Destroy this World instance. If destroying the world to load a different world, supply it here to prevent GC of the new world or it's sublevels.
 	 */
-	void DestroyWorld( bool bInformEngineOfWorld );
+	void DestroyWorld( bool bInformEngineOfWorld, UWorld* NewWorld = nullptr );
 
 	/**
 	 *  Interface to allow WorldSettings to request immediate garbage collection
diff --git a/Engine/Source/Runtime/Engine/Private/World.cpp b/Engine/Source/Runtime/Engine/Private/World.cpp
--- a/Engine/Source/Runtime/Engine/Private/World.cpp
+++ b/Engine/Source/Runtime/Engine/Private/World.cpp
@@ -890,13 +890,13 @@ void UWorld::InitializeNewWorld(const InitializationValues IVS)
 }
 
 
-void UWorld::DestroyWorld( bool bInformEngineOfWorld )
+void UWorld::DestroyWorld( bool bInformEngineOfWorld, UWorld* NewWorld )
 {
 	// Clean up existing world and remove it from root set so it can be garbage collected.
 	bIsLevelStreamingFrozen = false;
 	bShouldForceUnloadStreamingLevels = true;
 	FlushLevelStreaming( NULL, true );
-	CleanupWorld();
+	CleanupWorld(true, true, NewWorld);
 
 	check( NetworkActors.Num() == 0 );
 
@@ -911,7 +911,7 @@ void UWorld::DestroyWorld( bool bInformEngineOfWorld )
 	for (int32 LevelIndex=0; LevelIndex < GetNumLevels(); ++LevelIndex)
 	{
 		UWorld* World = CastChecked<UWorld>(GetLevel(LevelIndex)->GetOuter());
-		if (World != this)
+		if (World != this && World != NewWorld)
 		{
 			World->ClearFlags(RF_Standalone);
 		}
@@ -2730,7 +2730,7 @@ bool UWorld::IsNavigationRebuilt() const
 	return GetNavigationSystem() == NULL || GetNavigationSystem()->IsNavigationBuilt(GetWorldSettings());
 }
 
-void UWorld::CleanupWorld(bool bSessionEnded, bool bCleanupResources)
+void UWorld::CleanupWorld(bool bSessionEnded, bool bCleanupResources, UWorld* NewWorld)
 {
 	check(IsVisibilityRequestPending() == false);
 
@@ -2777,7 +2777,7 @@ void UWorld::CleanupWorld(bool bSessionEnded, bool bCleanupResources)
 
 	// Clear standalone flag when switching maps in the Editor. This causes resources placed in the map
 	// package to be garbage collected together with the world.
-	if( GIsEditor && !IsTemplate() )
+	if( GIsEditor && !IsTemplate() && this != NewWorld )
 	{
 		TArray<UObject*> WorldObjects;
 
@@ -2798,7 +2798,7 @@ void UWorld::CleanupWorld(bool bSessionEnded, bool bCleanupResources)
 		UWorld* World = CastChecked<UWorld>(GetLevel(LevelIndex)->GetOuter());
 		if (World != this)
 		{
-			World->CleanupWorld(bSessionEnded, bCleanupResources);
+			World->CleanupWorld(bSessionEnded, bCleanupResources, NewWorld);
 		}
 	}
 }
