UE4: Removed some special case code for umap files. They can now be treated as normal assets

diff --git a/Engine/Source/Editor/UnrealEd/Private/Toolkits/AssetEditorManager.cpp b/Engine/Source/Editor/UnrealEd/Private/Toolkits/AssetEditorManager.cpp
--- a/Engine/Source/Editor/UnrealEd/Private/Toolkits/AssetEditorManager.cpp
+++ b/Engine/Source/Editor/UnrealEd/Private/Toolkits/AssetEditorManager.cpp
@@ -314,32 +314,19 @@ void FAssetEditorManager::HandleRequestOpenAssetMessage( const FAssetEditorReque
 
 void FAssetEditorManager::OpenEditorForAsset(const FString& AssetPathName)
 {
-	FString FileType = FPaths::GetExtension(AssetPathName, /*bIncludeDot=*/true);
+	// An asset needs loading
+	UPackage* Package = LoadPackage(NULL, *AssetPathName, LOAD_NoRedirects);
 
-	// Check if a map needs loading
-	if (FileType.ToLower() == FPackageName::GetMapPackageExtension().ToLower())
+	if (Package)
 	{
-		bool bLoadAsTemplate = false;
-		bool bShowProgress = false;
+		Package->FullyLoad();
 
-		FEditorFileUtils::LoadMap(AssetPathName, bLoadAsTemplate, bShowProgress);
-	}
-	else
-	{
-		// An asset needs loading
-		UPackage* Package = LoadPackage(NULL, *AssetPathName, LOAD_NoRedirects);
+		FString AssetName = FPaths::GetBaseFilename(AssetPathName);
+		UObject* Object = FindObject<UObject>(Package, *AssetName);
 
-		if (Package)
+		if (Object != NULL)
 		{
-			Package->FullyLoad();
-
-			FString AssetName = FPaths::GetBaseFilename(AssetPathName);
-			UObject* Object = FindObject<UObject>(Package, *AssetName);
-
-			if (Object != NULL)
-			{
-				OpenEditorForAsset(Object);
-			}
+			OpenEditorForAsset(Object);
 		}
 	}
 }
