UE4: Fixed a crash when NewWorld is NULL in EditorDestroyWorld.

diff --git a/Engine/Source/Editor/UnrealEd/Private/EditorServer.cpp b/Engine/Source/Editor/UnrealEd/Private/EditorServer.cpp
--- a/Engine/Source/Editor/UnrealEd/Private/EditorServer.cpp
+++ b/Engine/Source/Editor/UnrealEd/Private/EditorServer.cpp
@@ -1649,7 +1649,7 @@ void UEditorEngine::EditorDestroyWorld( FWorldContext & Context, const FText& Cl
 	for( TObjectIterator<UWorld> It; It; ++It )
 	{
 		UWorld* RemainingWorld = *It;
-		const bool bIsNewWorld = (RemainingWorld == NewWorld);
+		const bool bIsNewWorld = (NewWorld && RemainingWorld == NewWorld);
 		const bool bIsPersistantWorldType = (RemainingWorld->WorldType == EWorldType::Inactive) || (RemainingWorld->WorldType == EWorldType::Preview);
 		if (!bIsNewWorld && !bIsPersistantWorldType && !WorldHasValidContext(RemainingWorld))
 		{
@@ -1664,10 +1664,11 @@ void UEditorEngine::EditorDestroyWorld( FWorldContext & Context, const FText& Cl
 
 	if (WorldPackage != NULL)
 	{
+		UPackage* NewWorldPackage = NewWorld ? NewWorld->GetOutermost() : nullptr;
 		for( TObjectIterator<UPackage> It; It; ++It )
 		{
 			UPackage* RemainingPackage = *It;
-			const bool bIsNewWorldPackage = (RemainingPackage == NewWorld->GetOutermost());
+			const bool bIsNewWorldPackage = (NewWorldPackage && RemainingPackage == NewWorldPackage);
 			if (!bIsNewWorldPackage && RemainingPackage == WorldPackage)
 			{
 				StaticExec(NULL, *FString::Printf(TEXT("OBJ REFS CLASS=PACKAGE NAME=%s"), *RemainingPackage->GetPathName()));
