bp: user defined structs: we need the ability to expose struct "transforms" and "vectors" as widgets in the editor ttp#334937
Structure Editor Utils refactored.

diff --git a/Engine/Source/Editor/Kismet/Private/UserDefinedStructureEditor.cpp b/Engine/Source/Editor/Kismet/Private/UserDefinedStructureEditor.cpp
--- a/Engine/Source/Editor/Kismet/Private/UserDefinedStructureEditor.cpp
+++ b/Engine/Source/Editor/Kismet/Private/UserDefinedStructureEditor.cpp
@@ -405,6 +405,37 @@ public:
 		}
 	}
 
+	// 3d widget
+	EVisibility Is3dWidgetOptionVisible() const
+	{
+		auto StructureDetailsSP = StructureDetails.Pin();
+		if (StructureDetailsSP.IsValid())
+		{
+			return FStructureEditorUtils::CanEnable3dWidget(StructureDetailsSP->GetUserDefinedStruct(), FieldGuid) ? EVisibility::Visible : EVisibility::Collapsed;
+		}
+		return EVisibility::Collapsed;
+	}
+
+	ESlateCheckBoxState::Type OnGet3dWidgetEnabled() const
+	{
+		auto StructureDetailsSP = StructureDetails.Pin();
+		if (StructureDetailsSP.IsValid())
+		{
+			return FStructureEditorUtils::Is3dWidgetEnabled(StructureDetailsSP->GetUserDefinedStruct(), FieldGuid) ? ESlateCheckBoxState::Checked : ESlateCheckBoxState::Unchecked;
+		}
+		return ESlateCheckBoxState::Undetermined;
+	}
+
+	void On3dWidgetEnabledCommitted(ESlateCheckBoxState::Type InNewState)
+	{
+		auto StructureDetailsSP = StructureDetails.Pin();
+		if (StructureDetailsSP.IsValid() && (ESlateCheckBoxState::Undetermined != InNewState))
+		{
+			FStructureEditorUtils::Change3dWidgetEnabled(StructureDetailsSP->GetUserDefinedStruct(), FieldGuid, ESlateCheckBoxState::Checked == InNewState);
+			OnChanged();
+		}
+	}
+
 	void OnArgDefaultValueCommitted(const FText& NewText, ETextCommit::Type InTextCommit)
 	{
 		auto StructureDetailsSP = StructureDetails.Pin();
@@ -589,6 +620,24 @@ public:
 			.OnCheckStateChanged(this, &FUserDefinedStructureFieldLayout::OnEditableOnBPInstanceCommitted)
 			.IsChecked(this, &FUserDefinedStructureFieldLayout::OnGetEditableOnBPInstanceState)
 		];
+
+		ChildrenBuilder.AddChildContent(*LOCTEXT("3dWidget", "3d Widget").ToString())
+		.NameContent()
+		[
+			SNew(STextBlock)
+			.Text(LOCTEXT("3dWidget", "3d Widget"))
+			.Font(IDetailLayoutBuilder::GetDetailFont())
+		]
+		.ValueContent()
+		[
+			SNew(SCheckBox)
+			.OnCheckStateChanged(this, &FUserDefinedStructureFieldLayout::On3dWidgetEnabledCommitted)
+			.IsChecked(this, &FUserDefinedStructureFieldLayout::OnGet3dWidgetEnabled)
+		]
+		.Visibility(
+			TAttribute<EVisibility>::Create(
+				TAttribute<EVisibility>::FGetter::CreateSP(
+					this, &FUserDefinedStructureFieldLayout::Is3dWidgetOptionVisible)));
 	}
 
 	virtual void Tick( float DeltaTime ) OVERRIDE {}
diff --git a/Engine/Source/Editor/KismetCompiler/Private/UserDefinedStructureCompilerUtils.cpp b/Engine/Source/Editor/KismetCompiler/Private/UserDefinedStructureCompilerUtils.cpp
--- a/Engine/Source/Editor/KismetCompiler/Private/UserDefinedStructureCompilerUtils.cpp
+++ b/Engine/Source/Editor/KismetCompiler/Private/UserDefinedStructureCompilerUtils.cpp
@@ -5,6 +5,7 @@
 #include "KismetCompiler.h"
 #include "Editor/UnrealEd/Public/Kismet2/StructureEditorUtils.h"
 #include "Editor/UnrealEd/Public/Kismet2/BlueprintEditorUtils.h"
+#include "Editor/UnrealEd/Public/EditorModes.h"
 
 #define LOCTEXT_NAMESPACE "StructureCompiler"
 
@@ -150,6 +151,10 @@ struct FUserDefinedStructureCompilerInner
 			{
 				NewProperty->SetPropertyFlags(CPF_DisableEditOnInstance);
 			}
+			if (VarDesc.bEnable3dWidget)
+			{
+				NewProperty->SetMetaData(FEdMode::MD_MakeEditWidget, TEXT("true"));
+			}
 			NewProperty->SetMetaData(TEXT("DisplayName"), *VarDesc.FriendlyName);
 			NewProperty->SetMetaData(TEXT("Category"), *VarDesc.Category);
 			NewProperty->SetMetaData(FBlueprintMetadata::MD_Tooltip, *VarDesc.ToolTip);
diff --git a/Engine/Source/Editor/UnrealEd/Classes/UserDefinedStructure/UserDefinedStructEditorData.h b/Engine/Source/Editor/UnrealEd/Classes/UserDefinedStructure/UserDefinedStructEditorData.h
--- a/Engine/Source/Editor/UnrealEd/Classes/UserDefinedStructure/UserDefinedStructEditorData.h
+++ b/Engine/Source/Editor/UnrealEd/Classes/UserDefinedStructure/UserDefinedStructEditorData.h
@@ -47,6 +47,9 @@ struct FStructVariableDescription
 	UPROPERTY()
 	bool bDontEditoOnInstance;
 
+	UPROPERTY()
+	bool bEnable3dWidget;
+
 	UNREALED_API bool SetPinType(const struct FEdGraphPinType& VarType);
 
 	UNREALED_API FEdGraphPinType ToPinType() const;
@@ -55,6 +58,7 @@ struct FStructVariableDescription
 		: bIsArray(false)
 		, bInvalidMember(false)
 		, bDontEditoOnInstance(false)
+		, bEnable3dWidget(false)
 	{ }
 };
 
diff --git a/Engine/Source/Editor/UnrealEd/Private/EditorModes.cpp b/Engine/Source/Editor/UnrealEd/Private/EditorModes.cpp
--- a/Engine/Source/Editor/UnrealEd/Private/EditorModes.cpp
+++ b/Engine/Source/Editor/UnrealEd/Private/EditorModes.cpp
@@ -970,7 +970,7 @@ void FEdMode::DrawHUD(FLevelEditorViewportClient* ViewportClient,FViewport* View
 
 				//do some string fixing
 				const uint32 VectorIndex = WidgetInfos[i].PropertyIndex;
-				FString WidgetDisplayName = WidgetName + ((VectorIndex != INDEX_NONE) ? FString::Printf(TEXT("[%d]"), VectorIndex) : TEXT(""));
+				const FString WidgetDisplayName = WidgetInfos[i].DisplayName + ((VectorIndex != INDEX_NONE) ? FString::Printf(TEXT("[%d]"), VectorIndex) : TEXT(""));
 				FinalString = FinalString.IsEmpty() ? WidgetDisplayName : FinalString;
 
 				if(Proj.W > 0.f)
@@ -1066,6 +1066,11 @@ AActor* FEdMode::GetFirstSelectedActorInstance() const
 	return NULL;
 }
 
+bool FEdMode::CanCreateWidgetForStructure(const UStruct* InPropStruct)
+{
+	return InPropStruct && (InPropStruct->GetFName() == NAME_Vector || InPropStruct->GetFName() == NAME_Transform);
+}
+
 bool FEdMode::CanCreateWidgetForProperty(UProperty* InProp)
 {
 	UStructProperty* TestProperty = Cast<UStructProperty>(InProp);
@@ -1077,7 +1082,7 @@ bool FEdMode::CanCreateWidgetForProperty(UProperty* InProp)
 			TestProperty = Cast<UStructProperty>(ArrayProperty->Inner);
 		}
 	}
-	return (TestProperty != NULL) && (TestProperty->Struct->GetFName() == NAME_Vector || TestProperty->Struct->GetFName() == NAME_Transform);
+	return (TestProperty != NULL) && CanCreateWidgetForStructure(TestProperty->Struct);
 }
 
 bool FEdMode::ShouldCreateWidgetForProperty(UProperty* InProp)
@@ -1104,6 +1109,7 @@ void FEdMode::GetPropertyWidgetInfos(const UStruct* InStruct, const void* InCont
 		UProperty* CurrentProp = *PropertyIt;
 		if(	ShouldCreateWidgetForProperty(CurrentProp) )
 		{
+			const FString DisplayName = CurrentProp->GetMetaData(TEXT("DisplayName"));
 			if( UArrayProperty* ArrayProp = Cast<UArrayProperty>(CurrentProp) )
 			{
 				check(InContainer != NULL);
@@ -1119,6 +1125,7 @@ void FEdMode::GetPropertyWidgetInfos(const UStruct* InStruct, const void* InCont
 
 					//fill it in with the struct name
 					WidgetInfo.PropertyName = PropertyNamePrefix + CurrentProp->GetFName().ToString();
+					WidgetInfo.DisplayName = DisplayName.IsEmpty() ? WidgetInfo.PropertyName : (PropertyNamePrefix + DisplayName);
 
 					//And see if we have any meta data that matches the MD_ValidateWidgetUsing name
 					WidgetInfo.PropertyValidationName = FName(*CurrentProp->GetMetaData(MD_ValidateWidgetUsing));
@@ -1140,6 +1147,7 @@ void FEdMode::GetPropertyWidgetInfos(const UStruct* InStruct, const void* InCont
 
 				//fill it in with the struct name
 				WidgetInfo.PropertyName = PropertyNamePrefix + CurrentProp->GetFName().ToString();
+				WidgetInfo.DisplayName = DisplayName.IsEmpty() ? WidgetInfo.PropertyName : (PropertyNamePrefix + DisplayName);
 
 				//And see if we have any meta data that matches the MD_ValidateWidgetUsing name
 				WidgetInfo.PropertyValidationName = FName(*CurrentProp->GetMetaData(MD_ValidateWidgetUsing));
diff --git a/Engine/Source/Editor/UnrealEd/Private/Kismet2/StructureEditorUtils.cpp b/Engine/Source/Editor/UnrealEd/Private/Kismet2/StructureEditorUtils.cpp
--- a/Engine/Source/Editor/UnrealEd/Private/Kismet2/StructureEditorUtils.cpp
+++ b/Engine/Source/Editor/UnrealEd/Private/Kismet2/StructureEditorUtils.cpp
@@ -8,12 +8,13 @@
 #include "EdGraphSchema_K2.h"
 #include "ObjectTools.h"
 #include "Editor/UnrealEd/Public/Kismet2/CompilerResultsLog.h"
+#include "Editor/UnrealEd/Public/EditorModes.h"
 #include "Editor/KismetCompiler/Public/KismetCompilerModule.h"
 
 #define LOCTEXT_NAMESPACE "Structure"
 
 //////////////////////////////////////////////////////////////////////////
-// FEnumEditorManager
+// FStructEditorManager
 FStructureEditorUtils::FStructEditorManager::FStructEditorManager() : FListenerManager<UUserDefinedStruct>()
 {}
 
@@ -253,13 +254,13 @@ bool FStructureEditorUtils::RemoveVariable(UUserDefinedStruct* Struct, FGuid Var
 {
 	if(Struct)
 	{
-		const FScopedTransaction Transaction( LOCTEXT("RemoveVariable", "Remove Variable") );
-		ModifyStructData(Struct);
-
 		const auto OldNum = GetVarDesc(Struct).Num();
 		const bool bAllowToMakeEmpty = false;
 		if (bAllowToMakeEmpty || (OldNum > 1))
 		{
+			const FScopedTransaction Transaction(LOCTEXT("RemoveVariable", "Remove Variable"));
+			ModifyStructData(Struct);
+
 			GetVarDesc(Struct).RemoveAll(FFindByGuidHelper<FStructVariableDescription>(VarGuid));
 			if (OldNum != GetVarDesc(Struct).Num())
 			{
@@ -279,9 +280,6 @@ bool FStructureEditorUtils::RenameVariable(UUserDefinedStruct* Struct, FGuid Var
 {
 	if (Struct)
 	{
-		const FScopedTransaction Transaction( LOCTEXT("RenameVariable", "Rename Variable") );
-		ModifyStructData(Struct);
-
 		const FName NewNameBase(*NewDisplayNameStr);
 		if (NewNameBase.IsValidXName(INVALID_OBJECTNAME_CHARACTERS) &&
 			IsUniqueVariableDisplayName(Struct, NewDisplayNameStr))
@@ -289,9 +287,12 @@ bool FStructureEditorUtils::RenameVariable(UUserDefinedStruct* Struct, FGuid Var
 			const FName NewName = FMemberVariableNameHelper::Generate(Struct, NewDisplayNameStr);
 			if (NULL == GetVarDesc(Struct).FindByPredicate(FFindByNameHelper<FStructVariableDescription>(NewName)))
 			{
-				FStructVariableDescription* VarDesc = GetVarDesc(Struct).FindByPredicate(FFindByGuidHelper<FStructVariableDescription>(VarGuid));
+				auto VarDesc = GetVarDescByGuid(Struct, VarGuid);
 				if (VarDesc)
 				{
+					const FScopedTransaction Transaction(LOCTEXT("RenameVariable", "Rename Variable"));
+					ModifyStructData(Struct);
+
 					VarDesc->FriendlyName = NewDisplayNameStr;
 					VarDesc->VarName = NewName;
 					OnStructureChanged(Struct);
@@ -307,9 +308,6 @@ bool FStructureEditorUtils::ChangeVariableType(UUserDefinedStruct* Struct, FGuid
 {
 	if (Struct)
 	{
-		const FScopedTransaction Transaction( LOCTEXT("ChangeVariableType", "Change Variable Type") );
-		ModifyStructData(Struct);
-
 		FString ErrorMessage;
 		if(!CanHaveAMemberVariableOfType(Struct, NewType, &ErrorMessage))
 		{
@@ -317,12 +315,15 @@ bool FStructureEditorUtils::ChangeVariableType(UUserDefinedStruct* Struct, FGuid
 			return false;
 		}
 
-		FStructVariableDescription* VarDesc = GetVarDesc(Struct).FindByPredicate(FFindByGuidHelper<FStructVariableDescription>(VarGuid));
+		auto VarDesc = GetVarDescByGuid(Struct, VarGuid);
 		if(VarDesc)
 		{
 			const bool bChangedType = (VarDesc->ToPinType() != NewType);
 			if (bChangedType)
 			{
+				const FScopedTransaction Transaction(LOCTEXT("ChangeVariableType", "Change Variable Type"));
+				ModifyStructData(Struct);
+
 				VarDesc->VarName = FMemberVariableNameHelper::Generate(Struct, VarDesc->FriendlyName);
 				VarDesc->DefaultValue = FString();
 				VarDesc->SetPinType(NewType);
@@ -337,21 +338,18 @@ bool FStructureEditorUtils::ChangeVariableType(UUserDefinedStruct* Struct, FGuid
 
 bool FStructureEditorUtils::ChangeVariableDefaultValue(UUserDefinedStruct* Struct, FGuid VarGuid, const FString& NewDefaultValue)
 {
-	if(Struct)
+	const UEdGraphSchema_K2* K2Schema = GetDefault<UEdGraphSchema_K2>();
+	auto VarDesc = GetVarDescByGuid(Struct, VarGuid);
+	if (VarDesc 
+		&& (NewDefaultValue != VarDesc->DefaultValue) 
+		&& K2Schema->DefaultValueSimpleValidation(VarDesc->ToPinType(), FString(), NewDefaultValue, NULL, FText::GetEmpty()))
 	{
-		const FScopedTransaction Transaction( LOCTEXT("ChangeVariableDefaultValue", "Change Variable Default Value") );
+		const FScopedTransaction Transaction(LOCTEXT("ChangeVariableDefaultValue", "Change Variable Default Value"));
 		ModifyStructData(Struct);
 
-		const UEdGraphSchema_K2* K2Schema = GetDefault<UEdGraphSchema_K2>();
-		FStructVariableDescription* VarDesc = GetVarDesc(Struct).FindByPredicate(FFindByGuidHelper<FStructVariableDescription>(VarGuid));
-		if (VarDesc 
-			&& (NewDefaultValue != VarDesc->DefaultValue) 
-			&& K2Schema->DefaultValueSimpleValidation(VarDesc->ToPinType(), FString(), NewDefaultValue, NULL, FText::GetEmpty()))
-		{
-			VarDesc->DefaultValue = NewDefaultValue;
-			OnStructureChanged(Struct);
-			return true;
-		}
+		VarDesc->DefaultValue = NewDefaultValue;
+		OnStructureChanged(Struct);
+		return true;
 	}
 	return false;
 }
@@ -374,16 +372,8 @@ bool FStructureEditorUtils::IsUniqueVariableDisplayName(const UUserDefinedStruct
 
 FString FStructureEditorUtils::GetVariableDisplayName(const UUserDefinedStruct* Struct, FGuid VarGuid)
 {
-	if (Struct)
-	{
-		const FStructVariableDescription* VarDesc = GetVarDesc(Struct).FindByPredicate(FFindByGuidHelper<FStructVariableDescription>(VarGuid));
-		if (VarDesc)
-		{
-			return VarDesc->FriendlyName;
-		}
-	}
-
-	return FString();
+	const auto VarDesc = GetVarDescByGuid(Struct, VarGuid);
+	return VarDesc ? VarDesc->FriendlyName : FString();
 }
 
 bool FStructureEditorUtils::UserDefinedStructEnabled()
@@ -560,57 +550,43 @@ bool FStructureEditorUtils::ChangeTooltip(UUserDefinedStruct* Struct, const FStr
 
 FString FStructureEditorUtils::GetVariableTooltip(const UUserDefinedStruct* Struct, FGuid VarGuid)
 {
-	if (Struct)
-	{
-		const FStructVariableDescription* VarDesc = GetVarDesc(Struct).FindByPredicate(FFindByGuidHelper<FStructVariableDescription>(VarGuid));
-		if (VarDesc)
-		{
-			return VarDesc->ToolTip;
-		}
-	}
-
-	return FString();
+	const auto VarDesc = GetVarDescByGuid(Struct, VarGuid);
+	return VarDesc ? VarDesc->ToolTip : FString();
 }
 
 bool FStructureEditorUtils::ChangeVariableTooltip(UUserDefinedStruct* Struct, FGuid VarGuid, const FString& InTooltip)
 {
-	if (Struct)
+	auto VarDesc = GetVarDescByGuid(Struct, VarGuid);
+	if (VarDesc && (InTooltip != VarDesc->ToolTip))
 	{
-		FStructVariableDescription* VarDesc = GetVarDesc(Struct).FindByPredicate(FFindByGuidHelper<FStructVariableDescription>(VarGuid));
-		if (VarDesc && (InTooltip != VarDesc->ToolTip))
-		{
-			const FScopedTransaction Transaction(LOCTEXT("ChangeVariableTooltip", "Change UDS Variable Tooltip"));
-			ModifyStructData(Struct);
-			VarDesc->ToolTip = InTooltip;
-
-			auto Property = FindField<UProperty>(Struct, VarDesc->VarName);
-			if (Property)
-			{
-				Property->SetMetaData(FBlueprintMetadata::MD_Tooltip, *VarDesc->ToolTip);
-			}
+		const FScopedTransaction Transaction(LOCTEXT("ChangeVariableTooltip", "Change UDS Variable Tooltip"));
+		ModifyStructData(Struct);
+		VarDesc->ToolTip = InTooltip;
 
-			return true;
+		auto Property = FindField<UProperty>(Struct, VarDesc->VarName);
+		if (Property)
+		{
+			Property->SetMetaData(FBlueprintMetadata::MD_Tooltip, *VarDesc->ToolTip);
 		}
+
+		return true;
 	}
 	return false;
 }
 
 bool FStructureEditorUtils::ChangeEditableOnBPInstance(UUserDefinedStruct* Struct, FGuid VarGuid, bool bInIsEditable)
 {
-	if (Struct)
+	const UEdGraphSchema_K2* K2Schema = GetDefault<UEdGraphSchema_K2>();
+	auto VarDesc = GetVarDescByGuid(Struct, VarGuid);
+	const bool bNewDontEditoOnInstance = !bInIsEditable;
+	if (VarDesc && (bNewDontEditoOnInstance != VarDesc->bDontEditoOnInstance))
 	{
 		const FScopedTransaction Transaction(LOCTEXT("ChangeVariableDefaultValue", "Change Variable Default Value"));
 		ModifyStructData(Struct);
 
-		const UEdGraphSchema_K2* K2Schema = GetDefault<UEdGraphSchema_K2>();
-		FStructVariableDescription* VarDesc = GetVarDesc(Struct).FindByPredicate(FFindByGuidHelper<FStructVariableDescription>(VarGuid));
-		const bool bNewDontEditoOnInstance = !bInIsEditable;
-		if (VarDesc && (bNewDontEditoOnInstance != VarDesc->bDontEditoOnInstance))
-		{
-			VarDesc->bDontEditoOnInstance = bNewDontEditoOnInstance;
-			OnStructureChanged(Struct);
-			return true;
-		}
+		VarDesc->bDontEditoOnInstance = bNewDontEditoOnInstance;
+		OnStructureChanged(Struct);
+		return true;
 	}
 	return false;
 }
@@ -625,4 +601,45 @@ void FStructureEditorUtils::ModifyStructData(UUserDefinedStruct* Struct)
 	}
 }
 
+bool FStructureEditorUtils::CanEnable3dWidget(const UUserDefinedStruct* Struct, FGuid VarGuid)
+{
+	const auto VarDesc = GetVarDescByGuid(Struct, VarGuid);
+	const auto PropertyStruct = VarDesc ? Cast<const UStruct>(VarDesc->SubCategoryObject.Get()) : NULL;
+	return FEdMode::CanCreateWidgetForStructure(PropertyStruct);
+}
+
+bool FStructureEditorUtils::Change3dWidgetEnabled(UUserDefinedStruct* Struct, FGuid VarGuid, bool bIsEnabled)
+{
+	auto VarDesc = GetVarDescByGuid(Struct, VarGuid);
+	const auto PropertyStruct = VarDesc ? Cast<const UStruct>(VarDesc->SubCategoryObject.Get()) : NULL;
+	if (FEdMode::CanCreateWidgetForStructure(PropertyStruct) && (VarDesc->bEnable3dWidget != bIsEnabled))
+	{
+		const FScopedTransaction Transaction(LOCTEXT("ChangeVariableDefaultValue", "Change Variable Default Value"));
+		ModifyStructData(Struct);
+
+		VarDesc->bEnable3dWidget = bIsEnabled;
+		auto Property = FindField<UProperty>(Struct, VarDesc->VarName);
+		if (Property)
+		{
+			if (VarDesc->bEnable3dWidget)
+			{
+				Property->SetMetaData(FEdMode::MD_MakeEditWidget, TEXT("true"));
+			}
+			else
+			{
+				Property->RemoveMetaData(FEdMode::MD_MakeEditWidget);
+			}
+		}
+		return true;
+	}
+	return false;
+}
+
+bool FStructureEditorUtils::Is3dWidgetEnabled(const UUserDefinedStruct* Struct, FGuid VarGuid)
+{
+	const auto VarDesc = GetVarDescByGuid(Struct, VarGuid);
+	const auto PropertyStruct = VarDesc ? Cast<const UStruct>(VarDesc->SubCategoryObject.Get()) : NULL;
+	return VarDesc && VarDesc->bEnable3dWidget && FEdMode::CanCreateWidgetForStructure(PropertyStruct);
+}
+
 #undef LOCTEXT_NAMESPACE
diff --git a/Engine/Source/Editor/UnrealEd/Public/EditorModes.h b/Engine/Source/Editor/UnrealEd/Public/EditorModes.h
--- a/Engine/Source/Editor/UnrealEd/Public/EditorModes.h
+++ b/Engine/Source/Editor/UnrealEd/Public/EditorModes.h
@@ -333,6 +333,8 @@ public:
 	static const FName MD_MakeEditWidget;
 	/** Specifies a function used for validation of the current value of a property.  The function returns a string that is empty if the value is valid, or contains an error description if the value is invalid */
 	static const FName MD_ValidateWidgetUsing;
+	/** Returns true if this structure can support creating a widget in the editor */
+	static bool CanCreateWidgetForStructure(const UStruct* InPropStruct);
 	/** Returns true if this property can support creating a widget in the editor */
 	static bool CanCreateWidgetForProperty(UProperty* InProp);
 	/** See if we should create a widget for the supplied property when selecting an actor instance */
@@ -374,6 +376,7 @@ protected:
 		FString PropertyName;
 		int32 PropertyIndex;
 		FName PropertyValidationName;
+		FString DisplayName;
 		bool bIsTransform;
 
 		FPropertyWidgetInfo()
diff --git a/Engine/Source/Editor/UnrealEd/Public/Kismet2/StructureEditorUtils.h b/Engine/Source/Editor/UnrealEd/Public/Kismet2/StructureEditorUtils.h
--- a/Engine/Source/Editor/UnrealEd/Public/Kismet2/StructureEditorUtils.h
+++ b/Engine/Source/Editor/UnrealEd/Public/Kismet2/StructureEditorUtils.h
@@ -109,12 +109,28 @@ public:
 
 	static bool ChangeEditableOnBPInstance(UUserDefinedStruct* Struct, FGuid VarGuid, bool bInIsEditable);
 
+	//3D Widget
+	static bool CanEnable3dWidget(const UUserDefinedStruct* Struct, FGuid VarGuid);
+
+	static bool Change3dWidgetEnabled(UUserDefinedStruct* Struct, FGuid VarGuid, bool bIsEnabled);
+
+	static bool Is3dWidgetEnabled(const UUserDefinedStruct* Struct, FGuid VarGuid);
 
 	//MISC
 	static TArray<FStructVariableDescription>& GetVarDesc(UUserDefinedStruct* Struct);
 
 	static const TArray<FStructVariableDescription>& GetVarDesc(const UUserDefinedStruct* Struct);
 
+	static FStructVariableDescription* GetVarDescByGuid(UUserDefinedStruct* Struct, FGuid VarGuid)
+	{
+		return Struct ? GetVarDesc(Struct).FindByPredicate(FFindByGuidHelper<FStructVariableDescription>(VarGuid)) : NULL;
+	}
+
+	static const FStructVariableDescription* GetVarDescByGuid(const UUserDefinedStruct* Struct, FGuid VarGuid)
+	{
+		return Struct ? GetVarDesc(Struct).FindByPredicate(FFindByGuidHelper<FStructVariableDescription>(VarGuid)) : NULL;
+	}
+
 	static void ModifyStructData(UUserDefinedStruct* Struct);
 
 	static bool UserDefinedStructEnabled();
