Remove dead code.

diff --git a/Engine/Source/Editor/UnrealEd/Private/Commandlets/CookCommandlet.cpp b/Engine/Source/Editor/UnrealEd/Private/Commandlets/CookCommandlet.cpp
--- a/Engine/Source/Editor/UnrealEd/Private/Commandlets/CookCommandlet.cpp
+++ b/Engine/Source/Editor/UnrealEd/Private/Commandlets/CookCommandlet.cpp
@@ -24,13 +24,6 @@
 DEFINE_LOG_CATEGORY_STATIC(LogCookCommandlet, Log, All);
 
 
-/** Helper to pass a recompile request to game thread */
-struct FRecompileRequest
-{
-	FShaderRecompileData RecompileData;
-	bool bComplete;
-};
-
 /* Static functions
  *****************************************************************************/
 
