Fixes a bug where the wrong animation is reimported when the FBX contains multiple animations. Adds an FString field to the reimport options that stores the source name of the animation so that it can be correctly reimported.

https://github.com/EpicGames/UnrealEngine/pull/168

diff --git a/Engine/Source/Editor/UnrealEd/Classes/Factories/FbxAnimSequenceImportData.h b/Engine/Source/Editor/UnrealEd/Classes/Factories/FbxAnimSequenceImportData.h
--- a/Engine/Source/Editor/UnrealEd/Classes/Factories/FbxAnimSequenceImportData.h
+++ b/Engine/Source/Editor/UnrealEd/Classes/Factories/FbxAnimSequenceImportData.h
@@ -39,6 +39,10 @@ class UFbxAnimSequenceImportData : public UFbxAssetImportData
 	UPROPERTY(EditAnywhere, Category=ImportSettings, meta=(DisplayName = "End Frame"))
 	int32	EndFrame;
 
+	/** Name of source animation that was imported, used to reimport correct animation from the FBX file*/
+	UPROPERTY(EditAnywhere, Category = ImportSettings, meta = (DisplayName = "Source Animation Name"))
+	FString SourceAnimationName;
+
 	/** Gets or creates fbx import data for the specified anim sequence */
 	static UFbxAnimSequenceImportData* GetImportDataForAnimSequence(UAnimSequence* AnimSequence, UFbxAnimSequenceImportData* TemplateForCreation);
 };
\ No newline at end of file
diff --git a/Engine/Source/Editor/UnrealEd/Private/SkeletalMeshEdit.cpp b/Engine/Source/Editor/UnrealEd/Private/SkeletalMeshEdit.cpp
--- a/Engine/Source/Editor/UnrealEd/Private/SkeletalMeshEdit.cpp
+++ b/Engine/Source/Editor/UnrealEd/Private/SkeletalMeshEdit.cpp
@@ -169,8 +169,20 @@ bool UEditorEngine::ReimportFbxAnimation( USkeleton * Skeleton, UAnimSequence *
 		}
 		else
 		{
-			// pick first one
-			FbxAnimStack* CurAnimStack = FbxImporter->Scene->GetSrcObject<FbxAnimStack>(0);
+			// find the correct animation based on import data
+			FbxAnimStack* CurAnimStack = nullptr;
+			
+			//ignore the source animation name if there's only one animation in the file.
+			//this is to make it easier for people who use content creation programs that only export one animation and/or ones that don't allow naming animations			
+			if (FbxImporter->Scene->GetSrcObjectCount(FbxAnimStack::ClassId) > 1 && !ImportData->SourceAnimationName.IsEmpty())
+			{
+				CurAnimStack = FbxCast<FbxAnimStack>(FbxImporter->Scene->FindSrcObject(FbxAnimStack::ClassId, TCHAR_TO_ANSI(*ImportData->SourceAnimationName), 0));
+			}
+			else
+			{
+				CurAnimStack = FbxCast<FbxAnimStack>(FbxImporter->Scene->GetSrcObject(FbxAnimStack::ClassId, 0));
+			}
+			
 			if (CurAnimStack)
 			{
 				// set current anim stack
