Refactor active visualizer logic into ComponentVisualizerManager struct

diff --git a/Engine/Source/Editor/ComponentVisualizers/Private/SplineComponentVisualizer.h b/Engine/Source/Editor/ComponentVisualizers/Private/SplineComponentVisualizer.h
--- a/Engine/Source/Editor/ComponentVisualizers/Private/SplineComponentVisualizer.h
+++ b/Engine/Source/Editor/ComponentVisualizers/Private/SplineComponentVisualizer.h
@@ -5,6 +5,7 @@
 #include "ComponentVisualizer.h"
 #include "Components/SplineComponent.h"
 
+/** Base class for clickable spline editing proxies */
 struct HSplineVisProxy : public HComponentVisProxy
 {
 	DECLARE_HIT_PROXY();
@@ -14,6 +15,7 @@ struct HSplineVisProxy : public HComponentVisProxy
 	{}
 };
 
+/** Proxy for a spline key */
 struct HSplineKeyProxy : public HSplineVisProxy
 {
 	DECLARE_HIT_PROXY();
@@ -27,12 +29,14 @@ struct HSplineKeyProxy : public HSplineVisProxy
 	int32 KeyIndex;
 };
 
+/** SplineComponent visualizer/edit functionality */
 class FSplineComponentVisualizer : public FComponentVisualizer
 {
 public:
 	FSplineComponentVisualizer()
 	: FComponentVisualizer()
 	, SelectedKeyIndex(INDEX_NONE)
+	, bAllowDuplication(true)
 	{}
 
 	// Begin FComponentVisualizer interface
@@ -44,14 +48,18 @@ public:
 	virtual bool HandleInputKey(FLevelEditorViewportClient* ViewportClient, FViewport* Viewport, FKey Key, EInputEvent Event) OVERRIDE;
 	// End FComponentVisualizer interface
 
-	/** */
+	/** Get the spline component we are currently editing */
 	USplineComponent* GetEditedSplineComponent();
 
 private:
 
+	/** Actor that owns the currently edited spline */
 	TWeakObjectPtr<AActor> SplineOwningActor;
+	/** Name of property on the actor that references the spline we are editing */
 	FName SplineCompPropName;
+	/** Index of key we have selected */
 	int32 SelectedKeyIndex;
 
+	/** Whether we currently allow duplication when dragging */
 	bool bAllowDuplication;
 };
diff --git a/Engine/Source/Editor/Kismet/Private/SCSEditorViewportClient.cpp b/Engine/Source/Editor/Kismet/Private/SCSEditorViewportClient.cpp
--- a/Engine/Source/Editor/Kismet/Private/SCSEditorViewportClient.cpp
+++ b/Engine/Source/Editor/Kismet/Private/SCSEditorViewportClient.cpp
@@ -332,6 +332,9 @@ void FSCSEditorViewportClient::ProcessClick(class FSceneView& View, class HHitPr
 
 			Invalidate();
 		}
+
+		// Pass to component vis manager
+		//GUnrealEd->ComponentVisManager.HandleProxyForComponentVis(HitProxy);
 	}
 }
 
@@ -431,21 +434,21 @@ bool FSCSEditorViewportClient::InputWidgetDelta( FViewport* Viewport, EAxisList:
 
 							if(PreviewBlueprint != NULL)
 							{
-							// Like PostEditMove(), but we only need to re-run construction scripts
-							if(PreviewBlueprint && PreviewBlueprint->bRunConstructionScriptOnDrag)
-							{
-								PreviewActor->RerunConstructionScripts();
-							}
+								// Like PostEditMove(), but we only need to re-run construction scripts
+								if(PreviewBlueprint && PreviewBlueprint->bRunConstructionScriptOnDrag)
+								{
+									PreviewActor->RerunConstructionScripts();
+								}
 
-							SceneComp->PostEditComponentMove(true); // @TODO HACK passing 'finished' every frame...
+								SceneComp->PostEditComponentMove(true); // @TODO HACK passing 'finished' every frame...
 
-							// If a constraint, copy back updated constraint frames to template
-							UPhysicsConstraintComponent* ConstraintComp = Cast<UPhysicsConstraintComponent>(SceneComp);
-							UPhysicsConstraintComponent* TemplateComp = Cast<UPhysicsConstraintComponent>(SelectedTemplate);
-							if(ConstraintComp && TemplateComp)
-							{
-								TemplateComp->ConstraintInstance.CopyConstraintGeometryFrom(&ConstraintComp->ConstraintInstance);
-							}
+								// If a constraint, copy back updated constraint frames to template
+								UPhysicsConstraintComponent* ConstraintComp = Cast<UPhysicsConstraintComponent>(SceneComp);
+								UPhysicsConstraintComponent* TemplateComp = Cast<UPhysicsConstraintComponent>(SelectedTemplate);
+								if(ConstraintComp && TemplateComp)
+								{
+									TemplateComp->ConstraintInstance.CopyConstraintGeometryFrom(&ConstraintComp->ConstraintInstance);
+								}
 
 								// Get the Blueprint class default object
 								AActor* CDO = NULL;
diff --git a/Engine/Source/Editor/UnrealEd/Classes/Editor/UnrealEdEngine.h b/Engine/Source/Editor/UnrealEd/Classes/Editor/UnrealEdEngine.h
--- a/Engine/Source/Editor/UnrealEd/Classes/Editor/UnrealEdEngine.h
+++ b/Engine/Source/Editor/UnrealEd/Classes/Editor/UnrealEdEngine.h
@@ -3,6 +3,7 @@
 #pragma once
 #include "UniquePtr.h"
 #include "IPackageAutoSaver.h"
+#include "ComponentVisualizerManager.h"
 #include "UnrealEdEngine.generated.h"
 
 UENUM()
@@ -157,6 +158,9 @@ class UNREALED_API UUnrealEdEngine : public UEditorEngine, public FNotifyHook
 	/** Map from component class to visualizer object to use */
 	TMap< FName, TSharedPtr<class FComponentVisualizer> > ComponentVisualizerMap;
 
+	/** Manages currently active visualizer and routes interactions to it */
+	FComponentVisualizerManager	ComponentVisManager;
+
 	// Begin UObject interface.
 	~UUnrealEdEngine();
 	virtual void FinishDestroy() OVERRIDE;
diff --git a/Engine/Source/Editor/UnrealEd/Private/ComponentVisualizerManager.cpp b/Engine/Source/Editor/UnrealEd/Private/ComponentVisualizerManager.cpp
new file mode 100644
--- /dev/null
+++ b/Engine/Source/Editor/UnrealEd/Private/ComponentVisualizerManager.cpp
@@ -0,0 +1,81 @@
+// Copyright 1998-2014 Epic Games, Inc. All Rights Reserved.
+
+#include "UnrealEd.h"
+#include "ComponentVisualizerManager.h"
+
+void FComponentVisualizerManager::HandleProxyForComponentVis(HHitProxy *HitProxy)
+{
+	if (HitProxy->IsA(HComponentVisProxy::StaticGetType()))
+	{
+		HComponentVisProxy* VisProxy = (HComponentVisProxy*)HitProxy;
+		const UActorComponent* ClickedComponent = VisProxy->Component.Get();
+		if (ClickedComponent != NULL)
+		{
+			TSharedPtr<FComponentVisualizer> Visualizer = GUnrealEd->FindComponentVisualizer(ClickedComponent->GetClass()->GetFName());
+			if (Visualizer.IsValid())
+			{
+				bool bIsActive = Visualizer->VisProxyHandleClick(VisProxy);
+				if (bIsActive)
+				{
+					// call EndEditing on any currently edited visualizer, if we are going to change it
+					if (EditedVisualizer.IsValid() && Visualizer.Get() != EditedVisualizer.Pin().Get())
+					{
+						EditedVisualizer.Pin()->EndEditing();
+					}
+
+					EditedVisualizer = Visualizer;
+				}
+			}
+		}
+	}
+	else
+	{
+		ClearActiveComponentVis();
+	}
+}
+
+void FComponentVisualizerManager::ClearActiveComponentVis()
+{
+	if (EditedVisualizer.IsValid())
+	{
+		EditedVisualizer.Pin()->EndEditing();
+		EditedVisualizer = NULL;
+	}
+}
+
+bool FComponentVisualizerManager::HandleInputKey(FLevelEditorViewportClient* ViewportClient, FViewport* Viewport, FKey Key, EInputEvent Event)
+{
+	if(EditedVisualizer.IsValid())
+	{
+		if(EditedVisualizer.Pin()->HandleInputKey(ViewportClient, Viewport, Key, Event))
+		{
+			return true;
+		}
+	}
+
+	return false;
+}
+
+bool FComponentVisualizerManager::HandleInputDelta(FLevelEditorViewportClient* InViewportClient, FViewport* InViewport, FVector& InDrag, FRotator& InRot, FVector& InScale)
+{
+	if (EditedVisualizer.IsValid() && InViewportClient->GetCurrentWidgetAxis() != EAxisList::None)
+	{
+		if (EditedVisualizer.Pin()->HandleInputDelta(InViewportClient, InViewport, InDrag, InRot, InScale))
+		{
+			return true;
+		}
+	}
+
+	return false;
+}
+
+
+bool FComponentVisualizerManager::GetWidgetLocation(FVector& OutLocation)
+{
+	if (EditedVisualizer.IsValid())
+	{
+		return EditedVisualizer.Pin()->GetWidgetLocation(OutLocation);
+	}
+
+	return false;
+}
diff --git a/Engine/Source/Editor/UnrealEd/Private/EditorModes.cpp b/Engine/Source/Editor/UnrealEd/Private/EditorModes.cpp
--- a/Engine/Source/Editor/UnrealEd/Private/EditorModes.cpp
+++ b/Engine/Source/Editor/UnrealEd/Private/EditorModes.cpp
@@ -278,12 +278,9 @@ bool FEdMode::CapturedMouseMove( FLevelEditorViewportClient* InViewportClient, F
 
 bool FEdMode::InputKey(FLevelEditorViewportClient* ViewportClient, FViewport* Viewport, FKey Key, EInputEvent Event)
 {
-	if (EditedVisualizer.IsValid())
+	if (GUnrealEd->ComponentVisManager.HandleInputKey(ViewportClient, Viewport, Key, Event))
 	{
-		if (EditedVisualizer.Pin()->HandleInputKey(ViewportClient, Viewport, Key, Event))
-		{
-			return true;
-		}
+		return true;
 	}
 
 	if( GetCurrentTool() && GetCurrentTool()->InputKey( ViewportClient, Viewport, Key, Event ) )
@@ -319,12 +316,9 @@ bool FEdMode::InputAxis(FLevelEditorViewportClient* InViewportClient, FViewport*
 
 bool FEdMode::InputDelta(FLevelEditorViewportClient* InViewportClient, FViewport* InViewport, FVector& InDrag, FRotator& InRot, FVector& InScale)
 {
-	if (EditedVisualizer.IsValid() && InViewportClient->GetCurrentWidgetAxis() != EAxisList::None)
+	if(GUnrealEd->ComponentVisManager.HandleInputDelta(InViewportClient, InViewport, InDrag, InRot, InScale))
 	{
-		if(EditedVisualizer.Pin()->HandleInputDelta(InViewportClient, InViewport, InDrag, InRot, InScale))
-		{
-			return true;
-		}
+		return true;
 	}
 	
 	if(UsesPropertyWidgets())
@@ -430,14 +424,11 @@ bool FEdMode::UsesTransformWidget(FWidget::EWidgetMode CheckMode) const
 
 FVector FEdMode::GetWidgetLocation() const
 {
-	if(EditedVisualizer.IsValid())
+	FVector ComponentVisWidgetLocation;
+	bool bUseComponentVisWidgetLocation = GUnrealEd->ComponentVisManager.GetWidgetLocation(ComponentVisWidgetLocation);
+	if(bUseComponentVisWidgetLocation)
 	{
-		FVector VisualizerWidgetPos;
-		bool bSetWidgetLocation = EditedVisualizer.Pin()->GetWidgetLocation(VisualizerWidgetPos);
-		if(bSetWidgetLocation)
-		{
-			return VisualizerWidgetPos;
-		}
+		return ComponentVisWidgetLocation;
 	}
 
 	if(UsesPropertyWidgets())
@@ -547,11 +538,7 @@ void FEdMode::ActorSelectionChangeNotify()
 	bEditedPropertyIsTransform = false;
 
 	// Clear active editing visualizer on selection change
-	if (EditedVisualizer.IsValid())
-	{
-		EditedVisualizer.Pin()->EndEditing();
-		EditedVisualizer = NULL;
-	}
+	GUnrealEd->ComponentVisManager.ClearActiveComponentVis();
 }
 
 bool FEdMode::HandleClick(FLevelEditorViewportClient* InViewportClient, HHitProxy *HitProxy, const FViewportClick &Click)
@@ -566,44 +553,20 @@ bool FEdMode::HandleClick(FLevelEditorViewportClient* InViewportClient, HHitProx
 			bEditedPropertyIsTransform = PropertyProxy->bPropertyIsTransform;
 			return true;
 		}
-		else if(HitProxy->IsA(HComponentVisProxy::StaticGetType()))
-		{
-			HComponentVisProxy* VisProxy = (HComponentVisProxy*)HitProxy;
-			const UActorComponent* ClickedComponent = VisProxy->Component.Get();
-			if(ClickedComponent != NULL)
-			{
-				TSharedPtr<FComponentVisualizer> Visualizer = GUnrealEd->FindComponentVisualizer(ClickedComponent->GetClass()->GetFName());
-				if(Visualizer.IsValid())
-				{
-					bool bIsActive = Visualizer->VisProxyHandleClick(VisProxy);
-					if(bIsActive)
-					{
-						// call EndEditing on any currently edited visualizer, if we are going to change it
-						if (EditedVisualizer.IsValid() && Visualizer.Get() != EditedVisualizer.Pin().Get())
-						{
-							EditedVisualizer.Pin()->EndEditing();
-						}
-
-						EditedVisualizer = Visualizer;
-					}
-				}
-			}
-		}
 		// Left clicking on an actor, stop editing a property
 		else if( HitProxy->IsA(HActor::StaticGetType()) )
 		{
 			EditedPropertyName = TEXT("");
 			EditedPropertyIndex = INDEX_NONE;
 			bEditedPropertyIsTransform = false;
-
-			if(EditedVisualizer.IsValid())
-			{
-				EditedVisualizer.Pin()->EndEditing();
-				EditedVisualizer = NULL;
-			}
 		}
 	}
 
+	if(HitProxy)
+	{	
+		GUnrealEd->ComponentVisManager.HandleProxyForComponentVis(HitProxy);
+	}
+
 	return 0;
 }
 
diff --git a/Engine/Source/Editor/UnrealEd/Public/ComponentVisualizerManager.h b/Engine/Source/Editor/UnrealEd/Public/ComponentVisualizerManager.h
new file mode 100644
--- /dev/null
+++ b/Engine/Source/Editor/UnrealEd/Public/ComponentVisualizerManager.h
@@ -0,0 +1,28 @@
+// Copyright 1998-2014 Epic Games, Inc. All Rights Reserved.
+
+#pragma once
+
+/** Class that managed active component visualizer and routes input to it */
+class UNREALED_API FComponentVisualizerManager
+{
+public:
+	FComponentVisualizerManager() {}
+	virtual ~FComponentVisualizerManager() {}
+
+
+	/** Activate a component visualizer given a clicked proxy */
+	void HandleProxyForComponentVis(HHitProxy *HitProxy);
+	/** Clear active component visualizer */
+	void ClearActiveComponentVis();
+
+	/** Pass key input to active visualizer */
+	bool HandleInputKey(FLevelEditorViewportClient* ViewportClient, FViewport* Viewport, FKey Key, EInputEvent Event);
+	/** Pass delta input to active visualizer */
+	bool HandleInputDelta(FLevelEditorViewportClient* InViewportClient, FViewport* InViewport, FVector& InDrag, FRotator& InRot, FVector& InScale);
+	/** Get widget location from active visualizer */
+	bool GetWidgetLocation(FVector& OutLocation);
+
+private:
+	/** Currently 'active' visualizer that we should pass input to etc */
+	TWeakPtr<class FComponentVisualizer> EditedVisualizer;
+};
\ No newline at end of file
diff --git a/Engine/Source/Editor/UnrealEd/Public/EditorModes.h b/Engine/Source/Editor/UnrealEd/Public/EditorModes.h
--- a/Engine/Source/Editor/UnrealEd/Public/EditorModes.h
+++ b/Engine/Source/Editor/UnrealEd/Public/EditorModes.h
@@ -401,9 +401,6 @@ protected:
 	 */
 	void GetPropertyWidgetInfos(const UStruct* InStruct, const void* InContainer, TArray<FPropertyWidgetInfo>& OutInfos, FString PropertyNamePrefix = TEXT("")) const;
 
-	/** Currently 'active' visualizer that we should pass input to etc */
-	TWeakPtr<FComponentVisualizer> EditedVisualizer;
-
 	/** Name of the property currently being edited */
 	FString EditedPropertyName;
 
