Fix fbx deprecation warnings

diff --git a/Engine/Source/Editor/UnrealEd/Private/SkeletalMeshEdit.cpp b/Engine/Source/Editor/UnrealEd/Private/SkeletalMeshEdit.cpp
--- a/Engine/Source/Editor/UnrealEd/Private/SkeletalMeshEdit.cpp
+++ b/Engine/Source/Editor/UnrealEd/Private/SkeletalMeshEdit.cpp
@@ -174,13 +174,13 @@ bool UEditorEngine::ReimportFbxAnimation( USkeleton * Skeleton, UAnimSequence *
 			
 			//ignore the source animation name if there's only one animation in the file.
 			//this is to make it easier for people who use content creation programs that only export one animation and/or ones that don't allow naming animations			
-			if (FbxImporter->Scene->GetSrcObjectCount(FbxAnimStack::ClassId) > 1 && !ImportData->SourceAnimationName.IsEmpty())
+			if (FbxImporter->Scene->GetSrcObjectCount(FbxCriteria::ObjectType(FbxAnimStack::ClassId)) > 1 && !ImportData->SourceAnimationName.IsEmpty())
 			{
-				CurAnimStack = FbxCast<FbxAnimStack>(FbxImporter->Scene->FindSrcObject(FbxAnimStack::ClassId, TCHAR_TO_ANSI(*ImportData->SourceAnimationName), 0));
+				CurAnimStack = FbxCast<FbxAnimStack>(FbxImporter->Scene->FindSrcObject(FbxCriteria::ObjectType(FbxAnimStack::ClassId), TCHAR_TO_ANSI(*ImportData->SourceAnimationName), 0));
 			}
 			else
 			{
-				CurAnimStack = FbxCast<FbxAnimStack>(FbxImporter->Scene->GetSrcObject(FbxAnimStack::ClassId, 0));
+				CurAnimStack = FbxCast<FbxAnimStack>(FbxImporter->Scene->GetSrcObject(FbxCriteria::ObjectType(FbxAnimStack::ClassId), 0));
 			}
 			
 			if (CurAnimStack)
