Fixed crash when box selecting if there is an actor in the scene with a BillboardComponent with no sprite texture set.
#ttp 335978 - CRASH: FLevelEditorViewportClient::ComponentIsTouchingSelectionBox() FDragTool_ActorBoxSelect::IntersectsBox()
#branch UE4
#proj Editor.UnrealEd
#add Added nullptr check in FLevelEditorViewportClient::ComponentIsTouchingSelectionBox to catch BillboardComponents with no sprite texture set.

diff --git a/Engine/Source/Editor/UnrealEd/Private/LevelEditorViewport.cpp b/Engine/Source/Editor/UnrealEd/Private/LevelEditorViewport.cpp
--- a/Engine/Source/Editor/UnrealEd/Private/LevelEditorViewport.cpp
+++ b/Engine/Source/Editor/UnrealEd/Private/LevelEditorViewport.cpp
@@ -3094,22 +3094,24 @@ bool FLevelEditorViewportClient::ComponentIsTouchingSelectionBox(AActor* InActor
 	else if (!bConsiderOnlyBSP && EngineShowFlags.BillboardSprites && InComponent->IsA(UBillboardComponent::StaticClass()))
 	{
 		const UBillboardComponent* SpriteComponent = CastChecked<const UBillboardComponent>(InComponent);
+		if (SpriteComponent->Sprite != nullptr)
+		{
+			float Scale = SpriteComponent->ComponentToWorld.GetMaximumAxisScale();
 
-		float Scale = SpriteComponent->ComponentToWorld.GetMaximumAxisScale();
-
-		// Construct a box representing the sprite
-		FBox SpriteBox(
-			InActor->GetActorLocation() - Scale * FMath::Max(SpriteComponent->Sprite->GetSizeX(),SpriteComponent->Sprite->GetSizeY()) * FVector(1,1,1),
-			InActor->GetActorLocation() + Scale * FMath::Max(SpriteComponent->Sprite->GetSizeX(),SpriteComponent->Sprite->GetSizeY()) * FVector(1,1,1) );
+			// Construct a box representing the sprite
+			FBox SpriteBox(
+				InActor->GetActorLocation() - Scale * FMath::Max(SpriteComponent->Sprite->GetSizeX(), SpriteComponent->Sprite->GetSizeY()) * FVector(1, 1, 1),
+				InActor->GetActorLocation() + Scale * FMath::Max(SpriteComponent->Sprite->GetSizeX(), SpriteComponent->Sprite->GetSizeY()) * FVector(1, 1, 1));
 
-		// If the selection box doesn't have to encompass the entire component and it intersects with the box constructed for the sprite, then it is valid.
-		// Additionally, if the selection box does have to encompass the entire component and both the min and max vectors of the sprite box are inside the selection box,
-		// then it is valid.
-		if (	( !bMustEncompassEntireComponent && InSelBBox.Intersect( SpriteBox ) ) 
-			||	( bMustEncompassEntireComponent && InSelBBox.IsInside( SpriteBox.Min ) && InSelBBox.IsInside( SpriteBox.Max ) ) )
-		{
-			bResult = true;
-			bAlreadyProcessed = true;
+			// If the selection box doesn't have to encompass the entire component and it intersects with the box constructed for the sprite, then it is valid.
+			// Additionally, if the selection box does have to encompass the entire component and both the min and max vectors of the sprite box are inside the selection box,
+			// then it is valid.
+			if ((!bMustEncompassEntireComponent && InSelBBox.Intersect(SpriteBox))
+				|| (bMustEncompassEntireComponent && InSelBBox.IsInside(SpriteBox.Min) && InSelBBox.IsInside(SpriteBox.Max)))
+			{
+				bResult = true;
+				bAlreadyProcessed = true;
+			}
 		}
 	}
 	else if( !bConsiderOnlyBSP && EngineShowFlags.SkeletalMeshes && InComponent->IsA( USkeletalMeshComponent::StaticClass() ) )
