Put guard in place to prevent brushes from being dragged into locked levels.
#ttp 334902 - EDITOR: Level Browser: Allowed to add bsp geometry brushes to locked levels
#branch UE4
#proj Editor.LevelEditor, Editor.UnrealEd
#add Added checks for whether the current level is locked in SLevelViewport::OnDrop (creating a Slate notification if so, instead of placing the object).  Doing this here also eliminates the multiple identical notifications which otherwise occur when dragging Actors into a locked level.
#add Had to make a special case of FBrushBuilderDragDropOp::OnDrop which is responsible for adding a dragged brush.  The notification will already have occurred from the change above.  This just inhibits the actual creation of the brush.

diff --git a/Engine/Source/Editor/LevelEditor/Private/SLevelViewport.cpp b/Engine/Source/Editor/LevelEditor/Private/SLevelViewport.cpp
--- a/Engine/Source/Editor/LevelEditor/Private/SLevelViewport.cpp
+++ b/Engine/Source/Editor/LevelEditor/Private/SLevelViewport.cpp
@@ -713,7 +713,19 @@ bool SLevelViewport::HandlePlaceDraggedObjects(const FDragDropEvent& DragDropEve
 
 FReply SLevelViewport::OnDrop( const FGeometry& MyGeometry, const FDragDropEvent& DragDropEvent )
 {
-	return HandlePlaceDraggedObjects(DragDropEvent, /*bCreateDropPreview=*/false) ? FReply::Handled() : FReply::Unhandled();
+	ULevel* CurrentLevel = (GetWorld()) ? GetWorld()->GetCurrentLevel() : NULL;
+
+	if (CurrentLevel && !FLevelUtils::IsLevelLocked(CurrentLevel))
+	{
+		return HandlePlaceDraggedObjects(DragDropEvent, /*bCreateDropPreview=*/false) ? FReply::Handled() : FReply::Unhandled();
+	}
+	else
+	{
+		FNotificationInfo Info(LOCTEXT("Error_OperationDisallowedOnLockedLevel", "The requested operation could not be completed because the level is locked."));
+		Info.ExpireDuration = 3.0f;
+		FSlateNotificationManager::Get().AddNotification(Info);
+		return FReply::Handled();
+	}
 }
 
 
diff --git a/Engine/Source/Editor/UnrealEd/Public/DragAndDrop/BrushBuilderDragDropOp.h b/Engine/Source/Editor/UnrealEd/Public/DragAndDrop/BrushBuilderDragDropOp.h
--- a/Engine/Source/Editor/UnrealEd/Public/DragAndDrop/BrushBuilderDragDropOp.h
+++ b/Engine/Source/Editor/UnrealEd/Public/DragAndDrop/BrushBuilderDragDropOp.h
@@ -1,6 +1,7 @@
 // Copyright 1998-2014 Epic Games, Inc. All Rights Reserved.
 
 #pragma once
+#include "LevelUtils.h"
 
 class FBrushBuilderDragDropOp : public FDragDropOperation
 {
@@ -31,19 +32,23 @@ public:
 	virtual void OnDrop( bool bDropWasHandled, const FPointerEvent& MouseEvent ) OVERRIDE
 	{
 		UWorld* World = GEditor->GetEditorWorldContext().World();
-		if(World != nullptr)
+		if (World != nullptr)
 		{
-			// Cut the BSP
-			if(bDropWasHandled)
+			ULevel* CurrentLevel = World->GetCurrentLevel();
+			if (CurrentLevel != nullptr && !FLevelUtils::IsLevelLocked(CurrentLevel))
 			{
-				World->GetBrush()->BrushBuilder = DuplicateObject<UBrushBuilder>(BrushBuilder.Get(), World->GetBrush()->GetOuter());
-				const TCHAR* Command = bIsAdditive ? TEXT("BRUSH ADD SELECTNEWBRUSH") : TEXT("BRUSH SUBTRACT SELECTNEWBRUSH");
-				GEditor->Exec(World, Command);
+				// Cut the BSP
+				if (bDropWasHandled)
+				{
+					World->GetBrush()->BrushBuilder = DuplicateObject<UBrushBuilder>(BrushBuilder.Get(), World->GetBrush()->GetOuter());
+					const TCHAR* Command = bIsAdditive ? TEXT("BRUSH ADD SELECTNEWBRUSH") : TEXT("BRUSH SUBTRACT SELECTNEWBRUSH");
+					GEditor->Exec(World, Command);
+				}
+
+				// Deselect & hide the builder brush
+				World->GetBrush()->SetIsTemporarilyHiddenInEditor(true);
+				GEditor->SelectActor(World->GetBrush(), false, false);
 			}
-
-			// Deselect & hide the builder brush
-			World->GetBrush()->SetIsTemporarilyHiddenInEditor(true);
-			GEditor->SelectActor(World->GetBrush(), false, false);
 		}
 
 		FDragDropOperation::OnDrop( bDropWasHandled, MouseEvent );
