Settings: moved viewport related enumeration types into viewport settings class

diff --git a/Engine/Source/Editor/UnrealEd/Classes/Editor/EditorUserSettings.h b/Engine/Source/Editor/UnrealEd/Classes/Editor/EditorUserSettings.h
--- a/Engine/Source/Editor/UnrealEd/Classes/Editor/EditorUserSettings.h
+++ b/Engine/Source/Editor/UnrealEd/Classes/Editor/EditorUserSettings.h
@@ -4,25 +4,6 @@
 
 #include "EditorUserSettings.generated.h"
 
-UENUM()
-enum EWASDType
-{
-	WASD_Always UMETA(DisplayName="Use WASD for Camera Controls"),
-	WASD_RMBOnly UMETA(DisplayName="Use WASD only when a Mouse Button is Pressed"),
-	WASD_Never UMETA(DisplayName="Never use WASD for Camera Controls"),
-	WASD_MAX,
-};
-
-
-UENUM()
-enum ERotationGridMode
-{
-	/** Using Divisions of 360 degrees (e.g 360/2. 360/3, 360/4, ... ) */
-	GridMode_DivisionsOf360,
-
-	/** Uses the user defined grid values */
-	GridMode_Common,
-};
 
 UCLASS(minimalapi, autoexpandcategories=(ViewportControls, ViewportLookAndFeel, LevelEditing, SourceControl, Content, Startup),
 	   hidecategories=(Object, Options, Grid, RotationGrid),
diff --git a/Engine/Source/Editor/UnrealEd/Classes/Settings/LevelEditorViewportSettings.h b/Engine/Source/Editor/UnrealEd/Classes/Settings/LevelEditorViewportSettings.h
--- a/Engine/Source/Editor/UnrealEd/Classes/Settings/LevelEditorViewportSettings.h
+++ b/Engine/Source/Editor/UnrealEd/Classes/Settings/LevelEditorViewportSettings.h
@@ -10,6 +10,33 @@
 #include "LevelEditorViewportSettings.generated.h"
 
 
+/**
+ * Enumerates modes for the viewport's rotation grid.
+ */
+UENUM()
+enum ERotationGridMode
+{
+	/** Using Divisions of 360 degrees (e.g 360/2. 360/3, 360/4, ... ) */
+	GridMode_DivisionsOf360,
+
+	/** Uses the user defined grid values */
+	GridMode_Common,
+};
+
+
+/**
+ * Enumerates camera control types for the W, A, S and D keys.
+ */
+UENUM()
+enum EWASDType
+{
+	WASD_Always UMETA(DisplayName="Use WASD for Camera Controls"),
+	WASD_RMBOnly UMETA(DisplayName="Use WASD only when a Mouse Button is Pressed"),
+	WASD_Never UMETA(DisplayName="Never use WASD for Camera Controls"),
+	WASD_MAX,
+};
+
+
 /**
  * Implements the Level Editor's per-instance view port settings.
  */
