Fixed material nodes positions inside generated material during exporting

diff --git a/Engine/Source/Editor/UnrealEd/Private/EditorExporters.cpp b/Engine/Source/Editor/UnrealEd/Private/EditorExporters.cpp
--- a/Engine/Source/Editor/UnrealEd/Private/EditorExporters.cpp
+++ b/Engine/Source/Editor/UnrealEd/Private/EditorExporters.cpp
@@ -2810,7 +2810,7 @@ namespace MaterialExportUtils
 			Material->Expressions.Add(MetallicExpression);
 			Material->Metallic.Expression = MetallicExpression;
 
-			MetallicExpression->MaterialExpressionEditorX = 250;
+			MetallicExpression->MaterialExpressionEditorX = -250;
 			MetallicExpression->MaterialExpressionEditorY = 0;
 		}
 		
@@ -2821,7 +2821,7 @@ namespace MaterialExportUtils
 			Material->Expressions.Add(RoughnessExpression);
 			Material->Roughness.Expression = RoughnessExpression;
 
-			RoughnessExpression->MaterialExpressionEditorX = 250;
+			RoughnessExpression->MaterialExpressionEditorX = -250;
 			RoughnessExpression->MaterialExpressionEditorY = 150;
 		}
 
@@ -2848,7 +2848,7 @@ namespace MaterialExportUtils
 			UMaterialExpressionTextureSample* BasecolorExpression = ConstructObject<UMaterialExpressionTextureSample>(UMaterialExpressionTextureSample::StaticClass(), Material);
 			BasecolorExpression->Texture = DiffuseTexture;
 			BasecolorExpression->SamplerType = EMaterialSamplerType::SAMPLERTYPE_Color;
-			BasecolorExpression->MaterialExpressionEditorX = 400;
+			BasecolorExpression->MaterialExpressionEditorX = -400;
 			BasecolorExpression->MaterialExpressionEditorY = -150;
 			Material->Expressions.Add(BasecolorExpression);
 			Material->BaseColor.Expression = BasecolorExpression;
@@ -2880,7 +2880,7 @@ namespace MaterialExportUtils
 			UMaterialExpressionTextureSample* NormalExpression = ConstructObject<UMaterialExpressionTextureSample>(UMaterialExpressionTextureSample::StaticClass(), Material);
 			NormalExpression->Texture = NormalTexture;
 			NormalExpression->SamplerType = EMaterialSamplerType::SAMPLERTYPE_Normal;
-			NormalExpression->MaterialExpressionEditorX = 400;
+			NormalExpression->MaterialExpressionEditorX = -400;
 			NormalExpression->MaterialExpressionEditorY = 300;
 			Material->Expressions.Add(NormalExpression);
 			Material->Normal.Expression = NormalExpression;
