Another attempt to fix CIS

diff --git a/Engine/Source/Editor/UnrealEd/Classes/Editor/EditorEngine.h b/Engine/Source/Editor/UnrealEd/Classes/Editor/EditorEngine.h
--- a/Engine/Source/Editor/UnrealEd/Classes/Editor/EditorEngine.h
+++ b/Engine/Source/Editor/UnrealEd/Classes/Editor/EditorEngine.h
@@ -5,6 +5,7 @@
 #include "EditorUserSettings.h"
 #include "Transactor.h"
 #include "../Settings/LevelEditorPlaySettings.h"
+#include "../Settings/LevelEditorViewportSettings.h"
 #include "EditorEngine.generated.h"
 
 //
diff --git a/Engine/Source/Editor/UnrealEd/Classes/Settings/LevelEditorViewportSettings.h b/Engine/Source/Editor/UnrealEd/Classes/Settings/LevelEditorViewportSettings.h
--- a/Engine/Source/Editor/UnrealEd/Classes/Settings/LevelEditorViewportSettings.h
+++ b/Engine/Source/Editor/UnrealEd/Classes/Settings/LevelEditorViewportSettings.h
@@ -7,6 +7,7 @@
 #pragma once
 
 #include "Viewports.h"
+#include "Editor/UnrealEdTypes.h"
 #include "LevelEditorViewportSettings.generated.h"
 
 
