Fix editor non-unity

diff --git a/Engine/Source/Editor/UnrealEd/Public/ComponentVisualizerManager.h b/Engine/Source/Editor/UnrealEd/Public/ComponentVisualizerManager.h
--- a/Engine/Source/Editor/UnrealEd/Public/ComponentVisualizerManager.h
+++ b/Engine/Source/Editor/UnrealEd/Public/ComponentVisualizerManager.h
@@ -2,6 +2,8 @@
 
 #pragma once
 
+#include "ComponentVisualizer.h"
+
 /** Class that managed active component visualizer and routes input to it */
 class UNREALED_API FComponentVisualizerManager
 {
