Ensure that before BP regeneration all used MetaData are loaded. It wasn't ensured after cl#2056251.

diff --git a/Engine/Source/Editor/UnrealEd/Private/Kismet2/BlueprintEditorUtils.cpp b/Engine/Source/Editor/UnrealEd/Private/Kismet2/BlueprintEditorUtils.cpp
--- a/Engine/Source/Editor/UnrealEd/Private/Kismet2/BlueprintEditorUtils.cpp
+++ b/Engine/Source/Editor/UnrealEd/Private/Kismet2/BlueprintEditorUtils.cpp
@@ -582,6 +582,14 @@ struct FRegenerationHelper
 			}
 		}
 
+		{
+			UPackage* Package = Struct->GetOutermost();
+			check(Package);
+			UMetaData* MetaData = Package->GetMetaData();
+			check(MetaData);
+			FrocedLoad(MetaData);
+		}
+
 		const int32 OldPropertiesSize = Struct->GetPropertiesSize();
 		for (UField* Field = Struct->Children; Field; Field = Field->Next)
 		{
