Scene outliner folder expansion states are now stored per-user

The expansion states are stored in temporary files under Game/Saved/Config/WorldState. For now the file name is derived from the package path as we don't have anything more persistant at the moment. The absence of a json file will caused default behaviour of fully expanded.

This addresses TTP#334483 - Scene Outliner: Remember expansion state of folders between sessions (save with map?)

diff --git a/Engine/Source/Editor/SceneOutliner/Private/SSceneOutliner.cpp b/Engine/Source/Editor/SceneOutliner/Private/SSceneOutliner.cpp
--- a/Engine/Source/Editor/SceneOutliner/Private/SSceneOutliner.cpp
+++ b/Engine/Source/Editor/SceneOutliner/Private/SSceneOutliner.cpp
@@ -321,6 +321,9 @@ namespace SceneOutliner
 				// Called when an item is scrolled into view
 				.OnItemScrolledIntoView( this, &SSceneOutliner::OnOutlinerTreeItemScrolledIntoView )
 
+				// Called when an item is expanded or collapsed
+				.OnExpansionChanged(this, &SSceneOutliner::OnItemExpansionChanged)
+
 				// Called to child items for any given parent item
 				.OnGetChildren( this, &SSceneOutliner::OnGetChildrenForOutlinerTree )
 
@@ -910,11 +913,11 @@ namespace SceneOutliner
 			if (!IsShowingOnlySelected() && ShouldShowFolders())
 			{
 				// Add any folders which might match the current search terms
-				for (const auto& Path : FActorFolders::Get().GetFoldersForWorld(*World))
+				for (const auto& Pair : FActorFolders::Get().GetFolderPropertiesForWorld(*World))
 				{
-					if (!FolderToTreeItemMap.Find(Path))
+					if (!FolderToTreeItemMap.Find(Pair.Key))
 					{
-						TSharedRef<TOutlinerFolderTreeItem> NewFolderItem = MakeShareable(new TOutlinerFolderTreeItem(Path));
+						TSharedRef<TOutlinerFolderTreeItem> NewFolderItem = MakeShareable(new TOutlinerFolderTreeItem(Pair.Key));
 						if (!AddFolderToTree(NewFolderItem))
 						{
 							// If we didn't add that folder, attempt to add its parents too
@@ -1111,6 +1114,11 @@ namespace SceneOutliner
 		// Add the item to the map, making it available for display in the tree when its parent is looking for its children.
 		FolderToTreeItemMap.Add(InFolderItem->Path, InFolderItem);
 
+		if (FActorFolderProps* Props = FActorFolders::Get().GetFolderProperties(*RepresentingWorld, InFolderItem->Path))
+		{
+			OutlinerTreeView->SetItemExpansion(InFolderItem, Props->bIsExpanded);
+		}
+
 		NewItemActions.ItemHasBeenCreated(OutlinerTreeView.ToSharedRef(), InFolderItem);
 
 		return true;
@@ -1157,14 +1165,6 @@ namespace SceneOutliner
 				ExpansionStateInfo.CollapsedItems.Add(Item.Key.Get());
 			}
 		}
-
-		for (const auto& Item : FolderToTreeItemMap)
-		{
-			if (!OutlinerTreeView->IsItemExpanded(Item.Value))
-			{
-				ExpansionStateInfo.CollapsedFolders.Add(Item.Key);
-			}
-		}
 	}
 
 	void SSceneOutliner::SetParentsExpansionState(const FParentActorsSet& ParentActors, const FParentsExpansionState& ExpansionStateInfo) const
@@ -1176,13 +1176,6 @@ namespace SceneOutliner
 			const bool bShouldExpand = !ExpansionStateInfo.CollapsedItems.Contains(Actor);
 			OutlinerTreeView->SetItemExpansion(Item, bShouldExpand);
 		}
-
-		for (const auto& Item : FolderToTreeItemMap)
-		{
-			// Expand folder according to the previous state
-			const bool bShouldExpand = !ExpansionStateInfo.CollapsedFolders.Contains(Item.Key);
-			OutlinerTreeView->SetItemExpansion(Item.Value, bShouldExpand);
-		}
 	}
 
 	void SSceneOutliner::PopulateSearchStrings( const TOutlinerTreeItem& TreeItem, OUT TArray< FString >& OutSearchStrings ) const
@@ -1749,7 +1742,7 @@ namespace SceneOutliner
 	void SSceneOutliner::AddMoveToFolderOutliner(FMenuBuilder& MenuBuilder) const
 	{
 		// We don't show this if there is no world (unusual) or if there aren't any folders in the world
-		const int32 NumFoldersInWorld = FActorFolders::Get().GetFoldersForWorld(*RepresentingWorld).Num();
+		const int32 NumFoldersInWorld = FActorFolders::Get().GetFolderPropertiesForWorld(*RepresentingWorld).Num();
 		if (!RepresentingWorld || !NumFoldersInWorld)
 		{
 			return;
@@ -2046,7 +2039,7 @@ namespace SceneOutliner
 		NewItemActions.WhenCreated(NewFolderName, NewItemActuator::Select | NewItemActuator::Rename);
 	}
 
-	void SSceneOutliner::OnBroadcastFolderCreate(const UWorld& InWorld, FName NewPath)
+	void SSceneOutliner::OnBroadcastFolderCreate(UWorld& InWorld, FName NewPath)
 	{
 		if (!ShouldShowFolders() || &InWorld != RepresentingWorld)
 		{
@@ -2085,7 +2078,7 @@ namespace SceneOutliner
 		}
 	}
 
-	void SSceneOutliner::OnBroadcastFolderDelete(const UWorld& InWorld, FName Path)
+	void SSceneOutliner::OnBroadcastFolderDelete(UWorld& InWorld, FName Path)
 	{
 		if (&InWorld != RepresentingWorld)
 		{
@@ -2552,6 +2545,36 @@ namespace SceneOutliner
 			}
 		}
 	}
+	
+	void SSceneOutliner::OnItemExpansionChanged(FOutlinerTreeItemPtr TreeItem, bool bIsExpanded) const
+	{
+		if (RepresentingWorld && TreeItem->Type == TOutlinerTreeItem::Folder)
+		{
+			// Update the central store of folder properties with this folder's new expansion state
+
+			auto FolderItem = StaticCastSharedPtr<TOutlinerFolderTreeItem>(TreeItem);
+
+			if (FActorFolderProps* Props = FActorFolders::Get().GetFolderProperties(*RepresentingWorld, FolderItem->Path))
+			{
+				Props->bIsExpanded = bIsExpanded;
+			}
+
+			// Expand any children that are also expanded
+			for (const auto& WeakChild : FolderItem->ChildItems)
+			{
+				auto Child = WeakChild.Pin();
+				if (Child.IsValid() && Child->Type == TOutlinerTreeItem::Folder)
+				{
+					auto ChildItem = StaticCastSharedPtr<TOutlinerFolderTreeItem>(Child);
+					const FActorFolderProps* FolderPtr = FActorFolders::Get().GetFolderProperties(*RepresentingWorld, ChildItem->Path);
+					if (FolderPtr && FolderPtr->bIsExpanded)
+					{
+						OutlinerTreeView->SetItemExpansion(ChildItem, true);
+					}
+				}
+			}
+		}
+	}
 
 	void SSceneOutliner::OnLevelAdded(ULevel* InLevel, UWorld* InWorld)
 	{
@@ -2988,7 +3011,6 @@ namespace SceneOutliner
 				else
 				{
 					const FScopedTransaction Transaction( LOCTEXT("UndoAction_DeleteSelection", "Delete selection") );
-//					FActorFolders::Get().Modify(*GWorld);
 
 					// Delete selected folders too
 					auto SelectedItems = OutlinerTreeView->GetSelectedItems();
diff --git a/Engine/Source/Editor/SceneOutliner/Private/SSceneOutliner.h b/Engine/Source/Editor/SceneOutliner/Private/SSceneOutliner.h
--- a/Engine/Source/Editor/SceneOutliner/Private/SSceneOutliner.h
+++ b/Engine/Source/Editor/SceneOutliner/Private/SSceneOutliner.h
@@ -179,6 +179,9 @@ namespace SceneOutliner
 		/** Called by STreeView when an item is scrolled into view */
 		void OnOutlinerTreeItemScrolledIntoView( FOutlinerTreeItemPtr TreeItem, const TSharedPtr<ITableRow>& Widget );
 
+		/** Called when an item in the tree has been collapsed or expanded */
+		void OnItemExpansionChanged(FOutlinerTreeItemPtr TreeItem, bool bIsExpanded) const;
+
 		/** Called by USelection::SelectionChangedEvent delegate when the level's selection changes */
 		void OnLevelSelectionChanged(UObject* Obj);
 
@@ -255,13 +258,13 @@ namespace SceneOutliner
 		bool AddChildToParent(FOutlinerTreeItemRef TreeItem, FName ParentPath, bool bIgnoreSearchFilter = false);
 
 		/** Called when a folder is to be created */
-		void OnBroadcastFolderCreate(const UWorld& InWorld, FName NewPath);
+		void OnBroadcastFolderCreate(UWorld& InWorld, FName NewPath);
 
 		/** Called when a folder is to be moved */
 		void OnBroadcastFolderMove(FName OldPath, FName NewPath);
 
 		/** Called when a folder is to be deleted */
-		void OnBroadcastFolderDelete(const UWorld& InWorld, FName Path);
+		void OnBroadcastFolderDelete(UWorld& InWorld, FName Path);
 
 		/** Detach the specified item from its parent if specified */
 		void DetachChildFromParent(TSharedRef<TOutlinerTreeItem> Child, FName ParentPath);
@@ -380,9 +383,6 @@ namespace SceneOutliner
 		{
 			/** The set of parent items which are not expanded */
 			TSet<const AActor*> CollapsedItems;
-
-			/** The set of folders which are not expanded */
-			TSet<FName> CollapsedFolders;
 		};
 
 
diff --git a/Engine/Source/Editor/UnrealEd/Private/EditorActorFolders.cpp b/Engine/Source/Editor/UnrealEd/Private/EditorActorFolders.cpp
--- a/Engine/Source/Editor/UnrealEd/Private/EditorActorFolders.cpp
+++ b/Engine/Source/Editor/UnrealEd/Private/EditorActorFolders.cpp
@@ -10,6 +10,18 @@ UEditorActorFolders::UEditorActorFolders(const FPostConstructInitializePropertie
 	: Super(PCIP)
 { }
 
+void UEditorActorFolders::Serialize(FArchive& Ar)
+{
+	Ar << Folders;
+}
+
+FString GetWorldStateFilename(UPackage* Package)
+{
+	const FString PathName = Package->GetPathName();
+	const uint32 PathNameCrc = FCrc::MemCrc32(*PathName, sizeof(TCHAR)*PathName.Len());
+	return FPaths::Combine(*FPaths::GameSavedDir(), TEXT("Config"), TEXT("WorldState"), *FString::Printf(TEXT("%u.json"), PathNameCrc));
+}
+
 /** Convert an old path to a new path, replacing an ancestor branch with something else */
 FName OldPathToNewPath(const FString& InOldBranch, const FString& InNewBranch, const FString& PathToMove)
 {
@@ -25,15 +37,11 @@ FActorFolders::FActorFolders()
 {
 	check(GEngine);
 	GEngine->OnLevelActorFolderChanged().AddRaw(this, &FActorFolders::OnActorFolderChanged);
-	GEngine->OnLevelActorListChanged().AddRaw( this, &FActorFolders::OnLevelActorListChanged );
+	GEngine->OnLevelActorListChanged().AddRaw(this, &FActorFolders::OnLevelActorListChanged);
 	GEngine->OnLevelActorAdded().AddRaw(this, &FActorFolders::OnLevelActorAdded);
-	FEditorDelegates::MapChange.AddRaw( this, &FActorFolders::OnMapChange );
 
-	UWorld* World = GWorld;
-	if (World)
-	{
-		RebuildFolderListForWorld(*World);
-	}
+	FEditorDelegates::MapChange.AddRaw(this, &FActorFolders::OnMapChange);
+	FEditorDelegates::PostSaveWorld.AddRaw(this, &FActorFolders::OnWorldSaved);
 }
 
 FActorFolders::~FActorFolders()
@@ -41,7 +49,9 @@ FActorFolders::~FActorFolders()
 	check(GEngine);
 	GEngine->OnLevelActorFolderChanged().RemoveAll(this);
 	GEngine->OnLevelActorListChanged().RemoveAll(this);
+
 	FEditorDelegates::MapChange.RemoveAll(this);
+	FEditorDelegates::PostSaveWorld.RemoveAll(this);
 }
 
 void FActorFolders::AddReferencedObjects(FReferenceCollector& Collector)
@@ -97,38 +107,67 @@ void FActorFolders::OnMapChange(uint32 MapChangeFlags)
 	RebuildFolderListForWorld(*GWorld);
 }
 
-void FActorFolders::OnLevelActorAdded(AActor* Actor)
+void FActorFolders::OnWorldSaved(uint32 SaveFlags, UWorld* World, bool bSuccess)
 {
-	check(Actor && Actor->GetWorld());
+	// Attempt to save the folder state
+	const UEditorActorFolders* const* ExisingFolders = TemporaryWorldFolders.Find(World);
 
-	if (!Actor->GetFolderPath().IsNone())
+	if (ExisingFolders)
 	{
-		UEditorActorFolders& Folders = GetFoldersForWorld_Internal(*Actor->GetWorld());
-		const FName& ActorFolder = Actor->GetFolderPath();
-		if (!Folders.Paths.Contains(ActorFolder))
+		const auto Filename = GetWorldStateFilename(World->GetOutermost());
+		TUniquePtr<FArchive> Ar(IFileManager::Get().CreateFileWriter(*Filename));
+		if (Ar)
 		{
-			const FScopedTransaction Transaction(LOCTEXT("UndoAction_ActorAdded", "Actor Folder Added"));
-			Folders.Modify();
-			Folders.Paths.Add(ActorFolder);
+			TSharedRef<FJsonObject> RootObject = MakeShareable(new FJsonObject);
+			TSharedRef<FJsonObject> JsonFolders = MakeShareable(new FJsonObject);
+
+			for (const auto& KeyValue : (*ExisingFolders)->Folders)
+			{
+				TSharedRef<FJsonObject> JsonFolder = MakeShareable(new FJsonObject);
+				JsonFolder->SetBoolField(TEXT("bIsExpanded"), KeyValue.Value.bIsExpanded);
+
+				JsonFolders->SetObjectField(KeyValue.Key.ToString(), JsonFolder);
+			}
+
+			RootObject->SetObjectField(TEXT("Folders"), JsonFolders);
+
+			{
+				auto Writer = TJsonWriterFactory<TCHAR>::Create(Ar.Get());
+				FJsonSerializer::Serialize(RootObject, Writer);
+				Ar->Close();
+			}
 		}
 	}
 }
 
-void FActorFolders::OnActorFolderChanged(const AActor* InActor, FName OldPath)
+void FActorFolders::OnLevelActorAdded(AActor* InActor)
 {
-	if (const auto* World = InActor->GetWorld())
+	check(InActor && InActor->GetWorld());
+
+	FScopedTransaction Transaction(LOCTEXT("UndoAction_ActorAdded", "Actor Folder Added"));
+
+	if (!AddFolderToWorld(*InActor->GetWorld(), InActor->GetFolderPath()))
 	{
-		const FScopedTransaction Transaction(LOCTEXT("UndoAction_FolderChanged", "Actor Folder Changed"));
+		Transaction.Cancel();
+	}
+}
 
-		UEditorActorFolders& Folders = GetFoldersForWorld_Internal(*World);
-		Folders.Modify();
+void FActorFolders::OnActorFolderChanged(const AActor* InActor, FName OldPath)
+{
+	check(InActor && InActor->GetWorld());
 
-		const auto NewPath = InActor->GetFolderPath();
-		if (!NewPath.IsNone() && !Folders.Paths.Contains(NewPath))
-		{
-			Folders.Paths.Add(NewPath);
-			OnFolderCreate.Broadcast(*World, NewPath);
-		}
+	FScopedTransaction Transaction(LOCTEXT("UndoAction_FolderChanged", "Actor Folder Changed"));
+
+	UWorld* World = InActor->GetWorld();
+	const auto NewPath = InActor->GetFolderPath();
+
+	if (AddFolderToWorld(*World, NewPath))
+	{
+		OnFolderCreate.Broadcast(*World, NewPath);
+	}
+	else
+	{
+		Transaction.Cancel();
 	}
 }
 
@@ -143,52 +182,84 @@ bool FActorFolders::PathIsChildOf(const FString& InPotentialChild, const FString
 
 void FActorFolders::RebuildFolderListForWorld(UWorld& InWorld)
 {
-	GetFoldersForWorld_Internal(InWorld).Paths = ScanWorldForFolders(InWorld).Array();
+	if (TemporaryWorldFolders.Contains(&InWorld))
+	{
+		// We don't empty the existing folders so that we keep empty ones.
+		// Explicitly deleted folders will already be removed from the list
+
+		// Iterate over every actor in memory. WARNING: This is potentially very expensive!
+		for (FActorIterator ActorIt(&InWorld); ActorIt; ++ActorIt)
+		{
+			AddFolderToWorld(InWorld, ActorIt->GetFolderPath());
+		}
+	}
+	else
+	{
+		// No folders exist for this world yet - creating them will ensure they're up to date
+		InitializeForWorld(InWorld);
+	}
+}
+
+const TMap<FName, FActorFolderProps>& FActorFolders::GetFolderPropertiesForWorld(UWorld& InWorld)
+{
+	return GetOrCreateFoldersForWorld(InWorld).Folders;
 }
 
-const TArray<FName>& FActorFolders::GetFoldersForWorld(const UWorld& InWorld)
+FActorFolderProps* FActorFolders::GetFolderProperties(UWorld& InWorld, FName InPath)
 {
-	return GetFoldersForWorld_Internal(InWorld).Paths;
+	return GetOrCreateFoldersForWorld(InWorld).Folders.Find(InPath);
 }
 
-UEditorActorFolders& FActorFolders::GetFoldersForWorld_Internal(const UWorld& InWorld)
+UEditorActorFolders& FActorFolders::GetOrCreateFoldersForWorld(UWorld& InWorld)
 {
 	if (UEditorActorFolders** Folders = TemporaryWorldFolders.Find(&InWorld))
 	{
 		return **Folders;
 	}
 
+	return InitializeForWorld(InWorld);
+}
+
+UEditorActorFolders& FActorFolders::InitializeForWorld(UWorld& InWorld)
+{
 	// We intentionally don't pass RF_Transactional to ConstructObject so that we don't record the creation of the object into the undo buffer
 	// (to stop it getting deleted on undo as we manage its lifetime), but we still want it to be RF_Transactional so we can record any changes later
 	UEditorActorFolders* Folders = ConstructObject<UEditorActorFolders>(UEditorActorFolders::StaticClass(), GetTransientPackage(), NAME_None, RF_NoFlags);
 	Folders->SetFlags(RF_Transactional);
 	TemporaryWorldFolders.Add(&InWorld, Folders);
-	return *Folders;
-}
-
-TSet<FName> FActorFolders::ScanWorldForFolders(UWorld& InWorld)
-{
-	TSet<FName> Folders;
 
-	// Iterate over every actor in memory. WARNING: This is potentially very expensive!
-	for( FActorIterator ActorIt(&InWorld); ActorIt; ++ActorIt )
+	// Ensure the list is entirely up to date with the world before we write our serialized properties into it.
+	for (FActorIterator ActorIt(&InWorld); ActorIt; ++ActorIt)
 	{
-		AActor* Actor = *ActorIt;
-		check(Actor);
-		const FName& Path = Actor->GetFolderPath();
-		if (!Path.IsNone())
-		{
-			Folders.Add(Actor->GetFolderPath());
-		}
+		AddFolderToWorld(InWorld, ActorIt->GetFolderPath());
 	}
 
-	// Add any empty folders
-	for (const auto& Folder : GetFoldersForWorld_Internal(InWorld).Paths)
+	// Attempt to load the folder properties from this user's saved world state directory
+	const auto Filename = GetWorldStateFilename(InWorld.GetOutermost());
+	TUniquePtr<FArchive> Ar(IFileManager::Get().CreateFileReader(*Filename));
+	if (Ar)
 	{
-		Folders.Add(Folder);
+		TSharedPtr<FJsonObject> RootObject = MakeShareable(new FJsonObject);
+
+		auto Reader = TJsonReaderFactory<TCHAR>::Create(Ar.Get());
+		if (FJsonSerializer::Deserialize(Reader, RootObject))
+		{
+			const TSharedPtr<FJsonObject>& JsonFolders = RootObject->GetObjectField(TEXT("Folders"));
+			for (const auto& KeyValue : JsonFolders->Values)
+			{
+				// Only pull in the folder's properties if this folder still exists in the world.
+				// This means that old stale folders won't re-appear in the world (they'll won't get serialized when the world is saved anyway)
+				if (FActorFolderProps* FolderInWorld = Folders->Folders.Find(*KeyValue.Key))
+				{
+					auto FolderProperties = KeyValue.Value->AsObject();
+					FolderInWorld->bIsExpanded = FolderProperties->GetBoolField(TEXT("bIsExpanded"));
+				}
+			}
+		}
+		Ar->Close();
 	}
 
-	return Folders;
+	return *Folders;
 }
 
 FName FActorFolders::GetDefaultFolderNameForSelection(UWorld& InWorld)
@@ -215,7 +286,7 @@ FName FActorFolders::GetDefaultFolderNameForSelection(UWorld& InWorld)
 FName FActorFolders::GetDefaultFolderName(UWorld& InWorld, FName ParentPath)
 {
 	// This is potentially very slow but necessary to find a unique name
-	const TSet<FName> ExistingFolders = ScanWorldForFolders(InWorld);
+	const auto& ExistingFolders = GetFolderPropertiesForWorld(InWorld);
 
 	// Create a valid base name for this folder
 	uint32 Suffix = 1;
@@ -257,26 +328,27 @@ void FActorFolders::CreateFolderContainingSelection(UWorld& InWorld, FName Path)
 
 void FActorFolders::CreateFolder(UWorld& InWorld, FName Path)
 {
-	const FScopedTransaction Transaction(LOCTEXT("UndoAction_CreateFolder", "Create Folder"));
+	FScopedTransaction Transaction(LOCTEXT("UndoAction_CreateFolder", "Create Folder"));
 
-	UEditorActorFolders& Folders = GetFoldersForWorld_Internal(InWorld);
-	if (!Folders.Paths.Contains(Path))
+	if (AddFolderToWorld(InWorld, Path))
 	{
-		Folders.Modify();
-		Folders.Paths.Add(Path);
 		OnFolderCreate.Broadcast(InWorld, Path);
 	}
+	else
+	{
+		Transaction.Cancel();
+	}
 }
 
-void FActorFolders::DeleteFolder(const UWorld& InWorld, FName FolderToDelete)
+void FActorFolders::DeleteFolder(UWorld& InWorld, FName FolderToDelete)
 {
 	const FScopedTransaction Transaction(LOCTEXT("UndoAction_DeleteFolder", "Delete Folder"));
 
-	UEditorActorFolders& Folders = GetFoldersForWorld_Internal(InWorld);
-	if (Folders.Paths.Contains(FolderToDelete))
+	UEditorActorFolders& Folders = GetOrCreateFoldersForWorld(InWorld);
+	if (Folders.Folders.Contains(FolderToDelete))
 	{
 		Folders.Modify();
-		Folders.Paths.Remove(FolderToDelete);
+		Folders.Folders.Remove(FolderToDelete);
 		OnFolderDelete.Broadcast(InWorld, FolderToDelete);
 	}
 }
@@ -296,19 +368,30 @@ bool FActorFolders::RenameFolderInWorld(UWorld& World, FName OldPath, FName NewP
 	TSet<FName> RenamedFolders;
 
 	// Move any folders we currently hold - old ones will be deleted later
-	UEditorActorFolders& FoldersInWorld = GetFoldersForWorld_Internal(World);
+	UEditorActorFolders& FoldersInWorld = GetOrCreateFoldersForWorld(World);
 	FoldersInWorld.Modify();
 
-	auto PathsCopy = FoldersInWorld.Paths;
-	for (const auto& Path : PathsCopy)
+	auto ExistingFoldersCopy = FoldersInWorld.Folders;
+	for (const auto& Pair : ExistingFoldersCopy)
 	{
+		auto Path = Pair.Key;
+
 		const FString FolderPath = Path.ToString();
 		if (OldPath == Path || PathIsChildOf(FolderPath, OldPathString))
 		{
 			const FName NewFolder = OldPathToNewPath(OldPathString, NewPathString, FolderPath);
-			if (!FoldersInWorld.Paths.Contains(NewFolder))
+			if (!FoldersInWorld.Folders.Contains(NewFolder))
 			{
-				FoldersInWorld.Paths.Add(NewFolder);
+				// Use the existing properties for the folder if we have them
+				if (FActorFolderProps* ExistingProperties = FoldersInWorld.Folders.Find(Path))
+				{
+					FoldersInWorld.Folders.Add(NewFolder, *ExistingProperties);
+				}
+				else
+				{
+					// Otherwise use default properties
+					FoldersInWorld.Folders.Add(NewFolder);
+				}
 				OnFolderCreate.Broadcast(World, NewFolder);
 			}
 			RenamedFolders.Add(Path);
@@ -336,11 +419,36 @@ bool FActorFolders::RenameFolderInWorld(UWorld& World, FName OldPath, FName NewP
 	// Cleanup any old folders
 	for (const auto& Path : RenamedFolders)
 	{
-		FoldersInWorld.Paths.Remove(Path);
+		FoldersInWorld.Folders.Remove(Path);
 		OnFolderDelete.Broadcast(World, Path);
 	}
 
 	return RenamedFolders.Num() != 0;
 }
 
+bool FActorFolders::AddFolderToWorld(UWorld& InWorld, FName Path)
+{
+	if (!Path.IsNone())
+	{
+		UEditorActorFolders& Folders = GetOrCreateFoldersForWorld(InWorld);
+
+		if (!Folders.Folders.Contains(Path))
+		{
+			// Add the parent as well
+			const FName ParentPath(*FPaths::GetPath(Path.ToString()));
+			if (!ParentPath.IsNone())
+			{
+				AddFolderToWorld(InWorld, ParentPath);
+			}
+
+			Folders.Modify();
+			Folders.Folders.Add(Path);
+
+			return true;
+		}
+	}
+
+	return false;
+}
+
 #undef LOCTEXT_NAMESPACE
\ No newline at end of file
diff --git a/Engine/Source/Editor/UnrealEd/Public/EditorActorFolders.h b/Engine/Source/Editor/UnrealEd/Public/EditorActorFolders.h
--- a/Engine/Source/Editor/UnrealEd/Public/EditorActorFolders.h
+++ b/Engine/Source/Editor/UnrealEd/Public/EditorActorFolders.h
@@ -5,18 +5,34 @@
 #include "EditorActorFolders.generated.h"
 
 /** Multicast delegates for broadcasting various folder events */
-DECLARE_MULTICAST_DELEGATE_TwoParams(FOnActorFolderCreate, const UWorld&, FName);
-DECLARE_MULTICAST_DELEGATE_ThreeParams(FOnActorFolderMove, const UWorld&, FName, FName);
-DECLARE_MULTICAST_DELEGATE_TwoParams(FOnActorFolderDelete, const UWorld&, FName);
+DECLARE_MULTICAST_DELEGATE_TwoParams(FOnActorFolderCreate, UWorld&, FName);
+DECLARE_MULTICAST_DELEGATE_TwoParams(FOnActorFolderDelete, UWorld&, FName);
+
+USTRUCT()
+struct FActorFolderProps
+{
+	GENERATED_USTRUCT_BODY()
+
+	FActorFolderProps() : bIsExpanded(true) {}
+
+	/** Serializer */
+	FORCEINLINE friend FArchive& operator<<(FArchive& Ar, FActorFolderProps& Folder)
+	{
+		return Ar << Folder.bIsExpanded;
+	}
+
+	bool bIsExpanded;
+};
 
 /** Actor Folder UObject. This is used to support undo/redo reliably */
 UCLASS()
 class UEditorActorFolders : public UObject
 {
 	GENERATED_UCLASS_BODY()
+public:
+	virtual void Serialize(FArchive& Ar) OVERRIDE;
 
-	UPROPERTY()
-	TArray<FName> Paths;
+	TMap<FName, FActorFolderProps> Folders;
 };
 
 /** Class responsible for managing an in-memory representation of actor folders in the editor */
@@ -45,8 +61,11 @@ struct UNREALED_API FActorFolders : public FGCObject
 	/** Check if the specified path is a child of the specified parent */
 	static bool PathIsChildOf(const FString& InPotentialChild, const FString& InParent);
 
-	/** Get an array of folders for the specified world */
-	const TArray<FName>& GetFoldersForWorld(const UWorld& InWorld);
+	/** Get a map of folder properties for the specified world (map of folder path -> properties) */
+	const TMap<FName, FActorFolderProps>& GetFolderPropertiesForWorld(UWorld& InWorld);
+
+	/** Get the folder properties for the specified path. Returns nullptr if no properties exist */
+	FActorFolderProps* GetFolderProperties(UWorld& InWorld, FName InPath);
 
 	/** Get a default folder name under the specified parent path */
 	FName GetDefaultFolderName(UWorld& InWorld, FName ParentPath = FName());
@@ -61,15 +80,18 @@ struct UNREALED_API FActorFolders : public FGCObject
 	void CreateFolderContainingSelection(UWorld& InWorld, FName Path);
 
 	/** Delete the specified folder in the world */
-	void DeleteFolder(const UWorld& InWorld, FName FolderToDelete);
+	void DeleteFolder(UWorld& InWorld, FName FolderToDelete);
 
 	/** Rename the specified path to a new name */
 	bool RenameFolderInWorld(UWorld& World, FName OldPath, FName NewPath);
 
 private:
 
-	/** Internal (non-const) implementation of GetFoldersForWorld */
-	UEditorActorFolders& GetFoldersForWorld_Internal(const UWorld& InWorld);
+	/** Get or create a folder container for the specified world */
+	UEditorActorFolders& GetOrCreateFoldersForWorld(UWorld& InWorld);
+
+	/** Create and update a folder container for the specified world */
+	UEditorActorFolders& InitializeForWorld(UWorld& InWorld);
 
 	/** Rebuild the folder list for the specified world. This can be very slow as it
 		iterates all actors in memory to rebuild the array of actors for this world */
@@ -84,15 +106,17 @@ private:
 	/** Called when the global map in the editor has changed */
 	void OnMapChange(uint32 MapChangeFlags);
 
+	/** Called after a world has been saved */
+	void OnWorldSaved(uint32 SaveFlags, UWorld* World, bool bSuccess);
+
 	/** Called when an actor is added to the world */
 	void OnLevelActorAdded(AActor* Actor);
 
 	/** Remove any references to folder arrays for dead worlds */
 	void Housekeeping();
 
-	/** Scan the specified world for folders. This can be very slow as it iterates
-		all actors in memory to rebuild the array of actors for this world */
-	TSet<FName> ScanWorldForFolders(UWorld& InWorld);
+	/** Add a folder to the folder map for the specified world. Does not trigger any events. */
+	bool AddFolderToWorld(UWorld& InWorld, FName Path);
 
 	/** Transient map of folders, keyed on world pointer */
 	TMap<TWeakObjectPtr<UWorld>, UEditorActorFolders*> TemporaryWorldFolders;
