Remove deferred release of clothing apex.

We can't defer this because apex uses unique names. The defer is not needed anyway because we're authoring so assumes threading is safe.

TTP#336184

diff --git a/Engine/Source/Editor/UnrealEd/Private/ApexClothingUtils.cpp b/Engine/Source/Editor/UnrealEd/Private/ApexClothingUtils.cpp
--- a/Engine/Source/Editor/UnrealEd/Private/ApexClothingUtils.cpp
+++ b/Engine/Source/Editor/UnrealEd/Private/ApexClothingUtils.cpp
@@ -1151,7 +1151,7 @@ NxClothingAsset* ApplyTransform(NxClothingAsset* ApexClothingAsset)
 		FCStringAnsi::Strncpy(ApexAssetName, ApexClothingAsset->getName(), sizeof(ApexAssetName) );
 
 		// destroy old asset
-		GPhysCommandHandler->DeferredRelease(ApexClothingAsset);
+		ApexClothingAsset->release();
 		ApexClothingAsset = NULL;
 
 		// create new asset from the authoring
@@ -1167,7 +1167,7 @@ NxClothingAsset* ApplyTransform(NxClothingAsset* ApexClothingAsset)
 		check(NewAsset);
 
 		// release authoring
-		GPhysCommandHandler->DeferredRelease(ApexClothingAssetAuthoring);
+		ApexClothingAssetAuthoring->release();
 
 		return NewAsset;
 	}
@@ -1580,7 +1580,7 @@ EClothUtilRetType ImportApexAssetFromApexFile(FString& ApexFile, USkeletalMesh*
 	//when re-importing, just skip this warning dialog
 	if(!bReimport && !IsProperApexFile(*ApexClothingAsset, SkelMesh))
 	{
-		GPhysCommandHandler->DeferredRelease(ApexClothingAsset);
+		ApexClothingAsset->release();
 		return CURT_Cancel;
 	}
 
