Temp fix incorrect skinning on some fbx imports.

Back out revision 137 from //depot/UE4/Engine/Source/Editor/UnrealEd/Private/Fbx/FbxSkeletalMeshImport.cpp

diff --git a/Engine/Source/Editor/UnrealEd/Private/Fbx/FbxSkeletalMeshImport.cpp b/Engine/Source/Editor/UnrealEd/Private/Fbx/FbxSkeletalMeshImport.cpp
--- a/Engine/Source/Editor/UnrealEd/Private/Fbx/FbxSkeletalMeshImport.cpp
+++ b/Engine/Source/Editor/UnrealEd/Private/Fbx/FbxSkeletalMeshImport.cpp
@@ -834,7 +834,9 @@ bool UnFbx::FFbxImporter::ImportBone(TArray<FbxNode*>& NodeArray, FSkeletalMeshI
 	GlobalsPerLink[0].SetIdentity();
 	
 	bool GlobalLinkFoundFlag;
-	FTransform LocalTransform;
+	FbxVector4 LocalLinkT;
+	FbxQuaternion LocalLinkQ;
+	FbxVector4	LocalLinkS;
 
 	bool bAnyLinksNotInBindPose = false;
 	FString LinksWithoutBindPoses;
@@ -941,13 +943,18 @@ bool UnFbx::FFbxImporter::ImportBone(TArray<FbxNode*>& NodeArray, FSkeletalMeshI
 
 		if (LinkIndex)
 		{
-			FTransform Parent = Converter.ConvertTransform(GlobalsPerLink[ParentIndex]);
-			FTransform Current = Converter.ConvertTransform(GlobalsPerLink[LinkIndex]);
-			LocalTransform = Current.GetRelativeTransform(Parent);
+			FbxAMatrix	Matrix;
+			Matrix = GlobalsPerLink[ParentIndex].Inverse() * GlobalsPerLink[LinkIndex];
+			LocalLinkT = Matrix.GetT();
+			LocalLinkQ = Matrix.GetQ();
+			LocalLinkS = Matrix.GetS();
 		}
 		else	// skeleton root
 		{
-			LocalTransform = Converter.ConvertTransform(GlobalsPerLink[LinkIndex]);
+			// for root, this is global coordinate
+			LocalLinkT = GlobalsPerLink[LinkIndex].GetT();
+			LocalLinkQ = GlobalsPerLink[LinkIndex].GetQ();
+			LocalLinkS = GlobalsPerLink[LinkIndex].GetS();
 		}
 		
 		// set bone
@@ -987,7 +994,9 @@ bool UnFbx::FFbxImporter::ImportBone(TArray<FbxNode*>& NodeArray, FSkeletalMeshI
 			}
 		}
 
-		JointMatrix.Transform = LocalTransform;
+		JointMatrix.Transform.SetTranslation(Converter.ConvertPos(LocalLinkT));
+		JointMatrix.Transform.SetRotation(Converter.ConvertRotToQuat(LocalLinkQ));
+		JointMatrix.Transform.SetScale3D(Converter.ConvertScale(LocalLinkS));
 	}
 	
 	if(bAnyLinksNotInBindPose)
@@ -1906,11 +1915,29 @@ bool UnFbx::FFbxImporter::FillSkelMeshImporterFromFbx( FSkeletalMeshImportData&
 	int32 ExistPointNum = ImportData.Points.Num();
 	ImportData.Points.AddUninitialized(ControlPointsCount);
 
-	SkinControlPointsToPose(ImportData, Node->GetMesh(), FbxShape, false);
 	// Construct the matrices for the conversion from right handed to left handed system
-
 	FbxAMatrix TotalMatrix;
 	FbxAMatrix TotalMatrixForNormal;
+	TotalMatrix = ComputeTotalMatrix(Node);
+	TotalMatrixForNormal = TotalMatrix.Inverse();
+	TotalMatrixForNormal = TotalMatrixForNormal.Transpose();
+
+	int32 ControlPointsIndex;
+	for( ControlPointsIndex = 0 ; ControlPointsIndex < ControlPointsCount ;ControlPointsIndex++ )
+	{
+		FbxVector4 Position;
+		if (FbxShape)
+		{
+			Position = FbxShape->GetControlPoints()[ControlPointsIndex];
+		}
+		else
+		{
+			Position = Mesh->GetControlPoints()[ControlPointsIndex];
+		}																	 
+		FbxVector4 FinalPosition;
+		FinalPosition = TotalMatrix.MultT(Position);
+		ImportData.Points[ControlPointsIndex+ExistPointNum] = Converter.ConvertPos(FinalPosition);
+	}
 	
 	bool OddNegativeScale = IsOddNegativeScale(TotalMatrix);
 	
