Move forward declaration of UInterpTrackMoveAxis out of UnFbx namespace

diff --git a/Engine/Source/Editor/UnrealEd/Private/FbxImporter.h b/Engine/Source/Editor/UnrealEd/Private/FbxImporter.h
--- a/Engine/Source/Editor/UnrealEd/Private/FbxImporter.h
+++ b/Engine/Source/Editor/UnrealEd/Private/FbxImporter.h
@@ -88,6 +88,7 @@
 class FSkeletalMeshImportData;
 class FSkelMeshOptionalImportData;
 class ASkeletalMeshActor;
+class UInterpTrackMoveAxis;
 struct FbxSceneInfo;
 
 DECLARE_LOG_CATEGORY_EXTERN(LogFbx, Log, All);
@@ -914,7 +915,7 @@ protected:
 	/**
 	 * Imports an FBX transform curve into a movement subtrack
 	 */
-	void ImportMoveSubTrack( FbxAnimCurve* FbxCurve, int32 FbxDimension, class UInterpTrackMoveAxis* SubTrack, int32 CurveIndex, bool bNegative, FbxAnimCurve* RealCurve, float DefaultVal );
+	void ImportMoveSubTrack( FbxAnimCurve* FbxCurve, int32 FbxDimension, UInterpTrackMoveAxis* SubTrack, int32 CurveIndex, bool bNegative, FbxAnimCurve* RealCurve, float DefaultVal );
 
 	/**
 	 * Imports a FBX animated element into a Matinee track.
