Move FAnimControlTrackKey forward declaration out of UnFbx namespace

diff --git a/Engine/Source/Editor/UnrealEd/Private/FbxExporter.h b/Engine/Source/Editor/UnrealEd/Private/FbxExporter.h
--- a/Engine/Source/Editor/UnrealEd/Private/FbxExporter.h
+++ b/Engine/Source/Editor/UnrealEd/Private/FbxExporter.h
@@ -36,6 +36,8 @@
 
 #include "Landscape/LandscapeProxy.h"
 
+struct FAnimControlTrackKey;
+
 namespace UnFbx
 {
 
@@ -150,7 +152,7 @@ public:
 	/**
 	 * Exports the list of UAnimSequences as a single animation based on the settings in the TrackKeys
 	 */
-	UNREALED_API void ExportAnimSequencesAsSingle( USkeletalMesh* SkelMesh, const ASkeletalMeshActor* SkelMeshActor, const FString& ExportName, const TArray<UAnimSequence*>& AnimSeqList, const TArray<struct FAnimControlTrackKey>& TrackKeys );
+	UNREALED_API void ExportAnimSequencesAsSingle( USkeletalMesh* SkelMesh, const ASkeletalMeshActor* SkelMeshActor, const FString& ExportName, const TArray<UAnimSequence*>& AnimSeqList, const TArray<FAnimControlTrackKey>& TrackKeys );
 
 private:
 	FFbxExporter();
