BP, UDS: Arrays in UDS can have default value

diff --git a/Engine/Source/Editor/BlueprintGraph/Private/EdGraphSchema_K2.cpp b/Engine/Source/Editor/BlueprintGraph/Private/EdGraphSchema_K2.cpp
--- a/Engine/Source/Editor/BlueprintGraph/Private/EdGraphSchema_K2.cpp
+++ b/Engine/Source/Editor/BlueprintGraph/Private/EdGraphSchema_K2.cpp
@@ -1655,19 +1655,14 @@ FString UEdGraphSchema_K2::IsPinDefaultValid(const UEdGraphPin* Pin, const FStri
 	const bool bIsReference = Pin->PinType.bIsReference;
 	const bool bIsAutoCreateRefTerm = IsAutoCreateRefTerm(Pin);
 
-	if (OwningBP->BlueprintType != BPTYPE_Interface)
+	if (bIsArray)
 	{
-		if( !bIsAutoCreateRefTerm )
-		{
-			if( bIsArray )
-			{
-				return TEXT("Literal values are not allowed for array parameters.  Use a Make Array node instead");
-			}
-			else if( bIsReference )
-			{
-				return TEXT("Literal values are not allowed for pass-by-reference parameters.");
-			}
-		}
+		return TEXT("Literal values are not allowed for array parameters.  Use a Make Array node instead");
+	}
+
+	if ((OwningBP->BlueprintType != BPTYPE_Interface) && !bIsAutoCreateRefTerm && bIsReference)
+	{
+		return TEXT("Literal values are not allowed for pass-by-reference parameters.");
 	}
 
 	FString ReturnMsg;
@@ -1700,16 +1695,18 @@ bool UEdGraphSchema_K2::DefaultValueSimpleValidation(const FEdGraphPinType& PinT
 #endif
 #define DVSV_RETURN_MSG(Str) if(NULL != OutMsg) { *OutMsg = Str; } return false;
 
-	const bool bIsArray = PinType.bIsArray;
-	const bool bIsReference = PinType.bIsReference;
 	const FString& PinCategory = PinType.PinCategory;
 	const FString& PinSubCategory = PinType.PinSubCategory;
 	const UObject* PinSubCategoryObject = PinType.PinSubCategoryObject.Get();
 
+	if (PinType.bIsArray)
+	{
+		// arrays are validated separately
+	}
 	//@TODO: FCString::Atoi, FCString::Atof, and appStringToBool will 'accept' any input, but we should probably catch and warn
 	// about invalid input (non numeric for int/byte/float, and non 0/1 or yes/no/true/false for bool)
 
-	if (PinCategory == PC_Boolean)
+	else if(PinCategory == PC_Boolean)
 	{
 		// All input is acceptable to some degree
 	}
diff --git a/Engine/Source/Editor/UnrealEd/Private/Kismet2/BlueprintEditorUtils.cpp b/Engine/Source/Editor/UnrealEd/Private/Kismet2/BlueprintEditorUtils.cpp
--- a/Engine/Source/Editor/UnrealEd/Private/Kismet2/BlueprintEditorUtils.cpp
+++ b/Engine/Source/Editor/UnrealEd/Private/Kismet2/BlueprintEditorUtils.cpp
@@ -5561,19 +5561,12 @@ bool FBlueprintEditorUtils::PropertyValueFromString(const UProperty* Property, c
 		{
 			CastChecked<UObjectPropertyBase>(Property)->SetObjectPropertyValue_InContainer(DefaultObject, FindObject<UObject>(ANY_PACKAGE, *Value));
 		}
-		else if( Property->IsA(UArrayProperty::StaticClass()) )
+		else if (Property->IsA(UArrayProperty::StaticClass()) 
+			|| Property->IsA(UMulticastDelegateProperty::StaticClass()))
 		{
-			const UArrayProperty* ArrayProp = CastChecked<UArrayProperty>(Property);
 			FStringOutputDevice ImportError;
-			ArrayProp->ImportText(*Value, ArrayProp->ContainerPtrToValuePtr<uint8>(DefaultObject), 0, NULL, &ImportError);
-			bParseSuccedded = ImportError.IsEmpty();
-		}
-		else if( Property->IsA(UMulticastDelegateProperty::StaticClass()) )
-		{
-			const UMulticastDelegateProperty* MCDelegateProp = CastChecked<UMulticastDelegateProperty>(Property);
-			FStringOutputDevice ImportError;
-			MCDelegateProp->ImportText(*Value, MCDelegateProp->ContainerPtrToValuePtr<uint8>(DefaultObject), 0, NULL, &ImportError);
-			bParseSuccedded = ImportError.IsEmpty();
+			const auto EndOfParsedBuff = Property->ImportText(Value.IsEmpty() ? TEXT("()") : *Value, Property->ContainerPtrToValuePtr<uint8>(DefaultObject), 0, NULL, &ImportError);
+			bParseSuccedded = EndOfParsedBuff && ImportError.IsEmpty();
 		}
 		else
 		{
@@ -5627,8 +5620,8 @@ bool FBlueprintEditorUtils::PropertyValueFromString(const UProperty* Property, c
 			bParseSuccedded = FStructureEditorUtils::Fill_MakeStructureDefaultValue(Cast<const UUserDefinedStruct>(Struct), StructData);
 
 			FStringOutputDevice ImportError;
-			StructProperty->ImportText(Value.IsEmpty() ? TEXT("()") : *Value, StructData, 0, NULL, &ImportError);
-			bParseSuccedded &= ImportError.IsEmpty();
+			const auto EndOfParsedBuff = StructProperty->ImportText(Value.IsEmpty() ? TEXT("()") : *Value, StructData, 0, NULL, &ImportError);
+			bParseSuccedded &= EndOfParsedBuff && ImportError.IsEmpty();
 		}
 	}
 
diff --git a/Engine/Source/Editor/UnrealEd/Private/Kismet2/StructureEditorUtils.cpp b/Engine/Source/Editor/UnrealEd/Private/Kismet2/StructureEditorUtils.cpp
--- a/Engine/Source/Editor/UnrealEd/Private/Kismet2/StructureEditorUtils.cpp
+++ b/Engine/Source/Editor/UnrealEd/Private/Kismet2/StructureEditorUtils.cpp
@@ -344,12 +344,24 @@ bool FStructureEditorUtils::ChangeVariableDefaultValue(UUserDefinedStruct* Struc
 		&& (NewDefaultValue != VarDesc->DefaultValue) 
 		&& K2Schema->DefaultValueSimpleValidation(VarDesc->ToPinType(), FString(), NewDefaultValue, NULL, FText::GetEmpty()))
 	{
-		const FScopedTransaction Transaction(LOCTEXT("ChangeVariableDefaultValue", "Change Variable Default Value"));
-		ModifyStructData(Struct);
+		bool bAdvancedValidation = true;
+		if (!NewDefaultValue.IsEmpty())
+		{
+			const auto Property = FindField<UProperty>(Struct, VarDesc->VarName);
+			FStructOnScope StructDefaultMem(Struct);
+			bAdvancedValidation = StructDefaultMem.GetStructMemory() && Property &&
+				FBlueprintEditorUtils::PropertyValueFromString(Property, NewDefaultValue, StructDefaultMem.GetStructMemory());
+		}
 
-		VarDesc->DefaultValue = NewDefaultValue;
-		OnStructureChanged(Struct);
-		return true;
+		if (bAdvancedValidation)
+		{
+			const FScopedTransaction Transaction(LOCTEXT("ChangeVariableDefaultValue", "Change Variable Default Value"));
+			ModifyStructData(Struct);
+
+			VarDesc->DefaultValue = NewDefaultValue;
+			OnStructureChanged(Struct);
+			return true;
+		}
 	}
 	return false;
 }
