UE4: Removed unused function PackageTools::SavePackages.

diff --git a/Engine/Source/Editor/UnrealEd/Private/PackageTools.cpp b/Engine/Source/Editor/UnrealEd/Private/PackageTools.cpp
--- a/Engine/Source/Editor/UnrealEd/Private/PackageTools.cpp
+++ b/Engine/Source/Editor/UnrealEd/Private/PackageTools.cpp
@@ -188,197 +188,6 @@ namespace PackageTools
 		// no need to refresh content browser here as UPackage::FullyLoad() already does this
 		return bSuccessfullyCompleted;
 	}
-
-	/**
-	 * Attempts to save the specified packages; helper function for by e.g. SaveSelectedPackages().
-	 *
-	 * @param		PackagesToSave				The content packages to save.
-	 * @param		bUnloadPackagesAfterSave	If true, unload each package if it was saved successfully.
-	 * @param		pLastSaveDirectory			if specified, initializes the "Save File" dialog with this value
-	 *											[out] will be filled in with the directory the user chose for this save operation.
-	 *
-	 * @return									true if all packages were saved successfully, false otherwise.
-	 */
-	bool SavePackages( const TArray<UPackage*>& PackagesToSave, bool bUnloadPackagesAfterSave, FString* pLastSaveDirectory /*=NULL*/ )
-	{
-		// Get outermost packages, in case groups were selected.
-		TArray<UPackage*> Packages;
-		for( int32 PackageIndex = 0 ; PackageIndex < PackagesToSave.Num() ; ++PackageIndex )
-		{
-			UPackage* Package = PackagesToSave[PackageIndex];
-			if ( Package != NULL )
-			{
-				Packages.AddUnique( Package->GetOutermost() ? (UPackage*)Package->GetOutermost() : Package );
-			}
-		}
-
-		// Packages must be fully loaded before they can be saved.
-		if( !HandleFullyLoadingPackages( Packages, NSLOCTEXT("UnrealEd", "Save", "Save") ) )
-		{
-			return false;
-		}
-
-		// Assume that all packages were saved.
-		bool bAllPackagesWereSaved = true;
-
-		// Present the user with a 'save file' dialog.
-		FString FileTypes( TEXT("All Files|*.*") );
-		FileTypes += FString::Printf( TEXT("|(*%s)|*%s"), *FPackageName::GetAssetPackageExtension(), *FPackageName::GetAssetPackageExtension() );
-		FileTypes += FString::Printf( TEXT("|(*%s)|*%s"), *FPackageName::GetMapPackageExtension(), *FPackageName::GetMapPackageExtension() );
-
-		TArray<UPackage*> PackagesToUnload;
-
-		GWarn->BeginSlowTask( NSLOCTEXT("UnrealEd", "SavingPackagesE", "Saving packages..."), true );
-		for(int32 i=0; i<Packages.Num(); i++)
-		{
-			UPackage* PackageToSave = Packages[i];
-			if ( PackageToSave != NULL )
-			{
-				GWarn->StatusUpdate( i, Packages.Num(), FText::Format( NSLOCTEXT("UnrealEd", "SavingAssetf", "Saving asset {0}"), FText::FromString(PackageToSave->GetName()) ) );
-
-				// Prevent level packages from being saved via the Generic Browser.
-				UWorld* TheWorld = UWorld::FindWorldInPackage(PackageToSave);
-				if( TheWorld )
-				{
-					FMessageDialog::Open( EAppMsgType::Ok, FText::Format(NSLOCTEXT("UnrealEd", "Error_CantSaveMapViaGB", "Can't save map files via the browser. Not saving '{0}'."), FText::FromString(PackageToSave->GetName())) );
-					bAllPackagesWereSaved = false;
-				}
-				else
-				{
-					FString SaveFileName;
-
-					FString PackageName = PackageToSave->GetName();
-					FString ExistingFile;
-
-					if( FPackageName::DoesPackageExist( PackageName, NULL, &ExistingFile ) )
-					{
-						// If the package already exists, save to the existing filename.
-						SaveFileName = ExistingFile;
-					}
-					else
-					{
-						ensureMsg(false, TEXT("All files should be found in the package file cache"));
-
-						FString LastSaveDirectory;
-						if ( pLastSaveDirectory != NULL )
-						{
-							LastSaveDirectory = *pLastSaveDirectory;
-						}
-
-						// Couldn't find the package name; present a file dialog so the user can specify a name.
-						const FString File = FString::Printf( TEXT("%s%s"), *PackageName, *FPackageName::GetAssetPackageExtension() );
-						UE_LOG(LogPackageTools, Log, TEXT("NO_EXISTING: %s, %s, %s"), *ExistingFile, *File, *LastSaveDirectory);
-
-						// Loop through until a valid filename is given or the user presses cancel
-						bool bFilenameIsValid = false;
-						while( !bFilenameIsValid )
-						{
-							TArray<FString> SaveFilenames;
-							IDesktopPlatform* DesktopPlatform = FDesktopPlatformModule::Get();
-							bool bSave = false;
-							if ( DesktopPlatform )
-							{
-								void* ParentWindowWindowHandle = NULL;
-								IMainFrameModule& MainFrameModule = FModuleManager::LoadModuleChecked<IMainFrameModule>(TEXT("MainFrame"));
-								const TSharedPtr<SWindow>& MainFrameParentWindow = MainFrameModule.GetParentWindow();
-								if ( MainFrameParentWindow.IsValid() && MainFrameParentWindow->GetNativeWindow().IsValid() )
-								{
-									ParentWindowWindowHandle = MainFrameParentWindow->GetNativeWindow()->GetOSWindowHandle();
-								}
-
-								bSave = DesktopPlatform->SaveFileDialog(
-									ParentWindowWindowHandle,
-									*NSLOCTEXT("UnrealEd", "SaveAsset", "Save Asset").ToString(),
-									*LastSaveDirectory,
-									*File,
-									*FileTypes,
-									EFileDialogFlags::None,
-									SaveFilenames
-									);
-							}
-
-							if( bSave )
-							{
-								LastSaveDirectory = FPaths::GetPath(SaveFilenames[0]);
-								SaveFileName = SaveFilenames[0];
-								
-								// If the supplied file name is missing an extension then give it the default package
-								// file extension.
-								if( SaveFileName.Len() > 0 && FPaths::GetExtension(SaveFileName).Len() == 0 )
-								{
-									SaveFileName += FPackageName::GetAssetPackageExtension();
-								}						
-								
-								FText ErrorMessage;
-								if( !FEditorFileUtils::IsFilenameValidForSaving( SaveFileName, ErrorMessage ) )
-								{
-									SaveFileName.Empty();
-									FMessageDialog::Open( EAppMsgType::Ok, ErrorMessage );
-									// Start the loop over, prompting for save again
-									continue;
-								}
-								else
-								{
-									// Filename is valid, dont promt the user again
-									bFilenameIsValid = true;
-								}
-
-								if ( pLastSaveDirectory != NULL )
-								{
-									*pLastSaveDirectory = LastSaveDirectory;
-								}
-							}
-							else
-							{
-								// user canceled the save dialog.  Do not prompt again
-								break;
-							}
-						}
-					}
-
-					if ( SaveFileName.Len() )
-					{
-						const FScopedBusyCursor BusyCursor;
-						UE_LOG(LogPackageTools, Log, TEXT("saving as %s"), *SaveFileName);
-						if( IFileManager::Get().IsReadOnly( *SaveFileName ) )
-						{
-							FMessageDialog::Open( EAppMsgType::Ok, FText::Format(NSLOCTEXT("UnrealEd", "Error_CouldntWriteToFile_F", "Couldn't write to file '{0}'. Maybe file is read-only?"), FText::FromString(SaveFileName)) );
-							bAllPackagesWereSaved = false;
-						}
-						else if( !GUnrealEd->Exec( GWorld, *FString::Printf(TEXT("OBJ SAVEPACKAGE PACKAGE=\"%s\" FILE=\"%s\""), *PackageName, *SaveFileName) ) )
-						{
-							// Couldn't save.
-							FMessageDialog::Open( EAppMsgType::Ok, NSLOCTEXT("UnrealEd", "Error_CouldntSaveAsset", "Couldn't save asset. Check log window for further details.") );
-							bAllPackagesWereSaved = false;
-						}
-						else
-						{
-							// The package saved successfully.  Unload it?
-							if ( bUnloadPackagesAfterSave )
-							{
-								PackagesToUnload.Add( PackageToSave );
-							}
-						}
-					}
-					else
-					{
-						// Couldn't save; no filename to save as.
-						bAllPackagesWereSaved = false;
-					}
-				}
-			}
-		}
-
-
-		if( PackagesToUnload.Num() )
-		{
-			PackageTools::UnloadPackages( PackagesToUnload );
-		}
-
-		GWarn->EndSlowTask();
-
-		return bAllPackagesWereSaved;
-	}
 	
 	/**
 	 * Loads the specified package file (or returns an existing package if it's already loaded.)
diff --git a/Engine/Source/Editor/UnrealEd/Public/PackageTools.h b/Engine/Source/Editor/UnrealEd/Public/PackageTools.h
--- a/Engine/Source/Editor/UnrealEd/Public/PackageTools.h
+++ b/Engine/Source/Editor/UnrealEd/Public/PackageTools.h
@@ -32,20 +32,6 @@ namespace PackageTools
 	UNREALED_API void GetObjectsInPackages( const TArray<UPackage*>* InPackages,
 							   TArray<UObject*>& OutObjects );
 
-	/**
-	 * Attempts to save the specified packages; helper function for by e.g. SaveSelectedPackages().
-	 *
-	 * @param		PackagesToSave				The content packages to save.
-	 * @param		bUnloadPackagesAfterSave	If true, unload each package if it was saved successfully.
-	 * @param		pLastSaveDirectory			if specified, initializes the "Save File" dialog with this value
-	 *											[out] will be filled in with the directory the user chose for this save operation.
-	 *
-	 * @return									true if all packages were saved successfully, false otherwise.
-	 */
-	UNREALED_API bool SavePackages( const TArray<UPackage*>& PackagesToSave, 
-						bool bUnloadPackagesAfterSave, 
-						FString* pLastSaveDirectory=NULL );
-
 	/**
 	 * Handles fully loading passed in packages.
 	 *
