Remove LandscapeProxy include from FbxExporter.h and replace with a predeclaration.
The LandscapeProxy functions are only needed by two files that include this header.

diff --git a/Engine/Source/Editor/UnrealEd/Private/Fbx/FbxMainExport.cpp b/Engine/Source/Editor/UnrealEd/Private/Fbx/FbxMainExport.cpp
--- a/Engine/Source/Editor/UnrealEd/Private/Fbx/FbxMainExport.cpp
+++ b/Engine/Source/Editor/UnrealEd/Private/Fbx/FbxMainExport.cpp
@@ -48,6 +48,7 @@
 #include "Matinee/InterpTrackInstMove.h"
 
 #include "LandscapeDataAccess.h"
+#include "Landscape/LandscapeProxy.h"
 
 #include "FbxExporter.h"
 #include "RawMesh.h"
diff --git a/Engine/Source/Editor/UnrealEd/Private/FbxExporter.h b/Engine/Source/Editor/UnrealEd/Private/FbxExporter.h
--- a/Engine/Source/Editor/UnrealEd/Private/FbxExporter.h
+++ b/Engine/Source/Editor/UnrealEd/Private/FbxExporter.h
@@ -34,8 +34,7 @@
 #include "MatineeExporter.h"
 #include "FbxImporter.h"
 
-#include "Landscape/LandscapeProxy.h"
-
+class ALandscapeProxy;
 struct FAnimControlTrackKey;
 
 namespace UnFbx
diff --git a/Engine/Source/Editor/UnrealEd/Private/UnrealEdSrv.cpp b/Engine/Source/Editor/UnrealEd/Private/UnrealEdSrv.cpp
--- a/Engine/Source/Editor/UnrealEd/Private/UnrealEdSrv.cpp
+++ b/Engine/Source/Editor/UnrealEd/Private/UnrealEdSrv.cpp
@@ -18,13 +18,14 @@
 #include "ISourceControlModule.h"
 #include "FbxExporter.h"
 #include "DesktopPlatformModule.h"
-#include "Landscape/LandscapeInfo.h"
 #include "SnappingUtils.h"
 #include "MessageLog.h"
 #include "AssetSelection.h"
 #include "HighResScreenshot.h"
 #include "ActorEditorUtils.h"
 #include "Matinee/InterpData.h"
+#include "Landscape/LandscapeInfo.h"
+#include "Landscape/LandscapeProxy.h"
 #include "Landscape/LandscapeGizmoActiveActor.h"
 #include "Landscape/LandscapeComponent.h"
 #include "Landscape/LandscapeHeightfieldCollisionComponent.h"
