FIXIF: CRASH: LIVE: Blueprint causing a crash after updating to 4.1.1 ttp336410

diff --git a/Engine/Source/Editor/UnrealEd/Private/Kismet2/BlueprintEditorUtils.cpp b/Engine/Source/Editor/UnrealEd/Private/Kismet2/BlueprintEditorUtils.cpp
--- a/Engine/Source/Editor/UnrealEd/Private/Kismet2/BlueprintEditorUtils.cpp
+++ b/Engine/Source/Editor/UnrealEd/Private/Kismet2/BlueprintEditorUtils.cpp
@@ -1474,7 +1474,7 @@ UBlueprint* FBlueprintEditorUtils::FindBlueprintForNodeChecked(const UEdGraphNod
 // Helper function to get the blueprint that ultimately owns a graph.
 UBlueprint* FBlueprintEditorUtils::FindBlueprintForGraph(const UEdGraph* Graph)
 {
-	for (UObject* TestOuter = Graph->GetOuter(); TestOuter; TestOuter = TestOuter->GetOuter())
+	for (UObject* TestOuter = Graph ? Graph->GetOuter() : NULL; TestOuter; TestOuter = TestOuter->GetOuter())
 	{
 		if (UBlueprint* Result = Cast<UBlueprint>(TestOuter))
 		{
