UE4: CleanupAfterSuccessfulDelete will now be called after DeleteObjects in ObjectTools::DeleteAssets. This is so the objects that are targeted for delete are not GCed before they even enter the delete model.

diff --git a/Engine/Source/Editor/UnrealEd/Private/ObjectTools.cpp b/Engine/Source/Editor/UnrealEd/Private/ObjectTools.cpp
--- a/Engine/Source/Editor/UnrealEd/Private/ObjectTools.cpp
+++ b/Engine/Source/Editor/UnrealEd/Private/ObjectTools.cpp
@@ -1445,7 +1445,7 @@ namespace ObjectTools
 
 	int32 DeleteAssets( const TArray<FAssetData>& AssetsToDelete, bool bShowConfirmation )
 	{
-		TArray<UPackage*> PackageFilesToDelete;
+		TArray<TWeakObjectPtr<UPackage>> PackageFilesToDelete;
 		TArray<UObject*> ObjectsToDelete;
 		for ( int i = 0; i < AssetsToDelete.Num(); i++ )
 		{
@@ -1466,12 +1466,6 @@ namespace ObjectTools
 					continue;
 				}
 
-				if ( FPaths::GetExtension(PackageFilename, /*bIncludeDot=*/true).ToLower() != FPackageName::GetAssetPackageExtension() )
-				{
-					// Only delete UAsset packages because that is what we checked for in ShowDeleteConfirmationDialog()
-					continue;
-				}
-
 				UPackage* Package = FindPackage(nullptr, *AssetData.PackageName.ToString());
 				if ( Package )
 				{
@@ -1480,19 +1474,32 @@ namespace ObjectTools
 			}
 		}
 
-		const int32 NumPackagesToDelete = PackageFilesToDelete.Num();
-		if ( NumPackagesToDelete > 0 )
-		{
-			const bool bPerformReferenceCheck = true;
-			CleanupAfterSuccessfulDelete(PackageFilesToDelete, true);
-		}
-
 		int32 NumObjectsToDelete = ObjectsToDelete.Num();
 		if ( NumObjectsToDelete > 0 )
 		{
 			NumObjectsToDelete = DeleteObjects( ObjectsToDelete, bShowConfirmation );
 		}
 
+		const int32 NumPackagesToDelete = PackageFilesToDelete.Num();
+		if (NumPackagesToDelete > 0)
+		{
+			TArray<UPackage*> PackagePointers;
+			for ( const auto& PkgIt : PackageFilesToDelete )
+			{
+				UPackage* Package = PkgIt.Get();
+				if ( Package )
+				{
+					PackagePointers.Add(Package);
+				}
+			}
+			
+			if ( PackagePointers.Num() > 0 )
+			{
+				const bool bPerformReferenceCheck = true;
+				CleanupAfterSuccessfulDelete(PackagePointers, bPerformReferenceCheck);
+			}
+		}
+
 		return NumPackagesToDelete + NumObjectsToDelete;
 	}
 
