Fixing a clang warning.

Should not change the way things work now, however I am not sure whether current way was intentional.

diff --git a/Engine/Source/Editor/UnrealEd/Private/GroupActor.cpp b/Engine/Source/Editor/UnrealEd/Private/GroupActor.cpp
--- a/Engine/Source/Editor/UnrealEd/Private/GroupActor.cpp
+++ b/Engine/Source/Editor/UnrealEd/Private/GroupActor.cpp
@@ -73,7 +73,7 @@ void AGroupActor::PostEditUndo()
 bool AGroupActor::IsSelected() const
 {
 	// Group actors can only count as 'selected' if they are locked 
-	return IsLocked() && HasSelectedActors() || Super::IsSelected();
+	return (IsLocked() && HasSelectedActors()) || Super::IsSelected();
 }
 
 
