More clang fixes.

Mostly initialization order issues, but also a missing virtual destructor.

diff --git a/Engine/Source/Editor/Kismet/Public/WorkflowOrientedApp/WorkflowUObjectDocuments.h b/Engine/Source/Editor/Kismet/Public/WorkflowOrientedApp/WorkflowUObjectDocuments.h
--- a/Engine/Source/Editor/Kismet/Public/WorkflowOrientedApp/WorkflowUObjectDocuments.h
+++ b/Engine/Source/Editor/Kismet/Public/WorkflowOrientedApp/WorkflowUObjectDocuments.h
@@ -41,6 +41,8 @@ public:
 		return DocumentID.IsValid();
 	}
 
+	virtual ~FTabPayload_UObject() {};
+
 private:
 	// Buried constructor. Use Make instead!
 	FTabPayload_UObject(UObject* InDocumentID)
diff --git a/Engine/Source/Editor/MeshPaint/Public/MeshPaintEdMode.h b/Engine/Source/Editor/MeshPaint/Public/MeshPaintEdMode.h
--- a/Engine/Source/Editor/MeshPaint/Public/MeshPaintEdMode.h
+++ b/Engine/Source/Editor/MeshPaint/Public/MeshPaintEdMode.h
@@ -261,9 +261,9 @@ struct FTextureTargetListInfo
 	uint32 UVChannelIndex;
 	FTextureTargetListInfo(UTexture2D* InTextureData, int32 InUVChannelIndex, bool InbIsSelected = false)
 		:	TextureData(InTextureData)
-		,	UVChannelIndex(InUVChannelIndex)
 		,	bIsSelected(InbIsSelected)
 		,	UndoCount(0)
+		,	UVChannelIndex(InUVChannelIndex)
 	{}
 };
 
@@ -336,8 +336,8 @@ public:
 		/** Default ctor */
 		PaintTexture2DData() :
 			PaintingTexture2D( NULL ),
-			PaintingTexture2DDuplicate ( NULL ),
-			bIsPaintingTexture2DModified( false ),
+			bIsPaintingTexture2DModified(false),
+			PaintingTexture2DDuplicate(NULL),
 			PaintRenderTargetTexture( NULL ),
 			CloneRenderTargetTexture( NULL )
 		{}
diff --git a/Engine/Source/Editor/UnrealEd/Private/AssetDeleteModel.cpp b/Engine/Source/Editor/UnrealEd/Private/AssetDeleteModel.cpp
--- a/Engine/Source/Editor/UnrealEd/Private/AssetDeleteModel.cpp
+++ b/Engine/Source/Editor/UnrealEd/Private/AssetDeleteModel.cpp
@@ -412,12 +412,12 @@ void FAssetDeleteModel::PrepareToDelete(UObject* InObject)
 //-----------------------------------------------------------------
 
 FPendingDelete::FPendingDelete(UObject* InObject)
-	: Object(InObject)
+	: RemainingDiskReferences(0)
+	, RemainingMemoryReferences(0)
+	, Object(InObject)
 	, bReferencesChecked(false)
 	, bIsReferencedInMemory(false)
 	, bIsReferencedInMemoryByUndo(false)
-	, RemainingDiskReferences(0)
-	, RemainingMemoryReferences(0)
 	, bIsInternal(false)
 {
 	// Blueprints actually contain 3 assets, the UBlueprint, GeneratedClass and SkeletonGeneratedClass
diff --git a/Engine/Source/Editor/UnrealEd/Private/EditorModes.cpp b/Engine/Source/Editor/UnrealEd/Private/EditorModes.cpp
--- a/Engine/Source/Editor/UnrealEd/Private/EditorModes.cpp
+++ b/Engine/Source/Editor/UnrealEd/Private/EditorModes.cpp
@@ -193,11 +193,11 @@ FEdMode::FEdMode()
 	: Name( NSLOCTEXT("UnrealEd", "N/A", "N/A") )
 	, CurrentTool( NULL )
 	, ID( FBuiltinEditorModes::EM_None )
-	, EditedPropertyName( TEXT("") )
+	, bVisible(false)
+	, PriorityOrder(MAX_int32)
+	, EditedPropertyName(TEXT(""))
 	, EditedPropertyIndex( INDEX_NONE )
 	, bEditedPropertyIsTransform( false )
-	, bVisible( false )
-	, PriorityOrder( MAX_int32 )
 {
 	bDrawKillZ = true;
 }
@@ -1183,12 +1183,12 @@ FEditorModeTools::FEditorModeTools()
 	,	Snapping( 0 )
 	,	SnappedActor( 0 )
 	,	TranslateRotateXAxisAngle(0)
-	,	CoordSystem( COORD_World )
 	,	WidgetMode( FWidget::WM_Translate )
 	,	OverrideWidgetMode( FWidget::WM_None )
 	,	bShowWidget( 1 )
-	,	bIsTracking( false )
-	,	bHideViewportUI( false )
+	,	bHideViewportUI(false)
+	,	CoordSystem(COORD_World)
+	,	bIsTracking(false)
 {}
 
 FEditorModeTools::~FEditorModeTools()
diff --git a/Engine/Source/Editor/UnrealEd/Private/EditorViewportClient.cpp b/Engine/Source/Editor/UnrealEd/Private/EditorViewportClient.cpp
--- a/Engine/Source/Editor/UnrealEd/Private/EditorViewportClient.cpp
+++ b/Engine/Source/Editor/UnrealEd/Private/EditorViewportClient.cpp
@@ -80,11 +80,11 @@ public:
 FViewportCameraTransform::FViewportCameraTransform()
 	: TransitionCurve( new FCurveSequence( 0.0f, FocusConstants::TransitionTime, ECurveEaseFunction::CubicOut ) )
 	, ViewLocation( FVector::ZeroVector )
-	, DesiredLocation( FVector::ZeroVector )
 	, ViewRotation( FRotator::ZeroRotator )
-	, StartLocation( FVector::ZeroVector )
+	, DesiredLocation(FVector::ZeroVector)
 	, LookAt( FVector::ZeroVector )
-	, OrthoZoom( DEFAULT_ORTHOZOOM )
+	, StartLocation(FVector::ZeroVector)
+	, OrthoZoom(DEFAULT_ORTHOZOOM)
 {}
 
 void FViewportCameraTransform::SetLocation( const FVector& Position )
@@ -181,7 +181,8 @@ int32 FEditorViewportClient::GetCameraSpeedSetting() const
 float const FEditorViewportClient::SafePadding = 0.075f;
 
 FEditorViewportClient::FEditorViewportClient(FPreviewScene* InPreviewScene)
-	: ImmersiveDelegate()
+	: CameraSpeedSetting(4)
+	, ImmersiveDelegate()
 	, VisibilityDelegate()
 	, Viewport(NULL)
 	, ViewTransform()
@@ -196,29 +197,14 @@ FEditorViewportClient::FEditorViewportClient(FPreviewScene* InPreviewScene)
 	, ViewFOV(EditorViewportDefs::DefaultPerspectiveFOVAngle)
 	, FOVAngle(EditorViewportDefs::DefaultPerspectiveFOVAngle)
 	, AspectRatio(1.777777f)
-	, NearPlane(-1.0f)
-	, FarPlane(0.0f)
 	, bForcingUnlitForNewMap(false)
-	, bInGameViewMode(false)
 	, bWidgetAxisControlledByDrag(false)
 	, bNeedsRedraw(true)
 	, bNeedsLinkedRedraw(false)
 	, bNeedsInvalidateHitProxy(false)
 	, bUsingOrbitCamera(false)
-	, bIsSimulateInEditorViewport( false )
-	, bHasAudioFocus(false)
-	, bShouldCheckHitProxy(false)
-	, bUsesDrawHelper(true)
-	, PerspViewModeIndex(VMI_Lit)
-	, OrthoViewModeIndex(VMI_BrushWireframe)
-	, bIsTracking( false )
-	, bDraggingByHandle( false )
-	, CurrentGestureDragDelta(FVector::ZeroVector)
-	, CurrentGestureRotDelta(FRotator::ZeroRotator)
-	, GestureMoveForwardBackwardImpulse(0.0f)
-	, bDisableInput( false )
+	, bDisableInput(false)
 	, bDrawAxes(true)
-	, bForceAudioRealtime(false)
 	, bSetListenerPosition(false)
 	, LandscapeLODOverride(-1)
 	, Widget(new FWidget)
@@ -235,13 +221,27 @@ FEditorViewportClient::FEditorViewportClient(FPreviewScene* InPreviewScene)
 	, CachedMouseX(0)
 	, CachedMouseY(0)
 	, CurrentMousePos(-1, -1)
+	, bIsTracking(false)
+	, bDraggingByHandle(false)
+	, CurrentGestureDragDelta(FVector::ZeroVector)
+	, CurrentGestureRotDelta(FRotator::ZeroRotator)
+	, GestureMoveForwardBackwardImpulse(0.0f)
+	, bForceAudioRealtime(false)
 	, bIsRealtime(false)
 	, bStoredRealtime(false)
 	, bStoredShowStats(false)
 	, bShowStats(false)
-	, PreviewScene(InPreviewScene)
+	, bHasAudioFocus(false)
+	, bShouldCheckHitProxy(false)
+	, bUsesDrawHelper(true)
+	, bIsSimulateInEditorViewport(false)
 	, bCameraLock(false)
-	, CameraSpeedSetting(4)
+	, PreviewScene(InPreviewScene)
+	, PerspViewModeIndex(VMI_Lit)
+	, OrthoViewModeIndex(VMI_BrushWireframe)
+	, NearPlane(-1.0f)
+	, FarPlane(0.0f)
+	, bInGameViewMode(false)
 {
 
 	ViewState.Allocate();
diff --git a/Engine/Source/Editor/UnrealEd/Private/FeedbackContextEditor.cpp b/Engine/Source/Editor/UnrealEd/Private/FeedbackContextEditor.cpp
--- a/Engine/Source/Editor/UnrealEd/Private/FeedbackContextEditor.cpp
+++ b/Engine/Source/Editor/UnrealEd/Private/FeedbackContextEditor.cpp
@@ -171,8 +171,8 @@ const float FFeedbackContextEditor::UIUpdateGatingTime = 1.0f / 60.0f;
 
 FFeedbackContextEditor::FFeedbackContextEditor() : 
 	DialogRequestCount(0),
-	SlowTaskCount(0),
-	HasTaskBeenCancelled(false)
+	HasTaskBeenCancelled(false),
+	SlowTaskCount(0)
 {
 	DialogRequestStack.Reserve(15);
 }
