A bit more clang fixes.

diff --git a/Engine/Source/Editor/UnrealEd/Private/AssetThumbnail.cpp b/Engine/Source/Editor/UnrealEd/Private/AssetThumbnail.cpp
--- a/Engine/Source/Editor/UnrealEd/Private/AssetThumbnail.cpp
+++ b/Engine/Source/Editor/UnrealEd/Private/AssetThumbnail.cpp
@@ -17,9 +17,9 @@ public:
 		, _AllowFadeIn(false)
 		, _ForceGenericThumbnail(false)
 		, _AllowHintText(true)
-		, _HighlightedText( FText::GetEmpty() )
 		, _Label(EThumbnailLabel::ClassName)
-		, _HintColorAndOpacity( FLinearColor(0.0f, 0.0f, 0.0f, 0.0f) )
+		, _HighlightedText(FText::GetEmpty())
+		, _HintColorAndOpacity(FLinearColor(0.0f, 0.0f, 0.0f, 0.0f))
 		, _ClassThumbnailBrushOverride(NAME_None)
 		, _ShowClassBackground( true )
 		{}
@@ -551,10 +551,10 @@ private:
 
 
 FAssetThumbnail::FAssetThumbnail( UObject* InAsset, uint32 InWidth, uint32 InHeight, const TSharedPtr<class FAssetThumbnailPool>& InThumbnailPool )
-	: AssetData ( InAsset ? FAssetData(InAsset) : FAssetData() )
+	: ThumbnailPool(InThumbnailPool)
+	, AssetData(InAsset ? FAssetData(InAsset) : FAssetData())
 	, Width( InWidth )
 	, Height( InHeight )
-	, ThumbnailPool( InThumbnailPool )
 {
 	if ( InThumbnailPool.IsValid() )
 	{
