Added Tweak UI for Cloth physics properties in MeshDetails tab of Persona

diff --git a/Engine/Source/Editor/Persona/Private/PersonaMeshDetails.cpp b/Engine/Source/Editor/Persona/Private/PersonaMeshDetails.cpp
--- a/Engine/Source/Editor/Persona/Private/PersonaMeshDetails.cpp
+++ b/Engine/Source/Editor/Persona/Private/PersonaMeshDetails.cpp
@@ -398,7 +398,7 @@ void FPersonaMeshDetails::OnGenerateElementForClothingAsset( TSharedRef<IPropert
 		[
 			SNew( SButton )
 			.Text( LOCTEXT("ReimportButtonLabel", "Reimport") )
-			.OnClicked( this, &FPersonaMeshDetails::OnReimportApexFileClicked, ElementIndex )
+			.OnClicked(this, &FPersonaMeshDetails::OnReimportApexFileClicked, ElementIndex, DetailLayout)
 			.IsFocusable( false )
 			.ContentPadding(0)
 			.ForegroundColor( FSlateColor::UseForeground() )
@@ -454,6 +454,27 @@ void FPersonaMeshDetails::OnGenerateElementForClothingAsset( TSharedRef<IPropert
 	[
 		MakeApexDetailsWidget(ElementIndex)
 	];	
+	
+	USkeletalMesh* SkelMesh = PersonaPtr->GetMesh();
+	FClothingAssetData& AssetData = SkelMesh->ClothingAssets[ElementIndex];
+
+	// if properties are not changed by a user through UE4 editor, loads original properties from the asset
+	if (!AssetData.bClothPropertiesChanged)
+	{
+		ApexClothingUtils::GetPhysicsPropertiesFromApexAsset(AssetData.ApexClothingAsset->GetAsset(), AssetData.PhysicsProperties);
+	}
+
+	// cloth physics properties
+	TSharedRef<IPropertyHandle> ClothPhysicsProperties = StructProperty->GetChildHandle(FName("PhysicsProperties")).ToSharedRef();
+	uint32 NumChildren;
+	ClothPhysicsProperties->GetNumChildren(NumChildren);
+	FSimpleDelegate UpdatePhysicsPropertyDelegate = FSimpleDelegate::CreateSP(this, &FPersonaMeshDetails::UpdateClothPhysicsProperties, ElementIndex);
+	for (uint32 ChildIdx = 0; ChildIdx < NumChildren; ChildIdx++)
+	{
+		ClothPhysicsProperties->GetChildHandle(ChildIdx)->SetOnPropertyValueChanged(UpdatePhysicsPropertyDelegate);
+	}
+
+	ChildrenBuilder.AddChildProperty(ClothPhysicsProperties);
 }
 
 TSharedRef<SUniformGridPanel> FPersonaMeshDetails::MakeApexDetailsWidget(int32 AssetIndex) const
@@ -694,7 +715,7 @@ EVisibility FPersonaMeshDetails::IsClothingLODCheckBoxVisible(int32 MaterialInde
 	return NumLODsHasCorrespondSubmesh > 1 ? EVisibility::Visible : EVisibility::Hidden;
 }
 
-FReply FPersonaMeshDetails::OnReimportApexFileClicked(int32 AssetIndex)
+FReply FPersonaMeshDetails::OnReimportApexFileClicked(int32 AssetIndex, IDetailLayoutBuilder* DetailLayout)
 {
 	USkeletalMesh* SkelMesh = PersonaPtr->GetMesh();
 
@@ -702,11 +723,43 @@ FReply FPersonaMeshDetails::OnReimportApexFileClicked(int32 AssetIndex)
 
 	FString FileName(SkelMesh->ClothingAssets[AssetIndex].ApexFileName);
 
+	bool bNeedToLeaveProperties = false;
+	FClothPhysicsProperties CurClothPhysicsProperties;
+
+	// if this asset's cloth properties were changed in UE4 editor, then shows a warning dialog
+	if (SkelMesh->ClothingAssets[AssetIndex].bClothPropertiesChanged)
+	{
+		const FText Text = LOCTEXT("Warning_ClothPropertiesChanged", "This asset's physics properties have been modified inside UE4. \nWould you like to reapply the modified property values after reimporting?");
+		EAppReturnType::Type Ret = FMessageDialog::Open(EAppMsgType::YesNoCancel, Text);
+		if (EAppReturnType::Yes == Ret)
+		{
+			bNeedToLeaveProperties = true;
+			ApexClothingUtils::GetPhysicsPropertiesFromApexAsset(SkelMesh->ClothingAssets[AssetIndex].ApexClothingAsset->GetAsset(), CurClothPhysicsProperties);
+		}
+		else if (EAppReturnType::Cancel == Ret)
+		{
+			return FReply::Handled();
+		}
+	}
 	ApexClothingUtils::EClothUtilRetType RetType = ApexClothingUtils::ImportApexAssetFromApexFile(FileName, SkelMesh, AssetIndex);
 	switch(RetType)
 	{
 	case ApexClothingUtils::CURT_Ok:
 		UpdateComboBoxStrings();
+		if (bNeedToLeaveProperties)
+		{
+			// overwrites changed values instead of original values
+			SkelMesh->ClothingAssets[AssetIndex].bClothPropertiesChanged = true;
+			ApexClothingUtils::SetPhysicsPropertiesToApexAsset(SkelMesh->ClothingAssets[AssetIndex].ApexClothingAsset->GetAsset(), CurClothPhysicsProperties);
+		}
+		else
+		{
+			// will lose all changed values through UE4 editor
+			SkelMesh->ClothingAssets[AssetIndex].bClothPropertiesChanged = false;
+		}
+
+		DetailLayout->ForceRefreshDetails();
+
 		break;
 	case ApexClothingUtils::CURT_Fail:
 		// Failed to create or import
@@ -1078,6 +1131,16 @@ void FPersonaMeshDetails::CheckLODMaterialMapChanges()
 		OldLODInfo = SkelMesh->LODInfo;
 	}
 }
+
+void FPersonaMeshDetails::UpdateClothPhysicsProperties(int32 AssetIndex)
+{
+	USkeletalMesh* SkelMesh = PersonaPtr->GetMesh();
+	FClothingAssetData& Asset = SkelMesh->ClothingAssets[AssetIndex];
+
+	ApexClothingUtils::SetPhysicsPropertiesToApexAsset(Asset.ApexClothingAsset->GetAsset(), Asset.PhysicsProperties);
+	Asset.bClothPropertiesChanged = true;
+}
+
 #endif// #if WITH_APEX_CLOTHING
 
 
diff --git a/Engine/Source/Editor/Persona/Private/PersonaMeshDetails.h b/Engine/Source/Editor/Persona/Private/PersonaMeshDetails.h
--- a/Engine/Source/Editor/Persona/Private/PersonaMeshDetails.h
+++ b/Engine/Source/Editor/Persona/Private/PersonaMeshDetails.h
@@ -158,7 +158,7 @@ private:
 	FReply OnOpenClothingFileClicked(IDetailLayoutBuilder* DetailLayout);
 
 	/* Reimports a clothing asset */ 
-	FReply OnReimportApexFileClicked(int32 AssetIndex);
+	FReply OnReimportApexFileClicked(int32 AssetIndex, IDetailLayoutBuilder* DetailLayout);
 
 	/* Removes a clothing asset */ 
 	FReply OnRemoveApexFileClicked(int32 AssetIndex, IDetailLayoutBuilder* DetailLayout);
@@ -171,6 +171,8 @@ private:
 	/* if LODMaterialMap is changed, then re-map clothing sections by changed info */
 	void CheckLODMaterialMapChanges();
 
+	/* if physics properties are changed, then save to the clothing asset */
+	void UpdateClothPhysicsProperties(int32 AssetIndex);
 #endif // #if WITH_APEX_CLOTHING
 
 };
diff --git a/Engine/Source/Editor/UnrealEd/Classes/Animation/DebugSkelMeshComponent.h b/Engine/Source/Editor/UnrealEd/Classes/Animation/DebugSkelMeshComponent.h
--- a/Engine/Source/Editor/UnrealEd/Classes/Animation/DebugSkelMeshComponent.h
+++ b/Engine/Source/Editor/UnrealEd/Classes/Animation/DebugSkelMeshComponent.h
@@ -235,6 +235,9 @@ class UDebugSkelMeshComponent : public USkeletalMeshComponent
 
 	int32 FindCurrentSectionDisplayMode();
 
+	/** to avoid clothing reset while modifying properties in Persona */
+	virtual void CheckClothTeleport(float DeltaTime) OVERRIDE;
+
 #endif //#if WITH_APEX_CLOTHING
 };
 
diff --git a/Engine/Source/Editor/UnrealEd/Private/Animation/DebugSkelMeshComponent.cpp b/Engine/Source/Editor/UnrealEd/Private/Animation/DebugSkelMeshComponent.cpp
--- a/Engine/Source/Editor/UnrealEd/Private/Animation/DebugSkelMeshComponent.cpp
+++ b/Engine/Source/Editor/UnrealEd/Private/Animation/DebugSkelMeshComponent.cpp
@@ -610,4 +610,11 @@ int32 UDebugSkelMeshComponent::FindCurrentSectionDisplayMode()
 	return DisplayMode;
 }
 
+void UDebugSkelMeshComponent::CheckClothTeleport(float DeltaTime)
+{
+	// do nothing to avoid clothing reset while modifying properties
+	// modifying values can cause frame delay and clothes will be reset by a large delta time (low fps)
+	// doesn't need cloth teleport while previewing
+}
+
 #endif // #if WITH_APEX_CLOTHING
\ No newline at end of file
diff --git a/Engine/Source/Editor/UnrealEd/Private/ApexClothingUtils.cpp b/Engine/Source/Editor/UnrealEd/Private/ApexClothingUtils.cpp
--- a/Engine/Source/Editor/UnrealEd/Private/ApexClothingUtils.cpp
+++ b/Engine/Source/Editor/UnrealEd/Private/ApexClothingUtils.cpp
@@ -1031,7 +1031,8 @@ physx::apex::NxApexAssetAuthoring* ApexAuthoringFromAsset(NxApexAsset* Asset, Nx
 	::NxParameterized::Interface* NewInterface = ApexSDK->getParameterizedTraits()->createNxParameterized(OldInterface->className());
 	NewInterface->copy(*OldInterface);
 
-	return ApexSDK->createAssetAuthoring(NewInterface, "AuthoringFromAsset");
+	// pass NULL to use an auto-generated name
+	return ApexSDK->createAssetAuthoring(NewInterface, NULL);
 }
 
 NxClothingAsset* ApplyTransform(NxClothingAsset* ApexClothingAsset)
@@ -1040,9 +1041,7 @@ NxClothingAsset* ApplyTransform(NxClothingAsset* ApexClothingAsset)
 	NxClothingAssetAuthoring *ApexClothingAssetAuthoring 
 		= static_cast<NxClothingAssetAuthoring*>(ApexAuthoringFromAsset(ApexClothingAsset, GApexSDK));
 
-	check(ApexClothingAssetAuthoring);
-
-	if(ApexClothingAssetAuthoring)
+	if(ensure(ApexClothingAssetAuthoring))
 	{
 		const NxParameterized::Interface* AssetParams = ApexClothingAsset->getAssetNxParameterized();
 
@@ -1596,6 +1595,7 @@ EClothUtilRetType ImportApexAssetFromApexFile(FString& ApexFile, USkeletalMesh*
 		FClothingAssetData* Data = new(SkelMesh->ClothingAssets) FClothingAssetData;
 		Data->ApexFileName = ApexFile;
 		Data->AssetName = AssetName;
+		Data->bClothPropertiesChanged = false;
 		Data->ApexClothingAsset = MakeShareable(new FClothingAssetWrapper(ApexClothingAsset));
 	}
 	else
@@ -1966,6 +1966,98 @@ int32 GetNumRenderSubmeshes(NxClothingAsset *InAsset, int32 LODIndex)
 	return SubmeshCount;
 }
 
+void GetPhysicsPropertiesFromApexAsset(NxClothingAsset *InAsset, FClothPhysicsProperties& PropertyInfo)
+{
+	const NxParameterized::Interface* AssetParams = InAsset->getAssetNxParameterized();
+
+	uint32 MaterialIndex;
+	verify(NxParameterized::getParamU32(*AssetParams, "materialIndex", MaterialIndex));
+	
+	// ClothingMaterialLibraryParameters 
+	NxParameterized::Interface* MaterialLibraryParams;
+	NxParameterized::getParamRef(*AssetParams, "materialLibrary", MaterialLibraryParams);
+
+	if (MaterialLibraryParams != NULL)
+	{
+		int32 NumMaterials;
+		verify(NxParameterized::getParamArraySize(*MaterialLibraryParams, "materials", NumMaterials));
+
+		check(MaterialIndex < (uint32)NumMaterials);
+ 
+		char ParameterName[MAX_SPRINTF]; 
+
+		// stiffness properties
+		FCStringAnsi::Sprintf(ParameterName, "materials[%d].bendingStiffness", MaterialIndex); 
+		verify(NxParameterized::getParamF32(*MaterialLibraryParams, ParameterName, PropertyInfo.BendResistance));
+		FCStringAnsi::Sprintf(ParameterName, "materials[%d].shearingStiffness", MaterialIndex); 
+		verify(NxParameterized::getParamF32(*MaterialLibraryParams, ParameterName, PropertyInfo.ShearResistance));
+		FCStringAnsi::Sprintf(ParameterName, "materials[%d].hardStretchLimitation", MaterialIndex); 
+		verify(NxParameterized::getParamF32(*MaterialLibraryParams, ParameterName, PropertyInfo.StretchLimit));
+
+		// resistance properties
+		FCStringAnsi::Sprintf(ParameterName, "materials[%d].friction", MaterialIndex); 
+		verify(NxParameterized::getParamF32(*MaterialLibraryParams, ParameterName, PropertyInfo.Friction));
+		FCStringAnsi::Sprintf(ParameterName, "materials[%d].damping", MaterialIndex); 
+		verify(NxParameterized::getParamF32(*MaterialLibraryParams, ParameterName, PropertyInfo.Damping));
+		FCStringAnsi::Sprintf(ParameterName, "materials[%d].drag", MaterialIndex); 
+		verify(NxParameterized::getParamF32(*MaterialLibraryParams, ParameterName, PropertyInfo.Drag));
+
+		// scale properties
+		FCStringAnsi::Sprintf(ParameterName, "materials[%d].gravityScale", MaterialIndex);
+		verify(NxParameterized::getParamF32(*MaterialLibraryParams, ParameterName, PropertyInfo.GravityScale));
+		FCStringAnsi::Sprintf(ParameterName, "materials[%d].inertiaScale", MaterialIndex);
+		verify(NxParameterized::getParamF32(*MaterialLibraryParams, ParameterName, PropertyInfo.InertiaBlend));
+
+		// self-collision properties
+		FCStringAnsi::Sprintf(ParameterName, "materials[%d].selfcollisionThickness", MaterialIndex); 
+		verify(NxParameterized::getParamF32(*MaterialLibraryParams, ParameterName, PropertyInfo.SelfCollisionThickness));
+
+	}
+}
+
+void SetPhysicsPropertiesToApexAsset(NxClothingAsset *InAsset, FClothPhysicsProperties& PropertyInfo)
+{
+	const NxParameterized::Interface* AssetParams = InAsset->getAssetNxParameterized();
+
+	uint32 MaterialIndex;
+	verify(NxParameterized::getParamU32(*AssetParams, "materialIndex", MaterialIndex));
+
+	// ClothingMaterialLibraryParameters 
+	NxParameterized::Interface* MaterialLibraryParams;
+	NxParameterized::getParamRef(*AssetParams, "materialLibrary", MaterialLibraryParams);
+
+	if (MaterialLibraryParams != NULL)
+	{
+		char ParameterName[MAX_SPRINTF];
+
+		// stiffness properties
+		FCStringAnsi::Sprintf(ParameterName, "materials[%d].bendingStiffness", MaterialIndex);
+		verify(NxParameterized::setParamF32(*MaterialLibraryParams, ParameterName, PropertyInfo.BendResistance));
+		FCStringAnsi::Sprintf(ParameterName, "materials[%d].shearingStiffness", MaterialIndex);
+		verify(NxParameterized::setParamF32(*MaterialLibraryParams, ParameterName, PropertyInfo.ShearResistance));
+		FCStringAnsi::Sprintf(ParameterName, "materials[%d].hardStretchLimitation", MaterialIndex);
+		verify(NxParameterized::setParamF32(*MaterialLibraryParams, ParameterName, PropertyInfo.StretchLimit));
+
+		// resistance properties
+		FCStringAnsi::Sprintf(ParameterName, "materials[%d].friction", MaterialIndex);
+		verify(NxParameterized::setParamF32(*MaterialLibraryParams, ParameterName, PropertyInfo.Friction));
+		FCStringAnsi::Sprintf(ParameterName, "materials[%d].damping", MaterialIndex);
+		verify(NxParameterized::setParamF32(*MaterialLibraryParams, ParameterName, PropertyInfo.Damping));
+		FCStringAnsi::Sprintf(ParameterName, "materials[%d].drag", MaterialIndex);
+		verify(NxParameterized::setParamF32(*MaterialLibraryParams, ParameterName, PropertyInfo.Drag));
+
+		// scale properties
+		FCStringAnsi::Sprintf(ParameterName, "materials[%d].gravityScale", MaterialIndex);
+		verify(NxParameterized::setParamF32(*MaterialLibraryParams, ParameterName, PropertyInfo.GravityScale));
+		FCStringAnsi::Sprintf(ParameterName, "materials[%d].inertiaScale", MaterialIndex);
+		verify(NxParameterized::setParamF32(*MaterialLibraryParams, ParameterName, PropertyInfo.InertiaBlend));
+
+		// self-collision properties
+		FCStringAnsi::Sprintf(ParameterName, "materials[%d].selfcollisionThickness", MaterialIndex);
+		verify(NxParameterized::setParamF32(*MaterialLibraryParams, ParameterName, PropertyInfo.SelfCollisionThickness));
+	}
+}
+
 #else
 
 void RestoreOriginalClothingSection(USkeletalMesh* SkelMesh, uint32 LODIndex, uint32 SectionIndex, bool bRecreateSkelMeshComponent)
diff --git a/Engine/Source/Editor/UnrealEd/Public/ApexClothingUtils.h b/Engine/Source/Editor/UnrealEd/Public/ApexClothingUtils.h
--- a/Engine/Source/Editor/UnrealEd/Public/ApexClothingUtils.h
+++ b/Engine/Source/Editor/UnrealEd/Public/ApexClothingUtils.h
@@ -88,6 +88,9 @@ namespace ApexClothingUtils
 
 	UNREALED_API NxClothingAsset* CreateApexClothingAssetFromBuffer(const uint8* Buffer, int32 BufferSize);
 
+	// if don't specify MaterialIndex, default material index will be used 
+	UNREALED_API void GetPhysicsPropertiesFromApexAsset(NxClothingAsset *InAsset, FClothPhysicsProperties& OutPropertyInfo);
+	UNREALED_API void SetPhysicsPropertiesToApexAsset(NxClothingAsset *InAsset, FClothPhysicsProperties& InPropertyInfo);
 #endif // #if WITH_APEX_CLOTHING
 } // namespace ApexClothingUtils
 
diff --git a/Engine/Source/Runtime/Engine/Classes/Components/SkeletalMeshComponent.h b/Engine/Source/Runtime/Engine/Classes/Components/SkeletalMeshComponent.h
--- a/Engine/Source/Runtime/Engine/Classes/Components/SkeletalMeshComponent.h
+++ b/Engine/Source/Runtime/Engine/Classes/Components/SkeletalMeshComponent.h
@@ -930,7 +930,7 @@ public:
 	/** changes clothing LODs, if clothing LOD is disabled or LODIndex is greater than apex clothing LODs, simulation will be disabled */
 	void SetClothingLOD(int32 LODIndex);
 	/** check whether clothing teleport is needed or not to avoid a weird simulation result */
-	void CheckClothTeleport(float DeltaTime);
+	virtual void CheckClothTeleport(float DeltaTime);
 	/** 
 	 * Updates all clothing animation states including ComponentToWorld-related states.
 	 */
diff --git a/Engine/Source/Runtime/Engine/Classes/Engine/SkeletalMesh.h b/Engine/Source/Runtime/Engine/Classes/Engine/SkeletalMesh.h
--- a/Engine/Source/Runtime/Engine/Classes/Engine/SkeletalMesh.h
+++ b/Engine/Source/Runtime/Engine/Classes/Engine/SkeletalMesh.h
@@ -346,6 +346,50 @@ struct FClothBonePlane
 	FPlane PlaneData;
 };
 
+/**
+ * now exposed a part of properties based on 3DS Max plug-in
+ * property names are also changed into 3DS Max plug-in's one
+ */
+USTRUCT()
+struct FClothPhysicsProperties
+{
+	GENERATED_USTRUCT_BODY()
+
+	// Bending stiffness of the cloth in the range [0, 1]. 
+	UPROPERTY(EditAnywhere, Category = Stiffness, meta = (ClampMin = "0.0", ClampMax = "1.0", UIMin = "0.0", UIMax = "1.0"))
+	float BendResistance;
+
+	// Shearing stiffness of the cloth in the range [0, 1]. 
+	UPROPERTY(EditAnywhere, Category = Stiffness, meta = (ClampMin = "0.0", ClampMax = "1.0", UIMin = "0.0", UIMax = "1.0"))
+	float ShearResistance;
+
+	// Make cloth simulation less stretchy. A value smaller than 1 will turn it off. 
+	UPROPERTY(EditAnywhere, Category = Stiffness, meta = (ClampMin = "0.0", UIMin = "0.0", UIMax = "4.0"))
+	float StretchLimit;
+
+	// Friction coefficient in the range[0, 1]
+	UPROPERTY(EditAnywhere, Category = Stiffness, meta = (ClampMin = "0.0", ClampMax = "1.0", UIMin = "0.0", UIMax = "1.0"))
+	float Friction;
+	// Spring damping of the cloth in the range[0, 1]
+	UPROPERTY(EditAnywhere, Category = Stiffness, meta = (ClampMin = "0.0", ClampMax = "1.0", UIMin = "0.0", UIMax = "1.0"))
+	float Damping;
+
+	// Drag coefficient n the range [0, 1] 
+	UPROPERTY(EditAnywhere, Category = Stiffness, meta = (ClampMin = "0.0", ClampMax = "1.0", UIMin = "0.0", UIMax = "1.0"))
+	float Drag;
+
+	// Amount of gravity that is applied to the cloth. 
+	UPROPERTY(EditAnywhere, Category = Scale, meta = (ClampMin = "0.0", UIMin = "0.0", UIMax = "5.0"))
+	float GravityScale;
+	// Amount of inertia that is kept when using local space simulation. Internal name is inertia scale
+	UPROPERTY(EditAnywhere, Category = Scale, meta = (ClampMin = "0.0", UIMin = "0.0", UIMax = "5.0"))
+	float InertiaBlend;
+
+	// Minimal amount of distance particles will keep of each other.
+	UPROPERTY(EditAnywhere, Category = SelfCollision, meta = (ClampMin = "0.0", UIMin = "0.0", UIMax = "5.0"))
+	float SelfCollisionThickness;
+};
+
 USTRUCT()
 struct FClothingAssetData
 {
@@ -358,6 +402,13 @@ struct FClothingAssetData
 	UPROPERTY(EditAnywhere, Category=ClothingAssetData)
 	FString	ApexFileName;
 
+	/** the flag whether cloth physics properties are changed from UE4 editor or not */
+	UPROPERTY(EditAnywhere, Category = ClothingAssetData)
+	bool bClothPropertiesChanged;
+
+	UPROPERTY(EditAnywhere, Transient, Category = ClothingAssetData)
+	FClothPhysicsProperties PhysicsProperties;
+
 #if WITH_APEX_CLOTHING
 	TSharedPtr<FClothingAssetWrapper> ApexClothingAsset;
 
diff --git a/Engine/Source/Runtime/Engine/Private/SkeletalMeshComponent.cpp b/Engine/Source/Runtime/Engine/Private/SkeletalMeshComponent.cpp
--- a/Engine/Source/Runtime/Engine/Private/SkeletalMeshComponent.cpp
+++ b/Engine/Source/Runtime/Engine/Private/SkeletalMeshComponent.cpp
@@ -205,11 +205,6 @@ void USkeletalMeshComponent::CreateRenderState_Concurrent()
 	// Update bHasValidBodies flag
 	UpdateHasValidBodies();
 
-#if WITH_APEX_CLOTHING
-	//clothing actors will be re-created in TickClothing
-	ReleaseAllClothingResources();
-#endif// #if WITH_APEX_CLOTHING
-
 	Super::CreateRenderState_Concurrent();
 }
 
