Fixing "RenameAssets" console command sometimes creating malformed asset names (containing //)

diff --git a/Engine/Source/Editor/UnrealEd/Private/UnrealEdSrv.cpp b/Engine/Source/Editor/UnrealEd/Private/UnrealEdSrv.cpp
--- a/Engine/Source/Editor/UnrealEd/Private/UnrealEdSrv.cpp
+++ b/Engine/Source/Editor/UnrealEd/Private/UnrealEdSrv.cpp
@@ -1193,13 +1193,18 @@ bool UUnrealEdEngine::Exec( UWorld* InWorld, const TCHAR* Stream, FOutputDevice&
 				FString NewPathName = Asset.PackagePath.ToString();
 				if( NewAssetName.Contains( SearchTermStr ) )
 				{
-					NewAssetName = NewAssetName.Replace( *SearchTermStr, *ReplaceStr );
-					bRenamedAsset = true;
+					FString TempPathName = NewAssetName.Replace(*SearchTermStr, *ReplaceStr);
+					if (!TempPathName.IsEmpty())
+					{
+						NewAssetName = TempPathName;
+						bRenamedAsset = true;
+					}
 				}
 				
 				if( NewPathName.Contains( SearchTermStr ) )
 				{
 					FString TempPathName = NewPathName.Replace( *SearchTermStr, *ReplaceStr );
+					FPaths::RemoveDuplicateSlashes(TempPathName);
 
 					if( !TempPathName.IsEmpty() )
 					{
