More Engine.h trimming (EdGraph headers, BoneMaskFilter, Breakpoint, LevelScriptBP, ClipPAdEntry, CodecMovie, DestructibleFractureSettings, DeviceProfile)

diff --git a/Engine/Source/Editor/AnimGraph/Classes/AnimationTransitionSchema.h b/Engine/Source/Editor/AnimGraph/Classes/AnimationTransitionSchema.h
--- a/Engine/Source/Editor/AnimGraph/Classes/AnimationTransitionSchema.h
+++ b/Engine/Source/Editor/AnimGraph/Classes/AnimationTransitionSchema.h
@@ -18,4 +18,9 @@ class UAnimationTransitionSchema : public UEdGraphSchema_K2
 
 private:
 	void GetSourceStateActions(FGraphContextMenuBuilder& ContextMenuBuilder) const;
+
+	static UAnimStateTransitionNode* GetTransitionNodeFromGraph(const FAnimBlueprintDebugData& DebugData, const UEdGraph* Graph);
+
+	static UAnimStateNode* GetStateNodeFromGraph(const FAnimBlueprintDebugData& DebugData, const UEdGraph* Graph);
+
 };
diff --git a/Engine/Source/Editor/AnimGraph/Private/AnimationTransitionSchema.cpp b/Engine/Source/Editor/AnimGraph/Private/AnimationTransitionSchema.cpp
--- a/Engine/Source/Editor/AnimGraph/Private/AnimationTransitionSchema.cpp
+++ b/Engine/Source/Editor/AnimGraph/Private/AnimationTransitionSchema.cpp
@@ -47,7 +47,7 @@ void UAnimationTransitionSchema::GetSourceStateActions(FGraphContextMenuBuilder&
 
 		if (UAnimBlueprintGeneratedClass* AnimBlueprintClass = Blueprint->GetAnimBlueprintSkeletonClass())
 		{
-			if (UAnimStateTransitionNode* TransNode = AnimBlueprintClass->GetAnimBlueprintDebugData().GetTransitionNodeFromGraph(ContextMenuBuilder.CurrentGraph))
+			if (UAnimStateTransitionNode* TransNode = GetTransitionNodeFromGraph(AnimBlueprintClass->GetAnimBlueprintDebugData(), ContextMenuBuilder.CurrentGraph))
 			{
 				if (UAnimStateNode* SourceStateNode = Cast<UAnimStateNode>(TransNode->GetPreviousState()))
 				{
@@ -150,7 +150,7 @@ void UAnimationTransitionSchema::GetGraphDisplayInformation(const UEdGraph& Grap
 		UAnimBlueprint* Blueprint = CastChecked<UAnimBlueprint>(FBlueprintEditorUtils::FindBlueprintForGraph(&Graph));
 		if (UAnimBlueprintGeneratedClass* AnimBlueprintClass = Blueprint->GetAnimBlueprintSkeletonClass())
 		{
-			TransNode = AnimBlueprintClass->GetAnimBlueprintDebugData().GetTransitionNodeFromGraph(&Graph);
+			TransNode = GetTransitionNodeFromGraph(AnimBlueprintClass->GetAnimBlueprintDebugData(), &Graph);
 		}
 	}
 
@@ -164,4 +164,29 @@ void UAnimationTransitionSchema::GetGraphDisplayInformation(const UEdGraph& Grap
 	DisplayInfo.DisplayName = DisplayInfo.PlainName;
 }
 
+UAnimStateTransitionNode* UAnimationTransitionSchema::GetTransitionNodeFromGraph(const FAnimBlueprintDebugData& DebugData, const UEdGraph* Graph)
+{
+	if (const TWeakObjectPtr<UAnimStateTransitionNode>* TransNodePtr = DebugData.TransitionGraphToNodeMap.Find(Graph))
+	{
+		return TransNodePtr->Get();
+	}
+
+	if (const TWeakObjectPtr<UAnimStateTransitionNode>* TransNodePtr = DebugData.TransitionBlendGraphToNodeMap.Find(Graph))
+	{
+		return TransNodePtr->Get();
+	}
+
+	return NULL;
+}
+
+UAnimStateNode* UAnimationTransitionSchema::GetStateNodeFromGraph(const FAnimBlueprintDebugData& DebugData, const UEdGraph* Graph)
+{
+	if (const TWeakObjectPtr<UAnimStateNode>* StateNodePtr = DebugData.StateGraphToNodeMap.Find(Graph))
+	{
+		return StateNodePtr->Get();
+	}
+
+	return NULL;
+}
+
 #undef LOCTEXT_NAMESPACE
\ No newline at end of file
diff --git a/Engine/Source/Editor/BlueprintGraph/Classes/K2Node_MacroInstance.h b/Engine/Source/Editor/BlueprintGraph/Classes/K2Node_MacroInstance.h
--- a/Engine/Source/Editor/BlueprintGraph/Classes/K2Node_MacroInstance.h
+++ b/Engine/Source/Editor/BlueprintGraph/Classes/K2Node_MacroInstance.h
@@ -3,6 +3,7 @@
 
 #pragma once
 #include "K2Node_Tunnel.h"
+#include "EdGraph/EdGraph.h"
 #include "K2Node_MacroInstance.generated.h"
 
 UCLASS(MinimalAPI)
diff --git a/Engine/Source/Editor/BlueprintGraph/Private/EdGraphSchema_K2.cpp b/Engine/Source/Editor/BlueprintGraph/Private/EdGraphSchema_K2.cpp
--- a/Engine/Source/Editor/BlueprintGraph/Private/EdGraphSchema_K2.cpp
+++ b/Engine/Source/Editor/BlueprintGraph/Private/EdGraphSchema_K2.cpp
@@ -7,6 +7,7 @@
 
 #include "BlueprintGraphPrivatePCH.h"
 
+#include "Engine/LevelScriptBlueprint.h"
 #include "GraphEditorActions.h"
 #include "GraphEditorSettings.h"
 #include "ScopedTransaction.h"
diff --git a/Engine/Source/Editor/BlueprintGraph/Private/K2ActionMenuBuilder.cpp b/Engine/Source/Editor/BlueprintGraph/Private/K2ActionMenuBuilder.cpp
--- a/Engine/Source/Editor/BlueprintGraph/Private/K2ActionMenuBuilder.cpp
+++ b/Engine/Source/Editor/BlueprintGraph/Private/K2ActionMenuBuilder.cpp
@@ -1,6 +1,7 @@
 // Copyright 1998-2014 Epic Games, Inc. All Rights Reserved.
 
 #include "BlueprintGraphPrivatePCH.h"
+#include "Engine/LevelScriptBlueprint.h"
 #include "K2ActionMenuBuilder.h"
 #include "ComponentAssetBroker.h"
 #include "KismetEditorUtilities.h"
diff --git a/Engine/Source/Editor/BlueprintGraph/Private/K2Node.cpp b/Engine/Source/Editor/BlueprintGraph/Private/K2Node.cpp
--- a/Engine/Source/Editor/BlueprintGraph/Private/K2Node.cpp
+++ b/Engine/Source/Editor/BlueprintGraph/Private/K2Node.cpp
@@ -2,6 +2,8 @@
 
 
 #include "BlueprintGraphPrivatePCH.h"
+#include "Engine/Breakpoint.h"
+
 #include "K2Node.h"
 #include "KismetDebugUtilities.h" // for HasDebuggingData(), GetWatchText()
 
diff --git a/Engine/Source/Editor/DestructibleMeshEditor/Private/SDestructibleMeshEditorViewport.cpp b/Engine/Source/Editor/DestructibleMeshEditor/Private/SDestructibleMeshEditorViewport.cpp
--- a/Engine/Source/Editor/DestructibleMeshEditor/Private/SDestructibleMeshEditorViewport.cpp
+++ b/Engine/Source/Editor/DestructibleMeshEditor/Private/SDestructibleMeshEditorViewport.cpp
@@ -203,7 +203,7 @@ void FDestructibleMeshEditorViewportClient::Fracture()
 		else if (DestructibleMesh->ApexDestructibleAsset != NULL)
 		{
 			DestructibleMesh = ImportDestructibleMeshFromApexDestructibleAsset(DestructibleMesh->GetOuter(), *DestructibleMesh->ApexDestructibleAsset, DestructibleMesh->GetFName(), DestructibleMesh->GetFlags(), NULL,
- 																			   EImportOptions::PreserveSettings);
+ 																			   EDestructibleImportOptions::PreserveSettings);
 		}
 		
 		DestructibleMesh->FractureSettings->CreateVoronoiSitesInRootMesh();
diff --git a/Engine/Source/Editor/DetailCustomizations/Private/ActorDetails.cpp b/Engine/Source/Editor/DetailCustomizations/Private/ActorDetails.cpp
--- a/Engine/Source/Editor/DetailCustomizations/Private/ActorDetails.cpp
+++ b/Engine/Source/Editor/DetailCustomizations/Private/ActorDetails.cpp
@@ -13,6 +13,7 @@
 #include "ClassIconFinder.h"
 #include "Runtime/AssetRegistry/Public/AssetRegistryModule.h"
 #include "BlueprintGraphDefinitions.h"
+#include "Engine/Breakpoint.h"
 #include "ActorMaterialCategory.h"
 #include "IContentBrowserSingleton.h"
 #include "ContentBrowserModule.h"
diff --git a/Engine/Source/Editor/DetailCustomizations/Private/DeviceProfileDetails.cpp b/Engine/Source/Editor/DetailCustomizations/Private/DeviceProfileDetails.cpp
--- a/Engine/Source/Editor/DetailCustomizations/Private/DeviceProfileDetails.cpp
+++ b/Engine/Source/Editor/DetailCustomizations/Private/DeviceProfileDetails.cpp
@@ -5,6 +5,8 @@
 =============================================================================*/
 
 #include "DetailCustomizationsPrivatePCH.h"
+#include "DeviceProfiles/DeviceProfile.h"
+#include "DeviceProfiles/DeviceProfileManager.h"
 #include "DeviceProfileDetails.h"
 
 
diff --git a/Engine/Source/Editor/DetailCustomizations/Private/DeviceProfileDetails.h b/Engine/Source/Editor/DetailCustomizations/Private/DeviceProfileDetails.h
--- a/Engine/Source/Editor/DetailCustomizations/Private/DeviceProfileDetails.h
+++ b/Engine/Source/Editor/DetailCustomizations/Private/DeviceProfileDetails.h
@@ -6,6 +6,7 @@
 
 #pragma once
 
+class UDeviceProfile;
 
 /* FDeviceProfileParentPropertyDetails
  *****************************************************************************/
diff --git a/Engine/Source/Editor/DeviceProfileEditor/Private/DetailsPanel/SDeviceProfileDetailsPanel.cpp b/Engine/Source/Editor/DeviceProfileEditor/Private/DetailsPanel/SDeviceProfileDetailsPanel.cpp
--- a/Engine/Source/Editor/DeviceProfileEditor/Private/DetailsPanel/SDeviceProfileDetailsPanel.cpp
+++ b/Engine/Source/Editor/DeviceProfileEditor/Private/DetailsPanel/SDeviceProfileDetailsPanel.cpp
@@ -7,6 +7,7 @@
 
 
 #include "DeviceProfileEditorPCH.h"
+#include "DeviceProfiles/DeviceProfile.h"
 
 
 #define LOCTEXT_NAMESPACE "DeviceProfileDetailsPanel"
diff --git a/Engine/Source/Editor/DeviceProfileEditor/Private/DetailsPanel/SDeviceProfileDetailsPanel.h b/Engine/Source/Editor/DeviceProfileEditor/Private/DetailsPanel/SDeviceProfileDetailsPanel.h
--- a/Engine/Source/Editor/DeviceProfileEditor/Private/DetailsPanel/SDeviceProfileDetailsPanel.h
+++ b/Engine/Source/Editor/DeviceProfileEditor/Private/DetailsPanel/SDeviceProfileDetailsPanel.h
@@ -9,6 +9,8 @@
 
 #include "PropertyEditing.h"
 
+class UDeviceProfile;
+
 #define LOCTEXT_NAMESPACE "DeviceProfileDetailsPanel"
 
 
diff --git a/Engine/Source/Editor/DeviceProfileEditor/Private/DeviceProfileConsoleVariableColumn.cpp b/Engine/Source/Editor/DeviceProfileEditor/Private/DeviceProfileConsoleVariableColumn.cpp
--- a/Engine/Source/Editor/DeviceProfileEditor/Private/DeviceProfileConsoleVariableColumn.cpp
+++ b/Engine/Source/Editor/DeviceProfileEditor/Private/DeviceProfileConsoleVariableColumn.cpp
@@ -10,6 +10,7 @@
 #include "PropertyHandle.h"
 
 // Misc includes
+#include "DeviceProfiles/DeviceProfile.h"
 #include "ScopedTransaction.h"
 
 
diff --git a/Engine/Source/Editor/DeviceProfileEditor/Private/SDeviceProfileCreateProfilePanel.cpp b/Engine/Source/Editor/DeviceProfileEditor/Private/SDeviceProfileCreateProfilePanel.cpp
--- a/Engine/Source/Editor/DeviceProfileEditor/Private/SDeviceProfileCreateProfilePanel.cpp
+++ b/Engine/Source/Editor/DeviceProfileEditor/Private/SDeviceProfileCreateProfilePanel.cpp
@@ -6,6 +6,8 @@
 
 
 #include "DeviceProfileEditorPCH.h"
+#include "DeviceProfiles/DeviceProfileManager.h"
+#include "DeviceProfiles/DeviceProfile.h"
 #include "ISourceControlModule.h"
 #include "TargetPlatform.h"
 
diff --git a/Engine/Source/Editor/DeviceProfileEditor/Private/SDeviceProfileEditor.cpp b/Engine/Source/Editor/DeviceProfileEditor/Private/SDeviceProfileEditor.cpp
--- a/Engine/Source/Editor/DeviceProfileEditor/Private/SDeviceProfileEditor.cpp
+++ b/Engine/Source/Editor/DeviceProfileEditor/Private/SDeviceProfileEditor.cpp
@@ -8,6 +8,8 @@
 
 #include "DeviceProfileEditorPCH.h"
 
+#include "DeviceProfiles/DeviceProfileManager.h"
+#include "DeviceProfiles/DeviceProfile.h"
 #include "ISourceControlModule.h"
 #include "PropertyEditorModule.h"
 #include "PropertyPath.h"
diff --git a/Engine/Source/Editor/DeviceProfileEditor/Private/SDeviceProfileEditorSingleProfileView.cpp b/Engine/Source/Editor/DeviceProfileEditor/Private/SDeviceProfileEditorSingleProfileView.cpp
--- a/Engine/Source/Editor/DeviceProfileEditor/Private/SDeviceProfileEditorSingleProfileView.cpp
+++ b/Engine/Source/Editor/DeviceProfileEditor/Private/SDeviceProfileEditorSingleProfileView.cpp
@@ -6,6 +6,7 @@
 
 
 #include "DeviceProfileEditorPCH.h"
+#include "DeviceProfiles/DeviceProfile.h"
 #include "PropertyEditorModule.h"
 
 
diff --git a/Engine/Source/Editor/DeviceProfileEditor/Private/SDeviceProfileSelectionPanel.cpp b/Engine/Source/Editor/DeviceProfileEditor/Private/SDeviceProfileSelectionPanel.cpp
--- a/Engine/Source/Editor/DeviceProfileEditor/Private/SDeviceProfileSelectionPanel.cpp
+++ b/Engine/Source/Editor/DeviceProfileEditor/Private/SDeviceProfileSelectionPanel.cpp
@@ -6,6 +6,8 @@
 
 
 #include "DeviceProfileEditorPCH.h"
+#include "DeviceProfiles/DeviceProfile.h"
+#include "DeviceProfiles/DeviceProfileManager.h"
 #include "PropertyEditorModule.h"
 
 
diff --git a/Engine/Source/Editor/DeviceProfileServices/Private/DeviceProfileServicesUIManager.cpp b/Engine/Source/Editor/DeviceProfileServices/Private/DeviceProfileServicesUIManager.cpp
--- a/Engine/Source/Editor/DeviceProfileServices/Private/DeviceProfileServicesUIManager.cpp
+++ b/Engine/Source/Editor/DeviceProfileServices/Private/DeviceProfileServicesUIManager.cpp
@@ -5,7 +5,8 @@
 =============================================================================*/
 
 #include "DeviceProfileServicesPCH.h"
-
+#include "DeviceProfiles/DeviceProfileManager.h"
+#include "DeviceProfiles/DeviceProfile.h"
 
 DEFINE_LOG_CATEGORY_STATIC(LogDeviceProfileServices, Log, All);
 
diff --git a/Engine/Source/Editor/DeviceProfileServices/Public/IDeviceProfileServicesUIManager.h b/Engine/Source/Editor/DeviceProfileServices/Public/IDeviceProfileServicesUIManager.h
--- a/Engine/Source/Editor/DeviceProfileServices/Public/IDeviceProfileServicesUIManager.h
+++ b/Engine/Source/Editor/DeviceProfileServices/Public/IDeviceProfileServicesUIManager.h
@@ -6,6 +6,7 @@
 
 #pragma once
 
+class UDeviceProfile;
 
 /**
  * Type definition for shared pointers to instances of IDeviceProfileServicesUIManager.
diff --git a/Engine/Source/Editor/GraphEditor/Private/KismetNodes/SGraphNodeK2Base.cpp b/Engine/Source/Editor/GraphEditor/Private/KismetNodes/SGraphNodeK2Base.cpp
--- a/Engine/Source/Editor/GraphEditor/Private/KismetNodes/SGraphNodeK2Base.cpp
+++ b/Engine/Source/Editor/GraphEditor/Private/KismetNodes/SGraphNodeK2Base.cpp
@@ -2,6 +2,7 @@
 
 
 #include "GraphEditorCommon.h"
+#include "Engine/Breakpoint.h"
 #include "SGraphNodeK2Base.h"
 #include "Editor/UnrealEd/Public/Kismet2/KismetDebugUtilities.h"
 #include "Editor/UnrealEd/Public/Kismet2/BlueprintEditorUtils.h"
diff --git a/Engine/Source/Editor/Kismet/Private/BlueprintEditor.cpp b/Engine/Source/Editor/Kismet/Private/BlueprintEditor.cpp
--- a/Engine/Source/Editor/Kismet/Private/BlueprintEditor.cpp
+++ b/Engine/Source/Editor/Kismet/Private/BlueprintEditor.cpp
@@ -1,6 +1,8 @@
 // Copyright 1998-2014 Epic Games, Inc. All Rights Reserved.
 
 #include "BlueprintEditorPrivatePCH.h"
+#include "Engine/LevelScriptBlueprint.h"
+#include "Engine/Breakpoint.h"
 #include "ScopedTransaction.h"
 #include "Editor/UnrealEd/Public/Kismet2/KismetDebugUtilities.h"
 #include "Editor/UnrealEd/Public/Kismet2/BlueprintEditorUtils.h"
diff --git a/Engine/Source/Editor/Kismet/Private/Debugging/SKismetDebuggingView.cpp b/Engine/Source/Editor/Kismet/Private/Debugging/SKismetDebuggingView.cpp
--- a/Engine/Source/Editor/Kismet/Private/Debugging/SKismetDebuggingView.cpp
+++ b/Engine/Source/Editor/Kismet/Private/Debugging/SKismetDebuggingView.cpp
@@ -2,6 +2,7 @@
 
 
 #include "BlueprintEditorPrivatePCH.h"
+#include "Engine/Breakpoint.h"
 #include "Editor/UnrealEd/Public/Kismet2/KismetDebugUtilities.h"
 #include "GraphEditor.h"
 #include "BlueprintUtilities.h"
diff --git a/Engine/Source/Editor/Kismet/Private/SBlueprintEditorToolbar.cpp b/Engine/Source/Editor/Kismet/Private/SBlueprintEditorToolbar.cpp
--- a/Engine/Source/Editor/Kismet/Private/SBlueprintEditorToolbar.cpp
+++ b/Engine/Source/Editor/Kismet/Private/SBlueprintEditorToolbar.cpp
@@ -21,6 +21,7 @@
 #include "SLevelOfDetailBranchNode.h"
 #include "STutorialWrapper.h"
 #include "SBlueprintEditorSelectedDebugObjectWidget.h"
+#include "Engine/LevelScriptBlueprint.h"
 
 #define LOCTEXT_NAMESPACE "KismetToolbar"
 
diff --git a/Engine/Source/Editor/KismetCompiler/Private/KismetCompiler.cpp b/Engine/Source/Editor/KismetCompiler/Private/KismetCompiler.cpp
--- a/Engine/Source/Editor/KismetCompiler/Private/KismetCompiler.cpp
+++ b/Engine/Source/Editor/KismetCompiler/Private/KismetCompiler.cpp
@@ -5,6 +5,7 @@
 =============================================================================*/
 
 #include "KismetCompilerPrivatePCH.h"
+#include "Engine/LevelScriptBlueprint.h"
 #include "KismetCompilerBackend.h"
 #include "Editor/UnrealEd/Public/Kismet2/KismetDebugUtilities.h"
 #include "Editor/UnrealEd/Public/Kismet2/KismetReinstanceUtilities.h"
diff --git a/Engine/Source/Editor/LevelEditor/Private/LevelEditorActions.cpp b/Engine/Source/Editor/LevelEditor/Private/LevelEditorActions.cpp
--- a/Engine/Source/Editor/LevelEditor/Private/LevelEditorActions.cpp
+++ b/Engine/Source/Editor/LevelEditor/Private/LevelEditorActions.cpp
@@ -3,6 +3,7 @@
 
 #include "LevelEditor.h"
 #include "Matinee/MatineeActor.h"
+#include "Engine/LevelScriptBlueprint.h"
 #include "EditorSupportDelegates.h"
 #include "SLevelEditor.h"
 #include "LevelEditorActions.h"
diff --git a/Engine/Source/Editor/LevelEditor/Private/LevelEditorToolBar.cpp b/Engine/Source/Editor/LevelEditor/Private/LevelEditorToolBar.cpp
--- a/Engine/Source/Editor/LevelEditor/Private/LevelEditorToolBar.cpp
+++ b/Engine/Source/Editor/LevelEditor/Private/LevelEditorToolBar.cpp
@@ -15,6 +15,7 @@
 #include "Editor/ContentBrowser/Public/ContentBrowserModule.h"
 #include "EngineBuildSettings.h"
 #include "Matinee/MatineeActor.h"
+#include "Engine/LevelScriptBlueprint.h"
 
 static TAutoConsoleVariable<int32> CVarFeatureLevelPreview(
 	TEXT("r.FeatureLevelPreview"),
diff --git a/Engine/Source/Editor/LevelEditor/Private/SLevelViewportToolBar.cpp b/Engine/Source/Editor/LevelEditor/Private/SLevelViewportToolBar.cpp
--- a/Engine/Source/Editor/LevelEditor/Private/SLevelViewportToolBar.cpp
+++ b/Engine/Source/Editor/LevelEditor/Private/SLevelViewportToolBar.cpp
@@ -13,6 +13,7 @@
 #include "DelegateFilter.h"
 #include "Editor/SceneOutliner/Public/ISceneOutlinerColumn.h"
 #include "DeviceProfileServices.h"
+#include "DeviceProfiles/DeviceProfile.h"
 #include "EditorViewportCommands.h"
 #include "SLevelEditor.h"
 #include "Editor/UnrealEd/Public/SEditorViewportToolBarMenu.h"
diff --git a/Engine/Source/Editor/Levels/Private/LevelViewModel.cpp b/Engine/Source/Editor/Levels/Private/LevelViewModel.cpp
--- a/Engine/Source/Editor/Levels/Private/LevelViewModel.cpp
+++ b/Engine/Source/Editor/Levels/Private/LevelViewModel.cpp
@@ -2,6 +2,7 @@
 
 #include "LevelsPrivatePCH.h"
 
+#include "Engine/LevelScriptBlueprint.h"
 #include "LevelViewModel.h"
 #include "LevelUtils.h"
 #include "EditorLevelUtils.h"
diff --git a/Engine/Source/Editor/ProjectSettingsViewer/Private/ProjectSettingsViewerModule.cpp b/Engine/Source/Editor/ProjectSettingsViewer/Private/ProjectSettingsViewerModule.cpp
--- a/Engine/Source/Editor/ProjectSettingsViewer/Private/ProjectSettingsViewerModule.cpp
+++ b/Engine/Source/Editor/ProjectSettingsViewer/Private/ProjectSettingsViewerModule.cpp
@@ -5,7 +5,7 @@
 =============================================================================*/
 
 #include "ProjectSettingsViewerPrivatePCH.h"
-
+#include "Engine/Console.h"
 
 #define LOCTEXT_NAMESPACE "FProjectSettingsViewerModule"
 
diff --git a/Engine/Source/Editor/Sequencer/Private/Sequencer.cpp b/Engine/Source/Editor/Sequencer/Private/Sequencer.cpp
--- a/Engine/Source/Editor/Sequencer/Private/Sequencer.cpp
+++ b/Engine/Source/Editor/Sequencer/Private/Sequencer.cpp
@@ -5,6 +5,7 @@
 #include "SequencerObjectSpawner.h"
 #include "Toolkits/IToolkitHost.h"
 #include "MovieScene.h"
+#include "Engine/LevelScriptBlueprint.h"
 #include "Editor/PropertyEditor/Public/PropertyEditorModule.h"
 #include "Editor/PropertyEditor/Public/IDetailsView.h"
 #include "Editor/SequencerWidgets/Public/ITimeSlider.h"
diff --git a/Engine/Source/Editor/UnrealEd/Classes/Commandlets/AutomatedLauncherCommandlet.h b/Engine/Source/Editor/UnrealEd/Classes/Commandlets/AutomatedLauncherCommandlet.h
--- a/Engine/Source/Editor/UnrealEd/Classes/Commandlets/AutomatedLauncherCommandlet.h
+++ b/Engine/Source/Editor/UnrealEd/Classes/Commandlets/AutomatedLauncherCommandlet.h
@@ -3,6 +3,7 @@
 /** Commandlet for running automated commands through the launcher */
 
 #pragma once
+#include "Commandlets/Commandlet.h"
 #include "AutomatedLauncherCommandlet.generated.h"
 
 UCLASS()
diff --git a/Engine/Source/Editor/UnrealEd/Classes/Commandlets/CompressAnimationsCommandlet.h b/Engine/Source/Editor/UnrealEd/Classes/Commandlets/CompressAnimationsCommandlet.h
--- a/Engine/Source/Editor/UnrealEd/Classes/Commandlets/CompressAnimationsCommandlet.h
+++ b/Engine/Source/Editor/UnrealEd/Classes/Commandlets/CompressAnimationsCommandlet.h
@@ -1,6 +1,7 @@
 // Copyright 1998-2014 Epic Games, Inc. All Rights Reserved.
 
 #pragma once
+#include "Commandlets/Commandlet.h"
 #include "CompressAnimationsCommandlet.generated.h"
 
 UCLASS()
diff --git a/Engine/Source/Editor/UnrealEd/Classes/Commandlets/CookCommandlet.h b/Engine/Source/Editor/UnrealEd/Classes/Commandlets/CookCommandlet.h
--- a/Engine/Source/Editor/UnrealEd/Classes/Commandlets/CookCommandlet.h
+++ b/Engine/Source/Editor/UnrealEd/Classes/Commandlets/CookCommandlet.h
@@ -5,7 +5,7 @@
 =============================================================================*/
 
 #pragma once
-
+#include "Commandlets/Commandlet.h"
 #include "CookCommandlet.generated.h"
 
 UCLASS()
diff --git a/Engine/Source/Editor/UnrealEd/Classes/Commandlets/DerivedDataCacheCommandlet.h b/Engine/Source/Editor/UnrealEd/Classes/Commandlets/DerivedDataCacheCommandlet.h
--- a/Engine/Source/Editor/UnrealEd/Classes/Commandlets/DerivedDataCacheCommandlet.h
+++ b/Engine/Source/Editor/UnrealEd/Classes/Commandlets/DerivedDataCacheCommandlet.h
@@ -5,7 +5,7 @@
 =============================================================================*/
 
 #pragma once
-
+#include "Commandlets/Commandlet.h"
 #include "DerivedDataCacheCommandlet.generated.h"
 
 
diff --git a/Engine/Source/Editor/UnrealEd/Classes/Commandlets/DiffAssetsCommandlet.h b/Engine/Source/Editor/UnrealEd/Classes/Commandlets/DiffAssetsCommandlet.h
--- a/Engine/Source/Editor/UnrealEd/Classes/Commandlets/DiffAssetsCommandlet.h
+++ b/Engine/Source/Editor/UnrealEd/Classes/Commandlets/DiffAssetsCommandlet.h
@@ -5,6 +5,7 @@
  */
 
 #pragma once
+#include "Commandlets/Commandlet.h"
 #include "DiffAssetsCommandlet.generated.h"
 
 UCLASS()
diff --git a/Engine/Source/Editor/UnrealEd/Classes/Commandlets/DiffPackagesCommandlet.h b/Engine/Source/Editor/UnrealEd/Classes/Commandlets/DiffPackagesCommandlet.h
--- a/Engine/Source/Editor/UnrealEd/Classes/Commandlets/DiffPackagesCommandlet.h
+++ b/Engine/Source/Editor/UnrealEd/Classes/Commandlets/DiffPackagesCommandlet.h
@@ -1,6 +1,7 @@
 // Copyright 1998-2014 Epic Games, Inc. All Rights Reserved.
 
 #pragma once
+#include "Commandlets/Commandlet.h"
 #include "DiffPackagesCommandlet.generated.h"
 
 /**
diff --git a/Engine/Source/Editor/UnrealEd/Classes/Commandlets/DumpBlueprintsInfoCommandlet.h b/Engine/Source/Editor/UnrealEd/Classes/Commandlets/DumpBlueprintsInfoCommandlet.h
--- a/Engine/Source/Editor/UnrealEd/Classes/Commandlets/DumpBlueprintsInfoCommandlet.h
+++ b/Engine/Source/Editor/UnrealEd/Classes/Commandlets/DumpBlueprintsInfoCommandlet.h
@@ -1,6 +1,7 @@
 // Copyright 1998-2014 Epic Games, Inc. All Rights Reserved.
 
 #pragma once
+#include "Commandlets/Commandlet.h"
 #include "DumpBlueprintsInfoCommandlet.generated.h"
 
 UCLASS()
diff --git a/Engine/Source/Editor/UnrealEd/Classes/Commandlets/DumpHiddenCategoriesCommandlet.h b/Engine/Source/Editor/UnrealEd/Classes/Commandlets/DumpHiddenCategoriesCommandlet.h
--- a/Engine/Source/Editor/UnrealEd/Classes/Commandlets/DumpHiddenCategoriesCommandlet.h
+++ b/Engine/Source/Editor/UnrealEd/Classes/Commandlets/DumpHiddenCategoriesCommandlet.h
@@ -1,6 +1,7 @@
 // Copyright 1998-2014 Epic Games, Inc. All Rights Reserved.
 
 #pragma once
+#include "Commandlets/Commandlet.h"
 #include "DumpHiddenCategoriesCommandlet.generated.h"
 
 UCLASS()
diff --git a/Engine/Source/Editor/UnrealEd/Classes/Commandlets/FileServerCommandlet.h b/Engine/Source/Editor/UnrealEd/Classes/Commandlets/FileServerCommandlet.h
--- a/Engine/Source/Editor/UnrealEd/Classes/Commandlets/FileServerCommandlet.h
+++ b/Engine/Source/Editor/UnrealEd/Classes/Commandlets/FileServerCommandlet.h
@@ -5,7 +5,7 @@
 =============================================================================*/
 
 #pragma once
-
+#include "Commandlets/Commandlet.h"
 #include "FileServerCommandlet.generated.h"
 
 
diff --git a/Engine/Source/Editor/UnrealEd/Classes/Commandlets/FixupRedirectsCommandlet.h b/Engine/Source/Editor/UnrealEd/Classes/Commandlets/FixupRedirectsCommandlet.h
--- a/Engine/Source/Editor/UnrealEd/Classes/Commandlets/FixupRedirectsCommandlet.h
+++ b/Engine/Source/Editor/UnrealEd/Classes/Commandlets/FixupRedirectsCommandlet.h
@@ -1,6 +1,7 @@
 // Copyright 1998-2014 Epic Games, Inc. All Rights Reserved.
 
 #pragma once
+#include "Commandlets/Commandlet.h"
 #include "FixupRedirectsCommandlet.generated.h"
 
 UCLASS()
diff --git a/Engine/Source/Editor/UnrealEd/Classes/Commandlets/GatherTextCommandlet.h b/Engine/Source/Editor/UnrealEd/Classes/Commandlets/GatherTextCommandlet.h
--- a/Engine/Source/Editor/UnrealEd/Classes/Commandlets/GatherTextCommandlet.h
+++ b/Engine/Source/Editor/UnrealEd/Classes/Commandlets/GatherTextCommandlet.h
@@ -1,7 +1,7 @@
 // Copyright 1998-2014 Epic Games, Inc. All Rights Reserved.
 
 #pragma once 
-
+#include "Commandlets/GatherTextCommandletBase.h"
 #include "GatherTextCommandlet.generated.h"
 
 class FJsonValue;
diff --git a/Engine/Source/Editor/UnrealEd/Classes/Commandlets/GatherTextCommandletBase.h b/Engine/Source/Editor/UnrealEd/Classes/Commandlets/GatherTextCommandletBase.h
--- a/Engine/Source/Editor/UnrealEd/Classes/Commandlets/GatherTextCommandletBase.h
+++ b/Engine/Source/Editor/UnrealEd/Classes/Commandlets/GatherTextCommandletBase.h
@@ -4,6 +4,7 @@
 
 #include "UnrealString.h"
 #include "InternationalizationManifest.h"
+#include "Commandlets/Commandlet.h"
 #include "GatherTextCommandletBase.generated.h"
 
 class FWordCountReportData
diff --git a/Engine/Source/Editor/UnrealEd/Classes/Commandlets/GatherTextFromAssetsCommandlet.h b/Engine/Source/Editor/UnrealEd/Classes/Commandlets/GatherTextFromAssetsCommandlet.h
--- a/Engine/Source/Editor/UnrealEd/Classes/Commandlets/GatherTextFromAssetsCommandlet.h
+++ b/Engine/Source/Editor/UnrealEd/Classes/Commandlets/GatherTextFromAssetsCommandlet.h
@@ -1,7 +1,7 @@
 // Copyright 1998-2014 Epic Games, Inc. All Rights Reserved.
 
 #pragma once 
-
+#include "Commandlets/Commandlet.h"
 #include "GatherTextFromAssetsCommandlet.generated.h"
 
 namespace EAssetTextGatherStatus
diff --git a/Engine/Source/Editor/UnrealEd/Classes/Commandlets/GatherTextFromMetadataCommandlet.h b/Engine/Source/Editor/UnrealEd/Classes/Commandlets/GatherTextFromMetadataCommandlet.h
--- a/Engine/Source/Editor/UnrealEd/Classes/Commandlets/GatherTextFromMetadataCommandlet.h
+++ b/Engine/Source/Editor/UnrealEd/Classes/Commandlets/GatherTextFromMetadataCommandlet.h
@@ -1,7 +1,7 @@
 // Copyright 1998-2014 Epic Games, Inc. All Rights Reserved.
 
 #pragma once 
-
+#include "Commandlets/Commandlet.h"
 #include "GatherTextFromMetaDataCommandlet.generated.h"
 
 /**
diff --git a/Engine/Source/Editor/UnrealEd/Classes/Commandlets/GatherTextFromSourceCommandlet.h b/Engine/Source/Editor/UnrealEd/Classes/Commandlets/GatherTextFromSourceCommandlet.h
--- a/Engine/Source/Editor/UnrealEd/Classes/Commandlets/GatherTextFromSourceCommandlet.h
+++ b/Engine/Source/Editor/UnrealEd/Classes/Commandlets/GatherTextFromSourceCommandlet.h
@@ -1,7 +1,7 @@
 // Copyright 1998-2014 Epic Games, Inc. All Rights Reserved.
 
 #pragma once 
-
+#include "Commandlets/Commandlet.h"
 #include "GatherTextFromSourceCommandlet.generated.h"
 
 /**
diff --git a/Engine/Source/Editor/UnrealEd/Classes/Commandlets/GenerateDistillFileSetsCommandlet.h b/Engine/Source/Editor/UnrealEd/Classes/Commandlets/GenerateDistillFileSetsCommandlet.h
--- a/Engine/Source/Editor/UnrealEd/Classes/Commandlets/GenerateDistillFileSetsCommandlet.h
+++ b/Engine/Source/Editor/UnrealEd/Classes/Commandlets/GenerateDistillFileSetsCommandlet.h
@@ -5,6 +5,7 @@
  */
 
 #pragma once
+#include "Commandlets/Commandlet.h"
 #include "GenerateDistillFileSetsCommandlet.generated.h"
 
 UCLASS()
diff --git a/Engine/Source/Editor/UnrealEd/Classes/Commandlets/GenerateGatherArchiveCommandlet.h b/Engine/Source/Editor/UnrealEd/Classes/Commandlets/GenerateGatherArchiveCommandlet.h
--- a/Engine/Source/Editor/UnrealEd/Classes/Commandlets/GenerateGatherArchiveCommandlet.h
+++ b/Engine/Source/Editor/UnrealEd/Classes/Commandlets/GenerateGatherArchiveCommandlet.h
@@ -4,6 +4,7 @@
 
 #include "Internationalization/InternationalizationArchive.h"
 #include "Internationalization/InternationalizationManifest.h"
+#include "Commandlets/Commandlet.h"
 #include "GenerateGatherArchiveCommandlet.generated.h"
 
 
diff --git a/Engine/Source/Editor/UnrealEd/Classes/Commandlets/GenerateGatherManifestCommandlet.h b/Engine/Source/Editor/UnrealEd/Classes/Commandlets/GenerateGatherManifestCommandlet.h
--- a/Engine/Source/Editor/UnrealEd/Classes/Commandlets/GenerateGatherManifestCommandlet.h
+++ b/Engine/Source/Editor/UnrealEd/Classes/Commandlets/GenerateGatherManifestCommandlet.h
@@ -1,7 +1,7 @@
 // Copyright 1998-2014 Epic Games, Inc. All Rights Reserved.
 
 #pragma once 
-
+#include "Commandlets/Commandlet.h"
 #include "GenerateGatherManifestCommandlet.generated.h"
 
 /**
diff --git a/Engine/Source/Editor/UnrealEd/Classes/Commandlets/GenerateTextLocalizationReportCommandlet.h b/Engine/Source/Editor/UnrealEd/Classes/Commandlets/GenerateTextLocalizationReportCommandlet.h
--- a/Engine/Source/Editor/UnrealEd/Classes/Commandlets/GenerateTextLocalizationReportCommandlet.h
+++ b/Engine/Source/Editor/UnrealEd/Classes/Commandlets/GenerateTextLocalizationReportCommandlet.h
@@ -2,6 +2,7 @@
 
 #pragma once 
 
+#include "Commandlets/Commandlet.h"
 #include "GenerateTextLocalizationReportCommandlet.generated.h"
 
 
diff --git a/Engine/Source/Editor/UnrealEd/Classes/Commandlets/GenerateTextLocalizationResourceCommandlet.h b/Engine/Source/Editor/UnrealEd/Classes/Commandlets/GenerateTextLocalizationResourceCommandlet.h
--- a/Engine/Source/Editor/UnrealEd/Classes/Commandlets/GenerateTextLocalizationResourceCommandlet.h
+++ b/Engine/Source/Editor/UnrealEd/Classes/Commandlets/GenerateTextLocalizationResourceCommandlet.h
@@ -2,6 +2,7 @@
 
 #pragma once 
 
+#include "Commandlets/Commandlet.h"
 #include "GenerateTextLocalizationResourceCommandlet.generated.h"
 
 /**
diff --git a/Engine/Source/Editor/UnrealEd/Classes/Commandlets/InternationalizationConditioningCommandlet.h b/Engine/Source/Editor/UnrealEd/Classes/Commandlets/InternationalizationConditioningCommandlet.h
--- a/Engine/Source/Editor/UnrealEd/Classes/Commandlets/InternationalizationConditioningCommandlet.h
+++ b/Engine/Source/Editor/UnrealEd/Classes/Commandlets/InternationalizationConditioningCommandlet.h
@@ -2,6 +2,7 @@
 
 #pragma once 
 
+#include "Commandlets/Commandlet.h"
 #include "InternationalizationConditioningCommandlet.generated.h"
 
 
diff --git a/Engine/Source/Editor/UnrealEd/Classes/Commandlets/InternationalizationExportCommandlet.h b/Engine/Source/Editor/UnrealEd/Classes/Commandlets/InternationalizationExportCommandlet.h
--- a/Engine/Source/Editor/UnrealEd/Classes/Commandlets/InternationalizationExportCommandlet.h
+++ b/Engine/Source/Editor/UnrealEd/Classes/Commandlets/InternationalizationExportCommandlet.h
@@ -2,6 +2,7 @@
 
 #pragma once 
 
+#include "Commandlets/Commandlet.h"
 #include "InternationalizationExportCommandlet.generated.h"
 
 class FCulture;
diff --git a/Engine/Source/Editor/UnrealEd/Classes/Commandlets/KismetUpdateCommandlet.h b/Engine/Source/Editor/UnrealEd/Classes/Commandlets/KismetUpdateCommandlet.h
--- a/Engine/Source/Editor/UnrealEd/Classes/Commandlets/KismetUpdateCommandlet.h
+++ b/Engine/Source/Editor/UnrealEd/Classes/Commandlets/KismetUpdateCommandlet.h
@@ -5,6 +5,7 @@
  */
 
 #pragma once
+#include "Commandlets/ResavePackagesCommandlet.h"
 #include "KismetUpdateCommandlet.generated.h"
 
 UCLASS()
diff --git a/Engine/Source/Editor/UnrealEd/Classes/Commandlets/ListMaterialsUsedWithMeshEmittersCommandlet.h b/Engine/Source/Editor/UnrealEd/Classes/Commandlets/ListMaterialsUsedWithMeshEmittersCommandlet.h
--- a/Engine/Source/Editor/UnrealEd/Classes/Commandlets/ListMaterialsUsedWithMeshEmittersCommandlet.h
+++ b/Engine/Source/Editor/UnrealEd/Classes/Commandlets/ListMaterialsUsedWithMeshEmittersCommandlet.h
@@ -1,6 +1,7 @@
 // Copyright 1998-2014 Epic Games, Inc. All Rights Reserved.
 
 #pragma once
+#include "Commandlets/Commandlet.h"
 #include "ListMaterialsUsedWithMeshEmittersCommandlet.generated.h"
 
 UCLASS()
diff --git a/Engine/Source/Editor/UnrealEd/Classes/Commandlets/ListStaticMeshesImportedFromSpeedTreesCommandlet.h b/Engine/Source/Editor/UnrealEd/Classes/Commandlets/ListStaticMeshesImportedFromSpeedTreesCommandlet.h
--- a/Engine/Source/Editor/UnrealEd/Classes/Commandlets/ListStaticMeshesImportedFromSpeedTreesCommandlet.h
+++ b/Engine/Source/Editor/UnrealEd/Classes/Commandlets/ListStaticMeshesImportedFromSpeedTreesCommandlet.h
@@ -1,6 +1,7 @@
 // Copyright 1998-2014 Epic Games, Inc. All Rights Reserved.
 
 #pragma once
+#include "Commandlets/Commandlet.h"
 #include "ListStaticMeshesImportedFromSpeedTreesCommandlet.generated.h"
 
 UCLASS()
diff --git a/Engine/Source/Editor/UnrealEd/Classes/Commandlets/LoadPackageCommandlet.h b/Engine/Source/Editor/UnrealEd/Classes/Commandlets/LoadPackageCommandlet.h
--- a/Engine/Source/Editor/UnrealEd/Classes/Commandlets/LoadPackageCommandlet.h
+++ b/Engine/Source/Editor/UnrealEd/Classes/Commandlets/LoadPackageCommandlet.h
@@ -1,6 +1,7 @@
 // Copyright 1998-2014 Epic Games, Inc. All Rights Reserved.
 
 #pragma once
+#include "Commandlets/Commandlet.h"
 #include "LoadPackageCommandlet.generated.h"
 
 UCLASS()
diff --git a/Engine/Source/Editor/UnrealEd/Classes/Commandlets/PkgInfoCommandlet.h b/Engine/Source/Editor/UnrealEd/Classes/Commandlets/PkgInfoCommandlet.h
--- a/Engine/Source/Editor/UnrealEd/Classes/Commandlets/PkgInfoCommandlet.h
+++ b/Engine/Source/Editor/UnrealEd/Classes/Commandlets/PkgInfoCommandlet.h
@@ -1,6 +1,7 @@
 // Copyright 1998-2014 Epic Games, Inc. All Rights Reserved.
 
 #pragma once
+#include "Commandlets/Commandlet.h"
 #include "PkgInfoCommandlet.generated.h"
 
 UCLASS()
diff --git a/Engine/Source/Editor/UnrealEd/Classes/Commandlets/RecompileBlueprintsCommandlet.h b/Engine/Source/Editor/UnrealEd/Classes/Commandlets/RecompileBlueprintsCommandlet.h
--- a/Engine/Source/Editor/UnrealEd/Classes/Commandlets/RecompileBlueprintsCommandlet.h
+++ b/Engine/Source/Editor/UnrealEd/Classes/Commandlets/RecompileBlueprintsCommandlet.h
@@ -1,6 +1,7 @@
 // Copyright 1998-2014 Epic Games, Inc. All Rights Reserved.
 
 #pragma once
+#include "Commandlets/ResavePackagesCommandlet.h"
 #include "RecompileBlueprintsCommandlet.generated.h"
 
 UCLASS()
diff --git a/Engine/Source/Editor/UnrealEd/Classes/Commandlets/ReplaceActorCommandlet.h b/Engine/Source/Editor/UnrealEd/Classes/Commandlets/ReplaceActorCommandlet.h
--- a/Engine/Source/Editor/UnrealEd/Classes/Commandlets/ReplaceActorCommandlet.h
+++ b/Engine/Source/Editor/UnrealEd/Classes/Commandlets/ReplaceActorCommandlet.h
@@ -1,6 +1,7 @@
 // Copyright 1998-2014 Epic Games, Inc. All Rights Reserved.
 
 #pragma once
+#include "Commandlets/Commandlet.h"
 #include "ReplaceActorCommandlet.generated.h"
 
 /** Commandlet for replacing one kind of actor with another kind of actor, copying changed properties from the most-derived common superclass */
diff --git a/Engine/Source/Editor/UnrealEd/Classes/Commandlets/ResavePackagesCommandlet.h b/Engine/Source/Editor/UnrealEd/Classes/Commandlets/ResavePackagesCommandlet.h
--- a/Engine/Source/Editor/UnrealEd/Classes/Commandlets/ResavePackagesCommandlet.h
+++ b/Engine/Source/Editor/UnrealEd/Classes/Commandlets/ResavePackagesCommandlet.h
@@ -1,6 +1,7 @@
 // Copyright 1998-2014 Epic Games, Inc. All Rights Reserved.
 
 #pragma once
+#include "Commandlets/Commandlet.h"
 #include "ResavePackagesCommandlet.generated.h"
 
 UCLASS()
diff --git a/Engine/Source/Editor/UnrealEd/Classes/Commandlets/UpdateGameProjectCommandlet.h b/Engine/Source/Editor/UnrealEd/Classes/Commandlets/UpdateGameProjectCommandlet.h
--- a/Engine/Source/Editor/UnrealEd/Classes/Commandlets/UpdateGameProjectCommandlet.h
+++ b/Engine/Source/Editor/UnrealEd/Classes/Commandlets/UpdateGameProjectCommandlet.h
@@ -1,6 +1,7 @@
 // Copyright 1998-2014 Epic Games, Inc. All Rights Reserved.
 
 #pragma once
+#include "Commandlets/Commandlet.h"
 #include "UpdateGameProjectCommandlet.generated.h"
 
 UCLASS()
diff --git a/Engine/Source/Editor/UnrealEd/Classes/Commandlets/WrangleContentCommandlet.h b/Engine/Source/Editor/UnrealEd/Classes/Commandlets/WrangleContentCommandlet.h
--- a/Engine/Source/Editor/UnrealEd/Classes/Commandlets/WrangleContentCommandlet.h
+++ b/Engine/Source/Editor/UnrealEd/Classes/Commandlets/WrangleContentCommandlet.h
@@ -1,6 +1,7 @@
 // Copyright 1998-2014 Epic Games, Inc. All Rights Reserved.
 
 #pragma once
+#include "Commandlets/Commandlet.h"
 #include "WrangleContentCommandlet.generated.h"
 
 UCLASS()
diff --git a/Engine/Source/Editor/UnrealEd/Classes/MaterialGraph/MaterialGraph.h b/Engine/Source/Editor/UnrealEd/Classes/MaterialGraph/MaterialGraph.h
--- a/Engine/Source/Editor/UnrealEd/Classes/MaterialGraph/MaterialGraph.h
+++ b/Engine/Source/Editor/UnrealEd/Classes/MaterialGraph/MaterialGraph.h
@@ -1,7 +1,7 @@
 // Copyright 1998-2014 Epic Games, Inc. All Rights Reserved.
 
 #pragma once
-
+#include "EdGraph/EdGraph.h"
 #include "MaterialGraph.generated.h"
 
 DECLARE_DELEGATE_RetVal( bool, FRealtimeStateGetter );
diff --git a/Engine/Source/Editor/UnrealEd/Classes/MaterialGraph/MaterialGraphNode_Comment.h b/Engine/Source/Editor/UnrealEd/Classes/MaterialGraph/MaterialGraphNode_Comment.h
--- a/Engine/Source/Editor/UnrealEd/Classes/MaterialGraph/MaterialGraphNode_Comment.h
+++ b/Engine/Source/Editor/UnrealEd/Classes/MaterialGraph/MaterialGraphNode_Comment.h
@@ -1,7 +1,7 @@
 // Copyright 1998-2014 Epic Games, Inc. All Rights Reserved.
 
 #pragma once
-
+#include "EdGraph/EdGraphNode_Comment.h"
 #include "MaterialGraphNode_Comment.generated.h"
 
 UCLASS(MinimalAPI)
diff --git a/Engine/Source/Editor/UnrealEd/Classes/MaterialGraph/MaterialGraphSchema.h b/Engine/Source/Editor/UnrealEd/Classes/MaterialGraph/MaterialGraphSchema.h
--- a/Engine/Source/Editor/UnrealEd/Classes/MaterialGraph/MaterialGraphSchema.h
+++ b/Engine/Source/Editor/UnrealEd/Classes/MaterialGraph/MaterialGraphSchema.h
@@ -1,7 +1,7 @@
 // Copyright 1998-2014 Epic Games, Inc. All Rights Reserved.
 
 #pragma once
-
+#include "EdGraph/EdGraphSchema.h"
 #include "MaterialGraphSchema.generated.h"
 
 /** Action to add an expression node to the graph */
diff --git a/Engine/Source/Editor/UnrealEd/Classes/Niagara/EdGraphSchema_Niagara.h b/Engine/Source/Editor/UnrealEd/Classes/Niagara/EdGraphSchema_Niagara.h
--- a/Engine/Source/Editor/UnrealEd/Classes/Niagara/EdGraphSchema_Niagara.h
+++ b/Engine/Source/Editor/UnrealEd/Classes/Niagara/EdGraphSchema_Niagara.h
@@ -1,6 +1,7 @@
 // Copyright 1998-2014 Epic Games, Inc. All Rights Reserved.
 
 #pragma once
+#include "EdGraph/EdGraphSchema.h"
 #include "EdGraphSchema_Niagara.generated.h"
 
 /** Action to add a node to the graph */
diff --git a/Engine/Source/Editor/UnrealEd/Classes/Niagara/NiagaraGraph.h b/Engine/Source/Editor/UnrealEd/Classes/Niagara/NiagaraGraph.h
--- a/Engine/Source/Editor/UnrealEd/Classes/Niagara/NiagaraGraph.h
+++ b/Engine/Source/Editor/UnrealEd/Classes/Niagara/NiagaraGraph.h
@@ -1,6 +1,7 @@
 // Copyright 1998-2014 Epic Games, Inc. All Rights Reserved.
 
 #pragma once
+#include "EdGraph/EdGraph.h"
 #include "NiagaraGraph.generated.h"
 
 UCLASS(MinimalAPI)
diff --git a/Engine/Source/Editor/UnrealEd/Classes/SoundClassGraph/SoundClassGraph.h b/Engine/Source/Editor/UnrealEd/Classes/SoundClassGraph/SoundClassGraph.h
--- a/Engine/Source/Editor/UnrealEd/Classes/SoundClassGraph/SoundClassGraph.h
+++ b/Engine/Source/Editor/UnrealEd/Classes/SoundClassGraph/SoundClassGraph.h
@@ -1,6 +1,7 @@
 // Copyright 1998-2014 Epic Games, Inc. All Rights Reserved.
 
 #pragma once
+#include "EdGraph/EdGraph.h"
 #include "SoundClassGraph.generated.h"
 
 UCLASS(MinimalAPI)
diff --git a/Engine/Source/Editor/UnrealEd/Classes/SoundCueGraph/SoundCueGraph.h b/Engine/Source/Editor/UnrealEd/Classes/SoundCueGraph/SoundCueGraph.h
--- a/Engine/Source/Editor/UnrealEd/Classes/SoundCueGraph/SoundCueGraph.h
+++ b/Engine/Source/Editor/UnrealEd/Classes/SoundCueGraph/SoundCueGraph.h
@@ -1,6 +1,7 @@
 // Copyright 1998-2014 Epic Games, Inc. All Rights Reserved.
 
 #pragma once
+#include "EdGraph/EdGraph.h"
 #include "SoundCueGraph.generated.h"
 
 UCLASS(MinimalAPI)
diff --git a/Engine/Source/Editor/UnrealEd/Private/Factories/ApexDestructibleAssetImport.cpp b/Engine/Source/Editor/UnrealEd/Private/Factories/ApexDestructibleAssetImport.cpp
--- a/Engine/Source/Editor/UnrealEd/Private/Factories/ApexDestructibleAssetImport.cpp
+++ b/Engine/Source/Editor/UnrealEd/Private/Factories/ApexDestructibleAssetImport.cpp
@@ -15,6 +15,7 @@
 DEFINE_LOG_CATEGORY_STATIC(LogApexDestructibleAssetImport, Log, All);
 
 #include "Engine.h"
+#include "Engine/DestructibleFractureSettings.h"
 #include "TextureLayout.h"
 #include "SkelImport.h"
 #include "EditorPhysXSupport.h"
@@ -720,7 +721,7 @@ NxDestructibleAsset* CreateApexDestructibleAssetFromFile(const FString& Filename
 	return ApexDestructibleAsset;
 }
 
-bool SetApexDestructibleAsset(UDestructibleMesh& DestructibleMesh, NxDestructibleAsset& ApexDestructibleAsset, FSkeletalMeshImportData* OutData, EImportOptions::Type Options)
+bool SetApexDestructibleAsset(UDestructibleMesh& DestructibleMesh, NxDestructibleAsset& ApexDestructibleAsset, FSkeletalMeshImportData* OutData, EDestructibleImportOptions::Type Options)
 {
 	DestructibleMesh.PreEditChange(NULL);
 
@@ -746,7 +747,7 @@ bool SetApexDestructibleAsset(UDestructibleMesh& DestructibleMesh, NxDestructibl
 	// Removing const cast ... we'll have to make it non-const anyway when we modify it
 	DestructibleMesh.ApexDestructibleAsset = &ApexDestructibleAsset;
 
-	if ( !(Options&EImportOptions::PreserveSettings) )
+	if ( !(Options&EDestructibleImportOptions::PreserveSettings) )
 	{
 		// Resize the depth parameters array to the appropriate size
 		DestructibleMesh.DefaultDestructibleParameters.DepthParameters.Init(FDestructibleDepthParameters(), ApexDestructibleAsset.getDepthCount());
@@ -969,14 +970,14 @@ UNREALED_API bool BuildDestructibleMeshFromFractureSettings(UDestructibleMesh& D
 
 	if (NewApexDestructibleAsset != NULL)
 	{
-		Success = SetApexDestructibleAsset(DestructibleMesh, *NewApexDestructibleAsset, OutData, EImportOptions::PreserveSettings);
+		Success = SetApexDestructibleAsset(DestructibleMesh, *NewApexDestructibleAsset, OutData, EDestructibleImportOptions::PreserveSettings);
 	}
 #endif // WITH_APEX
 
 	return Success;
 }
 
-UDestructibleMesh* ImportDestructibleMeshFromApexDestructibleAsset(UObject* InParent, NxDestructibleAsset& ApexDestructibleAsset, FName Name, EObjectFlags Flags, FSkeletalMeshImportData* OutData, EImportOptions::Type Options)
+UDestructibleMesh* ImportDestructibleMeshFromApexDestructibleAsset(UObject* InParent, NxDestructibleAsset& ApexDestructibleAsset, FName Name, EObjectFlags Flags, FSkeletalMeshImportData* OutData, EDestructibleImportOptions::Type Options)
 {
 	// The APEX Destructible Asset contains an APEX Render Mesh Asset, get a pointer to this
 	const physx::NxRenderMeshAsset* ApexRenderMesh = ApexDestructibleAsset.getRenderMeshAsset();
@@ -1002,7 +1003,7 @@ UDestructibleMesh* ImportDestructibleMeshFromApexDestructibleAsset(UObject* InPa
 		DestructibleMesh = CastChecked<UDestructibleMesh>(StaticConstructObject(UDestructibleMesh::StaticClass(), InParent, Name, Flags));
 	}
 	
-	if (!(Options & EImportOptions::PreserveSettings))
+	if (!(Options & EDestructibleImportOptions::PreserveSettings))
 	{
 		// Store the current file path and timestamp for re-import purposes
 		// @todo AssetImportData make a data class for Apex destructible assets
diff --git a/Engine/Source/Editor/UnrealEd/Private/Factories/EditorFactories.cpp b/Engine/Source/Editor/UnrealEd/Private/Factories/EditorFactories.cpp
--- a/Engine/Source/Editor/UnrealEd/Private/Factories/EditorFactories.cpp
+++ b/Engine/Source/Editor/UnrealEd/Private/Factories/EditorFactories.cpp
@@ -45,6 +45,10 @@
 
 #include "Particles/ParticleSystem.h"
 
+#if WITH_APEX
+#include "ApexDestructibleAssetImport.h"
+#endif // WITH_APEX
+
 #if PLATFORM_WINDOWS
 // Needed for DDS support.
 #include "AllowWindowsPlatformTypes.h"
@@ -6832,7 +6836,7 @@ EReimportResult::Type UReimportDestructibleMeshFactory::Reimport( UObject* Obj )
 	{
 		// Succesfully created the NxDestructibleAsset, now create a UDestructibleMesh
 		UDestructibleMesh* ReimportedDestructibleMesh = ImportDestructibleMeshFromApexDestructibleAsset(DestructibleMesh->GetOuter(), *ApexDestructibleAsset, DestructibleMesh->GetFName(), DestructibleMesh->GetFlags(), NULL,
-																										EImportOptions::PreserveSettings);
+																										EDestructibleImportOptions::PreserveSettings);
 		if( ReimportedDestructibleMesh != NULL )
 		{
 			check( ReimportedDestructibleMesh == DestructibleMesh );
diff --git a/Engine/Source/Editor/UnrealEd/Private/Kismet2/BlueprintEditorUtils.cpp b/Engine/Source/Editor/UnrealEd/Private/Kismet2/BlueprintEditorUtils.cpp
--- a/Engine/Source/Editor/UnrealEd/Private/Kismet2/BlueprintEditorUtils.cpp
+++ b/Engine/Source/Editor/UnrealEd/Private/Kismet2/BlueprintEditorUtils.cpp
@@ -3,6 +3,7 @@
 #include "UnrealEd.h"
 
 #include "Matinee/MatineeActor.h"
+#include "Engine/LevelScriptBlueprint.h"
 
 #include "LatentActions.h"
 #include "Kismet2/CompilerResultsLog.h"
diff --git a/Engine/Source/Editor/UnrealEd/Private/Kismet2/Kismet2.cpp b/Engine/Source/Editor/UnrealEd/Private/Kismet2/Kismet2.cpp
--- a/Engine/Source/Editor/UnrealEd/Private/Kismet2/Kismet2.cpp
+++ b/Engine/Source/Editor/UnrealEd/Private/Kismet2/Kismet2.cpp
@@ -2,6 +2,7 @@
 
 
 #include "UnrealEd.h"
+#include "Engine/LevelScriptBlueprint.h"
 #include "BlueprintUtilities.h"
 #include "AnimGraphDefinitions.h"
 #include "AnimationGraph.h"
diff --git a/Engine/Source/Editor/UnrealEd/Private/Kismet2/KismetDebugUtilities.cpp b/Engine/Source/Editor/UnrealEd/Private/Kismet2/KismetDebugUtilities.cpp
--- a/Engine/Source/Editor/UnrealEd/Private/Kismet2/KismetDebugUtilities.cpp
+++ b/Engine/Source/Editor/UnrealEd/Private/Kismet2/KismetDebugUtilities.cpp
@@ -1,6 +1,7 @@
 // Copyright 1998-2014 Epic Games, Inc. All Rights Reserved.
 
 #include "UnrealEd.h"
+#include "Engine/Breakpoint.h"
 #include "ActorEditorUtils.h"
 #include "BlueprintUtilities.h"
 #include "AnimGraphDefinitions.h"
diff --git a/Engine/Source/Editor/UnrealEd/Public/ApexDestructibleAssetImport.h b/Engine/Source/Editor/UnrealEd/Public/ApexDestructibleAssetImport.h
--- a/Engine/Source/Editor/UnrealEd/Public/ApexDestructibleAssetImport.h
+++ b/Engine/Source/Editor/UnrealEd/Public/ApexDestructibleAssetImport.h
@@ -1,6 +1,8 @@
 // Copyright 1998-2014 Epic Games, Inc. All Rights Reserved.
 #pragma once
 
+#include "Engine/DestructibleFractureSettings.h"
+
 /*=============================================================================
 	ApexDestructibleAssetImport.h:
 
@@ -23,19 +25,7 @@ namespace physx
 	};
 };
 
-/** 
- * Options for APEX asset import.
- **/
-namespace EImportOptions
-{
-	enum Type
-	{
-		// Just imports the APEX asset
-		None				= 0,
-		// Preserves settings in DestructibleMesh 
-		PreserveSettings	= 1<<0,
-	};
-};
+
 
 /**
  * Creates an APEX NxDestructibleAsset from a binary buffer.
@@ -68,11 +58,11 @@ physx::apex::NxDestructibleAsset* CreateApexDestructibleAssetFromFile(const FStr
  * @param DestructibleMesh - the UDestructibleMesh to which to apply the new APEX Desrtructible Asset
  * @param ApexDestructibleAsset - the NxDestructibleAsset to use
  * @param OutData - Optional import data to populate
- * @param Options - Additional options to modify the import behavior ( see EImportOptions::Type )
+ * @param Options - Additional options to modify the import behavior ( see EDestructibleImportOptions::Type )
  *
  * @return true if successful
  */
-UNREALED_API bool SetApexDestructibleAsset(UDestructibleMesh& DestructibleMesh, physx::apex::NxDestructibleAsset& ApexDestructibleAsset, FSkeletalMeshImportData* OutData, EImportOptions::Type Options = EImportOptions::None);
+UNREALED_API bool SetApexDestructibleAsset(UDestructibleMesh& DestructibleMesh, physx::apex::NxDestructibleAsset& ApexDestructibleAsset, FSkeletalMeshImportData* OutData, EDestructibleImportOptions::Type Options = EDestructibleImportOptions::None);
 
 /**
  * Creates a DestructibleMesh from an APEX Destructible Asset with the given name and flags.
@@ -82,11 +72,11 @@ UNREALED_API bool SetApexDestructibleAsset(UDestructibleMesh& DestructibleMesh,
  * @param Name - the Unreal name for the UDestructibleMesh
  * @param Flags - object flags for the UDestructibleMesh
  * @param OutData - Optional import data to populate
- * @param Options - Additional options to modify the import behavior ( see EImportOptions::Type )
+ * @param Options - Additional options to modify the import behavior ( see EDestructibleImportOptions::Type )
  *
  * @return The newly created UDestructibleMesh if successful, NULL otherwise
  */
-UNREALED_API UDestructibleMesh* ImportDestructibleMeshFromApexDestructibleAsset(UObject* InParent, physx::apex::NxDestructibleAsset& ApexDestructibleAsset, FName Name, EObjectFlags Flags, FSkeletalMeshImportData* OutData, EImportOptions::Type Options = EImportOptions::None);
+UNREALED_API UDestructibleMesh* ImportDestructibleMeshFromApexDestructibleAsset(UObject* InParent, physx::apex::NxDestructibleAsset& ApexDestructibleAsset, FName Name, EObjectFlags Flags, FSkeletalMeshImportData* OutData, EDestructibleImportOptions::Type Options = EDestructibleImportOptions::None);
 
 /**
  * Builds a UDestructibleMesh from its internal fracture settings
diff --git a/Engine/Source/Editor/WorldBrowser/Private/LevelModel.cpp b/Engine/Source/Editor/WorldBrowser/Private/LevelModel.cpp
--- a/Engine/Source/Editor/WorldBrowser/Private/LevelModel.cpp
+++ b/Engine/Source/Editor/WorldBrowser/Private/LevelModel.cpp
@@ -1,6 +1,7 @@
 // Copyright 1998-2014 Epic Games, Inc. All Rights Reserved.
 #include "WorldBrowserPrivatePCH.h"
 
+#include "Engine/LevelScriptBlueprint.h"
 #include "Toolkits/AssetEditorManager.h"
 #include "LevelModel.h"
 #include "LevelCollectionModel.h"
diff --git a/Engine/Source/Runtime/Engine/Classes/Animation/AnimBlueprintGeneratedClass.h b/Engine/Source/Runtime/Engine/Classes/Animation/AnimBlueprintGeneratedClass.h
--- a/Engine/Source/Runtime/Engine/Classes/Animation/AnimBlueprintGeneratedClass.h
+++ b/Engine/Source/Runtime/Engine/Classes/Animation/AnimBlueprintGeneratedClass.h
@@ -129,30 +129,7 @@ public:
 		SnapshotBuffer = NULL;
 	}
 
-	UAnimStateTransitionNode* GetTransitionNodeFromGraph(const UEdGraph* Graph) const
-	{
-		if (const TWeakObjectPtr<UAnimStateTransitionNode>* TransNodePtr = TransitionGraphToNodeMap.Find(Graph))
-		{
-			return TransNodePtr->Get();
-		}
-
-		if (const TWeakObjectPtr<UAnimStateTransitionNode>* TransNodePtr = TransitionBlendGraphToNodeMap.Find(Graph))
-		{
-			return TransNodePtr->Get();
-		}
 
-		return NULL;
-	}
-
-	UAnimStateNode* GetStateNodeFromGraph(const UEdGraph* Graph) const
-	{
-		if (const TWeakObjectPtr<UAnimStateNode>* StateNodePtr = StateGraphToNodeMap.Find(Graph))
-		{
-			return StateNodePtr->Get();
-		}
-
-		return NULL;
-	}
 
 	bool IsReplayingSnapshot() const { return SnapshotIndex != INDEX_NONE; }
 	void TakeSnapshot(UAnimInstance* Instance);
diff --git a/Engine/Source/Runtime/Engine/Classes/Commandlets/SmokeTestCommandlet.h b/Engine/Source/Runtime/Engine/Classes/Commandlets/SmokeTestCommandlet.h
--- a/Engine/Source/Runtime/Engine/Classes/Commandlets/SmokeTestCommandlet.h
+++ b/Engine/Source/Runtime/Engine/Classes/Commandlets/SmokeTestCommandlet.h
@@ -1,6 +1,7 @@
 // Copyright 1998-2014 Epic Games, Inc. All Rights Reserved.
 
 #pragma once
+#include "Commandlets/Commandlet.h"
 #include "SmokeTestCommandlet.generated.h"
 
 UCLASS()
diff --git a/Engine/Source/Runtime/Engine/Classes/DeviceProfiles/DeviceProfileManager.h b/Engine/Source/Runtime/Engine/Classes/DeviceProfiles/DeviceProfileManager.h
--- a/Engine/Source/Runtime/Engine/Classes/DeviceProfiles/DeviceProfileManager.h
+++ b/Engine/Source/Runtime/Engine/Classes/DeviceProfiles/DeviceProfileManager.h
@@ -8,6 +8,8 @@
 
 #include "DeviceProfileManager.generated.h"
 
+class UDeviceProfile;
+
 // Delegate used to refresh the UI when the profiles change
 DECLARE_MULTICAST_DELEGATE( FOnManagerUpdated );
 
diff --git a/Engine/Source/Runtime/Engine/Classes/EdGraph/EdGraphNode.h b/Engine/Source/Runtime/Engine/Classes/EdGraph/EdGraphNode.h
--- a/Engine/Source/Runtime/Engine/Classes/EdGraph/EdGraphNode.h
+++ b/Engine/Source/Runtime/Engine/Classes/EdGraph/EdGraphNode.h
@@ -200,15 +200,7 @@ public:
 	}
 
 	/** Returns the graph that contains this node */
-	class UEdGraph* GetGraph() const
-	{
-		UEdGraph* Graph = Cast<UEdGraph>(GetOuter());
-		if(Graph == NULL)
-		{
-			ensureMsgf(false, TEXT("EdGraphNode::GetGraph : '%s' does not have a UEdGraph as an Outer."), *GetPathName());
-		}
-		return Graph;
-	}
+	class UEdGraph* GetGraph() const;
 
 	/**
 	 * Allocate default pins for a given node, based only the NodeType, which should already be filled in.
diff --git a/Engine/Source/Runtime/Engine/Classes/EdGraph/EdGraphNode_Comment.h b/Engine/Source/Runtime/Engine/Classes/EdGraph/EdGraphNode_Comment.h
--- a/Engine/Source/Runtime/Engine/Classes/EdGraph/EdGraphNode_Comment.h
+++ b/Engine/Source/Runtime/Engine/Classes/EdGraph/EdGraphNode_Comment.h
@@ -1,6 +1,7 @@
 // Copyright 1998-2014 Epic Games, Inc. All Rights Reserved.
 
 #pragma once
+#include "EdGraph/EdGraphNode.h"
 #include "EdGraphNode_Comment.generated.h"
 
 typedef TArray<class UObject*> FCommentNodeSet;
diff --git a/Engine/Source/Runtime/Engine/Classes/Engine/CodecMovieFallback.h b/Engine/Source/Runtime/Engine/Classes/Engine/CodecMovieFallback.h
--- a/Engine/Source/Runtime/Engine/Classes/Engine/CodecMovieFallback.h
+++ b/Engine/Source/Runtime/Engine/Classes/Engine/CodecMovieFallback.h
@@ -2,6 +2,7 @@
 
 
 #pragma once
+#include "Engine/CodecMovie.h"
 #include "CodecMovieFallback.generated.h"
 
 UCLASS()
diff --git a/Engine/Source/Runtime/Engine/Classes/Engine/DestructibleFractureSettings.h b/Engine/Source/Runtime/Engine/Classes/Engine/DestructibleFractureSettings.h
--- a/Engine/Source/Runtime/Engine/Classes/Engine/DestructibleFractureSettings.h
+++ b/Engine/Source/Runtime/Engine/Classes/Engine/DestructibleFractureSettings.h
@@ -6,10 +6,6 @@
 
 #if WITH_APEX
 
-#if WITH_EDITOR
-#include "ApexDestructibleAssetImport.h"
-#endif
-
 // Forward declares
 namespace physx
 {
@@ -27,6 +23,20 @@ namespace physx
 };
 #endif // WITH_APEX
 
+/** 
+ * Options for APEX asset import.
+ **/
+namespace EDestructibleImportOptions
+{
+	enum Type
+	{
+		// Just imports the APEX asset
+		None				= 0,
+		// Preserves settings in DestructibleMesh 
+		PreserveSettings	= 1<<0,
+	};
+};
+
 /** Parameters to describe the application of U,V coordinates on a particular slice within a destructible. */
 USTRUCT()
 struct FFractureMaterial
@@ -221,7 +231,7 @@ class UDestructibleFractureSettings : public UObject
 	 *
 	 * @returns true if successful
 	 */
-	ENGINE_API	bool					BuildRootMeshFromApexDestructibleAsset(physx::apex::NxDestructibleAsset& ApexDestructibleAsset, EImportOptions::Type Options = EImportOptions::None);
+	ENGINE_API	bool					BuildRootMeshFromApexDestructibleAsset(physx::apex::NxDestructibleAsset& ApexDestructibleAsset, EDestructibleImportOptions::Type Options = EDestructibleImportOptions::None);
 
 	/**
 	 * Build set of Voronoi sites internal to the root mesh
diff --git a/Engine/Source/Runtime/Engine/Classes/Engine/DestructibleMesh.h b/Engine/Source/Runtime/Engine/Classes/Engine/DestructibleMesh.h
--- a/Engine/Source/Runtime/Engine/Classes/Engine/DestructibleMesh.h
+++ b/Engine/Source/Runtime/Engine/Classes/Engine/DestructibleMesh.h
@@ -12,6 +12,14 @@ namespace NxParameterized
 {
 	class Interface;
 }
+
+namespace physx
+{
+	namespace apex
+	{
+		class NxDestructibleAsset;
+	}
+}
 #endif
 
 /**
diff --git a/Engine/Source/Runtime/Engine/Classes/Engine/Engine.h b/Engine/Source/Runtime/Engine/Classes/Engine/Engine.h
--- a/Engine/Source/Runtime/Engine/Classes/Engine/Engine.h
+++ b/Engine/Source/Runtime/Engine/Classes/Engine/Engine.h
@@ -8,6 +8,7 @@
 #include "Engine.generated.h"
 
 class FScreenSaverInhibitor;
+class UDeviceProfileManager;
 
 UENUM()
 enum EFullyLoadPackageType
diff --git a/Engine/Source/Runtime/Engine/Private/Blueprint.cpp b/Engine/Source/Runtime/Engine/Private/Blueprint.cpp
--- a/Engine/Source/Runtime/Engine/Private/Blueprint.cpp
+++ b/Engine/Source/Runtime/Engine/Private/Blueprint.cpp
@@ -1,6 +1,8 @@
 // Copyright 1998-2014 Epic Games, Inc. All Rights Reserved.
 
 #include "EnginePrivate.h"
+#include "Engine/Breakpoint.h"
+#include "EdGraph/EdGraph.h"
 #include "BlueprintUtilities.h"
 #include "LatentActions.h"
 
diff --git a/Engine/Source/Runtime/Engine/Private/Breakpoint.cpp b/Engine/Source/Runtime/Engine/Private/Breakpoint.cpp
--- a/Engine/Source/Runtime/Engine/Private/Breakpoint.cpp
+++ b/Engine/Source/Runtime/Engine/Private/Breakpoint.cpp
@@ -2,6 +2,7 @@
 
 
 #include "EnginePrivate.h"
+#include "Engine/Breakpoint.h"
 #include "BlueprintUtilities.h"
 
 UBreakpoint::UBreakpoint(const class FPostConstructInitializeProperties& PCIP)
diff --git a/Engine/Source/Runtime/Engine/Private/ClipPadEntry.cpp b/Engine/Source/Runtime/Engine/Private/ClipPadEntry.cpp
--- a/Engine/Source/Runtime/Engine/Private/ClipPadEntry.cpp
+++ b/Engine/Source/Runtime/Engine/Private/ClipPadEntry.cpp
@@ -1,6 +1,7 @@
 // Copyright 1998-2014 Epic Games, Inc. All Rights Reserved.
 
 #include "EnginePrivate.h"
+#include "Engine/ClipPadEntry.h"
 
 UDEPRECATED_ClipPadEntry::UDEPRECATED_ClipPadEntry(const class FPostConstructInitializeProperties& PCIP)
 	: Super(PCIP)
diff --git a/Engine/Source/Runtime/Engine/Private/Commandlets/Commandlet.cpp b/Engine/Source/Runtime/Engine/Private/Commandlets/Commandlet.cpp
--- a/Engine/Source/Runtime/Engine/Private/Commandlets/Commandlet.cpp
+++ b/Engine/Source/Runtime/Engine/Private/Commandlets/Commandlet.cpp
@@ -1,6 +1,7 @@
 // Copyright 1998-2014 Epic Games, Inc. All Rights Reserved.
 
 #include "EnginePrivate.h"
+#include "Commandlets/Commandlet.h"
 
 UCommandlet::UCommandlet(const class FPostConstructInitializeProperties& PCIP)
 	: Super(PCIP)
diff --git a/Engine/Source/Runtime/Engine/Private/Commandlets/SmokeTestCommandlet.cpp b/Engine/Source/Runtime/Engine/Private/Commandlets/SmokeTestCommandlet.cpp
--- a/Engine/Source/Runtime/Engine/Private/Commandlets/SmokeTestCommandlet.cpp
+++ b/Engine/Source/Runtime/Engine/Private/Commandlets/SmokeTestCommandlet.cpp
@@ -6,6 +6,7 @@
 =============================================================================*/
 
 #include "EnginePrivate.h"
+#include "Commandlets/SmokeTestCommandlet.h"
 
 USmokeTestCommandlet::USmokeTestCommandlet(const class FPostConstructInitializeProperties& PCIP)
 	: Super(PCIP)
diff --git a/Engine/Source/Runtime/Engine/Private/DestructibleFractureSettings.cpp b/Engine/Source/Runtime/Engine/Private/DestructibleFractureSettings.cpp
--- a/Engine/Source/Runtime/Engine/Private/DestructibleFractureSettings.cpp
+++ b/Engine/Source/Runtime/Engine/Private/DestructibleFractureSettings.cpp
@@ -1,6 +1,7 @@
 // Copyright 1998-2014 Epic Games, Inc. All Rights Reserved.
 
 #include "EnginePrivate.h"
+#include "Engine/DestructibleFractureSettings.h"
 #include "PhysicsEngine/PhysXSupport.h"
 
 
@@ -408,7 +409,7 @@ bool UDestructibleFractureSettings::SetRootMesh(const TArray<NxExplicitRenderTri
 	return Success;
 }
 
-bool UDestructibleFractureSettings::BuildRootMeshFromApexDestructibleAsset(NxDestructibleAsset& ApexDestructibleAsset, EImportOptions::Type Options)
+bool UDestructibleFractureSettings::BuildRootMeshFromApexDestructibleAsset(NxDestructibleAsset& ApexDestructibleAsset, EDestructibleImportOptions::Type Options)
 {
 	bool Success = false;
 
@@ -418,7 +419,7 @@ bool UDestructibleFractureSettings::BuildRootMeshFromApexDestructibleAsset(NxDes
 
 		IExplicitHierarchicalMesh& EHM = ApexDestructibleAssetAuthoring->getExplicitHierarchicalMesh();
 
-		if (!(Options & EImportOptions::PreserveSettings))
+		if (!(Options & EDestructibleImportOptions::PreserveSettings))
 		{
 			// Now apply the y -> -y and v -> 1-v transformation to all vertex data
 			for (PxU32 PartIndex = 0; PartIndex < EHM.partCount(); ++PartIndex)
diff --git a/Engine/Source/Runtime/Engine/Private/DestructibleMesh.cpp b/Engine/Source/Runtime/Engine/Private/DestructibleMesh.cpp
--- a/Engine/Source/Runtime/Engine/Private/DestructibleMesh.cpp
+++ b/Engine/Source/Runtime/Engine/Private/DestructibleMesh.cpp
@@ -5,6 +5,7 @@
 =============================================================================*/
 
 #include "EnginePrivate.h"
+#include "Engine/DestructibleFractureSettings.h"
 #include "PhysicsEngine/PhysXSupport.h"
 
 #if WITH_APEX && WITH_EDITOR
@@ -57,7 +58,7 @@ void UDestructibleMesh::PostLoad()
 				// re create destructible asset if it exceeds
 				if (ApexDestructibleAsset != NULL)
 				{
-					SetApexDestructibleAsset(*this, *ApexDestructibleAsset, NULL, EImportOptions::PreserveSettings);
+					SetApexDestructibleAsset(*this, *ApexDestructibleAsset, NULL, EDestructibleImportOptions::PreserveSettings);
 				}
 #else
 				UE_LOG(LogDestructible, Warning, TEXT("Can't render %s asset because it exceeds max GPU skin bones supported(%d). You'll need to resave this in the editor."), *GetName(), MaxGPUSkinBones);
diff --git a/Engine/Source/Runtime/Engine/Private/DeviceProfiles/DeviceProfile.cpp b/Engine/Source/Runtime/Engine/Private/DeviceProfiles/DeviceProfile.cpp
--- a/Engine/Source/Runtime/Engine/Private/DeviceProfiles/DeviceProfile.cpp
+++ b/Engine/Source/Runtime/Engine/Private/DeviceProfiles/DeviceProfile.cpp
@@ -1,6 +1,7 @@
 // Copyright 1998-2014 Epic Games, Inc. All Rights Reserved.
 
 #include "EnginePrivate.h"
+#include "DeviceProfiles/DeviceProfile.h"
 
 UDeviceProfile::UDeviceProfile(const class FPostConstructInitializeProperties& PCIP)
 	: Super(PCIP)
diff --git a/Engine/Source/Runtime/Engine/Private/DeviceProfiles/DeviceProfileManager.cpp b/Engine/Source/Runtime/Engine/Private/DeviceProfiles/DeviceProfileManager.cpp
--- a/Engine/Source/Runtime/Engine/Private/DeviceProfiles/DeviceProfileManager.cpp
+++ b/Engine/Source/Runtime/Engine/Private/DeviceProfiles/DeviceProfileManager.cpp
@@ -1,7 +1,8 @@
 // Copyright 1998-2014 Epic Games, Inc. All Rights Reserved.
 
 #include "EnginePrivate.h"
-#include "Core.h"
+#include "DeviceProfiles/DeviceProfileManager.h"
+#include "DeviceProfiles/DeviceProfile.h"
 #include "TargetPlatform.h"
 #include "IDeviceProfileSelectorModule.h"
 #include "IConsoleManager.h"
diff --git a/Engine/Source/Runtime/Engine/Private/EdGraph/EdGraph.cpp b/Engine/Source/Runtime/Engine/Private/EdGraph/EdGraph.cpp
--- a/Engine/Source/Runtime/Engine/Private/EdGraph/EdGraph.cpp
+++ b/Engine/Source/Runtime/Engine/Private/EdGraph/EdGraph.cpp
@@ -1,6 +1,7 @@
 // Copyright 1998-2014 Epic Games, Inc. All Rights Reserved.
 
 #include "EnginePrivate.h"
+#include "EdGraph/EdGraph.h"
 #include "BlueprintUtilities.h"
 #if WITH_EDITOR
 #include "Editor/UnrealEd/Public/Kismet2/BlueprintEditorUtils.h"
diff --git a/Engine/Source/Runtime/Engine/Private/EdGraph/EdGraphNode.cpp b/Engine/Source/Runtime/Engine/Private/EdGraph/EdGraphNode.cpp
--- a/Engine/Source/Runtime/Engine/Private/EdGraph/EdGraphNode.cpp
+++ b/Engine/Source/Runtime/Engine/Private/EdGraph/EdGraphNode.cpp
@@ -1,6 +1,7 @@
 // Copyright 1998-2014 Epic Games, Inc. All Rights Reserved.
 
 #include "EnginePrivate.h"
+#include "EdGraph/EdGraph.h"
 #include "BlueprintUtilities.h"
 #if WITH_EDITOR
 #include "Editor/UnrealEd/Public/Kismet2/BlueprintEditorUtils.h"
@@ -136,6 +137,16 @@ void UEdGraphNode::SnapToGrid(float GridSnapSize)
 	NodePosY = GridSnapSize * FMath::RoundToInt(NodePosY/GridSnapSize);
 }
 
+class UEdGraph* UEdGraphNode::GetGraph() const
+{
+	UEdGraph* Graph = Cast<UEdGraph>(GetOuter());
+	if (Graph == NULL)
+	{
+		ensureMsgf(false, TEXT("EdGraphNode::GetGraph : '%s' does not have a UEdGraph as an Outer."), *GetPathName());
+	}
+	return Graph;
+}
+
 void UEdGraphNode::DestroyNode()
 {
 	UEdGraph* ParentGraph = GetGraph();
diff --git a/Engine/Source/Runtime/Engine/Private/EdGraph/EdGraphNode_Comment.cpp b/Engine/Source/Runtime/Engine/Private/EdGraph/EdGraphNode_Comment.cpp
--- a/Engine/Source/Runtime/Engine/Private/EdGraph/EdGraphNode_Comment.cpp
+++ b/Engine/Source/Runtime/Engine/Private/EdGraph/EdGraphNode_Comment.cpp
@@ -1,6 +1,8 @@
 // Copyright 1998-2014 Epic Games, Inc. All Rights Reserved.
 
 #include "EnginePrivate.h"
+#include "EdGraph/EdGraphNode_Comment.h"
+
 #include "BlueprintUtilities.h"
 #if WITH_EDITOR
 #include "Editor/UnrealEd/Public/Kismet2/BlueprintEditorUtils.h"
diff --git a/Engine/Source/Runtime/Engine/Private/EdGraph/EdGraphPin.cpp b/Engine/Source/Runtime/Engine/Private/EdGraph/EdGraphPin.cpp
--- a/Engine/Source/Runtime/Engine/Private/EdGraph/EdGraphPin.cpp
+++ b/Engine/Source/Runtime/Engine/Private/EdGraph/EdGraphPin.cpp
@@ -1,6 +1,7 @@
 // Copyright 1998-2014 Epic Games, Inc. All Rights Reserved.
 
 #include "EnginePrivate.h"
+#include "EdGraph/EdGraph.h"
 #include "BlueprintUtilities.h"
 #if WITH_EDITOR
 #include "Editor/UnrealEd/Public/Kismet2/BlueprintEditorUtils.h"
diff --git a/Engine/Source/Runtime/Engine/Private/EdGraph/EdGraphSchema.cpp b/Engine/Source/Runtime/Engine/Private/EdGraph/EdGraphSchema.cpp
--- a/Engine/Source/Runtime/Engine/Private/EdGraph/EdGraphSchema.cpp
+++ b/Engine/Source/Runtime/Engine/Private/EdGraph/EdGraphSchema.cpp
@@ -1,6 +1,7 @@
 // Copyright 1998-2014 Epic Games, Inc. All Rights Reserved.
 
 #include "EnginePrivate.h"
+#include "EdGraph/EdGraph.h"
 #include "BlueprintUtilities.h"
 #if WITH_EDITOR
 #include "Editor/UnrealEd/Public/Kismet2/BlueprintEditorUtils.h"
diff --git a/Engine/Source/Runtime/Engine/Private/EngineMessage.cpp b/Engine/Source/Runtime/Engine/Private/EngineMessage.cpp
--- a/Engine/Source/Runtime/Engine/Private/EngineMessage.cpp
+++ b/Engine/Source/Runtime/Engine/Private/EngineMessage.cpp
@@ -1,6 +1,7 @@
 // Copyright 1998-2014 Epic Games, Inc. All Rights Reserved.
 
 #include "EnginePrivate.h"
+#include "Engine/Console.h"
 
 UEngineMessage::UEngineMessage(const class FPostConstructInitializeProperties& PCIP) :
 	Super(PCIP)
diff --git a/Engine/Source/Runtime/Engine/Private/EngineUtils.cpp b/Engine/Source/Runtime/Engine/Private/EngineUtils.cpp
--- a/Engine/Source/Runtime/Engine/Private/EngineUtils.cpp
+++ b/Engine/Source/Runtime/Engine/Private/EngineUtils.cpp
@@ -4,6 +4,7 @@
 =============================================================================*/
 #include "EnginePrivate.h"
 #include "EngineUtils.h"
+#include "Engine/Console.h"
 
 #include "DiagnosticTable.h"
 #include "TargetPlatform.h"
diff --git a/Engine/Source/Runtime/Engine/Private/GameViewportClient.cpp b/Engine/Source/Runtime/Engine/Private/GameViewportClient.cpp
--- a/Engine/Source/Runtime/Engine/Private/GameViewportClient.cpp
+++ b/Engine/Source/Runtime/Engine/Private/GameViewportClient.cpp
@@ -1,6 +1,7 @@
 // Copyright 1998-2014 Epic Games, Inc. All Rights Reserved.
 
 #include "EnginePrivate.h"
+#include "Engine/Console.h"
 #include "ParticleDefinitions.h"
 #include "SoundDefinitions.h"
 #include "FXSystem.h"
diff --git a/Engine/Source/Runtime/Engine/Private/KismetSystemLibrary.cpp b/Engine/Source/Runtime/Engine/Private/KismetSystemLibrary.cpp
--- a/Engine/Source/Runtime/Engine/Private/KismetSystemLibrary.cpp
+++ b/Engine/Source/Runtime/Engine/Private/KismetSystemLibrary.cpp
@@ -1,6 +1,7 @@
 // Copyright 1998-2014 Epic Games, Inc. All Rights Reserved.
 
 #include "EnginePrivate.h"
+#include "Engine/Console.h"
 #include "LatentActions.h"
 #include "DelayAction.h"
 #include "InterpolateComponentToAction.h"
diff --git a/Engine/Source/Runtime/Engine/Private/Level.cpp b/Engine/Source/Runtime/Engine/Private/Level.cpp
--- a/Engine/Source/Runtime/Engine/Private/Level.cpp
+++ b/Engine/Source/Runtime/Engine/Private/Level.cpp
@@ -5,6 +5,7 @@ Level.cpp: Level-related functions
 =============================================================================*/
 
 #include "EnginePrivate.h"
+#include "Engine/LevelScriptBlueprint.h"
 #include "Sound/SoundNodeWave.h"
 #include "Net/UnrealNetwork.h"
 #include "SoundDefinitions.h"
diff --git a/Engine/Source/Runtime/Engine/Private/LevelScriptActor.cpp b/Engine/Source/Runtime/Engine/Private/LevelScriptActor.cpp
--- a/Engine/Source/Runtime/Engine/Private/LevelScriptActor.cpp
+++ b/Engine/Source/Runtime/Engine/Private/LevelScriptActor.cpp
@@ -2,7 +2,7 @@
 
 
 #include "EnginePrivate.h"
-#include "SoundDefinitions.h"
+#include "Engine/LevelScriptBlueprint.h"
 #include "BlueprintUtilities.h"
 
 
diff --git a/Engine/Source/Runtime/Engine/Private/LevelScriptBlueprint.cpp b/Engine/Source/Runtime/Engine/Private/LevelScriptBlueprint.cpp
--- a/Engine/Source/Runtime/Engine/Private/LevelScriptBlueprint.cpp
+++ b/Engine/Source/Runtime/Engine/Private/LevelScriptBlueprint.cpp
@@ -1,6 +1,7 @@
 // Copyright 1998-2014 Epic Games, Inc. All Rights Reserved.
 
 #include "EnginePrivate.h"
+#include "Engine/LevelScriptBlueprint.h"
 #include "BlueprintUtilities.h"
 #include "LatentActions.h"
 
diff --git a/Engine/Source/Runtime/Engine/Private/MovieCodecs.cpp b/Engine/Source/Runtime/Engine/Private/MovieCodecs.cpp
--- a/Engine/Source/Runtime/Engine/Private/MovieCodecs.cpp
+++ b/Engine/Source/Runtime/Engine/Private/MovieCodecs.cpp
@@ -5,6 +5,7 @@
 =============================================================================*/
 
 #include "EnginePrivate.h"
+#include "Engine/CodecMovieFallback.h"
 
 /*-----------------------------------------------------------------------------
 	UCodecMovie
diff --git a/Engine/Source/Runtime/Engine/Private/PlayerController.cpp b/Engine/Source/Runtime/Engine/Private/PlayerController.cpp
--- a/Engine/Source/Runtime/Engine/Private/PlayerController.cpp
+++ b/Engine/Source/Runtime/Engine/Private/PlayerController.cpp
@@ -1,6 +1,7 @@
 // Copyright 1998-2014 Epic Games, Inc. All Rights Reserved.
 
 #include "EnginePrivate.h"
+#include "Engine/Console.h"
 #include "Net/UnrealNetwork.h"
 #include "GameFramework/GameNetworkManager.h"
 #include "ConfigCacheIni.h"
diff --git a/Engine/Source/Runtime/Engine/Private/TextureMovie.cpp b/Engine/Source/Runtime/Engine/Private/TextureMovie.cpp
--- a/Engine/Source/Runtime/Engine/Private/TextureMovie.cpp
+++ b/Engine/Source/Runtime/Engine/Private/TextureMovie.cpp
@@ -5,6 +5,7 @@
 =============================================================================*/
 
 #include "EnginePrivate.h"
+#include "Engine/CodecMovieFallback.h"
 
 /*-----------------------------------------------------------------------------
 UTextureMovie
diff --git a/Engine/Source/Runtime/Engine/Private/UnrealEngine.cpp b/Engine/Source/Runtime/Engine/Private/UnrealEngine.cpp
--- a/Engine/Source/Runtime/Engine/Private/UnrealEngine.cpp
+++ b/Engine/Source/Runtime/Engine/Private/UnrealEngine.cpp
@@ -7,6 +7,7 @@
 #include "EnginePrivate.h"
 
 #include "Net/UnrealNetwork.h"
+#include "Engine/Console.h"
 #include "FileManagerGeneric.h"
 #include "Database.h"
 #include "SkeletalMeshMerge.h"
@@ -38,6 +39,7 @@
 #include "StatsData.h"
 #include "AudioDevice.h"
 #include "ActiveSound.h"
+#include "DeviceProfiles/DeviceProfileManager.h"
 
 #include "Particles/Spawn/ParticleModuleSpawn.h"
 #include "Particles/TypeData/ParticleModuleTypeDataMesh.h"
diff --git a/Engine/Source/Runtime/Engine/Private/UserInterface/Console.cpp b/Engine/Source/Runtime/Engine/Private/UserInterface/Console.cpp
--- a/Engine/Source/Runtime/Engine/Private/UserInterface/Console.cpp
+++ b/Engine/Source/Runtime/Engine/Private/UserInterface/Console.cpp
@@ -5,6 +5,7 @@
 =============================================================================*/
 
 #include "EnginePrivate.h"
+#include "Engine/Console.h"
 #include "Slate.h"
 #include "DefaultValueHelper.h"
 
diff --git a/Engine/Source/Runtime/Engine/Private/World.cpp b/Engine/Source/Runtime/Engine/Private/World.cpp
--- a/Engine/Source/Runtime/Engine/Private/World.cpp
+++ b/Engine/Source/Runtime/Engine/Private/World.cpp
@@ -5,6 +5,8 @@
 =============================================================================*/
 
 #include "EnginePrivate.h"
+#include "Engine/LevelScriptBlueprint.h"
+#include "Engine/Console.h"
 #include "Matinee/MatineeActor.h"
 #include "Net/UnrealNetwork.h"
 #include "GameFramework/GameNetworkManager.h"
diff --git a/Engine/Source/Runtime/Engine/Public/AnimationRuntime.h b/Engine/Source/Runtime/Engine/Public/AnimationRuntime.h
--- a/Engine/Source/Runtime/Engine/Public/AnimationRuntime.h
+++ b/Engine/Source/Runtime/Engine/Public/AnimationRuntime.h
@@ -7,6 +7,8 @@
 
 #pragma once
 
+struct FInputBlendPose;
+
 /** In AnimationRunTime Library, we extract animation data based on Skeleton hierarchy, not ref pose hierarchy. 
 	Ref pose will need to be re-mapped later **/
 
diff --git a/Engine/Source/Runtime/Engine/Public/Engine.h b/Engine/Source/Runtime/Engine/Public/Engine.h
--- a/Engine/Source/Runtime/Engine/Public/Engine.h
+++ b/Engine/Source/Runtime/Engine/Public/Engine.h
@@ -223,12 +223,8 @@
 #include "AI/Navigation/AvoidanceManager.h"
 #include "Engine/BlendableInterface.h"
 #include "Engine/BlueprintCore.h"
-#include "EdGraph/EdGraphPin.h"
-#include "EdGraph/EdGraphNode.h"
-#include "EdGraph/EdGraph.h"
 #include "Engine/Blueprint.h"
 #include "Animation/AnimBlueprint.h"
-#include "Engine/LevelScriptBlueprint.h"
 #include "Sound/DialogueTypes.h"
 #include "Kismet/GameplayStatics.h"
 #include "Kismet/HeadMountedDisplayFunctionLibrary.h"
@@ -243,8 +239,6 @@
 #include "Engine/BlueprintGeneratedClass.h"
 #include "Animation/AnimBlueprintGeneratedClass.h"
 #include "PhysicsEngine/BodySetup.h"
-#include "Animation/AnimData/BoneMaskFilter.h"
-#include "Engine/Breakpoint.h"
 #include "Slate/ButtonStyleAsset.h"
 #include "Camera/CameraAnim.h"
 #include "Camera/CameraAnimInst.h"
@@ -253,13 +247,7 @@
 #include "Camera/CameraModifier_CameraShake.h"
 #include "GameFramework/CheatManager.h"
 #include "Slate/CheckboxStyleAsset.h"
-#include "Engine/ClipPadEntry.h"
-#include "Engine/CodecMovie.h"
-#include "Engine/CodecMovieFallback.h"
 #include "Engine/CollisionProfile.h"
-#include "Commandlets/Commandlet.h"
-#include "Commandlets/SmokeTestCommandlet.h"
-#include "Engine/Console.h"
 #include "Curves/CurveLinearColor.h"
 #include "Curves/CurveVector.h"
 #include "Curves/CurveEdPresetCurve.h"
@@ -267,9 +255,6 @@
 #include "GameFramework/DamageType.h"
 #include "Vehicles/TireType.h"
 #include "Engine/DataTable.h"
-#include "Engine/DestructibleFractureSettings.h"
-#include "DeviceProfiles/DeviceProfile.h"
-#include "DeviceProfiles/DeviceProfileManager.h"
 #include "Sound/DialogueVoice.h"
 #include "Sound/DialogueWave.h"
 #include "Distributions/Distribution.h"
@@ -297,8 +282,6 @@
 #include "Engine/InputVectorAxisDelegateBinding.h"
 #include "Engine/InputKeyDelegateBinding.h"
 #include "Engine/InputTouchDelegateBinding.h"
-#include "EdGraph/EdGraphNode_Comment.h"
-#include "EdGraph/EdGraphSchema.h"
 #include "Engine/Engine.h"
 #include "Engine/GameEngine.h"
 #include "Exporters/Exporter.h"
diff --git a/Engine/Source/Runtime/Launch/Private/LaunchEngineLoop.cpp b/Engine/Source/Runtime/Launch/Private/LaunchEngineLoop.cpp
--- a/Engine/Source/Runtime/Launch/Private/LaunchEngineLoop.cpp
+++ b/Engine/Source/Runtime/Launch/Private/LaunchEngineLoop.cpp
@@ -12,6 +12,7 @@
 #include "FileManagerGeneric.h"
 #include "TaskGraphInterfaces.h"
 
+
 #include "Projects.h"
 #include "UProjectInfo.h"
 #include "EngineVersion.h"
@@ -38,6 +39,7 @@
 	#include "ParticleHelper.h"
 	#include "Online.h"
 	#include "PlatformFeatures.h"
+	#include "DeviceProfiles/DeviceProfileManager.h"
 #if !UE_SERVER
 	#include "SlateRHIRendererModule.h"
 #endif
