Include EdGraph.h in BlueprintEditorUtils.h

diff --git a/Engine/Source/Editor/UnrealEd/Public/Kismet2/BlueprintEditorUtils.h b/Engine/Source/Editor/UnrealEd/Public/Kismet2/BlueprintEditorUtils.h
--- a/Engine/Source/Editor/UnrealEd/Public/Kismet2/BlueprintEditorUtils.h
+++ b/Engine/Source/Editor/UnrealEd/Public/Kismet2/BlueprintEditorUtils.h
@@ -3,7 +3,7 @@
 #pragma once
 
 #include "Editor/ClassViewer/Public/ClassViewerModule.h"
-
+#include "EdGraph/EdGraph.h"
 #include "EdGraphSchema_K2.h"
 
 /** 
