TTransArray object lifetime fix.

#add ITransaction::SaveArray now takes an extra parameter, to be called to construct the array elements.
#change TTransArray passes through an appropriate function for this parameter.
#change FObjectRecord::SerializeContents uses this parameter to construct its arguments.

diff --git a/Engine/Source/Editor/Matinee/Private/MatineeTransaction.cpp b/Engine/Source/Editor/Matinee/Private/MatineeTransaction.cpp
--- a/Engine/Source/Editor/Matinee/Private/MatineeTransaction.cpp
+++ b/Engine/Source/Editor/Matinee/Private/MatineeTransaction.cpp
@@ -29,11 +29,11 @@ void FMatineeTransaction::SaveObject( UObject* Object )
 		Object->IsA( UK2Node_MatineeController::StaticClass() ) )
 	{
 		// Save the object.
-		new( Records )FObjectRecord( this, Object, NULL, 0, 0, 0, 0, NULL, NULL );
+		new( Records )FObjectRecord( this, Object, NULL, 0, 0, 0, 0, NULL, NULL, NULL );
 	}
 }
 
-void FMatineeTransaction::SaveArray( UObject* Object, FScriptArray* Array, int32 Index, int32 Count, int32 Oper, int32 ElementSize, STRUCT_AR Serializer, STRUCT_DTOR Destructor )
+void FMatineeTransaction::SaveArray( UObject* Object, FScriptArray* Array, int32 Index, int32 Count, int32 Oper, int32 ElementSize, STRUCT_DC DefaultConstructor, STRUCT_AR Serializer, STRUCT_DTOR Destructor )
 {
 	// Never want this.
 }
\ No newline at end of file
diff --git a/Engine/Source/Editor/Matinee/Private/MatineeTransaction.h b/Engine/Source/Editor/Matinee/Private/MatineeTransaction.h
--- a/Engine/Source/Editor/Matinee/Private/MatineeTransaction.h
+++ b/Engine/Source/Editor/Matinee/Private/MatineeTransaction.h
@@ -14,7 +14,7 @@ public:
 
 	// Begin FTransaction Interface
 	virtual void SaveObject( UObject* Object ) OVERRIDE;
-	virtual void SaveArray( UObject* Object, FScriptArray* Array, int32 Index, int32 Count, int32 Oper, int32 ElementSize, STRUCT_AR Serializer, STRUCT_DTOR Destructor ) OVERRIDE;
+	virtual void SaveArray( UObject* Object, FScriptArray* Array, int32 Index, int32 Count, int32 Oper, int32 ElementSize, STRUCT_DC DefaultConstructor, STRUCT_AR Serializer, STRUCT_DTOR Destructor ) OVERRIDE;
 	// End FTransaction Interface
 };
 
diff --git a/Engine/Source/Editor/UnrealEd/Classes/Editor/Transactor.h b/Engine/Source/Editor/UnrealEd/Classes/Editor/Transactor.h
--- a/Engine/Source/Editor/UnrealEd/Classes/Editor/Transactor.h
+++ b/Engine/Source/Editor/UnrealEd/Classes/Editor/Transactor.h
@@ -84,6 +84,8 @@ protected:
 		int32				Oper;
 		/** Array: Size of each item in the array */
 		int32				ElementSize;
+		/** Array: DefaultConstructor for each item in the array */
+		STRUCT_DC			DefaultConstructor;
 		/** Array: Serializer to use for each item in the array */
 		STRUCT_AR			Serializer;
 		/** Array: Destructor for each item in the array */
@@ -96,17 +98,18 @@ protected:
 		// Constructors.
 		FObjectRecord()
 		{}
-		FObjectRecord( FTransaction* Owner, UObject* InObject, FScriptArray* InArray, int32 InIndex, int32 InCount, int32 InOper, int32 InElementSize, STRUCT_AR InSerializer, STRUCT_DTOR InDestructor )
-			:	Object		( InObject )
-			,	Array		( InArray )
-			,	Index		( InIndex )
-			,	Count		( InCount )
-			,	Oper		( InOper )
-			,	ElementSize	( InElementSize )
-			,	Serializer	( InSerializer )
-			,	Destructor	( InDestructor )
-			,	bRestored	( false )
-			,	bWantsBinarySerialization ( true )
+		FObjectRecord( FTransaction* Owner, UObject* InObject, FScriptArray* InArray, int32 InIndex, int32 InCount, int32 InOper, int32 InElementSize, STRUCT_DC InDefaultConstructor, STRUCT_AR InSerializer, STRUCT_DTOR InDestructor )
+			:	Object						( InObject )
+			,	Array						( InArray )
+			,	Index						( InIndex )
+			,	Count						( InCount )
+			,	Oper						( InOper )
+			,	ElementSize					( InElementSize )
+			,	DefaultConstructor			( InDefaultConstructor )
+			,	Serializer					( InSerializer )
+			,	Destructor					( InDestructor )
+			,	bRestored					( false )
+			,	bWantsBinarySerialization	( true )
 		{
 			// Blueprint compile-in-place can alter class layout so use tagged serialization for objects relying on a UBlueprint's Class
 			if (UBlueprintGeneratedClass* Class = Cast<UBlueprintGeneratedClass>(InObject->GetClass()))
@@ -301,7 +304,7 @@ public:
 
 	// FTransactionBase interface.
 	virtual void SaveObject( UObject* Object );
-	virtual void SaveArray( UObject* Object, FScriptArray* Array, int32 Index, int32 Count, int32 Oper, int32 ElementSize, STRUCT_AR Serializer, STRUCT_DTOR Destructor );
+	virtual void SaveArray( UObject* Object, FScriptArray* Array, int32 Index, int32 Count, int32 Oper, int32 ElementSize, STRUCT_DC DefaultConstructor, STRUCT_AR Serializer, STRUCT_DTOR Destructor );
 	virtual void SetPrimaryObject(UObject* InObject);
 
 	/**
diff --git a/Engine/Source/Editor/UnrealEd/Private/EditorTransaction.cpp b/Engine/Source/Editor/UnrealEd/Private/EditorTransaction.cpp
--- a/Engine/Source/Editor/UnrealEd/Private/EditorTransaction.cpp
+++ b/Engine/Source/Editor/UnrealEd/Private/EditorTransaction.cpp
@@ -19,6 +19,7 @@ void FTransaction::FObjectRecord::SerializeContents( FArchive& Ar, int32 InOper
 		check((SIZE_T)Array>=(SIZE_T)Object+sizeof(UObject));
 		check((SIZE_T)Array+sizeof(FScriptArray)<=(SIZE_T)Object+Object->GetClass()->GetPropertiesSize());
 		check(ElementSize!=0);
+		check(DefaultConstructor!=NULL);
 		check(Serializer!=NULL);
 		check(Index>=0);
 		check(Count>=0);
@@ -40,8 +41,11 @@ void FTransaction::FObjectRecord::SerializeContents( FArchive& Ar, int32 InOper
 			// "Undo/Redo Modify" or "Saving remove order" or "Undoing remove order" or "Redoing add order".
 			if( InOper==-1 && Ar.IsLoading() )
 			{
-				Array->Insert( Index, Count, ElementSize );
-				FMemory::Memzero( (uint8*)Array->GetData() + Index*ElementSize, Count*ElementSize );
+				Array->InsertZeroed( Index, Count, ElementSize );
+				for( int32 i=Index; i<Index+Count; i++ )
+				{
+					DefaultConstructor( (uint8*)Array->GetData() + i*ElementSize );
+				}
 			}
 
 			// Serialize changed items.
@@ -57,6 +61,7 @@ void FTransaction::FObjectRecord::SerializeContents( FArchive& Ar, int32 InOper
 		//UE_LOG(LogEditorTransaction, Log,  TEXT("Object %s"), *Object->GetFullName());
 		check(Index==0);
 		check(ElementSize==0);
+		check(DefaultConstructor==NULL);
 		check(Serializer==NULL);
 		Object->Serialize( Ar );
 	}
@@ -157,7 +162,7 @@ void FTransaction::SaveObject( UObject* Object )
 	{
 		ObjectMap.Add(Object,1);
 		// Save the object.
-		new( Records )FObjectRecord( this, Object, NULL, 0, 0, 0, 0, NULL, NULL );
+		new( Records )FObjectRecord( this, Object, NULL, 0, 0, 0, 0, NULL, NULL, NULL );
 	}
 	else
 	{
@@ -165,11 +170,12 @@ void FTransaction::SaveObject( UObject* Object )
 	}
 }
 
-void FTransaction::SaveArray( UObject* Object, FScriptArray* Array, int32 Index, int32 Count, int32 Oper, int32 ElementSize, STRUCT_AR Serializer, STRUCT_DTOR Destructor )
+void FTransaction::SaveArray( UObject* Object, FScriptArray* Array, int32 Index, int32 Count, int32 Oper, int32 ElementSize, STRUCT_DC DefaultConstructor, STRUCT_AR Serializer, STRUCT_DTOR Destructor )
 {
 	check(Object);
 	check(Array);
 	check(ElementSize);
+	check(DefaultConstructor);
 	check(Serializer);
 	check(Object->IsValidLowLevel());
 	check((SIZE_T)Array>=(SIZE_T)Object);
@@ -182,7 +188,7 @@ void FTransaction::SaveArray( UObject* Object, FScriptArray* Array, int32 Index,
 	if( Object->HasAnyFlags(RF_Transactional) && (Object->GetOutermost()->PackageFlags&PKG_PlayInEditor) == 0 )
 	{
 		// Save the array.
-		new( Records )FObjectRecord( this, Object, Array, Index, Count, Oper, ElementSize, Serializer, Destructor );
+		new( Records )FObjectRecord( this, Object, Array, Index, Count, Oper, ElementSize, DefaultConstructor, Serializer, Destructor );
 	}
 }
 
diff --git a/Engine/Source/Runtime/Core/Public/Containers/Array.h b/Engine/Source/Runtime/Core/Public/Containers/Array.h
--- a/Engine/Source/Runtime/Core/Public/Containers/Array.h
+++ b/Engine/Source/Runtime/Core/Public/Containers/Array.h
@@ -2417,7 +2417,7 @@ public:
 		const int32 Index = Super::AddUninitialized( Count );
 		if( GUndo )
 		{
-			GUndo->SaveArray( Owner, (FScriptArray*)this, Index, Count, 1, sizeof(T), SerializeItem, DestructItem );
+			GUndo->SaveArray( Owner, (FScriptArray*)this, Index, Count, 1, sizeof(T), DefaultConstructItem, SerializeItem, DestructItem );
 		}
 		return Index;
 	}
@@ -2426,14 +2426,14 @@ public:
 		Super::InsertUninitialized( Index, Count );
 		if( GUndo )
 		{
-			GUndo->SaveArray( Owner, (FScriptArray*)this, Index, Count, 1, sizeof(T), SerializeItem, DestructItem );
+			GUndo->SaveArray( Owner, (FScriptArray*)this, Index, Count, 1, sizeof(T), DefaultConstructItem, SerializeItem, DestructItem );
 		}
 	}
 	void RemoveAt( int32 Index, int32 Count=1 )
 	{
 		if( GUndo )
 		{
-			GUndo->SaveArray( Owner, (FScriptArray*)this, Index, Count, -1, sizeof(T), SerializeItem, DestructItem );
+			GUndo->SaveArray( Owner, (FScriptArray*)this, Index, Count, -1, sizeof(T), DefaultConstructItem, SerializeItem, DestructItem );
 		}
 		Super::RemoveAt( Index, Count );
 	}
@@ -2441,7 +2441,7 @@ public:
 	{
 		if( GUndo )
 		{
-			GUndo->SaveArray( Owner, (FScriptArray*)this, 0, this->ArrayNum, -1, sizeof(T), SerializeItem, DestructItem );
+			GUndo->SaveArray( Owner, (FScriptArray*)this, 0, this->ArrayNum, -1, sizeof(T), DefaultConstructItem, SerializeItem, DestructItem );
 		}
 		Super::Empty( Slack );
 	}
@@ -2500,13 +2500,13 @@ public:
 	void ModifyItem( int32 Index )
 	{
 		if( GUndo )
-			GUndo->SaveArray( Owner, (FScriptArray*)this, Index, 1, 0, sizeof(T), SerializeItem, DestructItem );
+			GUndo->SaveArray( Owner, (FScriptArray*)this, Index, 1, 0, sizeof(T), DefaultConstructItem, SerializeItem, DestructItem );
 	}
 	void ModifyAllItems()
 	{
 		if( GUndo )
 		{
-			GUndo->SaveArray( Owner, (FScriptArray*)this, 0, this->Num(), 0, sizeof(T), SerializeItem, DestructItem );
+			GUndo->SaveArray( Owner, (FScriptArray*)this, 0, this->Num(), 0, sizeof(T), DefaultConstructItem, SerializeItem, DestructItem );
 		}
 	}
 	friend FArchive& operator<<( FArchive& Ar, TTransArray& A )
@@ -2516,6 +2516,10 @@ public:
 		return Ar;
 	}
 protected:
+	static void DefaultConstructItem( void* TPtr )
+	{
+		new (TPtr) T;
+	}
 	static void SerializeItem( FArchive& Ar, void* TPtr )
 	{
 		Ar << *(T*)TPtr;
diff --git a/Engine/Source/Runtime/Core/Public/Misc/ITransaction.h b/Engine/Source/Runtime/Core/Public/Misc/ITransaction.h
--- a/Engine/Source/Runtime/Core/Public/Misc/ITransaction.h
+++ b/Engine/Source/Runtime/Core/Public/Misc/ITransaction.h
@@ -7,8 +7,9 @@
 #pragma once
 
 // Class for handling undo/redo transactions among objects.
-typedef void( *STRUCT_AR )( class FArchive& Ar, void* TPtr );
-typedef void( *STRUCT_DTOR )( void* TPtr );
+typedef void( *STRUCT_DC )( void* TPtr );						// default construct
+typedef void( *STRUCT_AR )( class FArchive& Ar, void* TPtr );	// serialize
+typedef void( *STRUCT_DTOR )( void* TPtr );						// destruct
 
 /**
  * Interface for transactions.
@@ -34,7 +35,7 @@ public:
 	 * @param Serializer -
 	 * @param Destructor -
 	 */
-	virtual void SaveArray( UObject* Object, class FScriptArray* Array, int32 Index, int32 Count, int32 Oper, int32 ElementSize, STRUCT_AR Serializer, STRUCT_DTOR Destructor ) = 0;
+	virtual void SaveArray( UObject* Object, class FScriptArray* Array, int32 Index, int32 Count, int32 Oper, int32 ElementSize, STRUCT_DC DefaultConstructor, STRUCT_AR Serializer, STRUCT_DTOR Destructor ) = 0;
 
 	/**
 	 * Saves an UObject to the transaction.
