Renaming local variables will properly update the Get/Set nodes.

Local variables references will now properly return their class scope and will rename using their member scope name, rather than the UStruct.

#ttp 336490 - BP: Renaming local variables does not rename their Get/Set nodes

diff --git a/Engine/Source/Editor/BlueprintGraph/Classes/K2Node.h b/Engine/Source/Editor/BlueprintGraph/Classes/K2Node.h
--- a/Engine/Source/Editor/BlueprintGraph/Classes/K2Node.h
+++ b/Engine/Source/Editor/BlueprintGraph/Classes/K2Node.h
@@ -499,7 +499,8 @@ public:
 	/** Get the class that owns this member */
 	UClass* GetMemberParentClass(UClass* SelfScope) const
 	{
-		return bSelfContext ? SelfScope : (UClass*)MemberParentClass;
+		// Local variables with a MemberScope act much the same as being SelfContext, their parent class is SelfScope.
+		return (bSelfContext || !MemberScope.IsEmpty())? SelfScope : (UClass*)MemberParentClass;
 	}
 
 	/** Get the scope of this member, using a node to derive the class */
@@ -514,6 +515,12 @@ public:
 		return FindField<UStruct>(InMemberParentClass, *MemberScope);
 	}
 
+	/** Get the name of the scope of this member */
+	FString GetMemberScopeName() const
+	{
+		return MemberScope;
+	}
+
 	/** 
 	 *	Returns the member UProperty/UFunction this reference is pointing to, or NULL if it no longer exists 
 	 *	Derives 'self' scope from supplied Blueprint node if required
diff --git a/Engine/Source/Editor/UnrealEd/Private/Kismet2/BlueprintEditorUtils.cpp b/Engine/Source/Editor/UnrealEd/Private/Kismet2/BlueprintEditorUtils.cpp
--- a/Engine/Source/Editor/UnrealEd/Private/Kismet2/BlueprintEditorUtils.cpp
+++ b/Engine/Source/Editor/UnrealEd/Private/Kismet2/BlueprintEditorUtils.cpp
@@ -219,7 +219,7 @@ static void RenameVariableReferences(UBlueprint* Blueprint, UClass* VariableClas
 
 					if(VariableNode->VariableReference.IsLocalScope())
 					{
-						VariableNode->VariableReference.SetLocalMember(NewVarName, VariableNode->VariableReference.GetMemberScope(NodeRefClass), VariableNode->VariableReference.GetMemberGuid());
+						VariableNode->VariableReference.SetLocalMember(NewVarName, VariableNode->VariableReference.GetMemberScopeName(), VariableNode->VariableReference.GetMemberGuid());
 					}
 					else if(VariableNode->VariableReference.IsSelfContext())
 					{
@@ -3284,9 +3284,15 @@ void FBlueprintEditorUtils::RemoveLocalVariable(UBlueprint* InBlueprint, const F
 	{
 		for( int32 VarIdx = 0; VarIdx < FunctionEntry->LocalVariables.Num(); ++VarIdx )
 		{
-			FunctionEntry->LocalVariables.RemoveAt(VarIdx);
-			FBlueprintEditorUtils::RemoveVariableNodes(InBlueprint, InVarName);
-			FBlueprintEditorUtils::MarkBlueprintAsStructurallyModified(InBlueprint);
+			if(FunctionEntry->LocalVariables[VarIdx].VarName == InVarName)
+			{
+				FunctionEntry->LocalVariables.RemoveAt(VarIdx);
+				FBlueprintEditorUtils::RemoveVariableNodes(InBlueprint, InVarName);
+				FBlueprintEditorUtils::MarkBlueprintAsStructurallyModified(InBlueprint);
+
+				// No other local variables will match, we are done
+				return;
+			}
 		}
 	}
 }
