Grammar fix in anim bp messagebox.

diff --git a/Engine/Source/Editor/UnrealEd/Private/Factories/AnimBlueprintFactory.cpp b/Engine/Source/Editor/UnrealEd/Private/Factories/AnimBlueprintFactory.cpp
--- a/Engine/Source/Editor/UnrealEd/Private/Factories/AnimBlueprintFactory.cpp
+++ b/Engine/Source/Editor/UnrealEd/Private/Factories/AnimBlueprintFactory.cpp
@@ -372,7 +372,7 @@ UObject* UAnimBlueprintFactory::FactoryCreateNew(UClass* Class, UObject* InParen
 	{
 		FFormatNamedArguments Args;
 		Args.Add( TEXT("ClassName"), (ParentClass != NULL) ? FText::FromString( ParentClass->GetName() ) : LOCTEXT("Null", "(null)") );
-		FMessageDialog::Open( EAppMsgType::Ok, FText::Format( LOCTEXT("CannotCreateAnimBlueprint", "Cannot create a Anim Blueprint based on the class '{0}'."), Args ) );
+		FMessageDialog::Open( EAppMsgType::Ok, FText::Format( LOCTEXT("CannotCreateAnimBlueprint", "Cannot create an Anim Blueprint based on the class '{0}'."), Args ) );
 		return NULL;
 	}
 	else
