Include child actors in transaction when moving parent actor

diff --git a/Engine/Source/Editor/UnrealEd/Private/LevelEditorViewport.cpp b/Engine/Source/Editor/UnrealEd/Private/LevelEditorViewport.cpp
--- a/Engine/Source/Editor/UnrealEd/Private/LevelEditorViewport.cpp
+++ b/Engine/Source/Editor/UnrealEd/Private/LevelEditorViewport.cpp
@@ -1684,31 +1684,52 @@ void FTrackingTransaction::Begin(const FText& Description)
 
 	TrackingTransactionState = ETransactionState::Active;
 
-	TSet<AGroupActor*> GroupActors;
+	TSet<AActor*> RelatedActors;
 
 	// Modify selected actors to record their state at the start of the transaction
 	for ( FSelectionIterator It( GEditor->GetSelectedActorIterator() ) ; It ; ++It )
 	{
-		AActor* Actor = static_cast<AActor*>( *It );
-		checkSlow( Actor->IsA(AActor::StaticClass()) );	
+		AActor* Actor = CastChecked<AActor>( *It );
 
 		Actor->Modify();
 
+		if (Actor->GetRootComponent())
+		{
+			TArray<USceneComponent*> RootComponents;
+			RootComponents.Add(Actor->GetRootComponent());
+			while (RootComponents.Num() > 0)
+			{
+				USceneComponent* RootComponent = RootComponents.Pop();
+				for (USceneComponent* AttachedComponent : RootComponent->AttachChildren)
+				{
+					if (AttachedComponent)
+					{
+						AActor* AttachedActor = AttachedComponent->GetOwner();
+						if (AttachedComponent == AttachedActor->GetRootComponent())
+						{
+							RelatedActors.Add(AttachedActor);
+							RootComponents.Add(AttachedComponent);
+						}
+					}
+				}
+			}
+		}
+
 		if (GEditor->bGroupingActive)
 		{
 			// if this actor is in a group, add the GroupActor into a list to be modified shortly
 			AGroupActor* ActorLockedRootGroup = AGroupActor::GetRootForActor(Actor, true);
 			if (ActorLockedRootGroup != nullptr)
 			{
-				GroupActors.Add(ActorLockedRootGroup);
+				RelatedActors.Add(ActorLockedRootGroup);
 			}
 		}
 	}
 
 	// Modify unique group actors
-	for (auto* GroupActor : GroupActors)
+	for (AActor* RelatedActor : RelatedActors)
 	{
-		GroupActor->Modify();
+		RelatedActor->Modify();
 	}
 }
 
