UE4: Replacing LoadObject call in SDlgPickAssetPath::ValidatePackage with FindObject. LoadObject creates a package in memory whether the object can be loaded or not. This will still find the object in memory to avoid problems like creating blueprints over existing objects.

diff --git a/Engine/Source/Editor/UnrealEd/Private/Dialogs/DlgPickAssetPath.cpp b/Engine/Source/Editor/UnrealEd/Private/Dialogs/DlgPickAssetPath.cpp
--- a/Engine/Source/Editor/UnrealEd/Private/Dialogs/DlgPickAssetPath.cpp
+++ b/Engine/Source/Editor/UnrealEd/Private/Dialogs/DlgPickAssetPath.cpp
@@ -133,7 +133,7 @@ bool SDlgPickAssetPath::ValidatePackage()
 		return false;
 	}
 
-	if (FPackageName::DoesPackageExist(GetFullAssetPath().ToString()) || LoadObject<UObject>(NULL, *(AssetPath.ToString() + "/" + AssetName.ToString() + "." + AssetName.ToString())) != NULL)
+	if (FPackageName::DoesPackageExist(GetFullAssetPath().ToString()) || FindObject<UObject>(NULL, *(AssetPath.ToString() + "/" + AssetName.ToString() + "." + AssetName.ToString())) != NULL)
 	{
 		FMessageDialog::Open(EAppMsgType::Ok, FText::Format(LOCTEXT("AssetAlreadyExists", "Asset {0} already exists."), GetFullAssetPath()));
 		return false;
