UE4: Replaced level SaveAs implementation when working with -WorldAssets to use a path picker instead of the OS save-as  dialog.

diff --git a/Engine/Source/Editor/UnrealEd/Private/EditorLevelUtils.cpp b/Engine/Source/Editor/UnrealEd/Private/EditorLevelUtils.cpp
--- a/Engine/Source/Editor/UnrealEd/Private/EditorLevelUtils.cpp
+++ b/Engine/Source/Editor/UnrealEd/Private/EditorLevelUtils.cpp
@@ -17,6 +17,8 @@
 #include "ActorEditorUtils.h"
 #include "Editor/Levels/Public/LevelEdMode.h"
 
+#include "Runtime/AssetRegistry/Public/AssetRegistryModule.h"
+
 DEFINE_LOG_CATEGORY(LogLevelTools);
 
 #define LOCTEXT_NAMESPACE "EditorLevelUtils"
@@ -470,21 +472,42 @@ namespace EditorLevelUtils
 		// Editor modes cannot be active when any level saving occurs.
 		GEditorModeTools().ActivateMode( FBuiltinEditorModes::EM_Default );
 
-		// Create a new world - so we can 'borrow' its level
-		UWorld* NewGWorld = UWorld::CreateWorld(EWorldType::None, false );
-		check(NewGWorld);
+		UWorld* NewWorld = nullptr;
+		if (FParse::Param(FCommandLine::Get(), TEXT("WorldAssets")))
+		{
+			// Create a new world
+			UWorldFactory* Factory = ConstructObject<UWorldFactory>(UWorldFactory::StaticClass());
+			Factory->WorldType = EWorldType::Inactive;
+			UPackage* Pkg = CreatePackage(NULL, NULL);
+			FName WorldName(TEXT("NewWorld"));
+			EObjectFlags Flags = RF_Public | RF_Standalone;
+			NewWorld = CastChecked<UWorld>(Factory->FactoryCreateNew(UWorld::StaticClass(), Pkg, WorldName, Flags, NULL, GWarn));
+			if ( NewWorld )
+			{
+				FAssetRegistryModule::AssetCreated(NewWorld);
+			}
+		}
+		else
+		{
+			// Create a new world - so we can 'borrow' its level
+			NewWorld = UWorld::CreateWorld(EWorldType::None, false );
+			check(NewWorld);
+		}
 
 		// Save the new world to disk.
-		const bool bNewWorldSaved = FEditorFileUtils::SaveLevel( NewGWorld->PersistentLevel, DefaultFilename );
+		const bool bNewWorldSaved = FEditorFileUtils::SaveLevel( NewWorld->PersistentLevel, DefaultFilename );
 		FString NewPackageName;
 		if ( bNewWorldSaved )
 		{
-			NewPackageName = NewGWorld->GetOutermost()->GetName();
+			NewPackageName = NewWorld->GetOutermost()->GetName();
 		}
 
-		// Destroy the new world we created and collect the garbage
-		NewGWorld->DestroyWorld( false );
-		CollectGarbage( GARBAGE_COLLECTION_KEEPFLAGS );
+		if (!FParse::Param(FCommandLine::Get(), TEXT("WorldAssets")))
+		{
+			// Destroy the new world we created and collect the garbage
+			NewWorld->DestroyWorld( false );
+			CollectGarbage( GARBAGE_COLLECTION_KEEPFLAGS );
+		}
 
 		// If the new world was saved successfully, import it as a streaming level.
 		ULevel* NewLevel = NULL;
diff --git a/Engine/Source/Editor/UnrealEd/Private/FileHelpers.cpp b/Engine/Source/Editor/UnrealEd/Private/FileHelpers.cpp
--- a/Engine/Source/Editor/UnrealEd/Private/FileHelpers.cpp
+++ b/Engine/Source/Editor/UnrealEd/Private/FileHelpers.cpp
@@ -21,6 +21,8 @@
 #include "LevelUtils.h"
 #include "MessageLog.h"
 
+#include "Dialogs/DlgPickAssetPath.h"
+
 DEFINE_LOG_CATEGORY_STATIC(LogFileHelpers, Log, All);
 
 //definition of flag used to do special work when we're attempting to load the "startup map"
@@ -555,6 +557,39 @@ bool RenameStreamingLevel( FString& LevelToRename, const FString& OldBaseLevelNa
 	return false;
 }
 
+static bool OpenLevelSaveAsDialog(const FString& InDefaultPath, const FString& InNewNameSuggestion, FString& OutPackageName)
+{
+	FString DefaultPath = InDefaultPath;
+	if (DefaultPath.IsEmpty())
+	{
+		DefaultPath = TEXT("/Game/Maps");
+	}
+
+	FString NewNameSuggestion = InNewNameSuggestion;
+	if (NewNameSuggestion.IsEmpty())
+	{
+		NewNameSuggestion = TEXT("NewMap");
+	}
+
+	FString PackageName = DefaultPath / NewNameSuggestion;
+	FString Name;
+	FAssetToolsModule& AssetToolsModule = FModuleManager::LoadModuleChecked<FAssetToolsModule>("AssetTools");
+	AssetToolsModule.Get().CreateUniqueAssetName(PackageName, TEXT(""), PackageName, Name);
+
+	TSharedPtr<SDlgPickAssetPath> PickAssetPathWidget =
+		SNew(SDlgPickAssetPath)
+		.Title(LOCTEXT("LevelPathPickerTitle", "Save Level As"))
+		.DefaultAssetPath(FText::FromString(PackageName));
+
+	if (EAppReturnType::Ok == PickAssetPathWidget->ShowModal())
+	{
+		OutPackageName = PickAssetPathWidget->GetFullAssetPath().ToString();
+		return true;
+	}
+
+	return false;
+}
+
 /**
  * Prompts the user with a dialog for selecting a filename.
  */
@@ -574,13 +609,34 @@ static bool SaveAsImplementation( UWorld* InWorld, const FString& DefaultFilenam
 	while( !bFilenameIsValid )
 	{
 		FString SaveFilename;
-		if( FileDialogHelpers::SaveFile(
+		bool bSaveFileLocationSelected = false;
+		if (FParse::Param(FCommandLine::Get(), TEXT("WorldAssets")))
+		{
+			FString DefaultPackagePath;
+			FPackageName::TryConvertFilenameToLongPackageName(DefaultDirectory, DefaultPackagePath);
+
+			FString PackageName;
+			bSaveFileLocationSelected = OpenLevelSaveAsDialog(
+				DefaultPackagePath,
+				FPaths::GetBaseFilename(DefaultFilename),
+				PackageName);
+
+			if ( bSaveFileLocationSelected )
+			{
+				SaveFilename = FPackageName::LongPackageNameToFilename(PackageName, FPackageName::GetMapPackageExtension());
+			}
+		}
+		else
+		{
+			bSaveFileLocationSelected = FileDialogHelpers::SaveFile(
 				NSLOCTEXT("UnrealEd", "SaveAs", "Save As").ToString(),
 				FEditorFileUtils::GetFilterString(FI_Save),
 				DefaultDirectory,
 				FPaths::GetCleanFilename(DefaultFilename),
 				SaveFilename
-			))
+				);
+		}
+		if( bSaveFileLocationSelected )
 		{
 			// Add a map file extension if none was supplied
 			if( FPaths::GetExtension(SaveFilename).IsEmpty() )
