Remove PhysicsPublic from Engine.h

diff --git a/Engine/Plugins/Experimental/Phya/Source/Phya/Private/PhyaCollisionHandler.cpp b/Engine/Plugins/Experimental/Phya/Source/Phya/Private/PhyaCollisionHandler.cpp
--- a/Engine/Plugins/Experimental/Phya/Source/Phya/Private/PhyaCollisionHandler.cpp
+++ b/Engine/Plugins/Experimental/Phya/Source/Phya/Private/PhyaCollisionHandler.cpp
@@ -1,7 +1,7 @@
 // Copyright 1998-2014 Epic Games, Inc. All Rights Reserved.
 
 #include "PhyaPluginPrivatePCH.h"
-
+#include "PhysicsPublic.h"
 #include "Phya.hpp"
 
 #define N_BODIES 5
diff --git a/Engine/Plugins/Runtime/CableComponent/Source/CableComponent/Private/CableComponent.cpp b/Engine/Plugins/Runtime/CableComponent/Source/CableComponent/Private/CableComponent.cpp
--- a/Engine/Plugins/Runtime/CableComponent/Source/CableComponent/Private/CableComponent.cpp
+++ b/Engine/Plugins/Runtime/CableComponent/Source/CableComponent/Private/CableComponent.cpp
@@ -1,6 +1,7 @@
 // Copyright 1998-2014 Epic Games, Inc. All Rights Reserved. 
 
 #include "CableComponentPluginPrivatePCH.h"
+#include "DynamicMeshBuilder.h"
 
 
 /** Vertex Buffer */
diff --git a/Engine/Plugins/Runtime/CustomMeshComponent/Source/CustomMeshComponent/Private/CustomMeshComponent.cpp b/Engine/Plugins/Runtime/CustomMeshComponent/Source/CustomMeshComponent/Private/CustomMeshComponent.cpp
--- a/Engine/Plugins/Runtime/CustomMeshComponent/Source/CustomMeshComponent/Private/CustomMeshComponent.cpp
+++ b/Engine/Plugins/Runtime/CustomMeshComponent/Source/CustomMeshComponent/Private/CustomMeshComponent.cpp
@@ -1,8 +1,8 @@
 // Copyright 1998-2014 Epic Games, Inc. All Rights Reserved. 
 
 #include "CustomMeshComponentPluginPrivatePCH.h"
-
 #include "CustomMeshComponent.h"
+#include "DynamicMeshBuilder.h"
 
 /** Vertex Buffer */
 class FCustomMeshVertexBuffer : public FVertexBuffer 
diff --git a/Engine/Source/Developer/PhysXFormats/Private/PhysXFormats.cpp b/Engine/Source/Developer/PhysXFormats/Private/PhysXFormats.cpp
--- a/Engine/Source/Developer/PhysXFormats/Private/PhysXFormats.cpp
+++ b/Engine/Source/Developer/PhysXFormats/Private/PhysXFormats.cpp
@@ -4,6 +4,7 @@
 #include "ModuleInterface.h"
 #include "ModuleManager.h"
 #include "Engine.h"
+#include "PhysicsPublic.h"
 #include "TargetPlatform.h"
 #include "PhysXFormats.h"
 
diff --git a/Engine/Source/Editor/Cascade/Private/Cascade.cpp b/Engine/Source/Editor/Cascade/Private/Cascade.cpp
--- a/Engine/Source/Editor/Cascade/Private/Cascade.cpp
+++ b/Engine/Source/Editor/Cascade/Private/Cascade.cpp
@@ -13,6 +13,7 @@
 #include "SCascadeEmitterCanvas.h"
 #include "CascadeActions.h"
 #include "Cascade.h"
+#include "PhysicsPublic.h"
 
 
 #include "Editor/WorkspaceMenuStructure/Public/WorkspaceMenuStructureModule.h"
diff --git a/Engine/Source/Editor/Cascade/Private/CascadePreviewViewportClient.cpp b/Engine/Source/Editor/Cascade/Private/CascadePreviewViewportClient.cpp
--- a/Engine/Source/Editor/Cascade/Private/CascadePreviewViewportClient.cpp
+++ b/Engine/Source/Editor/Cascade/Private/CascadePreviewViewportClient.cpp
@@ -12,6 +12,7 @@
 #include "Particles/ParticleLODLevel.h"
 #include "Particles/ParticleModuleRequired.h"
 #include "Particles/ParticleSpriteEmitter.h"
+#include "PhysicsPublic.h"
 
 #define LOCTEXT_NAMESPACE "CascadeViewportClient"
 
diff --git a/Engine/Source/Editor/DestructibleMeshEditor/Private/SDestructibleMeshEditorViewport.cpp b/Engine/Source/Editor/DestructibleMeshEditor/Private/SDestructibleMeshEditorViewport.cpp
--- a/Engine/Source/Editor/DestructibleMeshEditor/Private/SDestructibleMeshEditorViewport.cpp
+++ b/Engine/Source/Editor/DestructibleMeshEditor/Private/SDestructibleMeshEditorViewport.cpp
@@ -2,6 +2,7 @@
 
 #include "DestructibleMeshEditorPrivatePCH.h"
 
+#include "PhysicsPublic.h"
 #include "MouseDeltaTracker.h"
 #include "Runtime/Engine/Public/Slate/SceneViewport.h"
 #include "PreviewScene.h"
diff --git a/Engine/Source/Editor/GeometryMode/Private/GeometryEdMode.cpp b/Engine/Source/Editor/GeometryMode/Private/GeometryEdMode.cpp
--- a/Engine/Source/Editor/GeometryMode/Private/GeometryEdMode.cpp
+++ b/Engine/Source/Editor/GeometryMode/Private/GeometryEdMode.cpp
@@ -1,6 +1,7 @@
 // Copyright 1998-2014 Epic Games, Inc. All Rights Reserved.
 
 #include "GeometryModePrivatePCH.h"
+#include "DynamicMeshBuilder.h"
 
 
 IMPLEMENT_MODULE( FGeometryModeModule, GeometryMode );
diff --git a/Engine/Source/Editor/LandscapeEditor/Private/LandscapeEdMode.cpp b/Engine/Source/Editor/LandscapeEditor/Private/LandscapeEdMode.cpp
--- a/Engine/Source/Editor/LandscapeEditor/Private/LandscapeEdMode.cpp
+++ b/Engine/Source/Editor/LandscapeEditor/Private/LandscapeEdMode.cpp
@@ -17,6 +17,7 @@
 #include "LandscapeEdModeTools.h"
 #include "ScopedTransaction.h"
 #include "ImageWrapper.h"
+#include "DynamicMeshBuilder.h"
 
 //Slate dependencies
 #include "Editor/LevelEditor/Public/LevelEditor.h"
diff --git a/Engine/Source/Editor/PhAT/Classes/PhATEdSkeletalMeshComponent.h b/Engine/Source/Editor/PhAT/Classes/PhATEdSkeletalMeshComponent.h
--- a/Engine/Source/Editor/PhAT/Classes/PhATEdSkeletalMeshComponent.h
+++ b/Engine/Source/Editor/PhAT/Classes/PhATEdSkeletalMeshComponent.h
@@ -1,7 +1,7 @@
 // Copyright 1998-2014 Epic Games, Inc. All Rights Reserved.
 
 #pragma once
-
+#include "PhysicsPublic.h"
 #include "PhATEdSkeletalMeshComponent.generated.h"
 
 
diff --git a/Engine/Source/Editor/PhAT/Private/PhATHitProxies.h b/Engine/Source/Editor/PhAT/Private/PhATHitProxies.h
--- a/Engine/Source/Editor/PhAT/Private/PhATHitProxies.h
+++ b/Engine/Source/Editor/PhAT/Private/PhATHitProxies.h
@@ -2,6 +2,7 @@
 
 #pragma once
 
+#include "PhysicsPublic.h"
 
 /*-----------------------------------------------------------------------------
    Hit Proxies
diff --git a/Engine/Source/Editor/PhAT/Private/PhATPreviewViewportClient.h b/Engine/Source/Editor/PhAT/Private/PhATPreviewViewportClient.h
--- a/Engine/Source/Editor/PhAT/Private/PhATPreviewViewportClient.h
+++ b/Engine/Source/Editor/PhAT/Private/PhATPreviewViewportClient.h
@@ -2,6 +2,8 @@
 
 #pragma once
 
+#include "PhysicsPublic.h"
+
 class FPhAT;
 class FPhATSharedData;
 class SPhATPreviewViewport;
diff --git a/Engine/Source/Editor/PhAT/Private/PhATSharedData.cpp b/Engine/Source/Editor/PhAT/Private/PhATSharedData.cpp
--- a/Engine/Source/Editor/PhAT/Private/PhATSharedData.cpp
+++ b/Engine/Source/Editor/PhAT/Private/PhATSharedData.cpp
@@ -1,6 +1,7 @@
 // Copyright 1998-2014 Epic Games, Inc. All Rights Reserved.
 
 #include "PhATModule.h"
+#include "PhysicsPublic.h"
 #include "EditorSupportDelegates.h"
 #include "ScopedTransaction.h"
 #include "SPhATNewAssetDlg.h"
diff --git a/Engine/Source/Editor/PhAT/Private/PhATSharedData.h b/Engine/Source/Editor/PhAT/Private/PhATSharedData.h
--- a/Engine/Source/Editor/PhAT/Private/PhATSharedData.h
+++ b/Engine/Source/Editor/PhAT/Private/PhATSharedData.h
@@ -3,6 +3,7 @@
 #pragma once
 
 #include "PreviewScene.h"
+#include "PhysicsPublic.h"
 
 class UPhATEdSkeletalMeshComponent;
 class UPhysicsHandleComponent;
diff --git a/Engine/Source/Editor/UnrealEd/Private/ApexClothingUtils.cpp b/Engine/Source/Editor/UnrealEd/Private/ApexClothingUtils.cpp
--- a/Engine/Source/Editor/UnrealEd/Private/ApexClothingUtils.cpp
+++ b/Engine/Source/Editor/UnrealEd/Private/ApexClothingUtils.cpp
@@ -4,6 +4,7 @@
 
 DEFINE_LOG_CATEGORY_STATIC(LogApexClothingUtils, Log, All);
 
+#include "PhysicsPublic.h"
 #include "ApexClothingUtils.h"
 
 #if WITH_APEX
diff --git a/Engine/Source/Editor/UnrealEd/Private/EditorModes.cpp b/Engine/Source/Editor/UnrealEd/Private/EditorModes.cpp
--- a/Engine/Source/Editor/UnrealEd/Private/EditorModes.cpp
+++ b/Engine/Source/Editor/UnrealEd/Private/EditorModes.cpp
@@ -18,6 +18,7 @@
 #include "LevelEditor.h"
 #include "Toolkits/ToolkitManager.h"
 #include "EditorLevelUtils.h"
+#include "DynamicMeshBuilder.h"
 
 #include "Editor/FoliageEdit/Public/FoliageEditModule.h"
 #include "ActorEditorUtils.h"
diff --git a/Engine/Source/Editor/UnrealEd/Private/Factories/ApexDestructibleAssetImport.cpp b/Engine/Source/Editor/UnrealEd/Private/Factories/ApexDestructibleAssetImport.cpp
--- a/Engine/Source/Editor/UnrealEd/Private/Factories/ApexDestructibleAssetImport.cpp
+++ b/Engine/Source/Editor/UnrealEd/Private/Factories/ApexDestructibleAssetImport.cpp
@@ -15,6 +15,7 @@
 DEFINE_LOG_CATEGORY_STATIC(LogApexDestructibleAssetImport, Log, All);
 
 #include "Engine.h"
+#include "PhysicsPublic.h"
 #include "Engine/DestructibleFractureSettings.h"
 #include "TextureLayout.h"
 #include "SkelImport.h"
diff --git a/Engine/Source/Editor/UnrealEd/Private/Factories/EditorFactories.cpp b/Engine/Source/Editor/UnrealEd/Private/Factories/EditorFactories.cpp
--- a/Engine/Source/Editor/UnrealEd/Private/Factories/EditorFactories.cpp
+++ b/Engine/Source/Editor/UnrealEd/Private/Factories/EditorFactories.cpp
@@ -14,6 +14,7 @@
 #include "Sound/SoundNodeModulator.h"
 #include "Factories.h"
 #include "SoundDefinitions.h"
+#include "PhysicsPublic.h"
 #include "BlueprintUtilities.h"
 #include "Kismet2/KismetEditorUtilities.h"
 #include "BmpImageSupport.h"
diff --git a/Engine/Source/Editor/UnrealEd/Private/PlayLevel.cpp b/Engine/Source/Editor/UnrealEd/Private/PlayLevel.cpp
--- a/Engine/Source/Editor/UnrealEd/Private/PlayLevel.cpp
+++ b/Engine/Source/Editor/UnrealEd/Private/PlayLevel.cpp
@@ -26,6 +26,7 @@
 #include "TargetDeviceServices.h"
 #include "GameProjectGenerationModule.h"
 #include "SourceCodeNavigation.h"
+#include "PhysicsPublic.h"
 
 #include "Online.h"
 
diff --git a/Engine/Source/Editor/UnrealEd/Private/UnrealWidget.cpp b/Engine/Source/Editor/UnrealEd/Private/UnrealWidget.cpp
--- a/Engine/Source/Editor/UnrealEd/Private/UnrealWidget.cpp
+++ b/Engine/Source/Editor/UnrealEd/Private/UnrealWidget.cpp
@@ -2,6 +2,7 @@
 
 #include "UnrealEd.h"
 #include "SnappingUtils.h"
+#include "DynamicMeshBuilder.h"
 
 IMPLEMENT_HIT_PROXY(HWidgetAxis,HHitProxy);
 
diff --git a/Engine/Source/Editor/UnrealEd/Public/ApexClothingUtils.h b/Engine/Source/Editor/UnrealEd/Public/ApexClothingUtils.h
--- a/Engine/Source/Editor/UnrealEd/Public/ApexClothingUtils.h
+++ b/Engine/Source/Editor/UnrealEd/Public/ApexClothingUtils.h
@@ -6,8 +6,18 @@
 
 class USkeletalMeshComponent;
 class FPhysScene;
+struct FClothPhysicsProperties;
 
 #if WITH_APEX_CLOTHING
+
+namespace physx
+{
+	namespace apex
+	{
+		class NxClothingAsset;
+	}
+}
+
 struct FSubmeshInfo
 {
 	int32  SubmeshIndex;
@@ -73,7 +83,7 @@ namespace ApexClothingUtils
 
 	//if AssetIndex is -1, means newly added asset, otherwise re-import
 	UNREALED_API EClothUtilRetType ImportApexAssetFromApexFile(FString& ApexFile,USkeletalMesh* SkelMesh, int32 AssetIndex=-1);
-	UNREALED_API bool GetSubmeshInfoFromApexAsset(NxClothingAsset *InAsset, uint32 LODIndex, TArray<FSubmeshInfo>& OutSubmeshInfos);
+	UNREALED_API bool GetSubmeshInfoFromApexAsset(physx::apex::NxClothingAsset *InAsset, uint32 LODIndex, TArray<FSubmeshInfo>& OutSubmeshInfos);
 
 	//Imported LOD is decided according to bUseLOD in clothing asset
 	UNREALED_API bool ImportClothingSectionFromClothingAsset(USkeletalMesh* SkelMesh, uint32 SectionIndex, int32 AssetIndex, int32 AssetSubmeshIndex);
@@ -83,14 +93,14 @@ namespace ApexClothingUtils
 
 	UNREALED_API void BackupClothingDataFromSkeletalMesh(USkeletalMesh* SkelMesh, FClothingBackup& ClothingBackup);
 	UNREALED_API void ReapplyClothingDataToSkeletalMesh(USkeletalMesh* SkelMesh, FClothingBackup& ClothingBackup);
-	UNREALED_API int32 GetNumLODs(NxClothingAsset *InAsset);
-	UNREALED_API int32 GetNumRenderSubmeshes(NxClothingAsset *InAsset, int32 LODIndex);
+	UNREALED_API int32 GetNumLODs(physx::apex::NxClothingAsset *InAsset);
+	UNREALED_API int32 GetNumRenderSubmeshes(physx::apex::NxClothingAsset *InAsset, int32 LODIndex);
 
-	UNREALED_API NxClothingAsset* CreateApexClothingAssetFromBuffer(const uint8* Buffer, int32 BufferSize);
+	UNREALED_API physx::apex::NxClothingAsset* CreateApexClothingAssetFromBuffer(const uint8* Buffer, int32 BufferSize);
 
 	// if don't specify MaterialIndex, default material index will be used 
-	UNREALED_API void GetPhysicsPropertiesFromApexAsset(NxClothingAsset *InAsset, FClothPhysicsProperties& OutPropertyInfo);
-	UNREALED_API void SetPhysicsPropertiesToApexAsset(NxClothingAsset *InAsset, FClothPhysicsProperties& InPropertyInfo);
+	UNREALED_API void GetPhysicsPropertiesFromApexAsset(physx::apex::NxClothingAsset *InAsset, FClothPhysicsProperties& OutPropertyInfo);
+	UNREALED_API void SetPhysicsPropertiesToApexAsset(physx::apex::NxClothingAsset *InAsset, FClothPhysicsProperties& InPropertyInfo);
 #endif // #if WITH_APEX_CLOTHING
 } // namespace ApexClothingUtils
 
diff --git a/Engine/Source/Runtime/Engine/Private/AI/Navigation/NavTestRenderingComponent.cpp b/Engine/Source/Runtime/Engine/Private/AI/Navigation/NavTestRenderingComponent.cpp
--- a/Engine/Source/Runtime/Engine/Private/AI/Navigation/NavTestRenderingComponent.cpp
+++ b/Engine/Source/Runtime/Engine/Private/AI/Navigation/NavTestRenderingComponent.cpp
@@ -3,6 +3,7 @@
 #include "EnginePrivate.h"
 #include "DebugRenderSceneProxy.h"
 #include "AI/Navigation/NavTestRenderingComponent.h"
+#include "DynamicMeshBuilder.h"
 
 static const FColor NavMeshRenderColor_OpenSet(255,128,0,255);
 static const FColor NavMeshRenderColor_ClosedSet(255,196,0,255);
diff --git a/Engine/Source/Runtime/Engine/Private/AI/Navigation/RecastNavMeshGenerator.cpp b/Engine/Source/Runtime/Engine/Private/AI/Navigation/RecastNavMeshGenerator.cpp
--- a/Engine/Source/Runtime/Engine/Private/AI/Navigation/RecastNavMeshGenerator.cpp
+++ b/Engine/Source/Runtime/Engine/Private/AI/Navigation/RecastNavMeshGenerator.cpp
@@ -3,6 +3,8 @@
 #include "EnginePrivate.h"
 #if WITH_RECAST
 
+#include "PhysicsPublic.h"
+
 #if WITH_PHYSX
 	#include "../../PhysicsEngine/PhysXSupport.h"
 #endif
diff --git a/Engine/Source/Runtime/Engine/Private/Actor.cpp b/Engine/Source/Runtime/Engine/Private/Actor.cpp
--- a/Engine/Source/Runtime/Engine/Private/Actor.cpp
+++ b/Engine/Source/Runtime/Engine/Private/Actor.cpp
@@ -2,6 +2,7 @@
 
 
 #include "EnginePrivate.h"
+#include "PhysicsPublic.h"
 #include "ParticleDefinitions.h"
 #include "SoundDefinitions.h"
 #include "LatentActions.h"
diff --git a/Engine/Source/Runtime/Engine/Private/Collision/CollisionConversions.cpp b/Engine/Source/Runtime/Engine/Private/Collision/CollisionConversions.cpp
--- a/Engine/Source/Runtime/Engine/Private/Collision/CollisionConversions.cpp
+++ b/Engine/Source/Runtime/Engine/Private/Collision/CollisionConversions.cpp
@@ -3,6 +3,7 @@
 #include "EnginePrivate.h"
 
 #if WITH_PHYSX
+#include "PhysicsPublic.h"
 #include "../PhysicsEngine/PhysXSupport.h"
 #include "PhysXCollision.h"
 #include "CollisionDebugDrawing.h"
diff --git a/Engine/Source/Runtime/Engine/Private/Collision/CollisionDebugDrawing.cpp b/Engine/Source/Runtime/Engine/Private/Collision/CollisionDebugDrawing.cpp
--- a/Engine/Source/Runtime/Engine/Private/Collision/CollisionDebugDrawing.cpp
+++ b/Engine/Source/Runtime/Engine/Private/Collision/CollisionDebugDrawing.cpp
@@ -1,6 +1,7 @@
 // Copyright 1998-2014 Epic Games, Inc. All Rights Reserved.
 
 #include "EnginePrivate.h"
+#include "PhysicsPublic.h"
 #include "CollisionDebugDrawingPublic.h"
 
 static FColor TraceColor(255,255,255);
diff --git a/Engine/Source/Runtime/Engine/Private/Collision/PhysXCollision.cpp b/Engine/Source/Runtime/Engine/Private/Collision/PhysXCollision.cpp
--- a/Engine/Source/Runtime/Engine/Private/Collision/PhysXCollision.cpp
+++ b/Engine/Source/Runtime/Engine/Private/Collision/PhysXCollision.cpp
@@ -5,6 +5,7 @@
 =============================================================================*/
 
 #include "EnginePrivate.h"
+#include "PhysicsPublic.h"
 #include "Collision.h"
 
 #if WITH_PHYSX
diff --git a/Engine/Source/Runtime/Engine/Private/Collision/WorldCollision.cpp b/Engine/Source/Runtime/Engine/Private/Collision/WorldCollision.cpp
--- a/Engine/Source/Runtime/Engine/Private/Collision/WorldCollision.cpp
+++ b/Engine/Source/Runtime/Engine/Private/Collision/WorldCollision.cpp
@@ -5,6 +5,7 @@
 =============================================================================*/
 
 #include "EnginePrivate.h"
+#include "PhysicsPublic.h"
 #include "Collision.h"
 
 #if WITH_PHYSX
diff --git a/Engine/Source/Runtime/Engine/Private/DestructibleComponent.cpp b/Engine/Source/Runtime/Engine/Private/DestructibleComponent.cpp
--- a/Engine/Source/Runtime/Engine/Private/DestructibleComponent.cpp
+++ b/Engine/Source/Runtime/Engine/Private/DestructibleComponent.cpp
@@ -6,6 +6,7 @@
 
 
 #include "EnginePrivate.h"
+#include "PhysicsPublic.h"
 #include "PhysicsEngine/PhysXSupport.h"
 #include "Collision/PhysXCollision.h"
 #include "ParticleDefinitions.h"
diff --git a/Engine/Source/Runtime/Engine/Private/DestructibleFractureSettings.cpp b/Engine/Source/Runtime/Engine/Private/DestructibleFractureSettings.cpp
--- a/Engine/Source/Runtime/Engine/Private/DestructibleFractureSettings.cpp
+++ b/Engine/Source/Runtime/Engine/Private/DestructibleFractureSettings.cpp
@@ -1,6 +1,7 @@
 // Copyright 1998-2014 Epic Games, Inc. All Rights Reserved.
 
 #include "EnginePrivate.h"
+#include "PhysicsPublic.h"
 #include "Engine/DestructibleFractureSettings.h"
 #include "PhysicsEngine/PhysXSupport.h"
 
diff --git a/Engine/Source/Runtime/Engine/Private/DestructibleMesh.cpp b/Engine/Source/Runtime/Engine/Private/DestructibleMesh.cpp
--- a/Engine/Source/Runtime/Engine/Private/DestructibleMesh.cpp
+++ b/Engine/Source/Runtime/Engine/Private/DestructibleMesh.cpp
@@ -5,6 +5,7 @@
 =============================================================================*/
 
 #include "EnginePrivate.h"
+#include "PhysicsPublic.h"
 #include "StaticMeshResources.h"
 #include "Engine/DestructibleFractureSettings.h"
 #include "PhysicsEngine/PhysXSupport.h"
diff --git a/Engine/Source/Runtime/Engine/Private/DynamicMeshBuilder.cpp b/Engine/Source/Runtime/Engine/Private/DynamicMeshBuilder.cpp
--- a/Engine/Source/Runtime/Engine/Private/DynamicMeshBuilder.cpp
+++ b/Engine/Source/Runtime/Engine/Private/DynamicMeshBuilder.cpp
@@ -5,6 +5,7 @@
 =============================================================================*/
 
 #include "EnginePrivate.h"
+#include "DynamicMeshBuilder.h"
 
 /** The index buffer type used for dynamic meshes. */
 class FDynamicMeshIndexBuffer : public FDynamicPrimitiveResource, public FIndexBuffer
diff --git a/Engine/Source/Runtime/Engine/Private/InstancedStaticMesh.cpp b/Engine/Source/Runtime/Engine/Private/InstancedStaticMesh.cpp
--- a/Engine/Source/Runtime/Engine/Private/InstancedStaticMesh.cpp
+++ b/Engine/Source/Runtime/Engine/Private/InstancedStaticMesh.cpp
@@ -5,6 +5,7 @@
 =============================================================================*/
 
 #include "EnginePrivate.h"
+#include "PhysicsPublic.h"
 #include "ShaderParameters.h"
 #include "StaticLighting.h"
 
diff --git a/Engine/Source/Runtime/Engine/Private/Landscape/LandscapeCollision.cpp b/Engine/Source/Runtime/Engine/Private/Landscape/LandscapeCollision.cpp
--- a/Engine/Source/Runtime/Engine/Private/Landscape/LandscapeCollision.cpp
+++ b/Engine/Source/Runtime/Engine/Private/Landscape/LandscapeCollision.cpp
@@ -5,6 +5,7 @@ LandscapeCollision.cpp: Landscape collision
 =============================================================================*/
 
 #include "EnginePrivate.h"
+#include "PhysicsPublic.h"
 #include "Landscape/LandscapeDataAccess.h"
 #include "Landscape/LandscapeRender.h"
 #include "../PhysicsEngine/PhysXSupport.h"
diff --git a/Engine/Source/Runtime/Engine/Private/Landscape/LandscapeGizmoActor.cpp b/Engine/Source/Runtime/Engine/Private/Landscape/LandscapeGizmoActor.cpp
--- a/Engine/Source/Runtime/Engine/Private/Landscape/LandscapeGizmoActor.cpp
+++ b/Engine/Source/Runtime/Engine/Private/Landscape/LandscapeGizmoActor.cpp
@@ -7,6 +7,7 @@
 #include "Landscape/LandscapeGizmoActiveActor.h"
 #include "Landscape/LandscapeGizmoRenderComponent.h"
 #include "Landscape/Landscape.h"
+#include "DynamicMeshBuilder.h"
 
 namespace
 {
diff --git a/Engine/Source/Runtime/Engine/Private/Level.cpp b/Engine/Source/Runtime/Engine/Private/Level.cpp
--- a/Engine/Source/Runtime/Engine/Private/Level.cpp
+++ b/Engine/Source/Runtime/Engine/Private/Level.cpp
@@ -15,6 +15,7 @@ Level.cpp: Level-related functions
 #include "PrecomputedLightVolume.h"
 #include "TickTaskManagerInterface.h"
 #include "BlueprintUtilities.h"
+#include "DynamicMeshBuilder.h"
 #if WITH_EDITOR
 #include "Editor/UnrealEd/Public/Kismet2/KismetEditorUtilities.h"
 #include "Editor/UnrealEd/Public/Kismet2/BlueprintEditorUtils.h"
diff --git a/Engine/Source/Runtime/Engine/Private/LevelTick.cpp b/Engine/Source/Runtime/Engine/Private/LevelTick.cpp
--- a/Engine/Source/Runtime/Engine/Private/LevelTick.cpp
+++ b/Engine/Source/Runtime/Engine/Private/LevelTick.cpp
@@ -7,6 +7,7 @@
 #include "EnginePrivate.h"
 #include "Net/UnrealNetwork.h"
 #include "Collision.h"
+#include "PhysicsPublic.h"
 
 #include "ParticleDefinitions.h"
 //#include "SoundDefinitions.h"
diff --git a/Engine/Source/Runtime/Engine/Private/LocalPlayer.cpp b/Engine/Source/Runtime/Engine/Private/LocalPlayer.cpp
--- a/Engine/Source/Runtime/Engine/Private/LocalPlayer.cpp
+++ b/Engine/Source/Runtime/Engine/Private/LocalPlayer.cpp
@@ -8,6 +8,7 @@
 #include "SubtitleManager.h"
 #include "Net/UnrealNetwork.h"
 #include "OnlineSubsystemUtils.h"
+#include "PhysicsPublic.h"
 
 #include "RenderCore.h"
 #include "ColorList.h"
diff --git a/Engine/Source/Runtime/Engine/Private/NavMeshRenderingHelpers.h b/Engine/Source/Runtime/Engine/Private/NavMeshRenderingHelpers.h
--- a/Engine/Source/Runtime/Engine/Private/NavMeshRenderingHelpers.h
+++ b/Engine/Source/Runtime/Engine/Private/NavMeshRenderingHelpers.h
@@ -1,6 +1,7 @@
 // Copyright 1998-2014 Epic Games, Inc. All Rights Reserved.
 #pragma once
 #include "DebugRenderSceneProxy.h"
+#include "DynamicMeshBuilder.h"
 #include "AI/Navigation/RecastHelpers.h"
 
 static const FColor NavMeshRenderColor_Recast_TriangleEdges(255,255,255);
diff --git a/Engine/Source/Runtime/Engine/Private/PhysicsEngine/BodyInstance.cpp b/Engine/Source/Runtime/Engine/Private/PhysicsEngine/BodyInstance.cpp
--- a/Engine/Source/Runtime/Engine/Private/PhysicsEngine/BodyInstance.cpp
+++ b/Engine/Source/Runtime/Engine/Private/PhysicsEngine/BodyInstance.cpp
@@ -5,6 +5,7 @@
 =============================================================================*/ 
 
 #include "EnginePrivate.h"
+#include "PhysicsPublic.h"
 #include "Collision.h"
 
 #include "MessageLog.h"
diff --git a/Engine/Source/Runtime/Engine/Private/PhysicsEngine/BodySetup.cpp b/Engine/Source/Runtime/Engine/Private/PhysicsEngine/BodySetup.cpp
--- a/Engine/Source/Runtime/Engine/Private/PhysicsEngine/BodySetup.cpp
+++ b/Engine/Source/Runtime/Engine/Private/PhysicsEngine/BodySetup.cpp
@@ -5,6 +5,7 @@
 =============================================================================*/ 
 
 #include "EnginePrivate.h"
+#include "PhysicsPublic.h"
 #include "TargetPlatform.h"
 
 #if WITH_PHYSX
diff --git a/Engine/Source/Runtime/Engine/Private/PhysicsEngine/ConstraintInstance.cpp b/Engine/Source/Runtime/Engine/Private/PhysicsEngine/ConstraintInstance.cpp
--- a/Engine/Source/Runtime/Engine/Private/PhysicsEngine/ConstraintInstance.cpp
+++ b/Engine/Source/Runtime/Engine/Private/PhysicsEngine/ConstraintInstance.cpp
@@ -1,6 +1,7 @@
 // Copyright 1998-2014 Epic Games, Inc. All Rights Reserved.
 
 #include "EnginePrivate.h"
+#include "PhysicsPublic.h"
 #include "PhysXSupport.h"
 
 #define LOCTEXT_NAMESPACE "ConstraintInstance"
diff --git a/Engine/Source/Runtime/Engine/Private/PhysicsEngine/PhysCommandHandler.cpp b/Engine/Source/Runtime/Engine/Private/PhysicsEngine/PhysCommandHandler.cpp
--- a/Engine/Source/Runtime/Engine/Private/PhysicsEngine/PhysCommandHandler.cpp
+++ b/Engine/Source/Runtime/Engine/Private/PhysicsEngine/PhysCommandHandler.cpp
@@ -2,6 +2,7 @@
 
 
 #include "EnginePrivate.h"
+#include "PhysicsPublic.h"
 
 #if WITH_PHYSX
 	#include "PhysXSupport.h"
diff --git a/Engine/Source/Runtime/Engine/Private/PhysicsEngine/PhysDrawing.cpp b/Engine/Source/Runtime/Engine/Private/PhysicsEngine/PhysDrawing.cpp
--- a/Engine/Source/Runtime/Engine/Private/PhysicsEngine/PhysDrawing.cpp
+++ b/Engine/Source/Runtime/Engine/Private/PhysicsEngine/PhysDrawing.cpp
@@ -1,7 +1,9 @@
 // Copyright 1998-2014 Epic Games, Inc. All Rights Reserved.
 
 #include "EnginePrivate.h"
+#include "PhysicsPublic.h"
 #include "PhysXSupport.h"
+#include "DynamicMeshBuilder.h"
 
 static const int32 DrawCollisionSides = 16;
 static const int32 DrawConeLimitSides = 40;
diff --git a/Engine/Source/Runtime/Engine/Private/PhysicsEngine/PhysLevel.cpp b/Engine/Source/Runtime/Engine/Private/PhysicsEngine/PhysLevel.cpp
--- a/Engine/Source/Runtime/Engine/Private/PhysicsEngine/PhysLevel.cpp
+++ b/Engine/Source/Runtime/Engine/Private/PhysicsEngine/PhysLevel.cpp
@@ -1,6 +1,7 @@
 // Copyright 1998-2014 Epic Games, Inc. All Rights Reserved.
 
 #include "EnginePrivate.h"
+#include "PhysicsPublic.h"
 #include "ParticleDefinitions.h"
 #include "PrecomputedLightVolume.h"
 
diff --git a/Engine/Source/Runtime/Engine/Private/PhysicsEngine/PhysScene.cpp b/Engine/Source/Runtime/Engine/Private/PhysicsEngine/PhysScene.cpp
--- a/Engine/Source/Runtime/Engine/Private/PhysicsEngine/PhysScene.cpp
+++ b/Engine/Source/Runtime/Engine/Private/PhysicsEngine/PhysScene.cpp
@@ -2,6 +2,7 @@
 
 
 #include "EnginePrivate.h"
+#include "PhysicsPublic.h"
 
 #if WITH_PHYSX
 	#include "PhysXSupport.h"
diff --git a/Engine/Source/Runtime/Engine/Private/PhysicsEngine/PhysSubstepTasks.cpp b/Engine/Source/Runtime/Engine/Private/PhysicsEngine/PhysSubstepTasks.cpp
--- a/Engine/Source/Runtime/Engine/Private/PhysicsEngine/PhysSubstepTasks.cpp
+++ b/Engine/Source/Runtime/Engine/Private/PhysicsEngine/PhysSubstepTasks.cpp
@@ -1,6 +1,7 @@
 // Copyright 1998-2014 Epic Games, Inc. All Rights Reserved.
 
 #include "EnginePrivate.h"
+#include "PhysicsPublic.h"
 
 #if WITH_SUBSTEPPING
 
diff --git a/Engine/Source/Runtime/Engine/Private/PhysicsEngine/PhysUtils.cpp b/Engine/Source/Runtime/Engine/Private/PhysicsEngine/PhysUtils.cpp
--- a/Engine/Source/Runtime/Engine/Private/PhysicsEngine/PhysUtils.cpp
+++ b/Engine/Source/Runtime/Engine/Private/PhysicsEngine/PhysUtils.cpp
@@ -8,6 +8,7 @@
  ===========================================================================*/
 
 #include "EnginePrivate.h"
+#include "PhysicsPublic.h"
 #include "Model.h"
 #include "PhysXSupport.h"
 
diff --git a/Engine/Source/Runtime/Engine/Private/PhysicsEngine/PhysXLibs.cpp b/Engine/Source/Runtime/Engine/Private/PhysicsEngine/PhysXLibs.cpp
--- a/Engine/Source/Runtime/Engine/Private/PhysicsEngine/PhysXLibs.cpp
+++ b/Engine/Source/Runtime/Engine/Private/PhysicsEngine/PhysXLibs.cpp
@@ -5,6 +5,7 @@
 =============================================================================*/
 
 #include "EnginePrivate.h"
+#include "PhysicsPublic.h"
 
 #if WITH_PHYSX
 
diff --git a/Engine/Source/Runtime/Engine/Private/PhysicsEngine/PhysXSupport.cpp b/Engine/Source/Runtime/Engine/Private/PhysicsEngine/PhysXSupport.cpp
--- a/Engine/Source/Runtime/Engine/Private/PhysicsEngine/PhysXSupport.cpp
+++ b/Engine/Source/Runtime/Engine/Private/PhysicsEngine/PhysXSupport.cpp
@@ -5,7 +5,7 @@
 =============================================================================*/
 
 #include "EnginePrivate.h"
-
+#include "PhysicsPublic.h"
 
 #if WITH_PHYSX
 
diff --git a/Engine/Source/Runtime/Engine/Private/PhysicsEngine/PhysicalMaterial.cpp b/Engine/Source/Runtime/Engine/Private/PhysicsEngine/PhysicalMaterial.cpp
--- a/Engine/Source/Runtime/Engine/Private/PhysicsEngine/PhysicalMaterial.cpp
+++ b/Engine/Source/Runtime/Engine/Private/PhysicsEngine/PhysicalMaterial.cpp
@@ -5,7 +5,7 @@
 =============================================================================*/ 
 
 #include "EnginePrivate.h"
-
+#include "PhysicsPublic.h"
 
 #if WITH_PHYSX
 	#include "PhysXSupport.h"
diff --git a/Engine/Source/Runtime/Engine/Private/PhysicsEngine/PhysicsCollisionHandler.cpp b/Engine/Source/Runtime/Engine/Private/PhysicsEngine/PhysicsCollisionHandler.cpp
--- a/Engine/Source/Runtime/Engine/Private/PhysicsEngine/PhysicsCollisionHandler.cpp
+++ b/Engine/Source/Runtime/Engine/Private/PhysicsEngine/PhysicsCollisionHandler.cpp
@@ -1,6 +1,7 @@
 // Copyright 1998-2014 Epic Games, Inc. All Rights Reserved.
 
 #include "EnginePrivate.h"
+#include "PhysicsPublic.h"
 
 UPhysicsCollisionHandler::UPhysicsCollisionHandler(const class FPostConstructInitializeProperties& PCIP)
 	: Super(PCIP)
diff --git a/Engine/Source/Runtime/Engine/Private/PhysicsEngine/PhysicsHandleComponent.cpp b/Engine/Source/Runtime/Engine/Private/PhysicsEngine/PhysicsHandleComponent.cpp
--- a/Engine/Source/Runtime/Engine/Private/PhysicsEngine/PhysicsHandleComponent.cpp
+++ b/Engine/Source/Runtime/Engine/Private/PhysicsEngine/PhysicsHandleComponent.cpp
@@ -1,6 +1,7 @@
 // Copyright 1998-2014 Epic Games, Inc. All Rights Reserved.
 
 #include "EnginePrivate.h"
+#include "PhysicsPublic.h"
 #if WITH_PHYSX
 	#include "PhysXSupport.h"
 #endif // WITH_PHYSX
diff --git a/Engine/Source/Runtime/Engine/Private/PrimitiveComponent.cpp b/Engine/Source/Runtime/Engine/Private/PrimitiveComponent.cpp
--- a/Engine/Source/Runtime/Engine/Private/PrimitiveComponent.cpp
+++ b/Engine/Source/Runtime/Engine/Private/PrimitiveComponent.cpp
@@ -5,6 +5,7 @@
 =============================================================================*/
 
 #include "EnginePrivate.h"
+#include "PhysicsPublic.h"
 #include "LevelUtils.h"
 #if WITH_PHYSX
 #include "PhysicsEngine/PhysXSupport.h"
diff --git a/Engine/Source/Runtime/Engine/Private/PrimitiveDrawingUtils.cpp b/Engine/Source/Runtime/Engine/Private/PrimitiveDrawingUtils.cpp
--- a/Engine/Source/Runtime/Engine/Private/PrimitiveDrawingUtils.cpp
+++ b/Engine/Source/Runtime/Engine/Private/PrimitiveDrawingUtils.cpp
@@ -1,6 +1,7 @@
 // Copyright 1998-2014 Epic Games, Inc. All Rights Reserved.
 
 #include "EnginePrivate.h"
+#include "DynamicMeshBuilder.h"
 #include "SceneManagement.h"
 #include "StaticMeshResources.h"
 
diff --git a/Engine/Source/Runtime/Engine/Private/SceneComponent.cpp b/Engine/Source/Runtime/Engine/Private/SceneComponent.cpp
--- a/Engine/Source/Runtime/Engine/Private/SceneComponent.cpp
+++ b/Engine/Source/Runtime/Engine/Private/SceneComponent.cpp
@@ -5,6 +5,7 @@
 =============================================================================*/
 
 #include "EnginePrivate.h"
+#include "PhysicsPublic.h"
 #include "MessageLog.h"
 #include "Net/UnrealNetwork.h"
 
diff --git a/Engine/Source/Runtime/Engine/Private/SkeletalMesh.cpp b/Engine/Source/Runtime/Engine/Private/SkeletalMesh.cpp
--- a/Engine/Source/Runtime/Engine/Private/SkeletalMesh.cpp
+++ b/Engine/Source/Runtime/Engine/Private/SkeletalMesh.cpp
@@ -19,6 +19,7 @@
 #include "MapErrors.h"
 #include "SkeletalRenderGPUSkin.h"
 #include "RawIndexBuffer.h"
+#include "PhysicsPublic.h"
 
 #if WITH_EDITOR
 #include "MeshUtilities.h"
diff --git a/Engine/Source/Runtime/Engine/Private/SkeletalMeshComponentPhysics.cpp b/Engine/Source/Runtime/Engine/Private/SkeletalMeshComponentPhysics.cpp
--- a/Engine/Source/Runtime/Engine/Private/SkeletalMeshComponentPhysics.cpp
+++ b/Engine/Source/Runtime/Engine/Private/SkeletalMeshComponentPhysics.cpp
@@ -1,6 +1,7 @@
 // Copyright 1998-2014 Epic Games, Inc. All Rights Reserved.
 
 #include "EnginePrivate.h"
+#include "PhysicsPublic.h"
 #include "SkeletalRender.h"
 #include "SkeletalRenderPublic.h"
 
diff --git a/Engine/Source/Runtime/Engine/Private/TextRenderComponent.cpp b/Engine/Source/Runtime/Engine/Private/TextRenderComponent.cpp
--- a/Engine/Source/Runtime/Engine/Private/TextRenderComponent.cpp
+++ b/Engine/Source/Runtime/Engine/Private/TextRenderComponent.cpp
@@ -1,6 +1,7 @@
 // Copyright 1998-2014 Epic Games, Inc. All Rights Reserved.
 
 #include "EnginePrivate.h"
+#include "DynamicMeshBuilder.h"
 #include "../../../../Source/Runtime/Engine/Classes/Engine/TextRenderActor.h"
 
 ATextRenderActor::ATextRenderActor(const class FPostConstructInitializeProperties& PCIP)
diff --git a/Engine/Source/Runtime/Engine/Private/Vehicles/PhysXVehicleManager.cpp b/Engine/Source/Runtime/Engine/Private/Vehicles/PhysXVehicleManager.cpp
--- a/Engine/Source/Runtime/Engine/Private/Vehicles/PhysXVehicleManager.cpp
+++ b/Engine/Source/Runtime/Engine/Private/Vehicles/PhysXVehicleManager.cpp
@@ -1,6 +1,7 @@
 // Copyright 1998-2014 Epic Games, Inc. All Rights Reserved.
 
 #include "EnginePrivate.h"
+#include "PhysicsPublic.h"
 
 #if WITH_VEHICLE
 
diff --git a/Engine/Source/Runtime/Engine/Private/Vehicles/VehicleWheel.cpp b/Engine/Source/Runtime/Engine/Private/Vehicles/VehicleWheel.cpp
--- a/Engine/Source/Runtime/Engine/Private/Vehicles/VehicleWheel.cpp
+++ b/Engine/Source/Runtime/Engine/Private/Vehicles/VehicleWheel.cpp
@@ -1,6 +1,7 @@
 // Copyright 1998-2014 Epic Games, Inc. All Rights Reserved.
 
 #include "EnginePrivate.h"
+#include "PhysicsPublic.h"
 
 #if WITH_PHYSX
 #include "../PhysicsEngine/PhysXSupport.h"
diff --git a/Engine/Source/Runtime/Engine/Private/Vehicles/WheeledVehicleMovementComponent.cpp b/Engine/Source/Runtime/Engine/Private/Vehicles/WheeledVehicleMovementComponent.cpp
--- a/Engine/Source/Runtime/Engine/Private/Vehicles/WheeledVehicleMovementComponent.cpp
+++ b/Engine/Source/Runtime/Engine/Private/Vehicles/WheeledVehicleMovementComponent.cpp
@@ -1,6 +1,7 @@
 // Copyright 1998-2014 Epic Games, Inc. All Rights Reserved.
 
 #include "EnginePrivate.h"
+#include "PhysicsPublic.h"
 #include "Net/UnrealNetwork.h"
 #include "MessageLog.h"
 
diff --git a/Engine/Source/Runtime/Engine/Private/Vehicles/WheeledVehicleMovementComponent4W.cpp b/Engine/Source/Runtime/Engine/Private/Vehicles/WheeledVehicleMovementComponent4W.cpp
--- a/Engine/Source/Runtime/Engine/Private/Vehicles/WheeledVehicleMovementComponent4W.cpp
+++ b/Engine/Source/Runtime/Engine/Private/Vehicles/WheeledVehicleMovementComponent4W.cpp
@@ -1,10 +1,11 @@
 // Copyright 1998-2014 Epic Games, Inc. All Rights Reserved.
 
 #include "EnginePrivate.h"
+#include "PhysicsPublic.h"
 
 #if WITH_VEHICLE
 #include "../PhysicsEngine/PhysXSupport.h"
-#endif // WITH_PHYSX
+#endif // WITH_VEHICLE
 
 UWheeledVehicleMovementComponent4W::UWheeledVehicleMovementComponent4W(const class FPostConstructInitializeProperties& PCIP) : Super(PCIP)
 {
diff --git a/Engine/Source/Runtime/Engine/Private/World.cpp b/Engine/Source/Runtime/Engine/Private/World.cpp
--- a/Engine/Source/Runtime/Engine/Private/World.cpp
+++ b/Engine/Source/Runtime/Engine/Private/World.cpp
@@ -24,6 +24,7 @@
 #include "SoundDefinitions.h"
 #include "VisualLog.h"
 #include "LevelUtils.h"
+#include "PhysicsPublic.h"
 #if WITH_EDITOR
 	#include "Editor/UnrealEd/Public/Kismet2/KismetEditorUtilities.h"
 	#include "Editor/UnrealEd/Public/Kismet2/BlueprintEditorUtils.h"
diff --git a/Engine/Source/Runtime/Engine/Public/Engine.h b/Engine/Source/Runtime/Engine/Public/Engine.h
--- a/Engine/Source/Runtime/Engine/Public/Engine.h
+++ b/Engine/Source/Runtime/Engine/Public/Engine.h
@@ -389,7 +389,6 @@
 #include "NavigationModifier.h"
 #include "NavigationOctree.h"
 #include "ContentStreaming.h"				// Content streaming class definitions.
-#include "PhysicsPublic.h"
 #include "ComponentReregisterContext.h"	
 #include "DrawDebugHelpers.h"
 #include "UnrealEngine.h"					// Unreal engine helpers.
diff --git a/Engine/Source/Runtime/Engine/Public/PhysicsPublic.h b/Engine/Source/Runtime/Engine/Public/PhysicsPublic.h
--- a/Engine/Source/Runtime/Engine/Public/PhysicsPublic.h
+++ b/Engine/Source/Runtime/Engine/Public/PhysicsPublic.h
@@ -7,9 +7,7 @@
 
 #pragma once
 
-#include "DynamicMeshBuilder.h"
 #include "PhysxUserData.h"
-#include "LocalVertexFactory.h"
 
 /**
  * Physics stats
@@ -476,7 +474,7 @@ struct FPhysSceneShaderInfo
 
 #endif
 
-// Might be handy somewhere...
+/** Enum to indicate types of simple shapes */
 enum EKCollisionPrimitiveType
 {
 	KPT_Sphere = 0,
diff --git a/Engine/Source/Runtime/Launch/Private/Launch.cpp b/Engine/Source/Runtime/Launch/Private/Launch.cpp
--- a/Engine/Source/Runtime/Launch/Private/Launch.cpp
+++ b/Engine/Source/Runtime/Launch/Private/Launch.cpp
@@ -5,6 +5,7 @@
 =============================================================================*/
 
 #include "LaunchPrivatePCH.h"
+#include "PhysicsPublic.h"
 #include "ExceptionHandling.h"
 #include "ModuleManager.h"
 
diff --git a/Engine/Source/Runtime/Launch/Private/LaunchEngineLoop.cpp b/Engine/Source/Runtime/Launch/Private/LaunchEngineLoop.cpp
--- a/Engine/Source/Runtime/Launch/Private/LaunchEngineLoop.cpp
+++ b/Engine/Source/Runtime/Launch/Private/LaunchEngineLoop.cpp
@@ -38,6 +38,7 @@
 	#include "GlobalShader.h"
 	#include "ParticleHelper.h"
 	#include "Online.h"
+	#include "PhysicsPublic.h"
 	#include "PlatformFeatures.h"
 	#include "DeviceProfiles/DeviceProfileManager.h"
 	#include "Commandlets/Commandlet.h"
