Alt + Shift + R now brings up the reference viewer for the currently selected asset in the level editor and asset editor
#TTP 331416 - UE4: Would like to be able to set a hotkey to open the Reference viewer similar to "find in content browser"

#branch UE4

#proj Editor.UnrealEd
#proj Editor.LevelEditor

#add Added new global UICommand ViewReferences, bound to Alt + Shift + R.

#add Implemented the new UICommand in the AssetEditorToolkit and the LevelEditorr. This means that most of the editor is covered by this hotkey. Hitting it in an asset editor brings it up for the current asset, hitting it in the level editor brings it up for the selected asset.

#Further work: Further work is required to modify the content browsers uses of view references, for now it is deemed to be too much work for this task.

diff --git a/Engine/Source/Editor/LevelEditor/LevelEditor.Build.cs b/Engine/Source/Editor/LevelEditor/LevelEditor.Build.cs
--- a/Engine/Source/Editor/LevelEditor/LevelEditor.Build.cs
+++ b/Engine/Source/Editor/LevelEditor/LevelEditor.Build.cs
@@ -11,7 +11,8 @@ public class LevelEditor : ModuleRules
 				"AssetTools",
 				"Kismet",
 				"MainFrame",
-                "PlacementMode"
+                "PlacementMode",
+				"ReferenceViewer",
 			}
 		);
 
@@ -71,7 +72,8 @@ public class LevelEditor : ModuleRules
 				"DeviceProfileEditor",
 				"DeviceProfileServices",
                 "PlacementMode",
-				"UserFeedback"
+				"UserFeedback",
+				"ReferenceViewer",
 			}
 		);
 	}
diff --git a/Engine/Source/Editor/LevelEditor/Private/LevelEditor.cpp b/Engine/Source/Editor/LevelEditor/Private/LevelEditor.cpp
--- a/Engine/Source/Editor/LevelEditor/Private/LevelEditor.cpp
+++ b/Engine/Source/Editor/LevelEditor/Private/LevelEditor.cpp
@@ -615,7 +615,13 @@ void FLevelEditorModule::BindGlobalLevelEditorCommands()
 		FGlobalEditorCommonCommands::Get().FindInContentBrowser, 
 		FExecuteAction::CreateStatic( &FLevelEditorActionCallbacks::FindInContentBrowser_Clicked )
 		);
-						
+
+	ActionList.MapAction( 
+		FGlobalEditorCommonCommands::Get().ViewReferences, 
+		FExecuteAction::CreateStatic( &FLevelEditorActionCallbacks::ViewReferences_Execute ),
+		FCanExecuteAction::CreateStatic( &FLevelEditorActionCallbacks::CanViewReferences )
+		);
+
 	ActionList.MapAction( 
 		Commands.SnapCameraToActor, 
 		FExecuteAction::CreateStatic( &FLevelEditorActionCallbacks::ExecuteExecCommand, FString( TEXT("CAMERA SNAP") ) )
diff --git a/Engine/Source/Editor/LevelEditor/Private/LevelEditorActions.cpp b/Engine/Source/Editor/LevelEditor/Private/LevelEditorActions.cpp
--- a/Engine/Source/Editor/LevelEditor/Private/LevelEditorActions.cpp
+++ b/Engine/Source/Editor/LevelEditor/Private/LevelEditorActions.cpp
@@ -42,6 +42,7 @@
 #include "EngineAnalytics.h"
 #include "AnalyticsEventAttribute.h"
 #include "IAnalyticsProvider.h"
+#include "ReferenceViewer.h"
 
 #include "EditorActorFolders.h"
 #include "ActorPickerMode.h"
@@ -1064,6 +1065,38 @@ void FLevelEditorActionCallbacks::FindInContentBrowser_Clicked()
 	GEditor->SyncToContentBrowser();
 }
 
+void FLevelEditorActionCallbacks::ViewReferences_Execute()
+{
+	if( GEditor->GetSelectedActorCount() > 0 )
+	{
+		TArray< UObject* > ReferencedAssets;
+		GEditor->GetReferencedAssetsForEditorSelection( ReferencedAssets );
+
+		if (ReferencedAssets.Num() > 0)
+		{
+			TArray< FName > ViewableObjects;
+			for( auto ObjectIter = ReferencedAssets.CreateConstIterator(); ObjectIter; ++ObjectIter )
+			{
+				// Don't allow user to perform certain actions on objects that aren't actually assets (e.g. Level Script blueprint objects)
+				const auto EditingObject = *ObjectIter;
+				if( EditingObject != NULL && EditingObject->IsAsset() )
+				{
+					ViewableObjects.Add( EditingObject->GetOuter()->GetFName());
+				}
+			}
+
+			IReferenceViewerModule::Get().InvokeReferenceViewerTab(ViewableObjects);
+		}
+	}
+}
+
+bool FLevelEditorActionCallbacks::CanViewReferences()
+{
+	TArray< UObject* > ReferencedAssets;
+	GEditor->GetReferencedAssetsForEditorSelection(ReferencedAssets);
+	return ReferencedAssets.Num() > 0;
+}
+
 void FLevelEditorActionCallbacks::EditAsset_Clicked( const EToolkitMode::Type ToolkitMode, TWeakPtr< SLevelEditor > LevelEditor, bool bConfirmMultiple )
 {
 	if( GEditor->GetSelectedActorCount() > 0 )
diff --git a/Engine/Source/Editor/LevelEditor/Private/LevelEditorContextMenu.cpp b/Engine/Source/Editor/LevelEditor/Private/LevelEditorContextMenu.cpp
--- a/Engine/Source/Editor/LevelEditor/Private/LevelEditorContextMenu.cpp
+++ b/Engine/Source/Editor/LevelEditor/Private/LevelEditorContextMenu.cpp
@@ -221,6 +221,8 @@ void FLevelEditorContextMenu::FillMenu( FMenuBuilder& MenuBuilder, TWeakPtr<SLev
 
 			}
 
+			MenuBuilder.AddMenuEntry( FGlobalEditorCommonCommands::Get().ViewReferences );
+
 		}
 		MenuBuilder.EndSection();
 	}
diff --git a/Engine/Source/Editor/LevelEditor/Public/LevelEditorActions.h b/Engine/Source/Editor/LevelEditor/Public/LevelEditorActions.h
--- a/Engine/Source/Editor/LevelEditor/Public/LevelEditorActions.h
+++ b/Engine/Source/Editor/LevelEditor/Public/LevelEditorActions.h
@@ -758,6 +758,12 @@ public:
 	 */
 	static void FindInContentBrowser_Clicked();
 
+	/** Called when the ViewReferences command is executed */
+	static void ViewReferences_Execute();
+
+	/** If true ViewReferences_Execute can be called */
+	static bool CanViewReferences();
+
 	/** Called to when "Edit Asset" is clicked */
 	static void EditAsset_Clicked( const EToolkitMode::Type ToolkitMode, TWeakPtr< class SLevelEditor > LevelEditor, bool bAskMultiple );
 
diff --git a/Engine/Source/Editor/UnrealEd/Private/Toolkits/AssetEditorToolkit.cpp b/Engine/Source/Editor/UnrealEd/Private/Toolkits/AssetEditorToolkit.cpp
--- a/Engine/Source/Editor/UnrealEd/Private/Toolkits/AssetEditorToolkit.cpp
+++ b/Engine/Source/Editor/UnrealEd/Private/Toolkits/AssetEditorToolkit.cpp
@@ -11,6 +11,7 @@
 #include "ClassIconFinder.h"
 #include "IUserFeedbackModule.h"
 #include "IDocumentation.h"
+#include "ReferenceViewer.h"
 
 #define LOCTEXT_NAMESPACE "AssetEditorToolkit"
 
@@ -139,7 +140,12 @@ void FAssetEditorToolkit::InitAssetEditor( const EToolkitMode::Type Mode, const
 	ToolkitCommands->MapAction(
 		FGlobalEditorCommonCommands::Get().FindInContentBrowser,
 		FExecuteAction::CreateSP( this, &FAssetEditorToolkit::FindInContentBrowser_Execute ) );
-
+	
+	ToolkitCommands->MapAction(
+		FGlobalEditorCommonCommands::Get().ViewReferences,
+		FExecuteAction::CreateSP( this, &FAssetEditorToolkit::ViewReferences_Execute ),
+		FCanExecuteAction::CreateSP( this, &FAssetEditorToolkit::CanViewReferences ));
+	
 	ToolkitCommands->MapAction(
 		FGlobalEditorCommonCommands::Get().OpenDocumentation,
 		FExecuteAction::CreateSP( this, &FAssetEditorToolkit::BrowseDocumentation_Execute ) );
@@ -461,6 +467,29 @@ void FAssetEditorToolkit::BrowseDocumentation_Execute() const
 	IDocumentation::Get()->Open(GetDocumentationLink());
 }
 
+void FAssetEditorToolkit::ViewReferences_Execute()
+{
+	if (ensure( ViewableObjects.Num() > 0))
+	{
+		IReferenceViewerModule::Get().InvokeReferenceViewerTab(ViewableObjects);
+	}
+}
+
+bool FAssetEditorToolkit::CanViewReferences()
+{
+	ViewableObjects.Empty();
+	for (const auto EditingObject : EditingObjects)
+	{
+		// Don't allow user to perform certain actions on objects that aren't actually assets (e.g. Level Script blueprint objects)
+		if (EditingObject != NULL && EditingObject->IsAsset())
+		{
+			ViewableObjects.Add(EditingObject->GetOuter()->GetFName());
+		}
+	}
+
+	return ViewableObjects.Num() > 0;
+}
+
 FString FAssetEditorToolkit::GetDocumentationLink() const
 {
 	return FString(TEXT("%ROOT%"));
@@ -580,6 +609,7 @@ void FAssetEditorToolkit::FillDefaultAssetMenuCommands( FMenuBuilder& MenuBuilde
 	if( IsActuallyAnAsset() )
 	{
 		MenuBuilder.AddMenuEntry( FGlobalEditorCommonCommands::Get().FindInContentBrowser, "FindInContentBrowser", LOCTEXT("FindInContentBrowser", "Find in Content Browser...") );
+		MenuBuilder.AddMenuEntry( FGlobalEditorCommonCommands::Get().ViewReferences);
 
 		// Add a reimport menu entry for each supported editable object
 		for( auto ObjectIter = EditingObjects.CreateConstIterator(); ObjectIter; ++ObjectIter )
diff --git a/Engine/Source/Editor/UnrealEd/Private/Toolkits/GlobalEditorCommonCommands.cpp b/Engine/Source/Editor/UnrealEd/Private/Toolkits/GlobalEditorCommonCommands.cpp
--- a/Engine/Source/Editor/UnrealEd/Private/Toolkits/GlobalEditorCommonCommands.cpp
+++ b/Engine/Source/Editor/UnrealEd/Private/Toolkits/GlobalEditorCommonCommands.cpp
@@ -96,7 +96,8 @@ void FGlobalEditorCommonCommands::RegisterCommands()
 	//UI_COMMAND( SummonControlTabNavigation, "Tab Navigation", "Summons a list of open assets and tabs", EUserInterfaceActionType::Button, FInputGesture(EModifierKey::Control, EKeys::Tab) );
 	UI_COMMAND( SummonOpenAssetDialog, "Open Asset...", "Summons an asset picker", EUserInterfaceActionType::Button, FInputGesture(EModifierKey::Control, EKeys::P) );
 	UI_COMMAND( FindInContentBrowser, "Find in Content Browser", "Summons the Content Browser and navigates to the selected asset", EUserInterfaceActionType::Button, FInputGesture(EModifierKey::Control, EKeys::B));
-
+	UI_COMMAND( ViewReferences, "View References", "Launches the reference viewer showing the selected assets' references", EUserInterfaceActionType::Button, FInputGesture(EModifierKey::Shift | EModifierKey::Alt, EKeys::R));
+	
 	UI_COMMAND( OpenDocumentation, "Open Documentation...", "Opens documentation for this tool", EUserInterfaceActionType::Button, FInputGesture(EKeys::F1) );
 }
 
diff --git a/Engine/Source/Editor/UnrealEd/Public/Toolkits/AssetEditorToolkit.h b/Engine/Source/Editor/UnrealEd/Public/Toolkits/AssetEditorToolkit.h
--- a/Engine/Source/Editor/UnrealEd/Public/Toolkits/AssetEditorToolkit.h
+++ b/Engine/Source/Editor/UnrealEd/Public/Toolkits/AssetEditorToolkit.h
@@ -197,6 +197,12 @@ private:
 		FLayoutSaveRestore::SaveToConfig(GEditorLayoutIni, LayoutToSave);
 	}
 
+	/** Called when "View References" is called for this asset */
+	void ViewReferences_Execute();
+
+	/** If true ViewReferences_Execute can be called, also caches ViewableObjects */
+	bool CanViewReferences();
+
 protected:
 
 	/** For standalone asset editing tool-kits that were switched from world-centric mode on the fly, this stores
@@ -235,6 +241,9 @@ private:
 
 	/**	The tab ids for all the tabs used */
 	static const FName ToolbarTabId;
+
+	/** A cached list of selected objects that can be viewed in the reference viewer */
+	TArray< FName > ViewableObjects;
 };
 
 
diff --git a/Engine/Source/Editor/UnrealEd/Public/Toolkits/GlobalEditorCommonCommands.h b/Engine/Source/Editor/UnrealEd/Public/Toolkits/GlobalEditorCommonCommands.h
--- a/Engine/Source/Editor/UnrealEd/Public/Toolkits/GlobalEditorCommonCommands.h
+++ b/Engine/Source/Editor/UnrealEd/Public/Toolkits/GlobalEditorCommonCommands.h
@@ -22,5 +22,6 @@ public:
 	TSharedPtr<FUICommandInfo> SummonControlTabNavigation;
 	TSharedPtr<FUICommandInfo> SummonOpenAssetDialog;
 	TSharedPtr<FUICommandInfo> OpenDocumentation;
+	TSharedPtr<FUICommandInfo> ViewReferences;
 };
 
diff --git a/Engine/Source/Editor/UnrealEd/UnrealEd.Build.cs b/Engine/Source/Editor/UnrealEd/UnrealEd.Build.cs
--- a/Engine/Source/Editor/UnrealEd/UnrealEd.Build.cs
+++ b/Engine/Source/Editor/UnrealEd/UnrealEd.Build.cs
@@ -46,6 +46,7 @@ public class UnrealEd : ModuleRules
 				"SoundClassEditor",
 				"ViewportSnapping",
 				"SourceCodeAccess",
+				"ReferenceViewer",
 			}
 		);
 
@@ -169,7 +170,8 @@ public class UnrealEd : ModuleRules
 				"UserFeedback",
 				"GameplayTagsEditor",
 				"UndoHistory",
-				"SourceCodeAccess"
+				"SourceCodeAccess",
+				"ReferenceViewer"
 			}
 		);
 
