UE4: Fixed a crash that occurs when switching maps during automation testing.

diff --git a/Engine/Source/Editor/UnrealEd/Private/EditorServer.cpp b/Engine/Source/Editor/UnrealEd/Private/EditorServer.cpp
--- a/Engine/Source/Editor/UnrealEd/Private/EditorServer.cpp
+++ b/Engine/Source/Editor/UnrealEd/Private/EditorServer.cpp
@@ -2013,6 +2013,17 @@ bool UEditorEngine::Map_Load(const TCHAR* Str, FOutputDevice& Ar)
 					}
 					EditorDestroyWorld( Context, LocalizedLoadingMap, ExistingWorld );
 
+					// Refresh ExistingPackage and Existing World now that GC has occurred.
+					ExistingPackage = FindPackage(NULL, *LongTempFname);
+					if (ExistingPackage)
+					{
+						ExistingWorld = UWorld::FindWorldInPackage(ExistingPackage);
+					}
+					else
+					{
+						ExistingWorld = NULL;
+					}
+
 					GWarn->StatusUpdate( -1, -1, LocalizedLoadingMap );
 				}
 
