TTP 337155 : BlueprintsInfoDump commandlet uses unique keys for action entries.

diff --git a/Engine/Source/Editor/UnrealEd/Private/Commandlets/DumpBlueprintsInfoCommandlet.cpp b/Engine/Source/Editor/UnrealEd/Private/Commandlets/DumpBlueprintsInfoCommandlet.cpp
--- a/Engine/Source/Editor/UnrealEd/Private/Commandlets/DumpBlueprintsInfoCommandlet.cpp
+++ b/Engine/Source/Editor/UnrealEd/Private/Commandlets/DumpBlueprintsInfoCommandlet.cpp
@@ -32,6 +32,7 @@ int32 UDumpBlueprintsInfoCommandlet::Main(FString const& Params)
 	FArchive* FileOut = IFileManager::Get().CreateFileWriter(*FilePath);
 	if (FileOut != nullptr)
 	{
+		// Get class name from command line switches
 		UClass* ClassFilter = nullptr;
 		for (FString& Switch : Switches)
 		{
@@ -43,6 +44,7 @@ int32 UDumpBlueprintsInfoCommandlet::Main(FString const& Params)
 			}
 		}
 
+		// Build class context container
 		FString LeadingFields("{\n\t\"ClassContext\" : \"");
 		if (ClassFilter != nullptr)
 		{
@@ -52,7 +54,9 @@ int32 UDumpBlueprintsInfoCommandlet::Main(FString const& Params)
 		{
 			LeadingFields += "<NO_CONTEXT>";
 		}
+		// Build palette contents header
 		LeadingFields += "\",\n\t\"Palette\" : {";
+		
 		FileOut->Serialize(TCHAR_TO_ANSI(*LeadingFields), LeadingFields.Len());
 
 		UPackage* TransientPackage = GetTransientPackage();
@@ -71,46 +75,67 @@ int32 UDumpBlueprintsInfoCommandlet::Main(FString const& Params)
 				continue;
 			}
 
+			// Get action name
 			TSharedPtr<FEdGraphSchemaAction> PrimeAction = Action.Actions[0];
 			FString ActionName = PrimeAction->MenuDescription.ToString();
-
-			FString ActionEntry("\n\t\t\"");
-			ActionEntry += ActionName + TEXT("\" : {");
-			
 			
+			// Get action category info
 			TArray<FString> MenuHierarchy;
 			Action.GetCategoryChain(MenuHierarchy);
 			bool bHasCategory = (MenuHierarchy.Num() > 0);
+			FString ActionCategory = TEXT("");
 
-			ActionEntry += TEXT("\n\t\t\t\"Category\" : \"");
 			if (bHasCategory)
-			{
+			{	
 				for (FString const& SubCategory : MenuHierarchy)
 				{
-					ActionEntry += SubCategory + TEXT("|");
+					ActionCategory += SubCategory + TEXT("|");
 				}
-				ActionEntry.RemoveAt(ActionEntry.Len() - 1); // remove the last '|'
 			}
+
+
+			// Build action key data
+			FString ActionEntry("\n\t\t\"");
+			ActionEntry += ActionCategory + ActionName + TEXT("\" : {");
+
+			// Build action name data
+			ActionEntry += TEXT("\n\t\t\t\"Name\" : \"");
+			ActionEntry += ActionName;
 			ActionEntry += TEXT("\"");
 
+			// Build action category key data
+			ActionEntry += TEXT(",\n\t\t\t\"Category\" : \"");
+			if (bHasCategory)
+			{
+				// Build action category data from info
+				ActionEntry += ActionCategory;
+				ActionEntry.RemoveAt(ActionEntry.Len() - 1); // remove the last '|'
+			}
+			// End action category data
+			ActionEntry += TEXT("\"");
 
+			// Get action node type info
 			UK2Node const* NodeTemplate = FK2SchemaActionUtils::ExtractNodeTemplateFromAction(PrimeAction);
 			if (NodeTemplate != nullptr)
 			{
-
+				// Build action node type data
 				ActionEntry += TEXT(",\n\t\t\t\"Node\" : \"");
 				ActionEntry += NodeTemplate->GetClass()->GetPathName();
 				ActionEntry += TEXT("\"");
 			}
+			// Finish action entry
 			ActionEntry += TEXT("\n\t\t}");
 
 			bool bIsLastEntry = (ActionIndex == PaletteBuilder.GetNumActions()-1);
 			if (!bIsLastEntry)
 			{
+				// Add comma for next entry
 				ActionEntry += TEXT(",");
 			}
+			// Write entry to file
 			FileOut->Serialize(TCHAR_TO_ANSI(*ActionEntry), ActionEntry.Len());
 		}
+		// Close remaining entries and write to file
 		FileOut->Serialize(TCHAR_TO_ANSI(TEXT("\n\t}\n}")), 5);
  
 		FileOut->Close();
