change lighting model to shading model

diff --git a/Engine/Source/Editor/UnrealEd/Private/Lightmass/LightmassRender.cpp b/Engine/Source/Editor/UnrealEd/Private/Lightmass/LightmassRender.cpp
--- a/Engine/Source/Editor/UnrealEd/Private/Lightmass/LightmassRender.cpp
+++ b/Engine/Source/Editor/UnrealEd/Private/Lightmass/LightmassRender.cpp
@@ -303,7 +303,7 @@ public:
 		{
 			UMaterial* ProxyMaterial = MaterialInterface->GetMaterial();
 			EBlendMode BlendMode = MaterialInterface->GetBlendMode();
-			EMaterialLightingModel LightingModel = MaterialInterface->GetLightingModel();
+			EMaterialShadingModel ShadingModel = MaterialInterface->GetShadingModel();
 			check(ProxyMaterial);
 			FLightmassMaterialCompiler ProxyCompiler(Compiler);
 			switch (PropertyToCompile)
@@ -347,7 +347,7 @@ public:
 				}
 				else if (BlendMode == BLEND_Modulate)
 				{
-					if (LightingModel == MLM_Unlit)
+					if (ShadingModel == MSM_Unlit)
 					{
 						return Compiler->ForceCast(MaterialInterface->CompileProperty(Compiler, MP_EmissiveColor),MCT_Float3,true,true);
 					}
@@ -359,7 +359,7 @@ public:
 				else if ((BlendMode == BLEND_Translucent) || (BlendMode == BLEND_Additive))
 				{
 					int32 ColoredOpacity = -1;
-					if (LightingModel == MLM_Unlit)
+					if (ShadingModel == MSM_Unlit)
 					{
 						ColoredOpacity = Compiler->ForceCast(MaterialInterface->CompileProperty(Compiler, MP_EmissiveColor),MCT_Float3,true,true);
 					}
@@ -447,7 +447,7 @@ public:
 	}
 	virtual bool IsMasked() const									{ return false; }
 	virtual enum EBlendMode GetBlendMode() const					{ return BLEND_Opaque; }
-	virtual enum EMaterialLightingModel GetLightingModel() const	{ return MLM_Unlit; }
+	virtual enum EMaterialShadingModel GetShadingModel() const		{ return MSM_Unlit; }
 	virtual float GetOpacityMaskClipValue() const					{ return 0.5f; }
 	virtual FString GetFriendlyName() const { return FString::Printf(TEXT("FLightmassMaterialRenderer %s"), MaterialInterface ? *MaterialInterface->GetName() : TEXT("NULL")); }
 
@@ -527,7 +527,7 @@ public:
 		OutUniformValue.A = 0.0f;
 
 		EBlendMode BlendMode = MaterialInterface->GetBlendMode();
-		EMaterialLightingModel LightingModel = MaterialInterface->GetLightingModel();
+		EMaterialShadingModel ShadingModel = MaterialInterface->GetShadingModel();
 		
 		check(Material);
 		bool bExpressionIsNULL = false;
@@ -575,7 +575,7 @@ public:
 				(BlendMode == BLEND_Additive))
 			{
 				bool bColorInputIsNULL = false;
-				if (LightingModel == MLM_Unlit)
+				if (ShadingModel == MSM_Unlit)
 				{
 					bColorInputIsNULL = !IsMaterialInputConnected(Material, MP_EmissiveColor);
 				}
@@ -819,14 +819,14 @@ bool FLightmassMaterialRenderer::GenerateMaterialData(
 	check(BaseMaterial);
 
 	EBlendMode BlendMode = InMaterial.GetBlendMode();
-	EMaterialLightingModel LightingModel = InMaterial.GetLightingModel();
- 	if ((LightingModel != MLM_DefaultLit) &&
-		(LightingModel != MLM_Unlit) &&
-		(LightingModel != MLM_Subsurface) &&
-		(LightingModel != MLM_PreintegratedSkin))
+	EMaterialShadingModel ShadingModel = InMaterial.GetShadingModel();
+ 	if ((ShadingModel != MSM_DefaultLit) &&
+		(ShadingModel != MSM_Unlit) &&
+		(ShadingModel != MSM_Subsurface) &&
+		(ShadingModel != MSM_PreintegratedSkin))
 	{
-		UE_LOG(LogLightmassRender, Warning, TEXT("LIGHTMASS: Material has an unsupported lighting model: %d on %s"), 
-			(int32)LightingModel,
+		UE_LOG(LogLightmassRender, Warning, TEXT("LIGHTMASS: Material has an unsupported shading model: %d on %s"), 
+			(int32)ShadingModel,
 			*(InMaterial.GetPathName()));
 	}
 
