Update the ImportData when reimporting an Fbx animation if import is successful.
UDN: Bug on reimport animations - https://udn.unrealengine.com/questions/208671/bug-on-reimport-animations.html
reviewer Thomas.Sarkanen

diff --git a/Engine/Source/Editor/UnrealEd/Private/Factories/EditorFactories.cpp b/Engine/Source/Editor/UnrealEd/Private/Factories/EditorFactories.cpp
--- a/Engine/Source/Editor/UnrealEd/Private/Factories/EditorFactories.cpp
+++ b/Engine/Source/Editor/UnrealEd/Private/Factories/EditorFactories.cpp
@@ -6011,6 +6011,10 @@ EReimportResult::Type UReimportFbxAnimSequenceFactory::Reimport( UObject* Obj )
 	{
 		UE_LOG(LogEditorFactories, Log, TEXT("-- imported successfully") );
 
+		// update the data in case the file source has changed
+		ImportData->SourceFilePath = FReimportManager::SanitizeImportFilename(UFactory::CurrentFilename, AnimSequence);
+		ImportData->SourceFileTimestamp = IFileManager::Get().GetTimeStamp(*UFactory::CurrentFilename).ToString();
+
 		// Try to find the outer package so we can dirty it up
 		if (AnimSequence->GetOuter())
 		{
