#UE4
Fix for crash when diffing an asset that was saved and never reopened and saved for a second time...

diff --git a/Engine/Source/Editor/UnrealEd/Private/Kismet2/BlueprintEditorUtils.cpp b/Engine/Source/Editor/UnrealEd/Private/Kismet2/BlueprintEditorUtils.cpp
--- a/Engine/Source/Editor/UnrealEd/Private/Kismet2/BlueprintEditorUtils.cpp
+++ b/Engine/Source/Editor/UnrealEd/Private/Kismet2/BlueprintEditorUtils.cpp
@@ -484,9 +484,11 @@ public:
 	{
 		if( TargetBlueprint != NULL && (TargetBlueprint->BlueprintType != BPTYPE_MacroLibrary) )
 		{
-			check(TargetBlueprint->SkeletonGeneratedClass);
+			if( ensureMsgf(TargetBlueprint->SkeletonGeneratedClass, TEXT("Blueprint %s is missing its skeleton generated class - known possible for assets on revision 1 older than 2088505"), *TargetBlueprint->GetName() ) )
+			{
+				TargetBlueprint->SkeletonGeneratedClass->GetDefaultObject()->Serialize(*this);
+			}
 			check(TargetBlueprint->GeneratedClass);
-			TargetBlueprint->SkeletonGeneratedClass->GetDefaultObject()->Serialize(*this);
 			TargetBlueprint->GeneratedClass->GetDefaultObject()->Serialize(*this);
 
 			TArray<UObject*> SubObjs;
diff --git a/Engine/Source/Editor/UnrealEd/Private/Kismet2/Kismet2.cpp b/Engine/Source/Editor/UnrealEd/Private/Kismet2/Kismet2.cpp
--- a/Engine/Source/Editor/UnrealEd/Private/Kismet2/Kismet2.cpp
+++ b/Engine/Source/Editor/UnrealEd/Private/Kismet2/Kismet2.cpp
@@ -193,6 +193,7 @@ UBlueprint* FKismetEditorUtilities::CreateBlueprint(UClass* ParentClass, UObject
 	NewBP->ParentClass = ParentClass;
 	NewBP->BlueprintSystemVersion = UBlueprint::GetCurrentBlueprintSystemVersion();
 	NewBP->bIsNewlyCreated = true;
+	NewBP->bLegacyNeedToPurgeSkelRefs = false;
 
 	const UEdGraphSchema_K2* K2Schema = GetDefault<UEdGraphSchema_K2>();
 
