Corrected typo in loging statment for text gathering commandlet.

diff --git a/Engine/Source/Editor/UnrealEd/Private/Commandlets/GatherTextCommandlet.cpp b/Engine/Source/Editor/UnrealEd/Private/Commandlets/GatherTextCommandlet.cpp
--- a/Engine/Source/Editor/UnrealEd/Private/Commandlets/GatherTextCommandlet.cpp
+++ b/Engine/Source/Editor/UnrealEd/Private/Commandlets/GatherTextCommandlet.cpp
@@ -1,4 +1,4 @@
-// Copyright 1998-2014 Epic Games, Inc. All Rights Reserved.
+﻿// Copyright 1998-2014 Epic Games, Inc. All Rights Reserved.
 
 #include "UnrealEd.h"
 #include "ISourceControlModule.h"
@@ -92,7 +92,7 @@ int32 UGatherTextCommandlet::Main( const FString& Params )
 		UClass* CommandletClass = FindObject<UClass>(ANY_PACKAGE,*CommandletClassName,false);
 		if (!CommandletClass)
 		{
-			UE_LOG(LogGatherTextCommandlet, Error,TEXT("The commandlet name %s in seciton %s is invalid."), *CommandletClassName, *SectionName);
+			UE_LOG(LogGatherTextCommandlet, Error,TEXT("The commandlet name %s in section %s is invalid."), *CommandletClassName, *SectionName);
 			continue;
 		}
 
