Modified localization data gathering to handle structures, arrays, and nesting when gathering from assets.

diff --git a/Engine/Source/Editor/UnrealEd/Classes/Commandlets/GatherTextFromAssetsCommandlet.h b/Engine/Source/Editor/UnrealEd/Classes/Commandlets/GatherTextFromAssetsCommandlet.h
--- a/Engine/Source/Editor/UnrealEd/Classes/Commandlets/GatherTextFromAssetsCommandlet.h
+++ b/Engine/Source/Editor/UnrealEd/Classes/Commandlets/GatherTextFromAssetsCommandlet.h
@@ -27,9 +27,8 @@ class UGatherTextFromAssetsCommandlet : public UGatherTextCommandletBase
 	GENERATED_UCLASS_BODY()
 
 	void ProcessPackages( const TArray< UPackage* >& PackagesToProcess );
-	void ProcessObject( UObject* Object, const UPackage* ObjectPackage );
 	void ProcessDialogueWave( const UDialogueWave* DialogueWave );
-	bool ProcessTextProperty( UTextProperty* TextProp, UObject* Object, const FString& ObjectPath, const bool bFixBroken, bool& OutRepaired );
+	bool ProcessTextProperty( UTextProperty* TextProp, FText* Data, UObject* Object, bool& OutRepaired );
 
 public:
 	// Begin UCommandlet Interface
diff --git a/Engine/Source/Editor/UnrealEd/Private/Commandlets/GatherTextFromAssetsCommandlet.cpp b/Engine/Source/Editor/UnrealEd/Private/Commandlets/GatherTextFromAssetsCommandlet.cpp
--- a/Engine/Source/Editor/UnrealEd/Private/Commandlets/GatherTextFromAssetsCommandlet.cpp
+++ b/Engine/Source/Editor/UnrealEd/Private/Commandlets/GatherTextFromAssetsCommandlet.cpp
@@ -347,12 +347,163 @@ UGatherTextFromAssetsCommandlet::UGatherTextFromAssetsCommandlet(const class FPo
 
 }
 
-bool UGatherTextFromAssetsCommandlet::ProcessTextProperty(UTextProperty* InTextProp, UObject* Object, const FString& ObjectPath, bool bInFixBroken, bool& OutFixed)
+namespace
+{
+	class FGatherTextFromObject
+	{
+	public:
+		static inline void Execute(UGatherTextFromAssetsCommandlet* const Commandlet, UObject* const Object, const UPackage* const ObjectPackage, const bool ShouldFixBroken)
+		{
+			if( !Object->HasAnyFlags( RF_Transient | RF_PendingKill ) )
+			{
+				FGatherTextFromObject(Commandlet, Object, ObjectPackage, ShouldFixBroken);
+			}
+		}
+
+	private:
+		FGatherTextFromObject(UGatherTextFromAssetsCommandlet* const InCommandlet, UObject* const InObject, const UPackage* const InObjectPackage, const bool InShouldFixBroken)
+			: Commandlet(InCommandlet)
+			, Object(InObject)
+			, ObjectPackage(InObjectPackage)
+			, ShouldFixBroken(InShouldFixBroken)
+		{
+			// Iterate over all fields of the object's class.
+			for (TFieldIterator<UProperty>PropIt(Object->GetClass(), EFieldIteratorFlags::IncludeSuper, EFieldIteratorFlags::ExcludeDeprecated, EFieldIteratorFlags::IncludeInterfaces); PropIt; ++PropIt)
+			{
+				ProcessProperty( *PropIt, PropIt->ContainerPtrToValuePtr<void>(Object) );
+			}
+		}
+
+		void ProcessProperty(UProperty* const Property, void* const ValueAddress)
+		{
+			// Property is a struct property.
+			UStructProperty* const StructProperty = Cast<UStructProperty>(Property);
+			if(StructProperty)
+			{
+				// Iterate over all properties of the struct's class.
+				for (TFieldIterator<UProperty>PropIt(StructProperty->Struct, EFieldIteratorFlags::IncludeSuper, EFieldIteratorFlags::ExcludeDeprecated, EFieldIteratorFlags::IncludeInterfaces); PropIt; ++PropIt)
+				{
+					ProcessProperty( *PropIt, PropIt->ContainerPtrToValuePtr<void>(ValueAddress) );
+				}
+			}
+
+			// Property is an array property.
+			UArrayProperty* const ArrayProperty = Cast<UArrayProperty>(Property);
+			if(ArrayProperty)
+			{
+				// Iterate over all objects of the array.
+				FScriptArrayHelper ScriptArrayHelper(ArrayProperty, ValueAddress);
+				const int32 ElementCount = ScriptArrayHelper.Num();
+				for(int32 i = 0; i < ElementCount; ++i)
+				{
+					ProcessProperty( ArrayProperty->Inner, ScriptArrayHelper.GetRawPtr(i) );
+				}
+			}
+
+			// Property is a text property.
+			UTextProperty* const TextProp = Cast<UTextProperty>(Property);
+			if(TextProp)
+			{
+				FText* const Text = reinterpret_cast<FText*>(ValueAddress);
+				if( FTextInspector::GetFlags(*Text) & ETextFlag::ConvertedProperty )
+				{
+					ObjectPackage->MarkPackageDirty();
+				}
+
+				bool Fixed = false;
+				if( Commandlet->ProcessTextProperty(TextProp, Text, Object, /*OUT*/Fixed ) )
+				{
+					if( Fixed )
+					{
+						ObjectPackage->MarkPackageDirty();
+					}
+				}
+			}
+		}
+
+	private:
+		UGatherTextFromAssetsCommandlet* const Commandlet;
+		UObject* const Object;
+		const UPackage* const ObjectPackage;
+		const bool ShouldFixBroken;
+	};
+}
+
+void UGatherTextFromAssetsCommandlet::ProcessPackages( const TArray< UPackage* >& PackagesToProcess )
+{
+	for( int32 i = 0; i < PackagesToProcess.Num(); ++i )
+	{
+		UPackage* Package = PackagesToProcess[i];
+		TArray<UObject*> Objects;
+		GetObjectsWithOuter(Package, Objects);
+
+		for( int32 j = 0; j < Objects.Num(); ++j )
+		{
+			UObject* Object = Objects[j];
+			if ( Object->IsA( UBlueprint::StaticClass() ) )
+			{
+				UBlueprint* Blueprint = Cast<UBlueprint>( Object );
+
+				if( Blueprint->GeneratedClass != NULL )
+				{
+					FGatherTextFromObject::Execute(this, Blueprint->GeneratedClass->GetDefaultObject(), Package, bFixBroken);
+				}
+				else
+				{
+					UE_LOG(LogGatherTextFromAssetsCommandlet, Warning, TEXT("%s - Invalid generated class!"), *Blueprint->GetFullName());
+				}
+			}
+			else if( Object->IsA( UDialogueWave::StaticClass() ) )
+			{
+				UDialogueWave* DialogueWave = Cast<UDialogueWave>( Object );
+				ProcessDialogueWave( DialogueWave );
+			}
+
+			FGatherTextFromObject::Execute(this, Object, Package, bFixBroken);
+		}
+	}
+}
+
+void UGatherTextFromAssetsCommandlet::ProcessDialogueWave( const UDialogueWave* DialogueWave )
+{
+	if( !DialogueWave || DialogueWave->HasAnyFlags( RF_Transient | RF_PendingKill ) )
+	{
+		return;
+	}
+
+	FString DialogueName = DialogueWave->GetName();
+	// Use a helper class to extract the dialogue info and prepare it for the manifest
+	FDialogueHelper DialogueHelper;
+	if( DialogueHelper.ProcessDialogueWave( DialogueWave ) )
+	{
+		const FString& SpokenSource = DialogueHelper.GetSpokenSource();
+
+		if( !SpokenSource.IsEmpty() )
+		{
+			{
+				TSharedPtr< const FContext > Base = DialogueHelper.GetBaseContext();
+
+				FString EntryDescription = FString::Printf( TEXT("In non-optional variation of %s"), *DialogueName);
+				ManifestInfo->AddEntry(EntryDescription, DialogueHelper.DialogueNamespace, SpokenSource, *Base);
+			}
+			
+			{
+				const TArray< TSharedPtr< FContext > > Variations = DialogueHelper.GetContextSpecificVariations();
+				for( TSharedPtr< FContext > Variation : Variations )
+				{
+					FString EntryDescription = FString::Printf( TEXT("In context specific variation of %s"), *DialogueName);
+					ManifestInfo->AddEntry(EntryDescription, DialogueHelper.DialogueNamespace, SpokenSource, *Variation);
+				}
+			}
+		}
+	}
+}
+
+bool UGatherTextFromAssetsCommandlet::ProcessTextProperty(UTextProperty* InTextProp, FText* Data, UObject* Object, bool& OutFixed)
 {
 	bool TextPropertyWasValid = true;
 
 	OutFixed = false;
-	FText* Data = InTextProp->ContainerPtrToValuePtr<FText>(Object);
 
 	// Transient check.
 	if( Data->Flags & ETextFlag::Transient )
@@ -363,7 +514,7 @@ bool UGatherTextFromAssetsCommandlet::ProcessTextProperty(UTextProperty* InTextP
 	else
 	{
 		FConflictTracker::FEntry NewEntry;
-		NewEntry.ObjectPath = ObjectPath;
+		NewEntry.ObjectPath = Object->GetPathName();
 		NewEntry.SourceString = Data->SourceString;
 		NewEntry.Status = EAssetTextGatherStatus::None;
 
@@ -371,7 +522,7 @@ bool UGatherTextFromAssetsCommandlet::ProcessTextProperty(UTextProperty* InTextP
 		if( !( Data->Key.IsValid() ) || Data->Key->IsEmpty() )
 		{
 			// Key fix.
-			if (bInFixBroken)
+			if (bFixBroken)
 			{
 				// Create key if needed.
 				if( !( Data->Key.IsValid() ) )
@@ -411,7 +562,7 @@ bool UGatherTextFromAssetsCommandlet::ProcessTextProperty(UTextProperty* InTextP
 				// Fix conflict if present and allowed.
 				if( ExistingEntry->Source.Text != ( Data->SourceString.IsValid() ? **(Data->SourceString) : TEXT("") ) )
 				{
-					if (bInFixBroken)
+					if (bFixBroken)
 					{
 						// Generate new GUID for key.
 						*(Data->Key) = FGuid::NewGuid().ToString();
@@ -436,7 +587,7 @@ bool UGatherTextFromAssetsCommandlet::ProcessTextProperty(UTextProperty* InTextP
 				// Check for valid string.
 				if( Data->SourceString.IsValid() && !( Data->SourceString->IsEmpty() ) )
 				{
-					FString SrcLocation = ObjectPath + TEXT(".") + InTextProp->GetName();
+					FString SrcLocation = Object->GetPathName() + TEXT(".") + InTextProp->GetName();
 
 					// Adjust the source location if needed.
 					{
@@ -480,111 +631,6 @@ bool UGatherTextFromAssetsCommandlet::ProcessTextProperty(UTextProperty* InTextP
 	return TextPropertyWasValid;
 }
 
-void UGatherTextFromAssetsCommandlet::ProcessPackages( const TArray< UPackage* >& PackagesToProcess )
-{
-	for( int32 i = 0; i < PackagesToProcess.Num(); ++i )
-	{
-		UPackage* Package = PackagesToProcess[i];
-		TArray<UObject*> Objects;
-		GetObjectsWithOuter(Package, Objects);
-
-		for( int32 j = 0; j < Objects.Num(); ++j )
-		{
-			UObject* Object = Objects[j];
-			if ( Object->IsA( UBlueprint::StaticClass() ) )
-			{
-				UBlueprint* Blueprint = Cast<UBlueprint>( Object );
-
-				if( Blueprint->GeneratedClass != NULL )
-				{
-					ProcessObject( Blueprint->GeneratedClass->GetDefaultObject(), Package );
-				}
-				else
-				{
-					UE_LOG(LogGatherTextFromAssetsCommandlet, Warning, TEXT("%s - Invalid generated class!"), *Blueprint->GetFullName());
-				}
-			}
-			else if( Object->IsA( UDialogueWave::StaticClass() ) )
-			{
-				UDialogueWave* DialogueWave = Cast<UDialogueWave>( Object );
-				ProcessDialogueWave( DialogueWave );
-			}
-
-			ProcessObject( Object, Package );
-		}
-	}
-}
-
-void UGatherTextFromAssetsCommandlet::ProcessObject(UObject* Object, const UPackage* ObjectPackage)
-{
-	// Skip transient objects and those about to be deleted
-	if( Object->HasAnyFlags( RF_Transient | RF_PendingKill ) )
-	{
-		return;
-	}
-
-	FString ObjectPath = Object->GetPathName();
-
-	for (TFieldIterator<UTextProperty> PropIt(Object->GetClass(), EFieldIteratorFlags::IncludeSuper); PropIt; ++PropIt)
-	{
-		UTextProperty* TextProp	=  Cast<UTextProperty>( *(PropIt) );
-
-		//To do - Check Source control here to make sure we can fix this asset before adding it, otherwise we'll end up with orphan text in there for assets that cant be fixed.
-
-		FText* Text = TextProp->ContainerPtrToValuePtr<FText>(Object);
-		if( Text->Flags & ETextFlag::ConvertedProperty )
-		{
-			ObjectPackage->MarkPackageDirty();
-		}
-
-		bool Fixed = false;
-		if( ProcessTextProperty(TextProp, Object, ObjectPath, bFixBroken, /*OUT*/Fixed ) )
-		{
-			if( Fixed )
-			{
-				ObjectPackage->MarkPackageDirty();
-			}
-		}
-	}
-}
-
-void UGatherTextFromAssetsCommandlet::ProcessDialogueWave( const UDialogueWave* DialogueWave )
-{
-	if( !DialogueWave || DialogueWave->HasAnyFlags( RF_Transient | RF_PendingKill ) )
-	{
-		return;
-	}
-
-	FString DialogueName = DialogueWave->GetName();
-	// Use a helper class to extract the dialogue info and prepare it for the manifest
-	FDialogueHelper DialogueHelper;
-	if( DialogueHelper.ProcessDialogueWave( DialogueWave ) )
-	{
-		const FString& SpokenSource = DialogueHelper.GetSpokenSource();
-
-		if( !SpokenSource.IsEmpty() )
-		{
-			{
-				TSharedPtr< const FContext > Base = DialogueHelper.GetBaseContext();
-
-				FString EntryDescription = FString::Printf( TEXT("In non-optional variation of %s"), *DialogueName);
-				ManifestInfo->AddEntry(EntryDescription, DialogueHelper.DialogueNamespace, SpokenSource, *Base);
-			}
-			
-			{
-				const TArray< TSharedPtr< FContext > > Variations = DialogueHelper.GetContextSpecificVariations();
-				for( TSharedPtr< FContext > Variation : Variations )
-				{
-					FString EntryDescription = FString::Printf( TEXT("In context specific variation of %s"), *DialogueName);
-					ManifestInfo->AddEntry(EntryDescription, DialogueHelper.DialogueNamespace, SpokenSource, *Variation);
-				}
-			}
-		}
-	}
-}
-
-
-
 int32 UGatherTextFromAssetsCommandlet::Main(const FString& Params)
 {
 	// Parse command line.
