Added options dialog for actor merging utility

diff --git a/Engine/Source/Developer/MeshUtilities/Private/MeshUtilities.cpp b/Engine/Source/Developer/MeshUtilities/Private/MeshUtilities.cpp
--- a/Engine/Source/Developer/MeshUtilities/Private/MeshUtilities.cpp
+++ b/Engine/Source/Developer/MeshUtilities/Private/MeshUtilities.cpp
@@ -103,7 +103,8 @@ private:
 	virtual void ShutdownModule() OVERRIDE;
 
 	virtual void MergeActors(
-		const TArray<AActor*>& SourceActors, 
+		const TArray<AActor*>& SourceActors,
+		const FMeshMergingSettings& InSettings,
 		const FString& PackageName,
 		TArray<UObject*>& 
 		OutAssetsToSync, 
@@ -2951,7 +2952,12 @@ private:
 	TArray<const FAtlasedTextureSlot*>		PackedLigthmapSlots;
 };
 
-void FMeshUtilities::MergeActors(const TArray<AActor*>& SourceActors, const FString& InPackageName, TArray<UObject*>& OutAssetsToSync, FVector& OutMergedActorLocation) const
+void FMeshUtilities::MergeActors(
+	const TArray<AActor*>& SourceActors, 
+	const FMeshMergingSettings& InSettings, 
+	const FString& InPackageName, 
+	TArray<UObject*>& OutAssetsToSync, 
+	FVector& OutMergedActorLocation) const
 {
 	TArray<UStaticMeshComponent*> ComponentsToMerge;
 	
@@ -3019,6 +3025,12 @@ void FMeshUtilities::MergeActors(const TArray<AActor*>& SourceActors, const FStr
 		return;	
 	}
 
+	// Should we use vertex colors?
+	if (!InSettings.bImportVertexColors)
+	{
+		bWithVertexColors = false;
+	}
+	
 	//For each raw mesh, re-map the material indices according to the MaterialMap
 	for (int32 MeshIndex = 0; MeshIndex < SourceMeshes.Num(); ++MeshIndex)
 	{
@@ -3033,6 +3045,9 @@ void FMeshUtilities::MergeActors(const TArray<AActor*>& SourceActors, const FStr
 
 	FRawMeshExt MergedMesh;
 
+	// Set target channel for lightmap UV
+	MergedMesh.LightMapCoordinateIndex = InSettings.TargetLightmapUVChannel;
+
 	// Pack lightmaps
 	static const uint32 MaxLightmapRes = 2048;
 	float MergedLightmapScale = 1.f;
@@ -3053,8 +3068,8 @@ void FMeshUtilities::MergeActors(const TArray<AActor*>& SourceActors, const FStr
 
 	// Use first mesh for naming and pivot
 	MergedMesh.AssetPackageName = SourceMeshes[0].AssetPackageName;
-	MergedMesh.Pivot = SourceMeshes[0].Pivot;
-		
+	MergedMesh.Pivot = InSettings.bPivotPointAtZero ? FVector::ZeroVector : SourceMeshes[0].Pivot;
+	
 	// Merge meshes into single mesh
 	for (int32 SourceMeshIdx = 0; SourceMeshIdx < SourceMeshes.Num(); ++SourceMeshIdx)
 	{
@@ -3074,7 +3089,6 @@ void FMeshUtilities::MergeActors(const TArray<AActor*>& SourceActors, const FStr
 
 		for (FVector VertexPos : SourceRawMesh.VertexPositions)
 		{
-			MergedMesh.Pivot = FVector::ZeroVector; // bake vertices in world space. TODO: optional
 			TargetRawMesh.VertexPositions.Add(VertexPos - MergedMesh.Pivot);
 		}
 						
@@ -3099,10 +3113,21 @@ void FMeshUtilities::MergeActors(const TArray<AActor*>& SourceActors, const FStr
 				FMemory::Memset(&TargetRawMesh.WedgeColors[ColorsOffset], 0xFF, ColorsNum*TargetRawMesh.WedgeColors.GetTypeSize());
 			}
 		}
-		
-		// Only first UV channel will be used 
-		TargetRawMesh.WedgeTexCoords[0].Append(SourceRawMesh.WedgeTexCoords[0]);
-		
+
+		if (InSettings.bImportAllUVChannels)
+		{
+	/*		const int32 NumChannels = ARRAY_COUNT(TargetRawMesh.WedgeTexCoords);
+			for (int32 ChannelIdx = 0; ChannelIdx < NumChannels; ++ChannelIdx)
+			{
+				int32 TargetChannelIdx+= (ChannelIdx >= MergedMesh.LightMapCoordinateIndex ?)
+			}*/
+			TargetRawMesh.WedgeTexCoords[0].Append(SourceRawMesh.WedgeTexCoords[0]);
+		}
+		else // Only first UV channel will be used 
+		{
+			TargetRawMesh.WedgeTexCoords[0].Append(SourceRawMesh.WedgeTexCoords[0]);
+		}
+				
 		// Transform lightmap UVs
 		if (MergedMesh.LightMapRes)
 		{
diff --git a/Engine/Source/Developer/MeshUtilities/Public/MeshUtilities.h b/Engine/Source/Developer/MeshUtilities/Public/MeshUtilities.h
--- a/Engine/Source/Developer/MeshUtilities/Public/MeshUtilities.h
+++ b/Engine/Source/Developer/MeshUtilities/Public/MeshUtilities.h
@@ -68,6 +68,40 @@ public:
 		) = 0;
 };
 
+/**
+ * Mesh merging settings
+ */
+struct FMeshMergingSettings
+{
+	/** */
+	bool bGnerateAtlasedLightmapUV;
+	
+	/** */
+	bool bImportVertexColors;
+
+	/** */
+	bool bImportAllUVChannels;
+
+	/** */
+	int32 TargetLightmapUVChannel;
+
+	/** */
+	bool bPivotPointAtZero;
+
+	
+	/** Default settings. */
+	FMeshMergingSettings()
+		: bGnerateAtlasedLightmapUV(true)
+		, bImportVertexColors(false)
+		, bImportAllUVChannels(false)
+		, TargetLightmapUVChannel(1)
+		, bPivotPointAtZero(false)
+	{
+	}
+};
+
+
+
 /**
  * Mesh reduction module interface.
  */
@@ -197,12 +231,14 @@ public:
 	 * and merge into signle mesh grouping them by unique materials
 	 *
 	 * @param SourceActors				List of actors to merge
+	 * @param InSettings				Settings to use
 	 * @param PackageName				Destination package name for a generated assets
 	 * @param OutAssetsToSync			Merged mesh assets
 	 * @param OutMergedActorLocation	World position of merged mesh
 	 */
 	virtual void MergeActors(
-		const TArray<AActor*>& SourceActors, 
+		const TArray<AActor*>& SourceActors,
+		const FMeshMergingSettings& InSettings,
 		const FString& PackageName, 
 		TArray<UObject*>& OutAssetsToSync, 
 		FVector& OutMergedActorLocation) const = 0;
diff --git a/Engine/Source/Editor/UnrealEd/Private/Dialogs/SMeshMergingDialog.cpp b/Engine/Source/Editor/UnrealEd/Private/Dialogs/SMeshMergingDialog.cpp
new file mode 100644
--- /dev/null
+++ b/Engine/Source/Editor/UnrealEd/Private/Dialogs/SMeshMergingDialog.cpp
@@ -0,0 +1,402 @@
+// Copyright 1998-2014 Epic Games, Inc. All Rights Reserved.
+
+#include "UnrealEd.h"
+#include "Dialogs/DlgPickAssetPath.h"
+#include "Dialogs/SMeshMergingDialog.h"
+#include "MeshUtilities.h"
+#include "ContentBrowserModule.h"
+#include "AssetRegistryModule.h"
+#include "Editor/LevelEditor/Public/LevelEditor.h"
+
+#define LOCTEXT_NAMESPACE "SMeshMergingDialog"
+
+//////////////////////////////////////////////////////////////////////////
+// SMeshMergingDialog
+
+SMeshMergingDialog::SMeshMergingDialog()
+	: bPlaceInWorld(false)
+{
+}
+
+SMeshMergingDialog::~SMeshMergingDialog()
+{
+	FLevelEditorModule& LevelEditor = FModuleManager::GetModuleChecked<FLevelEditorModule>( "LevelEditor" );
+	LevelEditor.OnActorSelectionChanged().RemoveAll(this);
+}
+
+void SMeshMergingDialog::Construct(const FArguments& InArgs)
+{
+	ParentWindow = InArgs._ParentWindow;
+	
+	FLevelEditorModule& LevelEditor = FModuleManager::GetModuleChecked<FLevelEditorModule>( "LevelEditor" );
+	LevelEditor.OnActorSelectionChanged().AddSP(this, &SMeshMergingDialog::OnActorSelectionChanged);
+
+	GenerateNewPackageName();
+
+	FMeshMergingSettings();
+
+	// Create widget layout
+	this->ChildSlot
+	[
+		SNew(SVerticalBox)
+
+		+SVerticalBox::Slot()
+		.Padding(2,2,2,4)
+		[
+			SNew(SBorder)
+			.BorderImage(FEditorStyle::GetBrush("ToolPanel.GroupBorder"))
+			[
+				SNew(SVerticalBox)
+
+				//
+				// Merging options
+				//
+				+SVerticalBox::Slot()
+				[
+					SNew(SVerticalBox)
+				
+					// Lightmaps atlasing
+					+SVerticalBox::Slot()
+					.AutoHeight()
+					.Padding(FEditorStyle::GetMargin("StandardDialog.ContentPadding"))
+					[
+						SNew(SCheckBox)
+						.Type(ESlateCheckBoxType::CheckBox)
+						.IsChecked(this, &SMeshMergingDialog::GetAtlasLightmapUV)
+						.OnCheckStateChanged(this, &SMeshMergingDialog::SetAtlasLightmapUV)
+						.Content()
+						[
+							SNew(STextBlock).Text(LOCTEXT("AtlasLightmapUVLabel", "Generate Atlased Lightmap UVs"))
+						]
+					]
+
+					// Vertex colors
+					+SVerticalBox::Slot()
+					.AutoHeight()
+					.Padding(FEditorStyle::GetMargin("StandardDialog.ContentPadding"))
+					[
+						SNew(SCheckBox)
+						.Type(ESlateCheckBoxType::CheckBox)
+						.IsChecked(this, &SMeshMergingDialog::GetImportVertexColors)
+						.OnCheckStateChanged(this, &SMeshMergingDialog::SetImportVertexColors)
+						.Content()
+						[
+							SNew(STextBlock).Text(LOCTEXT("ImportVertexColorsLabel", "Import Vertex Colors"))
+						]
+					]
+
+					// Make one section per material
+					+SVerticalBox::Slot()
+					.AutoHeight()
+					.Padding(FEditorStyle::GetMargin("StandardDialog.ContentPadding"))
+					[
+						SNew(SCheckBox)
+						.Type(ESlateCheckBoxType::CheckBox)
+						.IsChecked(this, &SMeshMergingDialog::GetOneSectionPerMaterial)
+						.OnCheckStateChanged(this, &SMeshMergingDialog::SetOneSectionPerMaterial)
+						.Content()
+						[
+							SNew(STextBlock).Text(LOCTEXT("MakeOneSectionPerMaterialLabel", "Make One Section Per Material"))
+						]
+					]
+
+					// Pivot at zero
+					+SVerticalBox::Slot()
+					.AutoHeight()
+					.Padding(FEditorStyle::GetMargin("StandardDialog.ContentPadding"))
+					[
+						SNew(SCheckBox)
+						.Type(ESlateCheckBoxType::CheckBox)
+						.IsChecked(this, &SMeshMergingDialog::GetPivotPointAtZero)
+						.OnCheckStateChanged(this, &SMeshMergingDialog::SetPivotPointAtZero)
+						.Content()
+						[
+							SNew(STextBlock).Text(LOCTEXT("PivotPointAtZeroLabel", "Pivot Point At (0,0,0)"))
+						]
+					]
+
+					// Place in world
+					+SVerticalBox::Slot()
+					.AutoHeight()
+					.Padding(FEditorStyle::GetMargin("StandardDialog.ContentPadding"))
+					[
+						SNew(SCheckBox)
+						.Type(ESlateCheckBoxType::CheckBox)
+						.IsChecked(this, &SMeshMergingDialog::GetPlaceInWorld)
+						.OnCheckStateChanged(this, &SMeshMergingDialog::SetPlaceInWorld)
+						.Content()
+						[
+							SNew(STextBlock).Text(LOCTEXT("PlaceInWorldLabel", "Place In World"))
+						]
+					]
+						
+					// Asset name and picker
+					+SVerticalBox::Slot()
+					.AutoHeight()
+					.Padding(FEditorStyle::GetMargin("StandardDialog.ContentPadding"))
+					[
+						SNew(STextBlock).Text(LOCTEXT("AssetNameLabel", "Asset Name:"))
+					]
+
+					+SVerticalBox::Slot()
+					.AutoHeight()
+					.Padding(FEditorStyle::GetMargin("StandardDialog.ContentPadding"))
+					[
+						SNew(SHorizontalBox)
+					
+						+SHorizontalBox::Slot()
+						.FillWidth(1.0)
+						.Padding(0,0,2,0)
+						.HAlign(HAlign_Fill)
+						.VAlign(VAlign_Center)
+						[
+							SNew(SEditableTextBox)
+							.Text(this, &SMeshMergingDialog::GetMergedMeshPackageName)
+							.OnTextCommitted(this, &SMeshMergingDialog::OnMergedMeshPackageNameTextCommited)
+						]
+					
+						+SHorizontalBox::Slot()
+						.AutoWidth()
+						.HAlign(HAlign_Right)
+						.VAlign(VAlign_Center)
+						[
+							SNew(SButton)
+							.OnClicked(this, &SMeshMergingDialog::OnSelectPackageNameClicked)
+							.Text(LOCTEXT("SelectPackageButton", "..."))
+						]
+					]
+				]
+			]
+		]
+
+		+SVerticalBox::Slot()
+		.AutoHeight()
+		.HAlign(HAlign_Right)
+		.VAlign(VAlign_Bottom)
+		.AutoHeight()
+		[
+			SNew(SUniformGridPanel)
+			.SlotPadding(FEditorStyle::GetMargin("StandardDialog.SlotPadding"))
+			.MinDesiredSlotWidth(FEditorStyle::GetFloat("StandardDialog.MinDesiredSlotWidth"))
+			.MinDesiredSlotHeight(FEditorStyle::GetFloat("StandardDialog.MinDesiredSlotHeight"))
+			+SUniformGridPanel::Slot(0,0)
+			[
+				SNew(SButton)
+				.Text(LOCTEXT("MergeLabel", "Merge"))
+				.HAlign(HAlign_Center)
+				.ContentPadding(FEditorStyle::GetMargin("StandardDialog.ContentPadding"))
+				.OnClicked(this, &SMeshMergingDialog::OnMergeClicked)
+			]
+			+SUniformGridPanel::Slot(1,0)
+			[
+				SNew(SButton)
+				.Text(LOCTEXT("Cancel", "Cancel"))
+				.HAlign(HAlign_Center)
+				.ContentPadding(FEditorStyle::GetMargin("StandardDialog.ContentPadding"))
+				.OnClicked(this, &SMeshMergingDialog::OnCancelClicked)
+			]
+		]
+	];
+}
+
+FReply SMeshMergingDialog::OnCancelClicked()
+{
+	ParentWindow.Get()->RequestDestroyWindow();
+	return FReply::Handled();
+}
+
+FReply SMeshMergingDialog::OnMergeClicked()
+{
+	RunMerging();
+	ParentWindow.Get()->RequestDestroyWindow();
+	return FReply::Handled();
+}
+
+ESlateCheckBoxState::Type SMeshMergingDialog::GetAtlasLightmapUV() const
+{
+	return (MergingSettings.bGnerateAtlasedLightmapUV ? ESlateCheckBoxState::Checked : ESlateCheckBoxState::Unchecked);
+}
+
+void SMeshMergingDialog::SetAtlasLightmapUV(ESlateCheckBoxState::Type NewValue)
+{
+	MergingSettings.bGnerateAtlasedLightmapUV = (ESlateCheckBoxState::Checked == NewValue);
+}
+
+ESlateCheckBoxState::Type SMeshMergingDialog::GetImportVertexColors() const
+{
+	return (MergingSettings.bImportVertexColors ? ESlateCheckBoxState::Checked : ESlateCheckBoxState::Unchecked);
+}
+
+void SMeshMergingDialog::SetImportVertexColors(ESlateCheckBoxState::Type NewValue)
+{
+	MergingSettings.bImportVertexColors = (ESlateCheckBoxState::Checked == NewValue);
+}
+
+void SMeshMergingDialog::OnActorSelectionChanged(const TArray<UObject*>& NewSelection)
+{
+	GenerateNewPackageName();
+}
+
+ESlateCheckBoxState::Type SMeshMergingDialog::GetOneSectionPerMaterial() const
+{
+	return ESlateCheckBoxState::Checked;
+}
+
+void SMeshMergingDialog::SetOneSectionPerMaterial(ESlateCheckBoxState::Type NewValue)
+{
+	// TODO: Need to add support for one mesh per material
+	// Default is one mesh with all materials (one section per material)
+}
+
+ESlateCheckBoxState::Type SMeshMergingDialog::GetPivotPointAtZero() const
+{
+	return (MergingSettings.bPivotPointAtZero ? ESlateCheckBoxState::Checked : ESlateCheckBoxState::Unchecked);
+}
+
+void SMeshMergingDialog::SetPivotPointAtZero(ESlateCheckBoxState::Type NewValue)
+{
+	MergingSettings.bPivotPointAtZero = (ESlateCheckBoxState::Checked == NewValue);
+}
+
+ESlateCheckBoxState::Type SMeshMergingDialog::GetPlaceInWorld() const
+{
+	return (bPlaceInWorld ? ESlateCheckBoxState::Checked : ESlateCheckBoxState::Unchecked);
+}
+
+void SMeshMergingDialog::SetPlaceInWorld(ESlateCheckBoxState::Type NewValue)
+{
+	bPlaceInWorld = (ESlateCheckBoxState::Checked == NewValue);
+}
+
+void SMeshMergingDialog::GenerateNewPackageName()
+{
+	USelection* SelectedActors = GEditor->GetSelectedActors();
+	MergedMeshPackageName.Empty();
+
+	// Iterate through selected actors and find first static mesh asset
+	// Use this static mesh path as destination package name for a merged mesh
+	for(FSelectionIterator Iter(*SelectedActors); Iter; ++Iter)
+	{
+		AActor* Actor = Cast<AActor>(*Iter);
+		if (Actor)
+		{
+			TArray<UStaticMeshComponent*> SMComponets; 
+			Actor->GetComponents<UStaticMeshComponent>(SMComponets);
+			for (UStaticMeshComponent* Component : SMComponets)
+			{
+				if (Component->StaticMesh)
+				{
+					MergedMeshPackageName = FPackageName::GetLongPackagePath(Component->StaticMesh->GetOutermost()->GetName());
+					MergedMeshPackageName+= FString(TEXT("/MergedMesh_")) + Component->StaticMesh->GetName();
+					break;
+				}
+			}
+		}
+
+		if (!MergedMeshPackageName.IsEmpty())
+		{
+			break;
+		}
+	}
+
+	if (MergedMeshPackageName.IsEmpty())
+	{
+		MergedMeshPackageName = FPackageName::FilenameToLongPackageName(FPaths::GameContentDir() + TEXT("MergedMesh"));
+	}
+
+	MergedMeshPackageName = MakeUniqueObjectName(NULL, UPackage::StaticClass(), *MergedMeshPackageName).ToString();
+}
+
+FText SMeshMergingDialog::GetMergedMeshPackageName() const
+{
+	return FText::FromString(MergedMeshPackageName);
+}
+
+void SMeshMergingDialog::OnMergedMeshPackageNameTextCommited(const FText& InText, ETextCommit::Type InCommitType)
+{
+	MergedMeshPackageName = InText.ToString();
+}
+
+FReply SMeshMergingDialog::OnSelectPackageNameClicked()
+{
+	TSharedRef<SDlgPickAssetPath> PickPathDlg = 
+		SNew(SDlgPickAssetPath)
+		.Title(LOCTEXT("SelectProxyPackage", "Select destination package"))
+		.DefaultAssetPath(FText::FromString(MergedMeshPackageName));
+
+	if (PickPathDlg->ShowModal() != EAppReturnType::Cancel)
+	{
+		MergedMeshPackageName = PickPathDlg->GetFullAssetPath().ToString();
+	}
+
+	return FReply::Handled();
+}
+
+void SMeshMergingDialog::RunMerging()
+{
+	IMeshUtilities& MeshUtilities = FModuleManager::Get().LoadModuleChecked<IMeshUtilities>("MeshUtilities");
+	USelection* SelectedActors = GEditor->GetSelectedActors();
+	TArray<AActor*> Actors;
+	TArray<ULevel*> UniqueLevels;
+	for (FSelectionIterator Iter(*SelectedActors); Iter; ++Iter)
+	{
+		AActor* Actor = Cast<AActor>(*Iter);
+		if (Actor)
+		{
+			Actors.Add(Actor);
+			UniqueLevels.AddUnique(Actor->GetLevel());
+		}
+	}
+
+	// This restriction is only for replacement of selected actors with merged mesh actor
+	if (UniqueLevels.Num() > 1 && bPlaceInWorld)
+	{
+		FText Message = NSLOCTEXT("UnrealEd", "FailedToMergeActorsSublevels_Msg", "The selected actors should be in the same level");
+		OpenMsgDlgInt(EAppMsgType::Ok, Message, NSLOCTEXT("UnrealEd", "FailedToMergeActors_Title", "Unable to merge actors"));
+		return;
+	}
+	
+	FMeshMergingSettings MergeSettings;
+	FVector MergedActorLocation;
+	TArray<UObject*> AssetsToSync;
+	MeshUtilities.MergeActors(Actors, MergeSettings, MergedMeshPackageName, AssetsToSync, MergedActorLocation);
+
+	if (AssetsToSync.Num())
+	{
+		FAssetRegistryModule& AssetRegistry = FModuleManager::Get().LoadModuleChecked<FAssetRegistryModule>("AssetRegistry");
+		int32 AssetCount = AssetsToSync.Num();
+		for (int32 AssetIndex = 0; AssetIndex < AssetCount; AssetIndex++)
+		{
+			AssetRegistry.AssetCreated(AssetsToSync[AssetIndex]);
+			GEditor->BroadcastObjectReimported(AssetsToSync[AssetIndex]);
+		}
+
+		//Also notify the content browser that the new assets exists
+		FContentBrowserModule& ContentBrowserModule = FModuleManager::Get().LoadModuleChecked<FContentBrowserModule>("ContentBrowser");
+		ContentBrowserModule.Get().SyncBrowserToAssets(AssetsToSync, true);
+
+		// Place new mesh in the world
+		if (bPlaceInWorld)
+		{
+			UWorld* World = UniqueLevels[0]->OwningWorld;
+			FActorSpawnParameters Params;
+			Params.OverrideLevel = UniqueLevels[0];
+			FRotator MergedActorRotation(ForceInit);
+
+			AStaticMeshActor* MergedActor = World->SpawnActor<AStaticMeshActor>(MergedActorLocation, MergedActorRotation, Params);
+			MergedActor->StaticMeshComponent->StaticMesh = Cast<UStaticMesh>(AssetsToSync[0]);
+			MergedActor->SetActorLabel(AssetsToSync[0]->GetName());
+
+			// Add source actors as children to merged actor and hide them
+			for (AActor* Actor : Actors)
+			{
+				Actor->AttachRootComponentToActor(MergedActor, NAME_None, EAttachLocation::KeepWorldPosition);
+				Actor->SetActorHiddenInGame(true);
+				Actor->SetIsTemporarilyHiddenInEditor(true);
+			}
+		}
+	}
+}
+
+
+#undef LOCTEXT_NAMESPACE
\ No newline at end of file
diff --git a/Engine/Source/Editor/UnrealEd/Private/EditorGroup.cpp b/Engine/Source/Editor/UnrealEd/Private/EditorGroup.cpp
--- a/Engine/Source/Editor/UnrealEd/Private/EditorGroup.cpp
+++ b/Engine/Source/Editor/UnrealEd/Private/EditorGroup.cpp
@@ -5,6 +5,7 @@
 #include "ScopedTransaction.h"
 #include "MainFrame.h"
 #include "Dialogs/SMeshProxyDialog.h"
+#include "Dialogs/SMeshMergingDialog.h"
 #include "Dialogs/DlgPickAssetPath.h"
 #include "MeshUtilities.h"
 #include "AssetRegistryModule.h"
@@ -206,8 +207,7 @@ void UUnrealEdEngine::edactRemoveFromGroup()
 void UUnrealEdEngine::edactMergeActors()
 {
 	IMeshUtilities& MeshUtilities = FModuleManager::Get().LoadModuleChecked<IMeshUtilities>("MeshUtilities");
-	IMeshReduction* MeshReduction = MeshUtilities.GetMeshReductionInterface();
-	if (MeshReduction && MeshReduction->IsSupported())
+	if (MeshUtilities.GetMeshMergingInterface() != nullptr)
 	{
 		TSharedPtr<IMeshProxyDialog> MeshProxyControls;
 		if(MeshProxyControls.IsValid() && MeshProxyControls->GetParentWindow().IsValid())
@@ -252,118 +252,25 @@ void UUnrealEdEngine::edactMergeActors()
 				FSlateApplication::Get().AddWindow(MeshProxyDialog);
 			}
 		}
+		
 		MeshProxyControls->GetParentWindow()->ShowWindow();
 		MeshProxyControls->MarkDirty();
-
 	}
 }
 
 void UUnrealEdEngine::edactMergeActorsByMaterials()
 {
-	IMeshUtilities& MeshUtilities = FModuleManager::Get().LoadModuleChecked<IMeshUtilities>("MeshUtilities");
-	USelection* SelectedActors = GetSelectedActors();
-	TArray<AActor*> Actors;
-	TArray<ULevel*> UniqueLevels;
-	for(FSelectionIterator Iter(*SelectedActors); Iter; ++Iter)
-	{
-		AActor* Actor = Cast<AActor>(*Iter);
-		if (Actor)
-		{
-			Actors.Add(Actor);
-			UniqueLevels.AddUnique(Actor->GetLevel());
-		}
-	}
-
-	// This restriction is only for replacement of selected actors with merged mesh actor
-	// Should be optional
-	if (UniqueLevels.Num() > 1)
-	{
-		FText Message = NSLOCTEXT("UnrealEd", "FailedToMergeActorsSublevels_Msg", "The selected actors should be in the same level");
-		OpenMsgDlgInt( EAppMsgType::Ok, Message, NSLOCTEXT("UnrealEd", "FailedToMergeActors_Title", "Unable to merge actors") );
-		return;
-	}
-
-	FString PackageName;
-
-	// TODO: optional
-	// Bring a dialog to choose destination package name
-	{
-		// Use this static mesh path as destination package name for a merged mesh
-		for (AActor* Actor : Actors)
-		{
-			TArray<UStaticMeshComponent*> SMComponets; 
-			Actor->GetComponents<UStaticMeshComponent>(SMComponets);
-			for (UStaticMeshComponent* Component : SMComponets)
-			{
-				if (Component->StaticMesh)
-				{
-					PackageName = FPackageName::GetLongPackagePath(Component->StaticMesh->GetOutermost()->GetName());
-					PackageName+= FString(TEXT("/MergedMesh_")) + Component->StaticMesh->GetName();
-					break;
-				}
-			}
-
-			if (!PackageName.IsEmpty())
-			{
-				break;
-			}
-		}
+	/** Create the window to host mesh megin widget */
+	TSharedRef<SWindow> MergeActorsWindow = SNew(SWindow)
+											.Title(FText::FromString("Merge actors"))
+											.ClientSize(FVector2D(400, 200));
 
-		if (PackageName.IsEmpty())
-		{
-			PackageName = FPackageName::FilenameToLongPackageName(FPaths::GameContentDir() + TEXT("MergedMesh"));
-		}
-	
-		TSharedRef<SDlgPickAssetPath> AssetPathDlg = 
-			SNew(SDlgPickAssetPath)
-			.Title(NSLOCTEXT("UnrealEd", "SelectDestinationPath", "Select destination path"))
-			.DefaultAssetPath(FText::FromString(PackageName));
+	/** Set the content of the window to our package dialog widget */
+	TSharedRef<SMeshMergingDialog> MergingDialog = SNew(SMeshMergingDialog)
+													.ParentWindow(MergeActorsWindow);
 
-		if (AssetPathDlg->ShowModal() == EAppReturnType::Cancel)
-		{
-			return;
-		}
+	MergeActorsWindow->SetContent(MergingDialog);
 
-		PackageName = AssetPathDlg->GetFullAssetPath().ToString();
-	}
-			
-	FVector MergedActorLocation;
-	TArray<UObject*> AssetsToSync;
-	MeshUtilities.MergeActors(Actors, PackageName, AssetsToSync, MergedActorLocation);
-		
-	if (AssetsToSync.Num())
-	{
-		FAssetRegistryModule& AssetRegistry = FModuleManager::Get().LoadModuleChecked<FAssetRegistryModule>("AssetRegistry");
-		int32 AssetCount = AssetsToSync.Num();
-		for(int32 AssetIndex = 0; AssetIndex < AssetCount; AssetIndex++ )
-		{
-			AssetRegistry.AssetCreated(AssetsToSync[AssetIndex]);
-			GEditor->BroadcastObjectReimported(AssetsToSync[AssetIndex]);
-		}
-
-		//Also notify the content browser that the new assets exists
-		FContentBrowserModule& ContentBrowserModule = FModuleManager::Get().LoadModuleChecked<FContentBrowserModule>("ContentBrowser");
-		ContentBrowserModule.Get().SyncBrowserToAssets( AssetsToSync, true );
-
-		// Place new mesh in the world
-		{
-			UWorld* World = UniqueLevels[0]->OwningWorld;
-			FActorSpawnParameters Params;
-			Params.OverrideLevel = UniqueLevels[0];
-			FRotator MergedActorRotation(ForceInit);
-
-			AStaticMeshActor* MergedActor = World->SpawnActor<AStaticMeshActor>(MergedActorLocation, MergedActorRotation, Params);
-			MergedActor->StaticMeshComponent->StaticMesh = Cast<UStaticMesh>(AssetsToSync[0]);
-			MergedActor->SetActorLabel(AssetsToSync[0]->GetName());
-
-			// Add source actors as children to merged actor and hide them
-			for (AActor* Actor : Actors)
-			{
-				Actor->AttachRootComponentToActor(MergedActor, NAME_None, EAttachLocation::KeepWorldPosition);
-				Actor->SetActorHiddenInGame(true);
-				Actor->SetIsTemporarilyHiddenInEditor(true);
-				Actor->GetRootComponent()->SetVisibility(false, true);
-			}
-		}
-	}
+	/** Show the package dialog window as a modal window */
+	GEditor->EditorAddModalWindow(MergeActorsWindow);
 }
\ No newline at end of file
