Missed file from CL# 2091686

diff --git a/Engine/Source/Editor/UnrealEd/Public/Dialogs/SMeshMergingDialog.h b/Engine/Source/Editor/UnrealEd/Public/Dialogs/SMeshMergingDialog.h
new file mode 100644
--- /dev/null
+++ b/Engine/Source/Editor/UnrealEd/Public/Dialogs/SMeshMergingDialog.h
@@ -0,0 +1,83 @@
+// Copyright 1998-2014 Epic Games, Inc. All Rights Reserved.
+#pragma once
+
+#include "Developer/MeshUtilities/Public/MeshUtilities.h"
+
+/*-----------------------------------------------------------------------------
+   SMeshMergingDialog
+-----------------------------------------------------------------------------*/
+class SMeshMergingDialog : public SCompoundWidget
+{
+public:
+	SLATE_BEGIN_ARGS(SMeshMergingDialog)
+	{
+	}
+
+	// The parent window hosting this dialog
+	SLATE_ATTRIBUTE(TSharedPtr<SWindow>, ParentWindow)
+
+	SLATE_END_ARGS()
+
+public:
+	/** **/
+	SMeshMergingDialog();
+	virtual ~SMeshMergingDialog();
+
+	/** SWidget functions */
+	void Construct(const FArguments& InArgs);
+
+private:
+	/** Called when the Cancel button is clicked */
+	FReply OnCancelClicked();
+
+	/** Called when the Merge button is clicked */
+	FReply OnMergeClicked();
+
+	/**  */
+	ESlateCheckBoxState::Type GetAtlasLightmapUV() const;
+	void SetAtlasLightmapUV(ESlateCheckBoxState::Type NewValue);
+
+	/**  */
+	ESlateCheckBoxState::Type GetImportVertexColors() const;
+	void SetImportVertexColors(ESlateCheckBoxState::Type NewValue);
+
+	/**  */
+	ESlateCheckBoxState::Type GetOneSectionPerMaterial() const;
+	void SetOneSectionPerMaterial(ESlateCheckBoxState::Type NewValue);
+
+	/**  */
+	ESlateCheckBoxState::Type GetPivotPointAtZero() const;
+	void SetPivotPointAtZero(ESlateCheckBoxState::Type NewValue);
+
+	ESlateCheckBoxState::Type GetPlaceInWorld() const;
+	void SetPlaceInWorld(ESlateCheckBoxState::Type NewValue);
+	
+	/** Called when actors selection is changed */
+	void OnActorSelectionChanged(const TArray<UObject*>& NewSelection);
+	
+	/** Generates destination package name using currently selected actors */
+	void GenerateNewPackageName();
+
+	/** Destination package name accessors */
+	FText GetMergedMeshPackageName() const;
+	void OnMergedMeshPackageNameTextCommited(const FText& InText, ETextCommit::Type InCommitType);
+	
+	/** Called when the select package name  button is clicked. Brings asset path picker dialog */
+	FReply OnSelectPackageNameClicked();
+
+	/** */
+	void RunMerging();
+
+private:
+	/** Pointer to the parent window */
+	TAttribute<TSharedPtr<SWindow>> ParentWindow;
+
+	/** Current mesh merging settings */
+	FMeshMergingSettings MergingSettings;
+
+	/** Merged mesh destination package name */
+	FString MergedMeshPackageName;
+
+	/** Whether to spawn merged actor in the world */
+	bool bPlaceInWorld;
+};
