Added SCC support for internationalization Portable Object(.po) export

diff --git a/Engine/Source/Editor/UnrealEd/Private/Commandlets/InternationalizationExportCommandlet.cpp b/Engine/Source/Editor/UnrealEd/Private/Commandlets/InternationalizationExportCommandlet.cpp
--- a/Engine/Source/Editor/UnrealEd/Private/Commandlets/InternationalizationExportCommandlet.cpp
+++ b/Engine/Source/Editor/UnrealEd/Private/Commandlets/InternationalizationExportCommandlet.cpp
@@ -1053,6 +1053,16 @@ bool UInternationalizationExportCommandlet::DoExport( const FString& SourcePath,
 						OutputFileName = DestinationPath / Filename;
 					}
 
+					if( SourceControlInfo.IsValid() )
+					{
+						FText SCCErrorText;
+						if (!SourceControlInfo->CheckOutFile(OutputFileName, SCCErrorText))
+						{
+							UE_LOG(LogInternationalizationExportCommandlet, Error, TEXT("Check out of file %s failed: %s"), *OutputFileName, *SCCErrorText.ToString());
+							return false;
+						}
+					}
+
 					//@TODO We force UTF8 at the moment but we want this to be based on the format found in the header info.
 					if( !FFileHelper::SaveStringToFile(OutputString, *OutputFileName, FFileHelper::EEncodingOptions::ForceUTF8) )
 					{
