Stop spammy cancelling of transactions during SpawnActor (in particular this resolves BeginTransaction getting called during the application of another transaction)

diff --git a/Engine/Source/Editor/UnrealEd/Private/EditorActorFolders.cpp b/Engine/Source/Editor/UnrealEd/Private/EditorActorFolders.cpp
--- a/Engine/Source/Editor/UnrealEd/Private/EditorActorFolders.cpp
+++ b/Engine/Source/Editor/UnrealEd/Private/EditorActorFolders.cpp
@@ -142,13 +142,19 @@ void FActorFolders::OnWorldSaved(uint32 SaveFlags, UWorld* World, bool bSuccess)
 
 void FActorFolders::OnLevelActorAdded(AActor* InActor)
 {
-	check(InActor && InActor->GetWorld());
+	check(InActor);
+
+	UWorld* World = InActor->GetWorld();
+	check(World);
 
-	FScopedTransaction Transaction(LOCTEXT("UndoAction_ActorAdded", "Actor Folder Added"));
+	const FName FolderPath = InActor->GetFolderPath();
 
-	if (!AddFolderToWorld(*InActor->GetWorld(), InActor->GetFolderPath()))
+	if (WillAddFolderToWorld(*World, FolderPath))
 	{
-		Transaction.Cancel();
+		FScopedTransaction Transaction(LOCTEXT("UndoAction_ActorAdded", "Actor Folder Added"));
+
+		const bool bAddedFolder = AddFolderToWorld(*World, FolderPath);
+		check(bAddedFolder);
 	}
 }
 
@@ -210,6 +216,11 @@ FActorFolderProps* FActorFolders::GetFolderProperties(UWorld& InWorld, FName InP
 	return GetOrCreateFoldersForWorld(InWorld).Folders.Find(InPath);
 }
 
+bool FActorFolders::FoldersExistForWorld(UWorld& InWorld) const
+{
+	return (TemporaryWorldFolders.Find(&InWorld) != NULL);
+}
+
 UEditorActorFolders& FActorFolders::GetOrCreateFoldersForWorld(UWorld& InWorld)
 {
 	if (UEditorActorFolders** Folders = TemporaryWorldFolders.Find(&InWorld))
@@ -426,6 +437,19 @@ bool FActorFolders::RenameFolderInWorld(UWorld& World, FName OldPath, FName NewP
 	return RenamedFolders.Num() != 0;
 }
 
+bool FActorFolders::WillAddFolderToWorld(UWorld& InWorld, FName Path) const
+{
+	bool bWillAddFolders = false;
+
+	if (!Path.IsNone() && FoldersExistForWorld(InWorld))
+	{
+		UEditorActorFolders& Folders = *TemporaryWorldFolders.FindChecked(&InWorld);
+		bWillAddFolders = !Folders.Folders.Contains(Path);
+	}
+
+	return bWillAddFolders;
+}
+
 bool FActorFolders::AddFolderToWorld(UWorld& InWorld, FName Path)
 {
 	if (!Path.IsNone())
diff --git a/Engine/Source/Editor/UnrealEd/Public/EditorActorFolders.h b/Engine/Source/Editor/UnrealEd/Public/EditorActorFolders.h
--- a/Engine/Source/Editor/UnrealEd/Public/EditorActorFolders.h
+++ b/Engine/Source/Editor/UnrealEd/Public/EditorActorFolders.h
@@ -90,6 +90,9 @@ struct UNREALED_API FActorFolders : public FGCObject
 
 private:
 
+	/** Returns true if folders have been created for the specified world */
+	bool FoldersExistForWorld(UWorld& InWorld) const;
+
 	/** Get or create a folder container for the specified world */
 	UEditorActorFolders& GetOrCreateFoldersForWorld(UWorld& InWorld);
 
@@ -118,6 +121,9 @@ private:
 	/** Remove any references to folder arrays for dead worlds */
 	void Housekeeping();
 
+	/** Returns true if calling AddFolderToWorld with the provided details would add a new folder */
+	bool WillAddFolderToWorld(UWorld& InWorld, FName Path) const;
+
 	/** Add a folder to the folder map for the specified world. Does not trigger any events. */
 	bool AddFolderToWorld(UWorld& InWorld, FName Path);
 
