Fix Transaction cancel destroying incorrect record counts

diff --git a/Engine/Source/Editor/UnrealEd/Classes/Editor/TransBuffer.h b/Engine/Source/Editor/UnrealEd/Classes/Editor/TransBuffer.h
--- a/Engine/Source/Editor/UnrealEd/Classes/Editor/TransBuffer.h
+++ b/Engine/Source/Editor/UnrealEd/Classes/Editor/TransBuffer.h
@@ -25,6 +25,9 @@ class UNREALED_API UTransBuffer
 	/** Number of actions in the current transaction */
 	int32 ActiveCount;
 
+	/** The cached count of the number of object records each time a transaction is begun */
+	TArray<int32> ActiveRecordCounts;
+
 	/** Maximum number of bytes the transaction buffer is allowed to occupy */
 	SIZE_T MaxMemory;
 
diff --git a/Engine/Source/Editor/UnrealEd/Classes/Editor/Transactor.h b/Engine/Source/Editor/UnrealEd/Classes/Editor/Transactor.h
--- a/Engine/Source/Editor/UnrealEd/Classes/Editor/Transactor.h
+++ b/Engine/Source/Editor/UnrealEd/Classes/Editor/Transactor.h
@@ -354,6 +354,7 @@ public:
 	 */
 	void GetTransactionObjects(TArray<UObject*>& Objects);
 	void RemoveRecords( int32 Count = 1 );
+	int32 GetRecordCount() const;
 
 	/**
 	 * Outputs the contents of the ObjectMap to the specified output device.
diff --git a/Engine/Source/Editor/UnrealEd/Private/EditorTransaction.cpp b/Engine/Source/Editor/UnrealEd/Private/EditorTransaction.cpp
--- a/Engine/Source/Editor/UnrealEd/Private/EditorTransaction.cpp
+++ b/Engine/Source/Editor/UnrealEd/Private/EditorTransaction.cpp
@@ -92,9 +92,14 @@ void FTransaction::FObjectRecord::Restore( FTransaction* Owner )
 	}
 }
 
+int32 FTransaction::GetRecordCount() const
+{
+	return Records.Num();
+}
+
 void FTransaction::RemoveRecords( int32 Count /* = 1  */ )
 {
-	if ( Records.Num() >= Count )
+	if ( Count > 0 && Records.Num() >= Count )
 	{
 		Records.RemoveAt( Records.Num() - Count, Count );
 
@@ -238,8 +243,8 @@ void FTransaction::Apply()
 			FBSPOps::bspBuildBounds( Model );
 			++NumModelsModified;
 		}
-		ChangedObject->PostEditUndo();
-	}
+			ChangedObject->PostEditUndo();
+		}
 	
 	// Rebuild BSP here instead of waiting for the next tick since
 	// multiple transaction events can occur in a single tick
@@ -249,7 +254,7 @@ void FTransaction::Apply()
 	}
 
 	// Flip it.
-	if( bFlip )
+	if (bFlip)
 	{
 		Inc *= -1;
 	}
@@ -326,7 +331,7 @@ void UTransBuffer::Serialize( FArchive& Ar )
 	{
 		Ar << UndoBuffer;
 	}
-	Ar << ResetReason << UndoCount << ActiveCount;
+	Ar << ResetReason << UndoCount << ActiveCount << ActiveRecordCounts;
 
 	CheckState();
 }
@@ -367,7 +372,7 @@ int32 UTransBuffer::Begin( const TCHAR* SessionContext, const FText& Description
 {
 	CheckState();
 	const int32 Result = ActiveCount;
-	if( ActiveCount++==0 )
+	if (ActiveCount++ == 0)
 	{
 		// Cancel redo buffer.
 		if( UndoCount )
@@ -385,6 +390,8 @@ int32 UTransBuffer::Begin( const TCHAR* SessionContext, const FText& Description
 		// Begin a new transaction.
 		GUndo = new(UndoBuffer)FTransaction( SessionContext, Description, 1 );
 	}
+	const int32 PriorRecordsCount = (Result > 0 ? ActiveRecordCounts[Result - 1] : 0);
+	ActiveRecordCounts.Add(UndoBuffer.Last().GetRecordCount() - PriorRecordsCount);
 	CheckState();
 	return Result;
 }
@@ -410,6 +417,7 @@ int32 UTransBuffer::End()
 #endif
 			GUndo = NULL;
 		}
+		ActiveRecordCounts.Pop();
 		CheckState();
 	}
 	return ActiveCount;
@@ -443,6 +451,7 @@ void UTransBuffer::Reset( const FText& Reason )
 	UndoCount    = 0;
 	ResetReason  = Reason;
 	ActiveCount  = 0;
+	ActiveRecordCounts.Empty();
 
 	CheckState();
 }
@@ -465,12 +474,19 @@ void UTransBuffer::Cancel( int32 StartIndex /*=0*/ )
 		}
 		else
 		{
+			int32 RecordsToKeep = 0;
+			for (int32 ActiveIndex = 0; ActiveIndex <= StartIndex; ++ActiveIndex)
+			{
+				RecordsToKeep += ActiveRecordCounts[ActiveIndex];
+			}
+
 			FTransaction& Transaction = UndoBuffer.Last();
-			Transaction.RemoveRecords(ActiveCount - StartIndex);
+			Transaction.RemoveRecords(Transaction.GetRecordCount() - RecordsToKeep);
 		}
 
 		// reset the active count
 		ActiveCount = StartIndex;
+		ActiveRecordCounts.SetNum(StartIndex);
 	}
 
 	CheckState();
@@ -652,6 +668,7 @@ void UTransBuffer::CheckState() const
 	// Validate the internal state.
 	check(UndoBuffer.Num()>=UndoCount);
 	check(ActiveCount>=0);
+	check(ActiveRecordCounts.Num() == ActiveCount);
 }
 
 
