Removed shadow variable in importing curve.

336560	SHADOW VARIABLE: Possible Bug in tangent calculation

diff --git a/Engine/Source/Editor/UnrealEd/Private/SkeletalMeshEdit.cpp b/Engine/Source/Editor/UnrealEd/Private/SkeletalMeshEdit.cpp
--- a/Engine/Source/Editor/UnrealEd/Private/SkeletalMeshEdit.cpp
+++ b/Engine/Source/Editor/UnrealEd/Private/SkeletalMeshEdit.cpp
@@ -713,7 +713,6 @@ bool UnFbx::FFbxImporter::ImportCurve(const FbxAnimCurve* FbxCurve, FFloatCurve
 				NewInterpMode = RCIM_Cubic;
 				// get tangents
 				{
-					float LeaveTangent, ArriveTangent;
 					LeaveTangent = Key.GetDataFloat(FbxAnimCurveDef::eRightSlope);
 
 					if ( KeyIndex > 0 )
