Fix opening functions in  Xcode 5.1.1 on OS X.
For atos specify the architecture and properly escape the path. When parsing the results read all the digits from the line number, not just the first. Don't abandon opening the file when only a filename is specified as atos only returns filenames, not filepaths.

diff --git a/Engine/Plugins/Developer/XCodeSourceCodeAccess/Source/XCodeSourceCodeAccess/Private/XCodeSourceCodeAccessor.cpp b/Engine/Plugins/Developer/XCodeSourceCodeAccess/Source/XCodeSourceCodeAccess/Private/XCodeSourceCodeAccessor.cpp
--- a/Engine/Plugins/Developer/XCodeSourceCodeAccess/Source/XCodeSourceCodeAccess/Private/XCodeSourceCodeAccessor.cpp
+++ b/Engine/Plugins/Developer/XCodeSourceCodeAccess/Source/XCodeSourceCodeAccess/Private/XCodeSourceCodeAccessor.cpp
@@ -81,12 +81,6 @@ bool FXCodeSourceCodeAccessor::OpenFileAtLine(const FString& FullPath, int32 Lin
 	{
 		ColumnNumber++;
 	}
-
-	if (!FPaths::FileExists(FullPath))
-	{
-		SourceCodeAccessModule.OnOpenFileFailed().Broadcast(FullPath);
-		return false;
-	}
 		 
 	if ( FModuleManager::Get().IsSolutionFilePresent() )
 	{
diff --git a/Engine/Source/Editor/UnrealEd/Private/SourceCodeNavigation.cpp b/Engine/Source/Editor/UnrealEd/Private/SourceCodeNavigation.cpp
--- a/Engine/Source/Editor/UnrealEd/Private/SourceCodeNavigation.cpp
+++ b/Engine/Source/Editor/UnrealEd/Private/SourceCodeNavigation.cpp
@@ -509,7 +509,7 @@ void FSourceCodeNavigationImpl::NavigateToFunctionSource( const FString& Functio
 							if(FunctionSymbolName == SymbolName)
 							{
 								uint64 Address = SymbolEntry.n_value;
-								FString AtoSCommand = FString::Printf(TEXT("-nowarning -d -o %s 0x%x"), *FullModulePath, Address);
+								FString AtoSCommand = FString::Printf(TEXT("-nowarning -arch x86_64 -d -o \"%s\" 0x%x"), *FullModulePath, Address);
 								int32 ReturnCode = 0;
 								FString Results;
 								FPlatformProcess::ExecProcess( TEXT("/usr/bin/atos"), *AtoSCommand, &ReturnCode, &Results, NULL );
@@ -526,7 +526,7 @@ void FSourceCodeNavigationImpl::NavigateToFunctionSource( const FString& Functio
 											int32 FileNamePos = LastIndex+1;
 											int32 FileNameLen = ColonIndex-FileNamePos;
 											FString FileName = Results.Mid(FileNamePos, FileNameLen);
-											FString LineNumber = Results.Mid(ColonIndex + 1, 1);
+											FString LineNumber = Results.Mid(ColonIndex + 1, CloseIndex-(ColonIndex + 1));
 											SourceCodeAccessor.OpenFileAtLine( FileName, FCString::Atoi(*LineNumber), 0 );
 										}
 									}
diff --git a/Engine/Source/Runtime/Core/Private/Apple/ApplePlatformStackWalk.cpp b/Engine/Source/Runtime/Core/Private/Apple/ApplePlatformStackWalk.cpp
--- a/Engine/Source/Runtime/Core/Private/Apple/ApplePlatformStackWalk.cpp
+++ b/Engine/Source/Runtime/Core/Private/Apple/ApplePlatformStackWalk.cpp
@@ -178,7 +178,7 @@ void FApplePlatformStackWalk::ProgramCounterToSymbolInfo( uint64 ProgramCounter,
 			FCStringAnsi::Sprintf(FBase, "%p", (void*)DylibInfo.dli_fbase);
 			FCStringAnsi::Sprintf(Address, "%p", (void*)ProgramCounter);
 			
-			execl("/usr/bin/atos", "-nowarning", "-d", "-l", FBase, "-o", DylibInfo.dli_fname, Address, NULL);
+			execl("/usr/bin/atos", "-nowarning", "-arch", "x86_64", "-d", "-l", FBase, "-o", DylibInfo.dli_fname, Address, NULL);
 			close(FileDesc[1]);
 		}
 		else
