Transaction annotations framework.
Objects can now provide annotation data beyond what is serialized from within their properties via the virtual function GetTransactionAnnotation.
This annotation data, when valid, is returned through a new PostEditUndo virtual function to be applied.

diff --git a/Engine/Source/Editor/UnrealEd/Classes/Editor/Transactor.h b/Engine/Source/Editor/UnrealEd/Classes/Editor/Transactor.h
--- a/Engine/Source/Editor/UnrealEd/Classes/Editor/Transactor.h
+++ b/Engine/Source/Editor/UnrealEd/Classes/Editor/Transactor.h
@@ -73,6 +73,8 @@ protected:
 		TArray<FName>		ReferencedNames;
 		/** The object to track */
 		UObject*			Object;
+		/** Annotation data for the object stored externally */
+		TSharedPtr<ITransactionObjectAnnotation> ObjectAnnotation;
 		/** Array: If an array object, reference to script array */
 		FScriptArray*		Array;
 		/** Array: Offset into the array */
@@ -98,27 +100,7 @@ protected:
 		// Constructors.
 		FObjectRecord()
 		{}
-		FObjectRecord( FTransaction* Owner, UObject* InObject, FScriptArray* InArray, int32 InIndex, int32 InCount, int32 InOper, int32 InElementSize, STRUCT_DC InDefaultConstructor, STRUCT_AR InSerializer, STRUCT_DTOR InDestructor )
-			:	Object						( InObject )
-			,	Array						( InArray )
-			,	Index						( InIndex )
-			,	Count						( InCount )
-			,	Oper						( InOper )
-			,	ElementSize					( InElementSize )
-			,	DefaultConstructor			( InDefaultConstructor )
-			,	Serializer					( InSerializer )
-			,	Destructor					( InDestructor )
-			,	bRestored					( false )
-			,	bWantsBinarySerialization	( true )
-		{
-			// Blueprint compile-in-place can alter class layout so use tagged serialization for objects relying on a UBlueprint's Class
-			if (UBlueprintGeneratedClass* Class = Cast<UBlueprintGeneratedClass>(InObject->GetClass()))
-			{
-				bWantsBinarySerialization = false; 
-			}
-			FWriter Writer( Data, ReferencedObjects, ReferencedNames, bWantsBinarySerialization );
-			SerializeContents( Writer, Oper );
-		}
+		FObjectRecord( FTransaction* Owner, UObject* InObject, FScriptArray* InArray, int32 InIndex, int32 InCount, int32 InOper, int32 InElementSize, STRUCT_DC InDefaultConstructor, STRUCT_AR InSerializer, STRUCT_DTOR InDestructor );
 
 		// Functions.
 		void SerializeContents( FArchive& Ar, int32 InOper );
diff --git a/Engine/Source/Editor/UnrealEd/Private/EditorTransaction.cpp b/Engine/Source/Editor/UnrealEd/Private/EditorTransaction.cpp
--- a/Engine/Source/Editor/UnrealEd/Private/EditorTransaction.cpp
+++ b/Engine/Source/Editor/UnrealEd/Private/EditorTransaction.cpp
@@ -10,6 +10,29 @@ DEFINE_LOG_CATEGORY_STATIC(LogEditorTransaction, Log, All);
 	A single transaction.
 -----------------------------------------------------------------------------*/
 
+FTransaction::FObjectRecord::FObjectRecord(FTransaction* Owner, UObject* InObject, FScriptArray* InArray, int32 InIndex, int32 InCount, int32 InOper, int32 InElementSize, STRUCT_DC InDefaultConstructor, STRUCT_AR InSerializer, STRUCT_DTOR InDestructor)
+	:	Object				( InObject )
+	,	Array				( InArray )
+	,	Index				( InIndex )
+	,	Count				( InCount )
+	,	Oper				( InOper )
+	,	ElementSize			( InElementSize )
+	,	DefaultConstructor	( InDefaultConstructor )
+	,	Serializer			( InSerializer )
+	,	Destructor			( InDestructor )
+	,	bRestored			( false )
+	,	bWantsBinarySerialization ( true )
+{
+	// Blueprint compile-in-place can alter class layout so use tagged serialization for objects relying on a UBlueprint's Class
+	if (UBlueprintGeneratedClass* Class = Cast<UBlueprintGeneratedClass>(InObject->GetClass()))
+	{
+		bWantsBinarySerialization = false; 
+	}
+	ObjectAnnotation = Object->GetTransactionAnnotation();
+	FWriter Writer( Data, ReferencedObjects, ReferencedNames, bWantsBinarySerialization );
+	SerializeContents( Writer, Oper );
+}
+
 void FTransaction::FObjectRecord::SerializeContents( FArchive& Ar, int32 InOper )
 {
 	if( Array )
@@ -75,8 +98,10 @@ void FTransaction::FObjectRecord::Restore( FTransaction* Owner )
 		TArray<uint8> FlipData;
 		TArray<UObject*> FlipReferencedObjects;
 		TArray<FName> FlipReferencedNames;
+		TSharedPtr<ITransactionObjectAnnotation> FlipObjectAnnotation;
 		if( Owner->bFlip )
 		{
+			FlipObjectAnnotation = Object->GetTransactionAnnotation();
 			FWriter Writer( FlipData, FlipReferencedObjects, FlipReferencedNames, bWantsBinarySerialization );
 			SerializeContents( Writer, -Oper );
 		}
@@ -84,6 +109,7 @@ void FTransaction::FObjectRecord::Restore( FTransaction* Owner )
 		SerializeContents( Reader, Oper );
 		if( Owner->bFlip )
 		{
+			Exchange( ObjectAnnotation, FlipObjectAnnotation );
 			Exchange( Data, FlipData );
 			Exchange( ReferencedObjects, FlipReferencedObjects );
 			Exchange( ReferencedNames, FlipReferencedNames );
@@ -217,16 +243,19 @@ void FTransaction::Apply()
 	const int32 End   = Inc==1 ? Records.Num() :              -1;
 
 	// Init objects.
-	TArray<UObject*> ChangedObjects;
+	TMap<UObject*, TSharedPtr<ITransactionObjectAnnotation>> ChangedObjects;
 	for( int32 i=Start; i!=End; i+=Inc )
 	{
-		Records[i].bRestored = false;
-		if(ChangedObjects.Find(Records[i].Object) == INDEX_NONE)
+		FObjectRecord& Record = Records[i];
+		Record.bRestored = false;
+
+		if (!ChangedObjects.Contains(Record.Object))
 		{
-			Records[i].Object->CheckDefaultSubobjects();
-			Records[i].Object->PreEditUndo();
-			ChangedObjects.Add(Records[i].Object);
+			Record.Object->CheckDefaultSubobjects();
+			Record.Object->PreEditUndo();
 		}
+
+		ChangedObjects.Add(Record.Object, Record.ObjectAnnotation);
 	}
 	for( int32 i=Start; i!=End; i+=Inc )
 	{
@@ -234,17 +263,25 @@ void FTransaction::Apply()
 	}
 
 	NumModelsModified = 0;		// Count the number of UModels that were changed.
-	for(int32 ObjectIndex = 0;ObjectIndex < ChangedObjects.Num();ObjectIndex++)
+	for (auto ChangedObjectIt : ChangedObjects)
 	{
-		UObject* ChangedObject = ChangedObjects[ObjectIndex];
+		UObject* ChangedObject = ChangedObjectIt.Key;
 		UModel* Model = Cast<UModel>(ChangedObject);
-		if( Model && Model->Nodes.Num() )
+		if (Model && Model->Nodes.Num())
 		{
-			FBSPOps::bspBuildBounds( Model );
+			FBSPOps::bspBuildBounds(Model);
 			++NumModelsModified;
 		}
+		TSharedPtr<ITransactionObjectAnnotation> ChangedObjectTransactionAnnotation = ChangedObjectIt.Value;
+		if (ChangedObjectTransactionAnnotation.IsValid())
+		{
+			ChangedObject->PostEditUndo(ChangedObjectTransactionAnnotation);
+		}
+		else
+		{
 			ChangedObject->PostEditUndo();
 		}
+	}
 	
 	// Rebuild BSP here instead of waiting for the next tick since
 	// multiple transaction events can occur in a single tick
@@ -258,9 +295,9 @@ void FTransaction::Apply()
 	{
 		Inc *= -1;
 	}
-	for(int32 ObjectIndex = 0;ObjectIndex < ChangedObjects.Num();ObjectIndex++)
+	for (auto ChangedObjectIt : ChangedObjects)
 	{
-		UObject* ChangedObject = ChangedObjects[ObjectIndex];
+		UObject* ChangedObject = ChangedObjectIt.Key;
 		ChangedObject->CheckDefaultSubobjects();
 	}
 }
diff --git a/Engine/Source/Runtime/Core/Public/Misc/ITransaction.h b/Engine/Source/Runtime/Core/Public/Misc/ITransaction.h
--- a/Engine/Source/Runtime/Core/Public/Misc/ITransaction.h
+++ b/Engine/Source/Runtime/Core/Public/Misc/ITransaction.h
@@ -11,6 +11,14 @@ typedef void( *STRUCT_DC )( void* TPtr );						// default construct
 typedef void( *STRUCT_AR )( class FArchive& Ar, void* TPtr );	// serialize
 typedef void( *STRUCT_DTOR )( void* TPtr );						// destruct
 
+
+/**
+ * Abstract base class for transaction object annotations
+ */
+class ITransactionObjectAnnotation
+{
+};
+
 /**
  * Interface for transactions.
  */
diff --git a/Engine/Source/Runtime/CoreUObject/Private/UObject/Obj.cpp b/Engine/Source/Runtime/CoreUObject/Private/UObject/Obj.cpp
--- a/Engine/Source/Runtime/CoreUObject/Private/UObject/Obj.cpp
+++ b/Engine/Source/Runtime/CoreUObject/Private/UObject/Obj.cpp
@@ -337,6 +337,11 @@ void UObject::PostEditUndo()
 	}
 }
 
+void UObject::PostEditUndo(TSharedPtr<ITransactionObjectAnnotation> TransactionAnnotation)
+{
+	UObject::PostEditUndo();
+}
+
 #endif // WITH_EDITOR
 
 bool UObject::CanCreateInCurrentContext(UObject* Template)
diff --git a/Engine/Source/Runtime/CoreUObject/Public/UObject/UObject.h b/Engine/Source/Runtime/CoreUObject/Public/UObject/UObject.h
--- a/Engine/Source/Runtime/CoreUObject/Public/UObject/UObject.h
+++ b/Engine/Source/Runtime/CoreUObject/Public/UObject/UObject.h
@@ -217,11 +217,19 @@ public:
 	 */
 	virtual void PostEditChangeChainProperty( struct FPropertyChangedChainEvent& PropertyChangedEvent );
 
+	/** Gathers external data required for applying an undo transaction */
+	virtual TSharedPtr<ITransactionObjectAnnotation> GetTransactionAnnotation() const { return NULL; }
+
 	/** Called before applying a transaction to the object.  Default implementation simply calls PreEditChange. */
 	virtual void PreEditUndo();
 
 	/** Called after applying a transaction to the object.  Default implementation simply calls PostEditChange. */
 	virtual void PostEditUndo();
+
+	/** Called after applying a transaction to the object in cases where transaction annotation was provided. Default implementation simply calls PostEditChange. */
+	virtual void PostEditUndo(TSharedPtr<ITransactionObjectAnnotation> TransactionAnnotation);
+
+
 #endif // WITH_EDITOR
 
 	// @todo document
