Fixed crash when importing animation with length of zero.

TTP 336115	CRITICAL: EDITOR: ANIMATION: CRASH: Importing FBX using Animated Time option

diff --git a/Engine/Source/Editor/UnrealEd/Private/SkeletalMeshEdit.cpp b/Engine/Source/Editor/UnrealEd/Private/SkeletalMeshEdit.cpp
--- a/Engine/Source/Editor/UnrealEd/Private/SkeletalMeshEdit.cpp
+++ b/Engine/Source/Editor/UnrealEd/Private/SkeletalMeshEdit.cpp
@@ -289,6 +289,12 @@ bool UnFbx::FFbxImporter::IsValidAnimationData(TArray<FbxNode*>& SortedLinks, TA
 		}
 
 		FbxTimeSpan AnimTimeSpan = GetAnimationTimeSpan(SortedLinks[0], CurAnimStack);
+		if (AnimTimeSpan.GetDuration() <= 0)
+		{
+			AddTokenizedErrorMessage(FTokenizedMessage::Create(EMessageSeverity::Warning, FText::Format(LOCTEXT("FBXImport_ZeroLength", "Animation Stack {0} does not contain any valid key. Try different time options when import."), FText::FromString(CurAnimStack->GetName()))));
+			continue;
+		}
+
 		ValidTakeCount++;
 		{
 			bool bBlendCurveFound = false;
@@ -461,6 +467,7 @@ UAnimSequence * UnFbx::FFbxImporter::ImportAnimations(USkeleton * Skeleton, UObj
 	int32 ValidTakeCount = 0;
 	if (IsValidAnimationData(SortedLinks, NodeArray, ValidTakeCount) == false)
 	{
+		AddTokenizedErrorMessage(FTokenizedMessage::Create(EMessageSeverity::Warning, LOCTEXT("FBXImport_InvalidAnimationData", "This does not contain any valid animation takes.")));
 		return NULL;
 	}
 
@@ -630,6 +637,12 @@ bool UnFbx::FFbxImporter::ValidateAnimStack(TArray<FbxNode*>& SortedLinks, TArra
 
 	AnimTimeSpan = GetAnimationTimeSpan(SortedLinks[0], CurAnimStack);
 	
+	// if no duration is found, return false
+	if (AnimTimeSpan.GetDuration() <= 0)
+	{
+		return false;
+	}
+
 	const FBXImportOptions* ImportOption = GetImportOptions();
 	// only add morph time if not setrange. If Set Range there is no reason to override time
 	if ( bImportMorph && ImportOption->AnimationLengthImportType != FBXALIT_SetRange)
