Restore open assets notifications restore button only operates on the assets on the first click.
#TTP 335545 - EDITOR: Can click buttons of restore open asset toast multiple times

#branch UE4

#proj Editor.UnrealEd

#change Moved the OpenEditorForAsset operation to within the conditional check for the notification pointer, this means that although the button remains clickable it only does anything once.

diff --git a/Engine/Source/Editor/UnrealEd/Private/Toolkits/AssetEditorManager.cpp b/Engine/Source/Editor/UnrealEd/Private/Toolkits/AssetEditorManager.cpp
--- a/Engine/Source/Editor/UnrealEd/Private/Toolkits/AssetEditorManager.cpp
+++ b/Engine/Source/Editor/UnrealEd/Private/Toolkits/AssetEditorManager.cpp
@@ -460,11 +460,12 @@ void FAssetEditorManager::OnConfirmRestorePreviouslyOpenAssets(TArray<FString> A
 		UEditorLoadingSavingSettings& Settings = *GetMutableDefault<UEditorLoadingSavingSettings>();
 		Settings.bRestoreOpenAssetTabsOnRestart = bSuppressNotification;
 		Settings.PostEditChange();
-	}
 
-	for( const FString& AssetName : AssetsToOpen )
-	{
-		OpenEditorForAsset(AssetName);
+		// we do this inside the condition so that it can only be done once. 
+		for (const FString& AssetName : AssetsToOpen)
+		{
+			OpenEditorForAsset(AssetName);
+		}
 	}
 }
 
