General improvements to network PIE game launching. Added waiting to connect message for editor. Fixed number of spawned instances. Improved window positions to accommodate multiple monitors.

diff --git a/Engine/Source/Editor/UnrealEd/Classes/Editor/EditorEngine.h b/Engine/Source/Editor/UnrealEd/Classes/Editor/EditorEngine.h
--- a/Engine/Source/Editor/UnrealEd/Classes/Editor/EditorEngine.h
+++ b/Engine/Source/Editor/UnrealEd/Classes/Editor/EditorEngine.h
@@ -2554,8 +2554,15 @@ private:
 
 protected:
 
-	// @todo gmp: temp hack for Rocket demo
-	void PlayOnLocalPc(FString MapNameOverride=FString(), FString URLParms=FString(), FString CmdLineParms=FString(), int32 ResX=0, int32 ResY=0);
+	/**
+	 * Launch a standalone instance on this PC.
+	 *
+	 * @param	MapNameOverride		Map name override
+	 * @param	WindowPos			Postion we want to put the window we create.
+	 * @param	PIENum				PIE instance count
+	 * @param	bIsServer			Is this instance a server.
+	 */
+	void PlayStandaloneLocalPc(FString MapNameOverride = FString(), FIntPoint* WindowPos = NULL, int32 PIENum = 0, bool bIsServer = false);
 
 	void PlayUsingLauncher();
 
diff --git a/Engine/Source/Editor/UnrealEd/Private/PlayLevel.cpp b/Engine/Source/Editor/UnrealEd/Private/PlayLevel.cpp
--- a/Engine/Source/Editor/UnrealEd/Private/PlayLevel.cpp
+++ b/Engine/Source/Editor/UnrealEd/Private/PlayLevel.cpp
@@ -623,56 +623,59 @@ void UEditorEngine::PlaySessionSingleStepped()
 	FEditorDelegates::SingleStepPIE.Broadcast(bIsSimulatingInEditor);
 }
 
-void AdvanceWindowPositionsForNextPIEWindow(int32 &WinX, int32 &WinY, int32 Width, int32 Height)
+void AdvanceWindowPositionsForNextPIEWindow(int32 &WinX, int32 &WinY, FIntPoint WinSize)
 {
-	const FIntPoint WindowPadding(0, 0);//(15, 32);
+	const FIntPoint WindowPadding(16, 16);
 
-	WinX += Width + WindowPadding.X;
+	WinX += WinSize.X + WindowPadding.X;
 
-	FSlateRect PreferredWorkArea = FSlateApplication::Get().GetPreferredWorkArea();
-	
+	FDisplayMetrics DisplayMetrics;
+	FSlateApplication::Get().GetDisplayMetrics(DisplayMetrics);
+	FSlateRect PreferredWorkArea(DisplayMetrics.VirtualDisplayRect.Left, DisplayMetrics.VirtualDisplayRect.Top, DisplayMetrics.VirtualDisplayRect.Right, DisplayMetrics.VirtualDisplayRect.Bottom);
 	// if no more windows fit horizontally, place them in a new row
-	if ((WinX + WindowPadding.X + Width) > PreferredWorkArea.Right)
+	if ((WinX + WindowPadding.X + WinSize.X) > PreferredWorkArea.Right)
 	{
-		WinX = 0;
-		WinY += Height + WindowPadding.X;
+		WinX = PreferredWorkArea.Left + WindowPadding.X;
+		WinY += WindowPadding.Y;
 	}
 
 	// if no more rows fit vertically, stack windows on top of each other
-	if (WinY + WindowPadding.Y + Height >= PreferredWorkArea.Bottom)
+	if (WinY + WindowPadding.Y + WinSize.Y >= PreferredWorkArea.Bottom)
 	{
 		WinX += WindowPadding.X;
-		WinY = WindowPadding.Y;
+		WinY = PreferredWorkArea.Top + WindowPadding.Y;
 	}
 }
 
-void GetMultipleInstancePositions(int32 index, int32 &LastX, int32 &LastY)
+/* returns the size of the window depending of the net mode. */
+void GetWindowSizeForInstanceType(FIntPoint &WindowSize, const ULevelEditorPlaySettings* PlayInSettings)
 {
-	ULevelEditorPlaySettings* PlayInSettings = Cast<ULevelEditorPlaySettings>(ULevelEditorPlaySettings::StaticClass()->GetDefaultObject());
-
-	if (PlayInSettings->MultipleInstancePositions.IsValidIndex(index) &&
-		(PlayInSettings->MultipleInstanceLastHeight == PlayInSettings->NewWindowHeight) &&
-		(PlayInSettings->MultipleInstanceLastWidth == PlayInSettings->NewWindowWidth))
+	if (PlayInSettings->PlayNetMode == PIE_Standalone)
 	{
-		PlayInSettings->NewWindowPosition =	PlayInSettings->MultipleInstancePositions[index];
-
-		LastX = PlayInSettings->NewWindowPosition.X;
-		LastY = PlayInSettings->NewWindowPosition.Y;
+		WindowSize.X = PlayInSettings->StandaloneWindowWidth;
+		WindowSize.Y = PlayInSettings->StandaloneWindowHeight;
 	}
 	else
 	{
-		PlayInSettings->NewWindowPosition = FIntPoint(LastX, LastY);
+		WindowSize.X = PlayInSettings->ClientWindowWidth;
+		WindowSize.Y = PlayInSettings->ClientWindowHeight;
 	}
-
-	AdvanceWindowPositionsForNextPIEWindow(LastX, LastY, PlayInSettings->NewWindowWidth, PlayInSettings->NewWindowHeight);
 }
 
+/* 
+ * Generate the command line for pie instance. Window position, size etc. 
+ *
+ * @param	WinX			Window X position. This will contain the X position to use for the next window. (Not changed for dedicated server window).
+ * @param	WinY			Window Y position. This will contain the X position to use for the next window. (Not changed for dedicated server window).
+ * @param	InstanceNum		PIE instance index.
+ * @param	IsServer		Is this instance a dedicated server. true if so else false.
+ */
 FString GenerateCmdLineForNextPieInstance(int32 &WinX, int32 &WinY, int32 &InstanceNum, bool IsServer)
 {
+	const ULevelEditorPlaySettings* PlayInSettings = GetDefault<ULevelEditorPlaySettings>();
 	// Get GameSettings INI override
 	FString GameUserSettingsOverride = GGameUserSettingsIni.Replace(TEXT("GameUserSettings"), *FString::Printf(TEXT("PIEGameUserSettings%d"), InstanceNum++));
-	const ULevelEditorPlaySettings* PlayInSettings = GetDefault<ULevelEditorPlaySettings>();
-
+	
 	// Construct parms:
 	//	-Override GameUserSettings.ini
 	//	-Force no steam
@@ -686,51 +689,78 @@ FString GenerateCmdLineForNextPieInstance(int32 &WinX, int32 &WinY, int32 &Insta
 	}
 	else
 	{
-		// hack: fudge window by WindowBorderSize, so it doesn't appear off-screen
-		FMargin WindowBorderSize(8.0f, 32.0f);
-		TSharedPtr<SWindow> SomeWindow = FSlateApplication::Get().GetActiveTopLevelWindow();
+		// Adjust window by WindowBorderSize and title bar so it doesn't appear off-screen
+		FMargin WindowBorderSize(16.0f, 80.0f);
+		TSharedPtr<SWindow> TopLevelWindow = FSlateApplication::Get().GetActiveTopLevelWindow();
 
-		if (SomeWindow.IsValid())
+		if (TopLevelWindow.IsValid())
 		{
-			WindowBorderSize = SomeWindow->GetWindowBorderSize();
+			WindowBorderSize = TopLevelWindow->GetWindowBorderSize();
 		}
 
 		// Listen server or clients: specify default win position and SAVEWINPOS so the final positions are saved
 		// in order to preserve PIE networking window setup
-		CmdLine += FString::Printf(TEXT("WinX=%d WinY=%d SAVEWINPOS=1 "), WinX + WindowBorderSize.Left, WinY + WindowBorderSize.Top);
+		int32 WindowXPos = WinX + WindowBorderSize.Left;
+		int32 WindowYPos = WinY + WindowBorderSize.Top;
+		CmdLine += FString::Printf(TEXT("WinX=%d WinY=%d SAVEWINPOS=1"), WindowXPos, WindowYPos);
+
+		// Get the size of the window based on the type
+		FIntPoint WinSize(0,0);
+		GetWindowSizeForInstanceType(WinSize, PlayInSettings);
 
 		// Advance window
-		AdvanceWindowPositionsForNextPIEWindow(WinX, WinY, PlayInSettings->NewWindowWidth, PlayInSettings->NewWindowHeight );
+		AdvanceWindowPositionsForNextPIEWindow(WinX, WinY, WinSize);
 	}
 	
 	return CmdLine;
 }
 
+void GetMultipleInstancePositions(int32 index, int32 &LastX, int32 &LastY)
+{
+	ULevelEditorPlaySettings* PlayInSettings = Cast<ULevelEditorPlaySettings>(ULevelEditorPlaySettings::StaticClass()->GetDefaultObject());
+
+	if (PlayInSettings->MultipleInstancePositions.IsValidIndex(index) &&
+		(PlayInSettings->MultipleInstanceLastHeight == PlayInSettings->NewWindowHeight) &&
+		(PlayInSettings->MultipleInstanceLastWidth == PlayInSettings->NewWindowWidth))
+	{
+		PlayInSettings->NewWindowPosition = PlayInSettings->MultipleInstancePositions[index];
+
+		LastX = PlayInSettings->NewWindowPosition.X;
+		LastY = PlayInSettings->NewWindowPosition.Y;
+	}
+	else
+	{
+		PlayInSettings->NewWindowPosition = FIntPoint(LastX, LastY);
+	}
+
+	FIntPoint WinSize(0, 0);
+	GetWindowSizeForInstanceType(WinSize, PlayInSettings);
+	AdvanceWindowPositionsForNextPIEWindow(LastX, LastY, WinSize);
+}
+
 void UEditorEngine::StartQueuedPlayMapRequest()
 {
 	const bool bWantSimulateInEditor = bIsSimulateInEditorQueued;
 
 	EndPlayOnLocalPc();
 
-	const ULevelEditorPlaySettings* PlayInSettings = GetDefault<ULevelEditorPlaySettings>();
-	
-	// Launch multiplayer instances if necessary
+	ULevelEditorPlaySettings* PlayInSettings = GetMutableDefault<ULevelEditorPlaySettings>();
+	// Launch multi-player instances if necessary
 	// (note that if you have 'RunUnderOneProcess' checked and do a bPlayOnLocalPcSession (standalone) - play standalone 'wins' - multiple instances will be launched for multiplayer)
 	if (PlayInSettings->PlayNetMode != PIE_Standalone && (!PlayInSettings->RunUnderOneProcess || bPlayOnLocalPcSession) && !bPlayUsingLauncher)
 	{
-		// Launch multiple instances for pie networking here
-		int32 WinX=0;
-		int32 WinY=0;
-		int32 PIENum = 0;
 		int32 NumClients = 0;
 
+		// If we start to the right of the editor work area, call advance it will find the next place we can place a new instance window
+		FSlateRect PreferredWorkArea = FSlateApplication::Get().GetPreferredWorkArea();		
+		FIntPoint WinPosition((int32)PreferredWorkArea.Right, (int32)PreferredWorkArea.Top);
+		AdvanceWindowPositionsForNextPIEWindow(WinPosition.X, WinPosition.Y, FIntPoint(0,0));
+
 		// Do we need to spawn a server?
 		if (PlayInSettings->PlayNetMode == PIE_Client)
-		{
-			FString CmdLine = GenerateCmdLineForNextPieInstance(WinX, WinY, PIENum, true);
-			FString URLOptions = TEXT("?Listen");
-			PlayOnLocalPc(TEXT(""), URLOptions, CmdLine, PlayInSettings->ClientWindowWidth, PlayInSettings->ClientWindowHeight);
-
+		{			
+			PlayStandaloneLocalPc(TEXT(""), &WinPosition, NumClients, true);
+			
 			if (!PlayInSettings->PlayNetDedicated)
 			{
 				// Listen server counts as a client
@@ -738,37 +768,31 @@ void UEditorEngine::StartQueuedPlayMapRequest()
 			}
 		}
 		
-		if(!bPlayOnLocalPcSession && !bStartMovieCapture && PlayInSettings->PlayNetMode == PIE_Client)
+		if (!bPlayOnLocalPcSession && !bStartMovieCapture && PlayInSettings->PlayNetMode == PIE_Client)
 		{
 			// Editor counts as a client
 			NumClients++;
 		}
 
-		// Spawn number of clients
-		for (int32 i=NumClients; i < PlayInSettings->PlayNumberOfClients; ++i)
-		{
-			FString CmdLine = GenerateCmdLineForNextPieInstance(WinX, WinY, PIENum, false);
-			PlayOnLocalPc(TEXT("127.0.0.1"), TEXT(""), CmdLine, PlayInSettings->ClientWindowWidth, PlayInSettings->ClientWindowHeight );
-		}
-
-		// Launch the local PIE session
 		if (bPlayOnLocalPcSession)
 		{
-			FString CmdLine = GenerateCmdLineForNextPieInstance(WinX, WinY, PIENum, PlayInSettings->PlayNetMode == PIE_ListenServer);
-
-			PlayOnLocalPc(
-				PlayInSettings->PlayNetMode == PIE_Client ? TEXT("127.0.0.1") : TEXT(""),
-				PlayInSettings->PlayNetMode == PIE_ListenServer ? TEXT("?Listen") : TEXT(""),
-				CmdLine);
+			// Spawn number of clients
+			for (int32 i = NumClients; i < PlayInSettings->PlayNumberOfClients; ++i)
+			{
+				PlayStandaloneLocalPc(TEXT("127.0.0.1"), &WinPosition, NumClients + 1, false);
+			}
 		}
-		else if (bStartMovieCapture)
+		else 
 		{
-			// @todo Fix for UE4.  This is a temp workaround. 
-			PlayForMovieCapture();
-		}
-		else
-		{
-			PlayInEditor( GetEditorWorldContext().World(), bWantSimulateInEditor );
+			if (bStartMovieCapture)
+			{
+				// @todo Fix for UE4.  This is a temp workaround. 
+				PlayForMovieCapture();
+			}
+			else
+			{
+				PlayInEditor(GetEditorWorldContext().World(), bWantSimulateInEditor);
+			}
 		}
 	}
 	else
@@ -776,7 +800,7 @@ void UEditorEngine::StartQueuedPlayMapRequest()
 		// Launch standalone PIE session
 		if (bPlayOnLocalPcSession)
 		{
-			PlayOnLocalPc();
+			PlayStandaloneLocalPc();
 		}
 		else if (bPlayUsingLauncher)
 		{
@@ -933,8 +957,22 @@ void UEditorEngine::EndPlayOnLocalPc( )
 }
 
 // @todo gmp: temp hack for Rocket demo
-void UEditorEngine::PlayOnLocalPc(FString MapNameOverride, FString URLParms, FString CmdLineParms, int32 ResX, int32 ResY)
+void UEditorEngine::PlayStandaloneLocalPc(FString MapNameOverride, FIntPoint* WindowPos, int32 PIENum, bool bIsServer)
 {
+	const ULevelEditorPlaySettings* PlayInSettings = GetDefault<ULevelEditorPlaySettings>();
+	//const ULevelEditorPlaySettings* PlayInSettings = InPlaySettings != NULL ? InPlaySettings : GetDefault<ULevelEditorPlaySettings>();
+	FString CmdLine;
+	if (WindowPos != NULL)
+	{
+		CmdLine = GenerateCmdLineForNextPieInstance(WindowPos->X, WindowPos->Y, PIENum, bIsServer);
+	}
+	
+	FString URLParms;		
+	if (bIsServer == true)
+	{
+		URLParms = bIsServer == true ? TEXT("?Listen") : FString();
+	}
+	
 	// select map to play
 	TArray<FString> SavedMapNames;
 	if (MapNameOverride.IsEmpty())
@@ -981,8 +1019,6 @@ void UEditorEngine::PlayOnLocalPc(FString MapNameOverride, FString URLParms, FSt
 	}
 
 	// apply additional settings
-	const ULevelEditorPlaySettings* PlayInSettings = GetDefault<ULevelEditorPlaySettings>();
-
 	if (bPlayUsingMobilePreview)
 	{
 		AdditionalParameters += TEXT(" -featureleveles2 -faketouches");
@@ -999,14 +1035,17 @@ void UEditorEngine::PlayOnLocalPc(FString MapNameOverride, FString URLParms, FSt
 		AdditionalParameters += PlayInSettings->AdditionalLaunchParameters;
 	}
 
+	FIntPoint WinSize(0, 0);
+	GetWindowSizeForInstanceType(WinSize, PlayInSettings);
+	
 	FString Params = FString::Printf(TEXT("%s %s -game -PIEVIACONSOLE -ResX=%d -ResY=%d %s%s %s"),
 		*GameNameOrProjectFile,
 		*BuildPlayWorldURL(*SavedMapNames[0], false, URLParms),
-		(ResX > 0 ? ResX : PlayInSettings->StandaloneWindowWidth),
-		(ResY > 0 ? ResY : PlayInSettings->StandaloneWindowHeight),
+		WinSize.X,
+		WinSize.Y,
 		*FCommandLine::GetSubprocessCommandline(),
 		*AdditionalParameters,
-		*CmdLineParms
+		*CmdLine
 	);
 
 	// launch the game process
