Brush wireframe is now rendered stippled when behind objects when being moved (additional).

Perforce somehow missed this file yesterday.

TTP# 329784 - EDITOR: BSP: Unable to see the geometry Wireframe when scaling down BSP.

diff --git a/Engine/Source/Editor/UnrealEd/Private/EditorBrushBuilder.cpp b/Engine/Source/Editor/UnrealEd/Private/EditorBrushBuilder.cpp
--- a/Engine/Source/Editor/UnrealEd/Private/EditorBrushBuilder.cpp
+++ b/Engine/Source/Editor/UnrealEd/Private/EditorBrushBuilder.cpp
@@ -231,6 +231,7 @@ void UEditorBrushBuilder::PostEditChangeProperty(struct FPropertyChangedEvent& P
 		ABrush* Brush = Cast<ABrush>(GetOuter());
 		if (Brush)
 		{
+			Brush->bInManipulation = PropertyChangedEvent.ChangeType == EPropertyChangeType::Interactive;
 			Build(Brush->GetWorld(), Brush);
 		}
 	}
