Fixed reversed autogenerated materials

Changed last remaining instances I could find of material layout code that generated materials with flipped expression positions - SpeedTree, Fbx, dragging textures onto meshes and when placing rotate about axis nodes.

diff --git a/Engine/Plugins/Editor/SpeedTreeImporter/Source/SpeedTreeImporter/Private/SpeedTreeImportFactory.cpp b/Engine/Plugins/Editor/SpeedTreeImporter/Source/SpeedTreeImporter/Private/SpeedTreeImportFactory.cpp
--- a/Engine/Plugins/Editor/SpeedTreeImporter/Source/SpeedTreeImporter/Private/SpeedTreeImportFactory.cpp
+++ b/Engine/Plugins/Editor/SpeedTreeImporter/Source/SpeedTreeImporter/Private/SpeedTreeImportFactory.cpp
@@ -422,7 +422,7 @@ static void LayoutMaterial(UMaterialInterface* MaterialInterface)
 	Material->EditorY = 0;
 
 	const int32 Height = 200;
-	const int32 Width = 250;
+	const int32 Width = -250;
 
 	// layout X to make sure each input is 1 step further than output connection
 	bool bContinue = true;
@@ -433,7 +433,7 @@ static void LayoutMaterial(UMaterialInterface* MaterialInterface)
 		for (int32 ExpressionIndex = 0; ExpressionIndex < Material->Expressions.Num(); ++ExpressionIndex)
 		{
 			UMaterialExpression* Expression = Material->Expressions[ExpressionIndex];
-			Expression->MaterialExpressionEditorX = FMath::Max(Expression->MaterialExpressionEditorX, Width);
+			Expression->MaterialExpressionEditorX = FMath::Min(Expression->MaterialExpressionEditorX, Width);
 
 			TArray<FExpressionInput*> Inputs = Expression->GetInputs();
 			for (int32 InputIndex = 0; InputIndex < Inputs.Num(); ++InputIndex)
@@ -441,9 +441,9 @@ static void LayoutMaterial(UMaterialInterface* MaterialInterface)
 				UMaterialExpression* Input = Inputs[InputIndex]->Expression;
 				if (Input != NULL)
 				{
-					if (Input->MaterialExpressionEditorX < Expression->MaterialExpressionEditorX + Width)
+					if (Input->MaterialExpressionEditorX > Expression->MaterialExpressionEditorX + Width)
 					{
-						Input->MaterialExpressionEditorX = FMath::Max(Input->MaterialExpressionEditorX, Expression->MaterialExpressionEditorX + Width);
+						Input->MaterialExpressionEditorX = FMath::Min(Input->MaterialExpressionEditorX, Expression->MaterialExpressionEditorX + Width);
 						bContinue = true;
 					}
 				}
diff --git a/Engine/Source/Editor/MaterialEditor/Private/MaterialEditor.cpp b/Engine/Source/Editor/MaterialEditor/Private/MaterialEditor.cpp
--- a/Engine/Source/Editor/MaterialEditor/Private/MaterialEditor.cpp
+++ b/Engine/Source/Editor/MaterialEditor/Private/MaterialEditor.cpp
@@ -2565,12 +2565,9 @@ UMaterialExpression* FMaterialEditor::CreateNewMaterialExpression(UClass* NewExp
 			check(!NewExpression->bIsParameterExpression);
 		}
 
-		// If the new expression is created connected to an input tab, offset it by this amount.
-		int32 NewConnectionOffset = 0;
-
 		// Set the expression location.
-		NewExpression->MaterialExpressionEditorX = NodePos.X + NewConnectionOffset;
-		NewExpression->MaterialExpressionEditorY = NodePos.Y + NewConnectionOffset;
+		NewExpression->MaterialExpressionEditorX = NodePos.X;
+		NewExpression->MaterialExpressionEditorY = NodePos.Y;
 
 		if (bAutoAssignResource)
 		{
@@ -2647,7 +2644,7 @@ UMaterialExpression* FMaterialEditor::CreateNewMaterialExpression(UClass* NewExp
 			Material->Expressions.Add( WorldPositionExpression );
 			WorldPositionExpression->Material = Material;
 			RotateAboutAxisExpression->Position.Expression = WorldPositionExpression;
-			WorldPositionExpression->MaterialExpressionEditorX = RotateAboutAxisExpression->MaterialExpressionEditorX + 250;
+			WorldPositionExpression->MaterialExpressionEditorX = RotateAboutAxisExpression->MaterialExpressionEditorX - 250;
 			WorldPositionExpression->MaterialExpressionEditorY = RotateAboutAxisExpression->MaterialExpressionEditorY + 73;
 			Material->MaterialGraph->AddExpression(WorldPositionExpression);
 			if ( bAutoSelect )
diff --git a/Engine/Source/Editor/UnrealEd/Private/Fbx/FbxMaterialImport.cpp b/Engine/Source/Editor/UnrealEd/Private/Fbx/FbxMaterialImport.cpp
--- a/Engine/Source/Editor/UnrealEd/Private/Fbx/FbxMaterialImport.cpp
+++ b/Engine/Source/Editor/UnrealEd/Private/Fbx/FbxMaterialImport.cpp
@@ -275,7 +275,7 @@ bool UnFbx::FFbxImporter::CreateAndLinkExpressionForMaterialProperty(
 							MyCoordExpression->UTiling = ScaleU;
 							MyCoordExpression->VTiling = ScaleV;
 							UnrealTextureExpression->Coordinates.Expression = MyCoordExpression;
-							MyCoordExpression->MaterialExpressionEditorX = FMath::TruncToInt(Location.X+175);
+							MyCoordExpression->MaterialExpressionEditorX = FMath::TruncToInt(Location.X-175);
 							MyCoordExpression->MaterialExpressionEditorY = FMath::TruncToInt(Location.Y);
 
 						}
diff --git a/Engine/Source/Editor/UnrealEd/Private/LevelEditorViewport.cpp b/Engine/Source/Editor/UnrealEd/Private/LevelEditorViewport.cpp
--- a/Engine/Source/Editor/UnrealEd/Private/LevelEditorViewport.cpp
+++ b/Engine/Source/Editor/UnrealEd/Private/LevelEditorViewport.cpp
@@ -459,11 +459,13 @@ static UObject* GetOrCreateMaterialFromTexture( UTexture* UnrealTexture )
 	UMaterial* UnrealMaterial = (UMaterial*)MaterialFactory->FactoryCreateNew(
 		UMaterial::StaticClass(), Package, *MaterialFullName, RF_Standalone | RF_Public, NULL, GWarn );
 
+	const int HSpace = -300;
+
 	// If we were able to figure out the material kind, we need to try and build a complex material
 	// involving multiple textures.  If not, just try and connect what we found to the base map.
 	if ( MaterialKind == EMaterialKind::Unknown )
 	{
-		TryAndCreateMaterialInput( UnrealMaterial, EMaterialKind::Base, UnrealTexture, UnrealMaterial->DiffuseColor, 300, 0 );
+		TryAndCreateMaterialInput( UnrealMaterial, EMaterialKind::Base, UnrealTexture, UnrealMaterial->DiffuseColor, HSpace, 0 );
 	}
 	else
 	{
@@ -504,10 +506,10 @@ static UObject* GetOrCreateMaterialFromTexture( UTexture* UnrealTexture )
 
 		// Connect and layout any textures we find into their respective inputs in the material.
 		const int VSpace = 170;
-		TryAndCreateMaterialInput( UnrealMaterial, EMaterialKind::Base, BaseTexture, UnrealMaterial->DiffuseColor, 300, VSpace * -1 );
-		TryAndCreateMaterialInput( UnrealMaterial, EMaterialKind::Specular, SpecularTexture, UnrealMaterial->Specular, 300, VSpace * 0 );
-		TryAndCreateMaterialInput( UnrealMaterial, EMaterialKind::Emissive, EmissiveTexture, UnrealMaterial->EmissiveColor, 300, VSpace * 1 );
-		TryAndCreateMaterialInput( UnrealMaterial, EMaterialKind::Normal, NormalTexture, UnrealMaterial->Normal, 300, VSpace * 2 );
+		TryAndCreateMaterialInput( UnrealMaterial, EMaterialKind::Base, BaseTexture, UnrealMaterial->DiffuseColor, HSpace, VSpace * -1 );
+		TryAndCreateMaterialInput( UnrealMaterial, EMaterialKind::Specular, SpecularTexture, UnrealMaterial->Specular, HSpace, VSpace * 0 );
+		TryAndCreateMaterialInput( UnrealMaterial, EMaterialKind::Emissive, EmissiveTexture, UnrealMaterial->EmissiveColor, HSpace, VSpace * 1 );
+		TryAndCreateMaterialInput( UnrealMaterial, EMaterialKind::Normal, NormalTexture, UnrealMaterial->Normal, HSpace, VSpace * 2 );
 	}
 
 	if ( UnrealMaterial != NULL )
