FModeTool_InterpEdit::InputAxis now only sets the axis if one hasn't already been set

TTP# 337667 - CRASH: EDITOR: Using Ctrl+RMB to move an actor while in Matinee

This prevents it overwriting the axis set by the FMouseDeltaTracker when performing a rotate, but still allows it to set its own axis when performing a normal drag.

diff --git a/Engine/Source/Editor/UnrealEd/Private/EditorModeInterpolation.cpp b/Engine/Source/Editor/UnrealEd/Private/EditorModeInterpolation.cpp
--- a/Engine/Source/Editor/UnrealEd/Private/EditorModeInterpolation.cpp
+++ b/Engine/Source/Editor/UnrealEd/Private/EditorModeInterpolation.cpp
@@ -213,19 +213,22 @@ bool FModeTool_InterpEdit::MouseMove(FLevelEditorViewportClient* ViewportClient,
 
 bool FModeTool_InterpEdit::InputAxis(FLevelEditorViewportClient* InViewportClient, FViewport* Viewport, int32 ControllerId, FKey Key, float Delta, float DeltaTime)
 {
-	// We need to set up a widget axis here to prevent our drag operation being co-opted by box/frustum selection
-	ELevelViewportType ViewportType = InViewportClient->ViewportType;
-	switch(ViewportType)
+	if ( InViewportClient->GetCurrentWidgetAxis() == EAxisList::None )
 	{
-	case LVT_OrthoXY:
-		InViewportClient->SetCurrentWidgetAxis(EAxisList::XY);
-		break;
-	case LVT_OrthoXZ:
-		InViewportClient->SetCurrentWidgetAxis(EAxisList::XZ);
-		break;
-	case LVT_OrthoYZ:
-		InViewportClient->SetCurrentWidgetAxis(EAxisList::YZ);
-		break;
+		// We need to set up a widget axis here to prevent our drag operation being co-opted by box/frustum selection
+		ELevelViewportType ViewportType = InViewportClient->ViewportType;
+		switch(ViewportType)
+		{
+		case LVT_OrthoXY:
+			InViewportClient->SetCurrentWidgetAxis(EAxisList::XY);
+			break;
+		case LVT_OrthoXZ:
+			InViewportClient->SetCurrentWidgetAxis(EAxisList::XZ);
+			break;
+		case LVT_OrthoYZ:
+			InViewportClient->SetCurrentWidgetAxis(EAxisList::YZ);
+			break;
+		}
 	}
 
 	return false;
