Move the definition for FMipLevelBatchedElementParameters::BindShaders_RenderThread to the .cpp file or Clang will attempt to inline it in other modules, like TextureEditor where the templated definition is unavailable.

diff --git a/Engine/Source/Editor/UnrealEd/Private/CubemapUnwrapUtils.cpp b/Engine/Source/Editor/UnrealEd/Private/CubemapUnwrapUtils.cpp
--- a/Engine/Source/Editor/UnrealEd/Private/CubemapUnwrapUtils.cpp
+++ b/Engine/Source/Editor/UnrealEd/Private/CubemapUnwrapUtils.cpp
@@ -124,6 +124,18 @@ void FCubemapTexturePropertiesPS<bHDROutput>::SetParameters( FRHICommandList* RH
 	SetShaderValue(RHICmdList, GetPixelShader(), Gamma, GammaValue);
 }
 
+void FMipLevelBatchedElementParameters::BindShaders_RenderThread(FRHICommandList* RHICmdList, const FMatrix& InTransform, const float InGamma, const FMatrix& ColorWeights, const FTexture* Texture)
+{
+	if(bHDROutput)
+	{
+		BindShaders_RenderThread<FCubemapTexturePropertiesPS<true> >(RHICmdList, InTransform, InGamma, ColorWeights, Texture);
+	}
+	else
+	{
+		BindShaders_RenderThread<FCubemapTexturePropertiesPS<false> >(RHICmdList, InTransform, InGamma, ColorWeights, Texture);
+	}
+}
+
 template<typename TPixelShader>
 void FMipLevelBatchedElementParameters::BindShaders_RenderThread( FRHICommandList* RHICmdList, const FMatrix& InTransform, const float InGamma, const FMatrix& ColorWeights, const FTexture* Texture )
 {
diff --git a/Engine/Source/Editor/UnrealEd/Public/CubemapUnwrapUtils.h b/Engine/Source/Editor/UnrealEd/Public/CubemapUnwrapUtils.h
--- a/Engine/Source/Editor/UnrealEd/Public/CubemapUnwrapUtils.h
+++ b/Engine/Source/Editor/UnrealEd/Public/CubemapUnwrapUtils.h
@@ -119,17 +119,7 @@ public:
 	}
 
 	/** Binds vertex and pixel shaders for this element */
-	virtual void BindShaders_RenderThread(FRHICommandList* RHICmdList, const FMatrix& InTransform, const float InGamma, const FMatrix& ColorWeights, const FTexture* Texture) OVERRIDE
-	{
-		if(bHDROutput)
-		{
-			BindShaders_RenderThread<FCubemapTexturePropertiesPS<true> >(RHICmdList, InTransform, InGamma, ColorWeights, Texture);
-		}
-		else
-		{
-			BindShaders_RenderThread<FCubemapTexturePropertiesPS<false> >(RHICmdList, InTransform, InGamma, ColorWeights, Texture);
-		}
-	}
+	virtual void BindShaders_RenderThread(FRHICommandList* RHICmdList, const FMatrix& InTransform, const float InGamma, const FMatrix& ColorWeights, const FTexture* Texture) OVERRIDE;
 
 private:
 	template<typename TPixelShader> void BindShaders_RenderThread(FRHICommandList* RHICmdList, const FMatrix& InTransform, const float InGamma, const FMatrix& ColorWeights, const FTexture* Texture);
