Fix auto restoring assets not working due to CL 2093156

diff --git a/Engine/Source/Editor/UnrealEd/Private/Toolkits/AssetEditorManager.cpp b/Engine/Source/Editor/UnrealEd/Private/Toolkits/AssetEditorManager.cpp
--- a/Engine/Source/Editor/UnrealEd/Private/Toolkits/AssetEditorManager.cpp
+++ b/Engine/Source/Editor/UnrealEd/Private/Toolkits/AssetEditorManager.cpp
@@ -383,7 +383,7 @@ void FAssetEditorManager::RestorePreviouslyOpenAssets()
 			if(bAutoRestore)
 			{
 				// Pretend that we showed the notification and that the user clicked "Restore Now"
-				OnConfirmRestorePreviouslyOpenAssets(OpenAssets);
+				OpenEditorsForAssets(OpenAssets);
 			}
 			else
 			{
@@ -487,10 +487,8 @@ void FAssetEditorManager::OnConfirmRestorePreviouslyOpenAssets(TArray<FString> A
 		Settings.PostEditChange();
 
 		// we do this inside the condition so that it can only be done once. 
-		for (const FString& AssetName : AssetsToOpen)
-		{
-			OpenEditorForAsset(AssetName);
-		}
+		OpenEditorsForAssets(AssetsToOpen);
+
 	}
 }
 
@@ -538,4 +536,12 @@ void FAssetEditorManager::SaveOpenAssetEditors(bool bOnShutdown)
 	}
 }
 
+void FAssetEditorManager::OpenEditorsForAssets(const TArray<FString>& AssetsToOpen)
+{
+	for (const FString& AssetName : AssetsToOpen)
+	{
+		OpenEditorForAsset(AssetName);
+	}
+}
+
 #undef LOCTEXT_NAMESPACE
\ No newline at end of file
diff --git a/Engine/Source/Editor/UnrealEd/Public/Toolkits/AssetEditorManager.h b/Engine/Source/Editor/UnrealEd/Public/Toolkits/AssetEditorManager.h
--- a/Engine/Source/Editor/UnrealEd/Public/Toolkits/AssetEditorManager.h
+++ b/Engine/Source/Editor/UnrealEd/Public/Toolkits/AssetEditorManager.h
@@ -44,6 +44,9 @@ public:
 	bool OpenEditorForAsset( UObject* Asset, const EToolkitMode::Type ToolkitMode = EToolkitMode::Standalone, TSharedPtr< class IToolkitHost > OpenedFromLevelEditor = TSharedPtr< IToolkitHost >() );
 	bool OpenEditorForAssets( const TArray< UObject* >& Asset, const EToolkitMode::Type ToolkitMode = EToolkitMode::Standalone, TSharedPtr< class IToolkitHost > OpenedFromLevelEditor = TSharedPtr< IToolkitHost >() );
 
+	// Opens editors for the supplied assets (via OpenEditorForAsset
+	void OpenEditorsForAssets(const TArray<FString>& AssetsToOpen);
+
 	// Returns the primary editor if one is already open for the specified asset.
 	// If there is one open and bFocusIfOpen is true, that editor will be brought to the foreground and focused if possible.
 	IAssetEditorInstance* FindEditorForAsset(UObject* Asset, bool bFocusIfOpen);
