Removed unnecessary ensure in FBlueprintCompileReinstancer::ReplaceInstancesOfClass().

diff --git a/Engine/Source/Editor/UnrealEd/Private/Kismet2/KismetReinstanceUtilities.cpp b/Engine/Source/Editor/UnrealEd/Private/Kismet2/KismetReinstanceUtilities.cpp
--- a/Engine/Source/Editor/UnrealEd/Private/Kismet2/KismetReinstanceUtilities.cpp
+++ b/Engine/Source/Editor/UnrealEd/Private/Kismet2/KismetReinstanceUtilities.cpp
@@ -288,7 +288,7 @@ void FBlueprintCompileReinstancer::ReplaceInstancesOfClass(UClass* OldClass, UCl
 					for (auto AttachIt = OldRootComponent->AttachChildren.CreateConstIterator(); AttachIt; ++AttachIt)
 					{
 						USceneComponent* Child = *AttachIt;
-						if (ensure(Child != nullptr))
+						if (Child != nullptr)
 						{
 							OldAttachChildren.Add(Child->GetOwner());
 						}
