During BP Regeneration the Skeleton class may be null.
Better fix for ttp#337588 and ttp#321815 and ttp#337552

diff --git a/Engine/Source/Editor/UnrealEd/Private/Kismet2/BlueprintEditorUtils.cpp b/Engine/Source/Editor/UnrealEd/Private/Kismet2/BlueprintEditorUtils.cpp
--- a/Engine/Source/Editor/UnrealEd/Private/Kismet2/BlueprintEditorUtils.cpp
+++ b/Engine/Source/Editor/UnrealEd/Private/Kismet2/BlueprintEditorUtils.cpp
@@ -675,7 +675,7 @@ struct FRegenerationHelper
 		Blueprint->GetAllGraphs(Graphs);
 		for (auto Graph : Graphs)
 		{
-			if (!FBlueprintEditorUtils::IsGraphIntermediate(Graph))
+			if (Graph && !FBlueprintEditorUtils::IsGraphIntermediate(Graph))
 			{
 				const bool bIsDelegateSignatureGraph = FBlueprintEditorUtils::IsDelegateSignatureGraph(Graph);
 
@@ -683,40 +683,43 @@ struct FRegenerationHelper
 				Graph->GetNodesOfClass(Nodes);
 				for (auto Node : Nodes)
 				{
-					TArray<UStruct*> LocalDependentStructures;
-					if (Node->HasExternalBlueprintDependencies(&LocalDependentStructures))
+					if (Node)
 					{
-						for (auto Struct : LocalDependentStructures)
+						TArray<UStruct*> LocalDependentStructures;
+						if (Node->HasExternalBlueprintDependencies(&LocalDependentStructures))
 						{
-							ProcessHierarchy(Struct, Dependencies);
-						}
+							for (auto Struct : LocalDependentStructures)
+							{
+								ProcessHierarchy(Struct, Dependencies);
+							}
 
-						if (auto MacroNode = Cast<UK2Node_MacroInstance>(Node))
-						{
-							if (UBlueprint* MacroSource = MacroNode->GetSourceBlueprint())
+							if (auto MacroNode = Cast<UK2Node_MacroInstance>(Node))
 							{
-								MacroSources.Add(MacroSource);
+								if (UBlueprint* MacroSource = MacroNode->GetSourceBlueprint())
+								{
+									MacroSources.Add(MacroSource);
+								}
 							}
 						}
-					}
 
-					LocalDependentStructures.Empty(LocalDependentStructures.Max());
-					Node->HasExternalUserDefinedStructDependencies(&LocalDependentStructures);
-					for (auto Struct : LocalDependentStructures)
-					{
-						ProcessHierarchy(Struct, Dependencies);
-					}
+						LocalDependentStructures.Empty(LocalDependentStructures.Max());
+						Node->HasExternalUserDefinedStructDependencies(&LocalDependentStructures);
+						for (auto Struct : LocalDependentStructures)
+						{
+							ProcessHierarchy(Struct, Dependencies);
+						}
 
-					auto FunctionEntry = Cast<UK2Node_FunctionEntry>(Node);
-					if (FunctionEntry && !bIsDelegateSignatureGraph)
-					{
-						const FName FunctionName = (FunctionEntry->CustomGeneratedFunctionName != NAME_None) 
-							? FunctionEntry->CustomGeneratedFunctionName 
-							: FunctionEntry->SignatureName;
-						UFunction* ParentFunction = Blueprint->ParentClass ? Blueprint->ParentClass->FindFunctionByName(FunctionName) : NULL;
-						if (ParentFunction && (Schema->FN_UserConstructionScript != FunctionName))
+						auto FunctionEntry = Cast<UK2Node_FunctionEntry>(Node);
+						if (FunctionEntry && !bIsDelegateSignatureGraph)
 						{
-							ProcessHierarchy(ParentFunction, Dependencies);
+							const FName FunctionName = (FunctionEntry->CustomGeneratedFunctionName != NAME_None) 
+								? FunctionEntry->CustomGeneratedFunctionName 
+								: FunctionEntry->SignatureName;
+							UFunction* ParentFunction = Blueprint->ParentClass ? Blueprint->ParentClass->FindFunctionByName(FunctionName) : NULL;
+							if (ParentFunction && (Schema->FN_UserConstructionScript != FunctionName))
+							{
+								ProcessHierarchy(ParentFunction, Dependencies);
+							}
 						}
 					}
 				}
@@ -942,8 +945,6 @@ UClass* FBlueprintEditorUtils::RegenerateBlueprintClass(UBlueprint* Blueprint, U
 
 	bool bNeedsSkelRefRemoval = false;
 
-	FBlueprintEditorUtils::PreloadMembers(Blueprint);
-
 	if( ShouldRegenerateBlueprint(Blueprint) && !Blueprint->bHasBeenRegenerated )
 	{
 		Blueprint->bIsRegeneratingOnLoad = true;
@@ -2931,39 +2932,42 @@ bool FBlueprintEditorUtils::IsPinTypeValid(const FEdGraphPinType& Type)
 void FBlueprintEditorUtils::GetClassVariableList(const UBlueprint* Blueprint, TArray<FName>& VisibleVariables, bool bIncludePrivateVars) 
 {
 	// Existing variables in the parent class and above
-	check(Blueprint->SkeletonGeneratedClass != NULL);
-	for (TFieldIterator<UProperty> PropertyIt(Blueprint->SkeletonGeneratedClass, EFieldIteratorFlags::IncludeSuper); PropertyIt; ++PropertyIt)
+	check(Blueprint->bIsRegeneratingOnLoad || (Blueprint->SkeletonGeneratedClass != NULL));
+	if (Blueprint->SkeletonGeneratedClass != NULL)
 	{
-		UProperty* Property = *PropertyIt;
-
-		if ((!Property->HasAnyPropertyFlags(CPF_Parm) && (bIncludePrivateVars || Property->HasAllPropertyFlags(CPF_BlueprintVisible))))
+		for (TFieldIterator<UProperty> PropertyIt(Blueprint->SkeletonGeneratedClass, EFieldIteratorFlags::IncludeSuper); PropertyIt; ++PropertyIt)
 		{
-			VisibleVariables.Add(Property->GetFName());
+			UProperty* Property = *PropertyIt;
+
+			if ((!Property->HasAnyPropertyFlags(CPF_Parm) && (bIncludePrivateVars || Property->HasAllPropertyFlags(CPF_BlueprintVisible))))
+			{
+				VisibleVariables.Add(Property->GetFName());
+			}
 		}
-	}
 
-	if (bIncludePrivateVars)
-	{
-		// Include SCS node variable names, timelines, and other member variables that may be pending compilation. Consider them to be "private" as they're not technically accessible for editing just yet.
-		TArray<UBlueprint*> ParentBPStack;
-		UBlueprint::GetBlueprintHierarchyFromClass(Blueprint->SkeletonGeneratedClass, ParentBPStack);
-		for (int32 StackIndex = ParentBPStack.Num() - 1; StackIndex >= 0; --StackIndex)
+		if (bIncludePrivateVars)
 		{
-			UBlueprint* ParentBP = ParentBPStack[StackIndex];
-			check(ParentBP != NULL);
+			// Include SCS node variable names, timelines, and other member variables that may be pending compilation. Consider them to be "private" as they're not technically accessible for editing just yet.
+			TArray<UBlueprint*> ParentBPStack;
+			UBlueprint::GetBlueprintHierarchyFromClass(Blueprint->SkeletonGeneratedClass, ParentBPStack);
+			for (int32 StackIndex = ParentBPStack.Num() - 1; StackIndex >= 0; --StackIndex)
+			{
+				UBlueprint* ParentBP = ParentBPStack[StackIndex];
+				check(ParentBP != NULL);
 
-			GetSCSVariableNameList(ParentBP, VisibleVariables);
+				GetSCSVariableNameList(ParentBP, VisibleVariables);
 
-			for(int32 VariableIndex = 0; VariableIndex < ParentBP->NewVariables.Num(); ++VariableIndex)
-			{
-				VisibleVariables.AddUnique(ParentBP->NewVariables[VariableIndex].VarName);
-			}
+				for (int32 VariableIndex = 0; VariableIndex < ParentBP->NewVariables.Num(); ++VariableIndex)
+				{
+					VisibleVariables.AddUnique(ParentBP->NewVariables[VariableIndex].VarName);
+				}
 
-			for (auto Timeline : ParentBP->Timelines)
-			{
-				if (Timeline)
+				for (auto Timeline : ParentBP->Timelines)
 				{
-					VisibleVariables.AddUnique(Timeline->GetFName());
+					if (Timeline)
+					{
+						VisibleVariables.AddUnique(Timeline->GetFName());
+					}
 				}
 			}
 		}
