Add GameplayAbilitiesEditor to root editor

diff --git a/Engine/Source/Editor/UnrealEd/UnrealEd.Build.cs b/Engine/Source/Editor/UnrealEd/UnrealEd.Build.cs
--- a/Engine/Source/Editor/UnrealEd/UnrealEd.Build.cs
+++ b/Engine/Source/Editor/UnrealEd/UnrealEd.Build.cs
@@ -169,6 +169,7 @@ public class UnrealEd : ModuleRules
 				"ViewportSnapping",
 				"UserFeedback",
 				"GameplayTagsEditor",
+                "GameplayAbilitiesEditor",
 				"UndoHistory",
 				"SourceCodeAccess",
 				"ReferenceViewer"
