Fixed TTP #307187 : Duplicated SkeletalMesh import LOD code

diff --git a/Engine/Source/Editor/UnrealEd/Private/Fbx/FbxFactory.cpp b/Engine/Source/Editor/UnrealEd/Private/Fbx/FbxFactory.cpp
--- a/Engine/Source/Editor/UnrealEd/Private/Fbx/FbxFactory.cpp
+++ b/Engine/Source/Editor/UnrealEd/Private/Fbx/FbxFactory.cpp
@@ -394,9 +394,16 @@ UObject* UFbxFactory::FactoryCreateBinary
 								USkeletalMesh* BaseSkeletalMesh = Cast<USkeletalMesh>(NewObject);
 								FName LODObjectName = NAME_None;
 								USkeletalMesh *LODObject = FbxImporter->ImportSkeletalMesh( GetTransientPackage(), SkelMeshNodeArray, LODObjectName, RF_NoFlags, ImportUI->SkeletalMeshImportData, FPaths::GetBaseFilename(Filename) );
-								FbxImporter->ImportSkeletalMeshLOD( LODObject, BaseSkeletalMesh, LODIndex);
+								bool bImportSucceeded = FbxImporter->ImportSkeletalMeshLOD(LODObject, BaseSkeletalMesh, LODIndex, false);
 
-								BaseSkeletalMesh->LODInfo[LODIndex].DisplayFactor = 1.0f / (MaxLODLevel * LODIndex);
+								if (bImportSucceeded)
+								{
+									BaseSkeletalMesh->LODInfo[LODIndex].DisplayFactor = 1.0f / (MaxLODLevel * LODIndex);
+								}
+								else
+								{
+									FbxImporter->AddTokenizedErrorMessage(FTokenizedMessage::Create(EMessageSeverity::Error, LOCTEXT("FailedToImport_SkeletalMeshLOD", "Failed to import Skeletal mesh LOD.")));
+								}
 							}
 						
 							// import morph target
diff --git a/Engine/Source/Editor/UnrealEd/Private/Fbx/FbxSkeletalMeshImport.cpp b/Engine/Source/Editor/UnrealEd/Private/Fbx/FbxSkeletalMeshImport.cpp
--- a/Engine/Source/Editor/UnrealEd/Private/Fbx/FbxSkeletalMeshImport.cpp
+++ b/Engine/Source/Editor/UnrealEd/Private/Fbx/FbxSkeletalMeshImport.cpp
@@ -1501,7 +1501,7 @@ USkeletalMesh* UnFbx::FFbxImporter::ReimportSkeletalMesh(USkeletalMesh* Mesh, UF
 			{
 				USkeletalMesh* BaseSkeletalMesh = Cast<USkeletalMesh>(NewMesh);
 				UObject *LODObject = ImportSkeletalMesh( GetTransientPackage(), SkelMeshNodeArray, NAME_None, RF_NoFlags, TemplateImportData, FReimportManager::ResolveImportFilename(TemplateImportData->SourceFilePath, BaseSkeletalMesh) );
-				ImportSkeletalMeshLOD( Cast<USkeletalMesh>(LODObject), BaseSkeletalMesh, LODIndex);
+				ImportSkeletalMeshLOD( Cast<USkeletalMesh>(LODObject), BaseSkeletalMesh, LODIndex, false);
 
 				// Set LOD Model's DisplayFactor
 				// if this LOD is newly added, then set DisplayFactor
@@ -2260,38 +2260,141 @@ bool UnFbx::FFbxImporter::FillSkelMeshImporterFromFbx( FSkeletalMeshImportData&
 	return true;
 }
 
-void UnFbx::FFbxImporter::ImportSkeletalMeshLOD(USkeletalMesh* InSkeletalMesh, USkeletalMesh* BaseSkeletalMesh, int32 DesiredLOD)
+void UnFbx::FFbxImporter::InsertNewLODToBaseSkeletalMesh(USkeletalMesh* InSkeletalMesh, USkeletalMesh* BaseSkeletalMesh, int32 DesiredLOD)
+{
+	FSkeletalMeshResource* ImportedResource = InSkeletalMesh->GetImportedResource();
+	FSkeletalMeshResource* DestImportedResource = BaseSkeletalMesh->GetImportedResource();
+
+	FStaticLODModel& NewLODModel = ImportedResource->LODModels[0];
+
+	// If we want to add this as a new LOD to this mesh - add to LODModels/LODInfo array.
+	if (DesiredLOD == DestImportedResource->LODModels.Num())
+	{
+		new(DestImportedResource->LODModels)FStaticLODModel();
+
+		// Add element to LODInfo array.
+		BaseSkeletalMesh->LODInfo.AddZeroed();
+		check(BaseSkeletalMesh->LODInfo.Num() == DestImportedResource->LODModels.Num());
+		BaseSkeletalMesh->LODInfo[DesiredLOD] = InSkeletalMesh->LODInfo[0];
+	}
+	else
+	{
+		// if it's overwriting existing LOD, need to update section information
+		// update to the right # of sections 
+		// Set up LODMaterialMap to number of materials in new mesh.
+		// ImportedResource->LOD 0 is the newly imported mesh
+		FSkeletalMeshLODInfo& LODInfo = BaseSkeletalMesh->LODInfo[DesiredLOD];
+		// if section # has been changed
+		if (LODInfo.TriangleSortSettings.Num() != NewLODModel.Sections.Num())
+		{
+			// Save old information so that I can copy it over
+			TArray<FTriangleSortSettings> OldTriangleSortSettings;
+
+			OldTriangleSortSettings = LODInfo.TriangleSortSettings;
+
+			// resize to the correct number
+			LODInfo.TriangleSortSettings.Empty(NewLODModel.Sections.Num());
+			// fill up data
+			for (int32 SectionIndex = 0; SectionIndex < NewLODModel.Sections.Num(); ++SectionIndex)
+			{
+				// if found from previous data, copy over
+				if (SectionIndex < OldTriangleSortSettings.Num())
+				{
+					LODInfo.TriangleSortSettings.Add(OldTriangleSortSettings[SectionIndex]);
+				}
+				else
+				{
+					// if not add default data
+					LODInfo.TriangleSortSettings.AddZeroed();
+				}
+			}
+		}
+	}
+
+	// Set up LODMaterialMap to number of materials in new mesh.
+	FSkeletalMeshLODInfo& LODInfo = BaseSkeletalMesh->LODInfo[DesiredLOD];
+	LODInfo.LODMaterialMap.Empty();
+
+	// Now set up the material mapping array.
+	for (int32 MatIdx = 0; MatIdx < InSkeletalMesh->Materials.Num(); MatIdx++)
+	{
+		// Try and find the auto-assigned material in the array.
+		int32 LODMatIndex = INDEX_NONE;
+		if (InSkeletalMesh->Materials[MatIdx].MaterialInterface != NULL)
+		{
+			LODMatIndex = BaseSkeletalMesh->Materials.Find(InSkeletalMesh->Materials[MatIdx]);
+		}
+
+		// If we didn't just use the index - but make sure its within range of the Materials array.
+		if (LODMatIndex == INDEX_NONE)
+		{
+			LODMatIndex = FMath::Clamp(MatIdx, 0, BaseSkeletalMesh->Materials.Num() - 1);
+		}
+
+		LODInfo.LODMaterialMap.Add(LODMatIndex);
+	}
+
+	// if new LOD has more material slot, add the extra to main skeletal
+	if (BaseSkeletalMesh->Materials.Num() < InSkeletalMesh->Materials.Num())
+	{
+		BaseSkeletalMesh->Materials.AddZeroed(InSkeletalMesh->Materials.Num() - BaseSkeletalMesh->Materials.Num());
+	}
+
+	// same from here as FbxImporter
+
+	// Release all resources before replacing the model
+	BaseSkeletalMesh->PreEditChange(NULL);
+
+	// Index buffer will be destroyed when we copy the LOD model so we must copy the index buffer and reinitialize it after the copy
+	FMultiSizeIndexContainerData Data;
+	NewLODModel.MultiSizeIndexContainer.GetIndexBufferData(Data);
+
+	// Assign new FStaticLODModel to desired slot in selected skeletal mesh.
+	DestImportedResource->LODModels[DesiredLOD] = NewLODModel;
+
+	DestImportedResource->LODModels[DesiredLOD].MultiSizeIndexContainer.RebuildIndexBuffer(Data);
+
+	// rebuild vertex buffers and reinit RHI resources
+	BaseSkeletalMesh->PostEditChange();
+}
+
+bool UnFbx::FFbxImporter::ImportSkeletalMeshLOD(USkeletalMesh* InSkeletalMesh, USkeletalMesh* BaseSkeletalMesh, int32 DesiredLOD, bool bNeedToReregister, TArray<UActorComponent*>* ReregisterAssociatedComponents)
 {
 	check(InSkeletalMesh);
 	check(BaseSkeletalMesh);
 
 	FSkeletalMeshResource* ImportedResource = InSkeletalMesh->GetImportedResource();
-	FSkeletalMeshResource* BaseImportedResource = BaseSkeletalMesh->GetImportedResource();
+	FSkeletalMeshResource* DestImportedResource = BaseSkeletalMesh->GetImportedResource();
 
 	// Now we copy the base FStaticLODModel from the imported skeletal mesh as the new LOD in the selected mesh.
 	check(ImportedResource->LODModels.Num() == 1);
 
 	// Names of root bones must match.
 	// If the names of root bones don't match, the LOD Mesh does not share skeleton with base Mesh. 
-	if( InSkeletalMesh->RefSkeleton.GetBoneName(0) != BaseSkeletalMesh->RefSkeleton.GetBoneName(0) )
+	if (InSkeletalMesh->RefSkeleton.GetBoneName(0) != BaseSkeletalMesh->RefSkeleton.GetBoneName(0))
 	{
-		AddTokenizedErrorMessage(FTokenizedMessage::Create(EMessageSeverity::Error, FText::Format(LOCTEXT("LODRootBoneDoesNotMatch", "RootBone '{0}' not found in base SkeletalMesh '{1}'.\nImport failed."), FText::FromName(InSkeletalMesh->RefSkeleton.GetBoneName(0)), FText::FromString(BaseSkeletalMesh->GetName()))));
-		return;
+		AddTokenizedErrorMessage(FTokenizedMessage::Create(EMessageSeverity::Error, FText::Format(LOCTEXT("LODRootNameIncorrect", "Root bone in LOD is '{0}' instead of '{1}'.\nImport failed."),
+			FText::FromName(InSkeletalMesh->RefSkeleton.GetBoneName(0)), FText::FromName(BaseSkeletalMesh->RefSkeleton.GetBoneName(0)))));
+
+		return false;
 	}
 
 	// We do some checking here that for every bone in the mesh we just imported, it's in our base ref skeleton, and the parent is the same.
-	for(int32 i=0; i<InSkeletalMesh->RefSkeleton.GetNum(); i++)
+	for (int32 i = 0; i < InSkeletalMesh->RefSkeleton.GetNum(); i++)
 	{
 		int32 LODBoneIndex = i;
 		FName LODBoneName = InSkeletalMesh->RefSkeleton.GetBoneName(LODBoneIndex);
 		int32 BaseBoneIndex = BaseSkeletalMesh->RefSkeleton.FindBoneIndex(LODBoneName);
-		if( BaseBoneIndex == INDEX_NONE )
+		if (BaseBoneIndex == INDEX_NONE)
 		{
-			AddTokenizedErrorMessage(FTokenizedMessage::Create(EMessageSeverity::Error, FText::Format(LOCTEXT("LODBoneDoesNotMatch", "Bone '{0}' not found in base SkeletalMesh '{1}'.\nImport failed."), FText::FromName(LODBoneName), FText::FromString(BaseSkeletalMesh->GetName()))));
-			return;
+			// If we could not find the bone from this LOD in base mesh - we fail.
+			AddTokenizedErrorMessage(FTokenizedMessage::Create(EMessageSeverity::Error, FText::Format(LOCTEXT("LODBoneDoesNotMatch", "Bone '{0}' not found in base SkeletalMesh '{1}'.\nImport failed."),
+				FText::FromName(LODBoneName), FText::FromString(BaseSkeletalMesh->GetName()))));
+
+			return false;
 		}
 
-		if(i>0)
+		if (i > 0)
 		{
 			int32 LODParentIndex = InSkeletalMesh->RefSkeleton.GetParentIndex(LODBoneIndex);
 			FName LODParentName = InSkeletalMesh->RefSkeleton.GetBoneName(LODParentIndex);
@@ -2299,59 +2402,59 @@ void UnFbx::FFbxImporter::ImportSkeletalMeshLOD(USkeletalMesh* InSkeletalMesh, U
 			int32 BaseParentIndex = BaseSkeletalMesh->RefSkeleton.GetParentIndex(BaseBoneIndex);
 			FName BaseParentName = BaseSkeletalMesh->RefSkeleton.GetBoneName(BaseParentIndex);
 
-			if(LODParentName != BaseParentName)
+			if (LODParentName != BaseParentName)
 			{
 				// If bone has different parents, display an error and don't allow import.
 				AddTokenizedErrorMessage(FTokenizedMessage::Create(EMessageSeverity::Error, FText::Format(LOCTEXT("LODBoneHasIncorrectParent", "Bone '{0}' in LOD has parent '{1}' instead of '{2}'"),
 					FText::FromName(LODBoneName), FText::FromName(LODParentName), FText::FromName(BaseParentName))));
-				return;
+
+				return false;
 			}
 		}
 	}
 
 	FStaticLODModel& NewLODModel = ImportedResource->LODModels[0];
 
-	// TODO: remove redundant code. When import LOD, some check is unnecessary. Also considerate update workflow!!
-
 	// Enforce LODs having only single-influence vertices.
 	bool bCheckSingleInfluence;
-	GConfig->GetBool( TEXT("AnimSetViewer"), TEXT("CheckSingleInfluenceLOD"), bCheckSingleInfluence, GEditorIni );
-	if( bCheckSingleInfluence && 
-		DesiredLOD > 0 )
+	GConfig->GetBool(TEXT("AnimSetViewer"), TEXT("CheckSingleInfluenceLOD"), bCheckSingleInfluence, GEditorIni);
+	if (bCheckSingleInfluence &&
+		DesiredLOD > 0)
 	{
-		for(int32 ChunkIndex = 0;ChunkIndex < NewLODModel.Chunks.Num();ChunkIndex++)
+		for (int32 ChunkIndex = 0; ChunkIndex < NewLODModel.Chunks.Num(); ChunkIndex++)
 		{
-			if(NewLODModel.Chunks[ChunkIndex].SoftVertices.Num() > 0)
+			if (NewLODModel.Chunks[ChunkIndex].SoftVertices.Num() > 0)
 			{
-				AddTokenizedErrorMessage(FTokenizedMessage::Create(EMessageSeverity::Error, LOCTEXT("LODHasSoftVertices", "Warning: The mesh LOD you are importing has some vertices with more than one influence.")));
+				AddTokenizedErrorMessage(FTokenizedMessage::Create(EMessageSeverity::Warning, FText(LOCTEXT("LODHasSoftVertices", "Warning: The mesh LOD you are importing has some vertices with more than one influence."))));
 			}
 		}
 	}
 
 	// If this LOD is going to be the lowest one, we check all bones we have sockets on are present in it.
-	if( DesiredLOD == BaseImportedResource->LODModels.Num() || 
-		DesiredLOD == BaseImportedResource->LODModels.Num()-1 )
+	if (DesiredLOD == DestImportedResource->LODModels.Num() ||
+		DesiredLOD == DestImportedResource->LODModels.Num() - 1)
 	{
-		const TArray<USkeletalMeshSocket*> Sockets = BaseSkeletalMesh->GetMeshOnlySocketList();
+		const TArray<USkeletalMeshSocket*>& Sockets = BaseSkeletalMesh->GetMeshOnlySocketList();
 
-		for(int32 i=0; i<Sockets.Num(); i++)
+		for (int32 i = 0; i < Sockets.Num(); i++)
 		{
 			// Find bone index the socket is attached to.
 			USkeletalMeshSocket* Socket = Sockets[i];
-			int32 SocketBoneIndex = InSkeletalMesh->RefSkeleton.FindBoneIndex( Socket->BoneName );
+			int32 SocketBoneIndex = InSkeletalMesh->RefSkeleton.FindBoneIndex(Socket->BoneName);
 
 			// If this LOD does not contain the socket bone, abort import.
-			if( SocketBoneIndex == INDEX_NONE )
+			if (SocketBoneIndex == INDEX_NONE)
 			{
 				AddTokenizedErrorMessage(FTokenizedMessage::Create(EMessageSeverity::Error, FText::Format(LOCTEXT("LODMissingSocketBone", "This LOD is missing bone '{0}' used by socket '{1}'.\nAborting import."),
-					FText::FromString(Socket->BoneName.ToString()), FText::FromName(Socket->SocketName))));
-				return;
+					FText::FromName(Socket->BoneName), FText::FromName(Socket->SocketName))));
+
+				return false;
 			}
 		}
 	}
 
 	// Fix up the ActiveBoneIndices array.
-	for(int32 i=0; i<NewLODModel.ActiveBoneIndices.Num(); i++)
+	for (int32 i = 0; i < NewLODModel.ActiveBoneIndices.Num(); i++)
 	{
 		int32 LODBoneIndex = NewLODModel.ActiveBoneIndices[i];
 		FName LODBoneName = InSkeletalMesh->RefSkeleton.GetBoneName(LODBoneIndex);
@@ -2360,10 +2463,10 @@ void UnFbx::FFbxImporter::ImportSkeletalMeshLOD(USkeletalMesh* InSkeletalMesh, U
 	}
 
 	// Fix up the chunk BoneMaps.
-	for(int32 ChunkIndex = 0;ChunkIndex < NewLODModel.Chunks.Num();ChunkIndex++)
+	for (int32 ChunkIndex = 0; ChunkIndex < NewLODModel.Chunks.Num(); ChunkIndex++)
 	{
 		FSkelMeshChunk& Chunk = NewLODModel.Chunks[ChunkIndex];
-		for(int32 i=0; i<Chunk.BoneMap.Num(); i++)
+		for (int32 i = 0; i < Chunk.BoneMap.Num(); i++)
 		{
 			int32 LODBoneIndex = Chunk.BoneMap[i];
 			FName LODBoneName = InSkeletalMesh->RefSkeleton.GetBoneName(LODBoneIndex);
@@ -2373,70 +2476,78 @@ void UnFbx::FFbxImporter::ImportSkeletalMeshLOD(USkeletalMesh* InSkeletalMesh, U
 	}
 
 	// Create the RequiredBones array in the LODModel from the ref skeleton.
-	NewLODModel.RequiredBones.Empty();
-	for(int32 i=0; i<InSkeletalMesh->RefSkeleton.GetNum(); i++)
+	for (int32 i = 0; i < NewLODModel.RequiredBones.Num(); i++)
 	{
-		FName LODBoneName = InSkeletalMesh->RefSkeleton.GetBoneName(i);
+		FName LODBoneName = InSkeletalMesh->RefSkeleton.GetBoneName(NewLODModel.RequiredBones[i]);
 		int32 BaseBoneIndex = BaseSkeletalMesh->RefSkeleton.FindBoneIndex(LODBoneName);
-		if(BaseBoneIndex != INDEX_NONE)
+		if (BaseBoneIndex != INDEX_NONE)
 		{
-			NewLODModel.RequiredBones.Add(BaseBoneIndex);
+			NewLODModel.RequiredBones[i] = BaseBoneIndex;
 		}
-	}
-	
-	{
-		// If we want to add this as a new LOD to this mesh - add to LODModels/LODInfo array.
-		if(DesiredLOD == BaseImportedResource->LODModels.Num())
+		else
 		{
-			new(BaseImportedResource->LODModels)FStaticLODModel();
-
-			// Add element to LODInfo array.
-			BaseSkeletalMesh->LODInfo.AddZeroed();
-			check( BaseSkeletalMesh->LODInfo.Num() == BaseImportedResource->LODModels.Num() );
-			BaseSkeletalMesh->LODInfo[DesiredLOD] = InSkeletalMesh->LODInfo[0];
+			NewLODModel.RequiredBones.RemoveAt(i--);
 		}
+	}
 
-		// Set up LODMaterialMap to number of materials in new mesh.
-		FSkeletalMeshLODInfo& LODInfo = BaseSkeletalMesh->LODInfo[DesiredLOD];
-		LODInfo.LODMaterialMap.Empty();
+	// Also sort the RequiredBones array to be strictly increasing.
+	NewLODModel.RequiredBones.Sort();
 
-		// Now set up the material mapping array.
-		for(int32 MatIdx = 0; MatIdx < InSkeletalMesh->Materials.Num(); MatIdx++)
-		{
-			// Try and find the auto-assigned material in the array.
-			int32 LODMatIndex = BaseSkeletalMesh->Materials.Find( InSkeletalMesh->Materials[MatIdx] );
+	// To be extra-nice, we apply the difference between the root transform of the meshes to the verts.
+	FMatrix LODToBaseTransform = InSkeletalMesh->GetRefPoseMatrix(0).Inverse() * BaseSkeletalMesh->GetRefPoseMatrix(0);
 
-			// If we didn't just use the index - but make sure its within range of the Materials array.
-			if(LODMatIndex == INDEX_NONE)
-			{
-				LODMatIndex = FMath::Clamp(MatIdx, 0, BaseSkeletalMesh->Materials.Num() - 1);
-			}
+	for (int32 ChunkIndex = 0; ChunkIndex < NewLODModel.Chunks.Num(); ChunkIndex++)
+	{
+		FSkelMeshChunk& Chunk = NewLODModel.Chunks[ChunkIndex];
+		// Fix up rigid verts.
+		for (int32 i = 0; i < Chunk.RigidVertices.Num(); i++)
+		{
+			Chunk.RigidVertices[i].Position = LODToBaseTransform.TransformPosition(Chunk.RigidVertices[i].Position);
+			Chunk.RigidVertices[i].TangentX = LODToBaseTransform.TransformVector(Chunk.RigidVertices[i].TangentX);
+			Chunk.RigidVertices[i].TangentY = LODToBaseTransform.TransformVector(Chunk.RigidVertices[i].TangentY);
+			Chunk.RigidVertices[i].TangentZ = LODToBaseTransform.TransformVector(Chunk.RigidVertices[i].TangentZ);
+		}
 
-			LODInfo.LODMaterialMap.Add(LODMatIndex);
+		// Fix up soft verts.
+		for (int32 i = 0; i < Chunk.SoftVertices.Num(); i++)
+		{
+			Chunk.SoftVertices[i].Position = LODToBaseTransform.TransformPosition(Chunk.SoftVertices[i].Position);
+			Chunk.SoftVertices[i].TangentX = LODToBaseTransform.TransformVector(Chunk.SoftVertices[i].TangentX);
+			Chunk.SoftVertices[i].TangentY = LODToBaseTransform.TransformVector(Chunk.SoftVertices[i].TangentY);
+			Chunk.SoftVertices[i].TangentZ = LODToBaseTransform.TransformVector(Chunk.SoftVertices[i].TangentZ);
 		}
+	}
 
-		BaseSkeletalMesh->PreEditChange(NULL);
+	if (bNeedToReregister)
+	{
+		// Shut down the skeletal mesh component that is previewing this mesh.
+		if (ReregisterAssociatedComponents)
+		{
+			FMultiComponentReregisterContext ReregisterContext(*ReregisterAssociatedComponents);
+			// wait until resources are released
+			FlushRenderingCommands();
 
-		// Copy data from the index buffer as it cant be directly copied.
-		FMultiSizeIndexContainerData BufferData;
-		NewLODModel.MultiSizeIndexContainer.GetIndexBufferData( BufferData );
+			InsertNewLODToBaseSkeletalMesh(InSkeletalMesh, BaseSkeletalMesh, DesiredLOD);
 
-		// Assign new FStaticLODModel to desired slot in selected skeletal mesh.
-		BaseImportedResource->LODModels[DesiredLOD] = NewLODModel;
+			// ReregisterContexts go out of scope here, reregistering associated components with the scene.
+		}
+		else
+		{
+			TComponentReregisterContext<USkinnedMeshComponent> ComponentReregisterContext;
 
-		// Regenerate the index buffer.
-		BaseImportedResource->LODModels[DesiredLOD].MultiSizeIndexContainer.RebuildIndexBuffer( BufferData );
-		
-		// rebuild vertex buffers and reinit RHI resources
-		
-		BaseSkeletalMesh->PostEditChange();		
+			InsertNewLODToBaseSkeletalMesh(InSkeletalMesh, BaseSkeletalMesh, DesiredLOD);
 
-		// ReregisterContexts go out of scope here, reregistering skel components to the scene.
+			// ReregisterContexts go out of scope here, reregistering skel components with the scene.
+		}
+	}
+	else
+	{
+		InsertNewLODToBaseSkeletalMesh(InSkeletalMesh, BaseSkeletalMesh, DesiredLOD);
 	}
 
+	return true;
 }
 
-
 /**
  * A class encapsulating morph target processing that occurs during import on a separate thread
  */
diff --git a/Engine/Source/Editor/UnrealEd/Private/FbxImporter.h b/Engine/Source/Editor/UnrealEd/Private/FbxImporter.h
--- a/Engine/Source/Editor/UnrealEd/Private/FbxImporter.h
+++ b/Engine/Source/Editor/UnrealEd/Private/FbxImporter.h
@@ -445,8 +445,11 @@ public:
 	 * @param InSkeletalMesh - LOD mesh object
 	 * @param BaseSkeletalMesh - base mesh object
 	 * @param DesiredLOD - LOD level
+	 * @param bNeedToReregister - if true, re-register this skeletal mesh to shut down the skeletal mesh component that is previewing this mesh. 
+									But you can set this to false when in the first loading before rendering this mesh for a performance issue 
+	   @param ReregisterAssociatedComponents - if NULL, just re-registers all SkinnedMeshComponents but if you set the specific components, will only re-registers those components
 	 */
-	void ImportSkeletalMeshLOD(USkeletalMesh* InSkeletalMesh, USkeletalMesh* BaseSkeletalMesh, int32 DesiredLOD);
+	bool ImportSkeletalMeshLOD(USkeletalMesh* InSkeletalMesh, USkeletalMesh* BaseSkeletalMesh, int32 DesiredLOD, bool bNeedToReregister = true, TArray<UActorComponent*>* ReregisterAssociatedComponents = NULL);
 
 	/**
 	 * Empties the FBX scene, releasing its memory.
@@ -586,6 +589,15 @@ private:
 	 */
 	void ImportMorphTargetsInternal( TArray<FbxNode*>& SkelMeshNodeArray, USkeletalMesh* BaseSkelMesh, UObject* Parent, const FString& InFilename, int32 LODIndex );
 
+	/**
+	* sub-method called from ImportSkeletalMeshLOD method
+	*
+	* @param InSkeletalMesh - newly created mesh used as LOD
+	* @param BaseSkeletalMesh - the destination mesh object 
+	* @param DesiredLOD - the LOD index to import into. A new LOD entry is created if one doesn't exist
+	*/
+	void InsertNewLODToBaseSkeletalMesh(USkeletalMesh* InSkeletalMesh, USkeletalMesh* BaseSkeletalMesh, int32 DesiredLOD);
+
 public:
 	// current Fbx scene we are importing. Make sure to release it after import
 	FbxScene* Scene;
diff --git a/Engine/Source/Editor/UnrealEd/Private/FbxMeshUtils.cpp b/Engine/Source/Editor/UnrealEd/Private/FbxMeshUtils.cpp
--- a/Engine/Source/Editor/UnrealEd/Private/FbxMeshUtils.cpp
+++ b/Engine/Source/Editor/UnrealEd/Private/FbxMeshUtils.cpp
@@ -501,7 +501,7 @@ namespace FbxMeshUtils
 					bool bImportSucceeded = false;
 					if( TempSkelMesh )
 					{
-						bImportSucceeded = ConvertSkeletalMeshToLOD(TempSkelMesh, SelectedSkelMesh, SelectedLOD);
+						bImportSucceeded = FFbxImporter->ImportSkeletalMeshLOD(TempSkelMesh, SelectedSkelMesh, SelectedLOD);
 
 						// Mark package containing skeletal mesh as dirty.
 						SelectedSkelMesh->MarkPackageDirty();
@@ -551,7 +551,7 @@ namespace FbxMeshUtils
 	void ImportMeshLODDialog( class UObject* SelectedMesh, int32 LODLevel )
 	{
 		if(!SelectedMesh)
-	{
+		{
 			return;
 		}
 
@@ -624,265 +624,6 @@ namespace FbxMeshUtils
 			}
 		}
 	}
-
-	bool ConvertSkeletalMeshToLOD(USkeletalMesh* InSkeletalMesh, USkeletalMesh* SelectedSkelMesh, int32 DesiredLOD)
-	{
-		check(InSkeletalMesh);
-
-		FSkeletalMeshResource* ImportedResource = InSkeletalMesh->GetImportedResource();
-		FSkeletalMeshResource* DestImportedResource = SelectedSkelMesh->GetImportedResource();
-
-		// Now we copy the base FStaticLODModel from the imported skeletal mesh as the new LOD in the selected mesh.
-		check(ImportedResource->LODModels.Num() == 1);
-
-		UnFbx::FFbxImporter* FFbxImporter = UnFbx::FFbxImporter::GetInstance();
-
-		// Names of root bones must match.
-		if(InSkeletalMesh->RefSkeleton.GetBoneName(0) != SelectedSkelMesh->RefSkeleton.GetBoneName(0))
-		{
-			FFbxImporter->AddTokenizedErrorMessage(FTokenizedMessage::Create(EMessageSeverity::Error, FText::Format(LOCTEXT("LODRootNameIncorrect", "Root bone in LOD is '{0}' instead of '{1}'.\nImport failed."),
-				FText::FromName(InSkeletalMesh->RefSkeleton.GetBoneName(0)), FText::FromName(SelectedSkelMesh->RefSkeleton.GetBoneName(0)))));
-
-			return false;
-		}
-
-		// We do some checking here that for every bone in the mesh we just imported, it's in our base ref skeleton, and the parent is the same.
-		for(int32 i=0; i<InSkeletalMesh->RefSkeleton.GetNum(); i++)
-		{
-			int32 LODBoneIndex = i;
-			FName LODBoneName = InSkeletalMesh->RefSkeleton.GetBoneName(LODBoneIndex);
-			int32 BaseBoneIndex = SelectedSkelMesh->RefSkeleton.FindBoneIndex(LODBoneName);
-			if( BaseBoneIndex == INDEX_NONE )
-			{
-				// If we could not find the bone from this LOD in base mesh - we fail.
-				FFbxImporter->AddTokenizedErrorMessage(FTokenizedMessage::Create(EMessageSeverity::Error, FText::Format(LOCTEXT("LODBoneDoesNotMatch", "Bone '{0}' not found in base SkeletalMesh '{1}'.\nImport failed."),
-					FText::FromName(LODBoneName), FText::FromString(SelectedSkelMesh->GetName()))));
-
-				return false;
-			}
-
-			if(i>0)
-			{
-				int32 LODParentIndex = InSkeletalMesh->RefSkeleton.GetParentIndex(LODBoneIndex);
-				FName LODParentName = InSkeletalMesh->RefSkeleton.GetBoneName(LODParentIndex);
-
-				int32 BaseParentIndex = SelectedSkelMesh->RefSkeleton.GetParentIndex(BaseBoneIndex);
-				FName BaseParentName = SelectedSkelMesh->RefSkeleton.GetBoneName(BaseParentIndex);
-
-				if(LODParentName != BaseParentName)
-				{
-					// If bone has different parents, display an error and don't allow import.
-					FFbxImporter->AddTokenizedErrorMessage(FTokenizedMessage::Create(EMessageSeverity::Error, FText::Format(LOCTEXT("LODBoneHasIncorrectParent", "Bone '{0}' in LOD has parent '{1}' instead of '{2}'"),
-						FText::FromName(LODBoneName), FText::FromName(LODParentName), FText::FromName(BaseParentName))));
-
-					return false;
-				}
-			}
-		}
-
-		FStaticLODModel& NewLODModel = ImportedResource->LODModels[0];
-
-		// Enforce LODs having only single-influence vertices.
-		bool bCheckSingleInfluence;
-		GConfig->GetBool( TEXT("AnimSetViewer"), TEXT("CheckSingleInfluenceLOD"), bCheckSingleInfluence, GEditorIni );
-		if( bCheckSingleInfluence && 
-			DesiredLOD > 0 )
-		{
-			for(int32 ChunkIndex = 0;ChunkIndex < NewLODModel.Chunks.Num();ChunkIndex++)
-			{
-				if(NewLODModel.Chunks[ChunkIndex].SoftVertices.Num() > 0)
-				{
-					FFbxImporter->AddTokenizedErrorMessage(FTokenizedMessage::Create(EMessageSeverity::Warning, FText(LOCTEXT("LODHasSoftVertices", "Warning: The mesh LOD you are importing has some vertices with more than one influence."))));
-				}
-			}
-		}
-
-		// If this LOD is going to be the lowest one, we check all bones we have sockets on are present in it.
-		if( DesiredLOD == DestImportedResource->LODModels.Num() || 
-			DesiredLOD == DestImportedResource->LODModels.Num()-1 )
-		{
-			const TArray<USkeletalMeshSocket*> Sockets = SelectedSkelMesh->GetMeshOnlySocketList();
-
-			for(int32 i=0; i<Sockets.Num(); i++)
-			{
-				// Find bone index the socket is attached to.
-				USkeletalMeshSocket* Socket = Sockets[i];
-				int32 SocketBoneIndex = InSkeletalMesh->RefSkeleton.FindBoneIndex( Socket->BoneName );
-
-				// If this LOD does not contain the socket bone, abort import.
-				if( SocketBoneIndex == INDEX_NONE )
-				{
-					FFbxImporter->AddTokenizedErrorMessage(FTokenizedMessage::Create(EMessageSeverity::Error, FText::Format(LOCTEXT("LODMissingSocketBone", "This LOD is missing bone '{0}' used by socket '{1}'.\nAborting import."),
-						FText::FromName(Socket->BoneName), FText::FromName(Socket->SocketName))));
-					
-					return false;
-				}
-			}
-		}
-
-		// Fix up the ActiveBoneIndices array.
-		for(int32 i=0; i<NewLODModel.ActiveBoneIndices.Num(); i++)
-		{
-			int32 LODBoneIndex = NewLODModel.ActiveBoneIndices[i];
-			FName LODBoneName = InSkeletalMesh->RefSkeleton.GetBoneName(LODBoneIndex);
-			int32 BaseBoneIndex = SelectedSkelMesh->RefSkeleton.FindBoneIndex(LODBoneName);
-			NewLODModel.ActiveBoneIndices[i] = BaseBoneIndex;
-		}
-
-		// Fix up the chunk BoneMaps.
-		for(int32 ChunkIndex = 0;ChunkIndex < NewLODModel.Chunks.Num();ChunkIndex++)
-		{
-			FSkelMeshChunk& Chunk = NewLODModel.Chunks[ChunkIndex];
-			for(int32 i=0; i<Chunk.BoneMap.Num(); i++)
-			{
-				int32 LODBoneIndex = Chunk.BoneMap[i];
-				FName LODBoneName = InSkeletalMesh->RefSkeleton.GetBoneName(LODBoneIndex);
-				int32 BaseBoneIndex = SelectedSkelMesh->RefSkeleton.FindBoneIndex(LODBoneName);
-				Chunk.BoneMap[i] = BaseBoneIndex;
-			}
-		}
-
-		// Create the RequiredBones array in the LODModel from the ref skeleton.
-		for(int32 i=0; i<NewLODModel.RequiredBones.Num(); i++)
-		{
-			FName LODBoneName = InSkeletalMesh->RefSkeleton.GetBoneName(NewLODModel.RequiredBones[i]);
-			int32 BaseBoneIndex = SelectedSkelMesh->RefSkeleton.FindBoneIndex(LODBoneName);
-			if(BaseBoneIndex != INDEX_NONE)
-			{
-				NewLODModel.RequiredBones[i] = BaseBoneIndex;
-			}
-			else
-			{
-				NewLODModel.RequiredBones.RemoveAt(i--);
-			}
-		}
-
-		// Also sort the RequiredBones array to be strictly increasing.
-		NewLODModel.RequiredBones.Sort();
-
-		// To be extra-nice, we apply the difference between the root transform of the meshes to the verts.
-		FMatrix LODToBaseTransform = InSkeletalMesh->GetRefPoseMatrix(0).Inverse() * SelectedSkelMesh->GetRefPoseMatrix(0);
-
-		for(int32 ChunkIndex = 0;ChunkIndex < NewLODModel.Chunks.Num();ChunkIndex++)
-		{
-			FSkelMeshChunk& Chunk = NewLODModel.Chunks[ChunkIndex];
-			// Fix up rigid verts.
-			for(int32 i=0; i<Chunk.RigidVertices.Num(); i++)
-			{
-				Chunk.RigidVertices[i].Position = LODToBaseTransform.TransformPosition( Chunk.RigidVertices[i].Position );
-				Chunk.RigidVertices[i].TangentX = LODToBaseTransform.TransformVector( Chunk.RigidVertices[i].TangentX );
-				Chunk.RigidVertices[i].TangentY = LODToBaseTransform.TransformVector( Chunk.RigidVertices[i].TangentY );
-				Chunk.RigidVertices[i].TangentZ = LODToBaseTransform.TransformVector( Chunk.RigidVertices[i].TangentZ );
-			}
-
-			// Fix up soft verts.
-			for(int32 i=0; i<Chunk.SoftVertices.Num(); i++)
-			{
-				Chunk.SoftVertices[i].Position = LODToBaseTransform.TransformPosition( Chunk.SoftVertices[i].Position );
-				Chunk.SoftVertices[i].TangentX = LODToBaseTransform.TransformVector( Chunk.SoftVertices[i].TangentX );
-				Chunk.SoftVertices[i].TangentY = LODToBaseTransform.TransformVector( Chunk.SoftVertices[i].TangentY );
-				Chunk.SoftVertices[i].TangentZ = LODToBaseTransform.TransformVector( Chunk.SoftVertices[i].TangentZ );
-			}
-		}
-
-		// Shut down the skeletal mesh component that is previewing this mesh.
-		{
-			TComponentReregisterContext<USkinnedMeshComponent> ComponentReregisterContext;
-
-			// If we want to add this as a new LOD to this mesh - add to LODModels/LODInfo array.
-			if(DesiredLOD == DestImportedResource->LODModels.Num())
-			{
-				new(DestImportedResource->LODModels)FStaticLODModel();
-
-				// Add element to LODInfo array.
-				SelectedSkelMesh->LODInfo.AddZeroed();
-				check( SelectedSkelMesh->LODInfo.Num() == DestImportedResource->LODModels.Num() );
-				SelectedSkelMesh->LODInfo[DesiredLOD] = InSkeletalMesh->LODInfo[0];
-			}
-			else
-			{
-				// if it's overwriting existing LOD, need to update section information
-				// update to the right # of sections 
-				// Set up LODMaterialMap to number of materials in new mesh.
-				// InSkeletalMesh->LOD 0 is the newly imported mesh
-				FSkeletalMeshLODInfo& LODInfo = SelectedSkelMesh->LODInfo[DesiredLOD];
-				FStaticLODModel& LODModel = ImportedResource->LODModels[0];
-				// if section # has been changed
- 				if (LODInfo.TriangleSortSettings.Num() != LODModel.Sections.Num())
- 				{
- 					// Save old information so that I can copy it over
- 					TArray<FTriangleSortSettings> OldTriangleSortSettings;
- 
- 					OldTriangleSortSettings = LODInfo.TriangleSortSettings;
- 
- 					// resize to the correct number
- 					LODInfo.TriangleSortSettings.Empty(LODModel.Sections.Num());
- 					// fill up data
- 					for ( int32 SectionIndex = 0 ; SectionIndex < LODModel.Sections.Num() ; ++SectionIndex )
- 					{
- 						// if found from previous data, copy over
- 						if ( SectionIndex < OldTriangleSortSettings.Num() )
- 						{
- 							LODInfo.TriangleSortSettings.Add( OldTriangleSortSettings[SectionIndex] );
- 						}
- 						else
- 						{
- 							// if not add default data
- 							LODInfo.TriangleSortSettings.AddZeroed();
- 						}
- 					}
- 				}
-			}
-
-			// Set up LODMaterialMap to number of materials in new mesh.
-			FSkeletalMeshLODInfo& LODInfo = SelectedSkelMesh->LODInfo[DesiredLOD];
-			LODInfo.LODMaterialMap.Empty();
-
-			// Now set up the material mapping array.
-			for(int32 MatIdx = 0; MatIdx < InSkeletalMesh->Materials.Num(); MatIdx++)
-			{
-				// Try and find the auto-assigned material in the array.
-				int32 LODMatIndex = INDEX_NONE;
-				if (InSkeletalMesh->Materials[MatIdx].MaterialInterface != NULL)
-				{
-					LODMatIndex = SelectedSkelMesh->Materials.Find( InSkeletalMesh->Materials[MatIdx] );
-				}
-
-				// If we didn't just use the index - but make sure its within range of the Materials array.
-				if(LODMatIndex == INDEX_NONE)
-				{
-					LODMatIndex = FMath::Clamp(MatIdx, 0, SelectedSkelMesh->Materials.Num() - 1);
-				}
-
-				LODInfo.LODMaterialMap.Add(LODMatIndex);
-			}
-
-			// if new LOD has more material slot, add the extra to main skeletal
-			if ( SelectedSkelMesh->Materials.Num() < InSkeletalMesh->Materials.Num() )
-			{
-				SelectedSkelMesh->Materials.AddZeroed(InSkeletalMesh->Materials.Num()-SelectedSkelMesh->Materials.Num());
-			}
-
-			// Release all resources before replacing the model
-			SelectedSkelMesh->PreEditChange(NULL);
-
-			// Index buffer will be destroyed when we copy the LOD model so we must copy the index buffer and reinitialize it after the copy
-			FMultiSizeIndexContainerData Data;
-			NewLODModel.MultiSizeIndexContainer.GetIndexBufferData( Data );
-
-			// Assign new FStaticLODModel to desired slot in selected skeletal mesh.
-			DestImportedResource->LODModels[DesiredLOD] = NewLODModel;
-
-			DestImportedResource->LODModels[DesiredLOD].MultiSizeIndexContainer.RebuildIndexBuffer( Data );
-
-			// rebuild vertex buffers and reinit RHI resources
-			SelectedSkelMesh->PostEditChange();		
-
-			// ReregisterContexts go out of scope here, reregistering skel components with the scene.
-	}
-
-		return true;
-	}
 }  //end namespace MeshUtils
 
 #undef LOCTEXT_NAMESPACE
\ No newline at end of file
diff --git a/Engine/Source/Editor/UnrealEd/Private/LODUtilities.cpp b/Engine/Source/Editor/UnrealEd/Private/LODUtilities.cpp
--- a/Engine/Source/Editor/UnrealEd/Private/LODUtilities.cpp
+++ b/Engine/Source/Editor/UnrealEd/Private/LODUtilities.cpp
@@ -247,7 +247,7 @@ void FLODUtilities::ImportMeshLOD( FSkeletalMeshUpdateContext& UpdateContext, in
 						bool bImportSucceeded = false;
 						if( TempSkelMesh )
 						{
-							bImportSucceeded = ImportMeshLOD(UpdateContext, TempSkelMesh, SelectedLOD);
+							bImportSucceeded = FFbxImporter->ImportSkeletalMeshLOD(TempSkelMesh, UpdateContext.SkeletalMesh, SelectedLOD, true, UpdateContext.AssociatedComponents);
 
 							//Notify calling system of change
 							UpdateContext.OnLODChanged.ExecuteIfBound();
@@ -281,261 +281,6 @@ void FLODUtilities::ImportMeshLOD( FSkeletalMeshUpdateContext& UpdateContext, in
 #endif
 }
 
-bool FLODUtilities::ImportMeshLOD(FSkeletalMeshUpdateContext& UpdateContext, USkeletalMesh* InputSkeletalMesh, int32 DesiredLOD)
-{
-	check(InputSkeletalMesh);
-
-	USkeletalMesh* DestSkeletalMesh = UpdateContext.SkeletalMesh;
-	FSkeletalMeshResource* InputSkelMeshResource = InputSkeletalMesh->GetImportedResource();
-	FSkeletalMeshResource* DestSkelMeshResource = UpdateContext.SkeletalMesh->GetImportedResource();
-
-	// Now we copy the base FStaticLODModel from the imported skeletal mesh as the new LOD in the selected mesh.
-	check(InputSkelMeshResource->LODModels.Num() == 1);
-
-	// Names of root bones must match.
-	if(InputSkeletalMesh->RefSkeleton.GetBoneName(0) != DestSkeletalMesh->RefSkeleton.GetBoneName(0))
-	{
-		FMessageDialog::Open( EAppMsgType::Ok, FText::Format(NSLOCTEXT("ImportSkeletalMeshLOD", "LODRootNameIncorrect", "Root bone in LOD is '{0}' instead of '{1}'.\nImport failed."),
-			FText::FromName(InputSkeletalMesh->RefSkeleton.GetBoneName(0)), FText::FromName(DestSkeletalMesh->RefSkeleton.GetBoneName(0))) );
-		return false;
-	}
-
-	// We do some checking here that for every bone in the mesh we just imported, it's in our base ref skeleton, and the parent is the same.
-	for(int32 i=0; i<InputSkeletalMesh->RefSkeleton.GetNum(); i++)
-	{
-		int32 LODBoneIndex = i;
-		FName LODBoneName = InputSkeletalMesh->RefSkeleton.GetBoneName(LODBoneIndex);
-		int32 BaseBoneIndex = DestSkeletalMesh->RefSkeleton.FindBoneIndex(LODBoneName);
-		if( BaseBoneIndex == INDEX_NONE )
-		{
-			// If we could not find the bone from this LOD in base mesh - we fail.
-			FMessageDialog::Open( EAppMsgType::Ok, FText::Format(NSLOCTEXT("ImportSkeletalMeshLOD", "LODBoneDoesNotMatch", "Bone '{0}' not found in base SkeletalMesh '{1}'.\nImport failed."),
-				FText::FromName(LODBoneName), FText::FromString(DestSkeletalMesh->GetName())) );
-			return false;
-		}
-
-		if(i>0)
-		{
-			int32 LODParentIndex = InputSkeletalMesh->RefSkeleton.GetParentIndex(LODBoneIndex);
-			FName LODParentName = InputSkeletalMesh->RefSkeleton.GetBoneName(LODParentIndex);
-
-			int32 BaseParentIndex = DestSkeletalMesh->RefSkeleton.GetParentIndex(BaseBoneIndex);
-			FName BaseParentName = DestSkeletalMesh->RefSkeleton.GetBoneName(BaseParentIndex);
-
-			if(LODParentName != BaseParentName)
-			{
-				// If bone has different parents, display an error and don't allow import.
-				FMessageDialog::Open( EAppMsgType::Ok, FText::Format(NSLOCTEXT("ImportSkeletalMeshLOD", "LODBoneHasIncorrectParent", "Bone '{0}' in LOD has parent '{1}' instead of '{2}'"),
-					FText::FromName(LODBoneName), FText::FromName(LODParentName), FText::FromName(BaseParentName)) );
-				return false;
-			}
-		}
-	}
-
-	FStaticLODModel& NewLODModel = InputSkelMeshResource->LODModels[0];
-
-	// Enforce LODs having only single-influence vertices.
-	bool bCheckSingleInfluence;
-	GConfig->GetBool( TEXT("AnimSetViewer"), TEXT("CheckSingleInfluenceLOD"), bCheckSingleInfluence, GEditorIni );
-	if( bCheckSingleInfluence && 
-		DesiredLOD > 0 )
-	{
-		for(int32 ChunkIndex = 0;ChunkIndex < NewLODModel.Chunks.Num();ChunkIndex++)
-		{
-			if(NewLODModel.Chunks[ChunkIndex].SoftVertices.Num() > 0)
-			{
-				FMessageDialog::Open( EAppMsgType::Ok, NSLOCTEXT("ImportSkeletalMeshLOD", "LODHasSoftVertices", "Warning: The mesh LOD you are importing has some vertices with more than one influence.") );
-			}
-		}
-	}
-
-	// If this LOD is going to be the lowest one, we check all bones we have sockets on are present in it.
-	if( DesiredLOD == DestSkelMeshResource->LODModels.Num() || 
-		DesiredLOD == DestSkelMeshResource->LODModels.Num()-1 )
-	{
-		const TArray<USkeletalMeshSocket*>& SocketList = DestSkeletalMesh->GetActiveSocketList();
-
-		for(int32 i=0; i<SocketList.Num(); i++)
-		{
-			// Find bone index the socket is attached to.
-			USkeletalMeshSocket* Socket = SocketList[i];
-			int32 SocketBoneIndex = InputSkeletalMesh->RefSkeleton.FindBoneIndex( Socket->BoneName );
-
-			// If this LOD does not contain the socket bone, abort import.
-			if( SocketBoneIndex == INDEX_NONE )
-			{
-				FMessageDialog::Open( EAppMsgType::Ok, FText::Format(NSLOCTEXT("ImportSkeletalMeshLOD", "LODMissingSocketBone", "This LOD is missing bone '{0}' used by socket '{1}'.\nAborting import."),
-					FText::FromName(Socket->BoneName), FText::FromName(Socket->SocketName)) );
-				return false;
-			}
-		}
-	}
-
-	// Fix up the ActiveBoneIndices array.
-	for(int32 i=0; i<NewLODModel.ActiveBoneIndices.Num(); i++)
-	{
-		int32 LODBoneIndex = NewLODModel.ActiveBoneIndices[i];
-		FName LODBoneName = InputSkeletalMesh->RefSkeleton.GetBoneName(LODBoneIndex);
-		int32 BaseBoneIndex = DestSkeletalMesh->RefSkeleton.FindBoneIndex(LODBoneName);
-		NewLODModel.ActiveBoneIndices[i] = BaseBoneIndex;
-	}
-
-	// Fix up the chunk BoneMaps.
-	for(int32 ChunkIndex = 0;ChunkIndex < NewLODModel.Chunks.Num();ChunkIndex++)
-	{
-		FSkelMeshChunk& Chunk = NewLODModel.Chunks[ChunkIndex];
-		for(int32 i=0; i<Chunk.BoneMap.Num(); i++)
-		{
-			int32 LODBoneIndex = Chunk.BoneMap[i];
-			FName LODBoneName = InputSkeletalMesh->RefSkeleton.GetBoneName(LODBoneIndex);
-			int32 BaseBoneIndex = DestSkeletalMesh->RefSkeleton.FindBoneIndex(LODBoneName);
-			Chunk.BoneMap[i] = BaseBoneIndex;
-		}
-	}
-
-	// Create the RequiredBones array in the LODModel from the ref skeleton.
-	for(int32 i=0; i<NewLODModel.RequiredBones.Num(); i++)
-	{
-		FName LODBoneName = InputSkeletalMesh->RefSkeleton.GetBoneName(NewLODModel.RequiredBones[i]);
-		int32 BaseBoneIndex = DestSkeletalMesh->RefSkeleton.FindBoneIndex(LODBoneName);
-		if(BaseBoneIndex != INDEX_NONE)
-		{
-			NewLODModel.RequiredBones[i] = BaseBoneIndex;
-		}
-		else
-		{
-			NewLODModel.RequiredBones.RemoveAt(i--);
-		}
-	}
-
-	// Also sort the RequiredBones array to be strictly increasing.
-	NewLODModel.RequiredBones.Sort();
-
-	// To be extra-nice, we apply the difference between the root transform of the meshes to the verts.
-	FMatrix LODToBaseTransform = InputSkeletalMesh->GetRefPoseMatrix(0).Inverse() * DestSkeletalMesh->GetRefPoseMatrix(0);
-
-	for(int32 ChunkIndex = 0;ChunkIndex < NewLODModel.Chunks.Num();ChunkIndex++)
-	{
-		FSkelMeshChunk& Chunk = NewLODModel.Chunks[ChunkIndex];
-		// Fix up rigid verts.
-		for(int32 i=0; i<Chunk.RigidVertices.Num(); i++)
-		{
-			Chunk.RigidVertices[i].Position = LODToBaseTransform.TransformPosition( Chunk.RigidVertices[i].Position );
-			Chunk.RigidVertices[i].TangentX = LODToBaseTransform.TransformVector( Chunk.RigidVertices[i].TangentX );
-			Chunk.RigidVertices[i].TangentY = LODToBaseTransform.TransformVector( Chunk.RigidVertices[i].TangentY );
-			Chunk.RigidVertices[i].TangentZ = LODToBaseTransform.TransformVector( Chunk.RigidVertices[i].TangentZ );
-		}
-
-		// Fix up soft verts.
-		for(int32 i=0; i<Chunk.SoftVertices.Num(); i++)
-		{
-			Chunk.SoftVertices[i].Position = LODToBaseTransform.TransformPosition( Chunk.SoftVertices[i].Position );
-			Chunk.SoftVertices[i].TangentX = LODToBaseTransform.TransformVector( Chunk.SoftVertices[i].TangentX );
-			Chunk.SoftVertices[i].TangentY = LODToBaseTransform.TransformVector( Chunk.SoftVertices[i].TangentY );
-			Chunk.SoftVertices[i].TangentZ = LODToBaseTransform.TransformVector( Chunk.SoftVertices[i].TangentZ );
-		}
-	}
-
-	// Shut down the skeletal mesh component that is previewing this mesh.
-	{
-		FMultiComponentReregisterContext ReregisterContext(UpdateContext.AssociatedComponents);
-		FlushRenderingCommands();
-
-		// If we want to add this as a new LOD to this mesh - add to LODModels/LODInfo array.
-		if(DesiredLOD == DestSkelMeshResource->LODModels.Num())
-		{
-			new(DestSkelMeshResource->LODModels)FStaticLODModel();
-
-			// Add element to LODInfo array.
-			DestSkeletalMesh->LODInfo.AddZeroed();
-			check( DestSkeletalMesh->LODInfo.Num() == DestSkelMeshResource->LODModels.Num() );
-			DestSkeletalMesh->LODInfo[DesiredLOD] = InputSkeletalMesh->LODInfo[0];
-		}
-		else
-		{
-			// if it's overwriting existing LOD, need to update section information
-			// update to the right # of sections 
-			// Set up LODMaterialMap to number of materials in new mesh.
-			// InputSkeletalMesh->LOD 0 is the newly imported mesh
-			FSkeletalMeshLODInfo& LODInfo = DestSkeletalMesh->LODInfo[DesiredLOD];
-			FStaticLODModel& LODModel = InputSkelMeshResource->LODModels[0];
-			// if section # has been changed
- 			if (LODInfo.TriangleSortSettings.Num() != LODModel.Sections.Num())
- 			{
- 				// Save old information so that I can copy it over
- 				TArray<FTriangleSortSettings> OldTriangleSortSettings;
- 
- 				OldTriangleSortSettings = LODInfo.TriangleSortSettings;
- 
- 				// resize to the correct number
- 				LODInfo.TriangleSortSettings.Empty(LODModel.Sections.Num());
- 				// fill up data
- 				for ( int32 SectionIndex = 0 ; SectionIndex < LODModel.Sections.Num() ; ++SectionIndex )
- 				{
- 					// if found from previous data, copy over
- 					if ( SectionIndex < OldTriangleSortSettings.Num() )
- 					{
- 						LODInfo.TriangleSortSettings.Add( OldTriangleSortSettings[SectionIndex] );
- 					}
- 					else
- 					{
- 						// if not add default data
- 						LODInfo.TriangleSortSettings.AddZeroed();
- 					}
- 				}
- 			}
-		}
-
-		// Set up LODMaterialMap to number of materials in new mesh.
-		FSkeletalMeshLODInfo& LODInfo = DestSkeletalMesh->LODInfo[DesiredLOD];
-		LODInfo.LODMaterialMap.Empty();
-
-		// Now set up the material mapping array.
-		for(int32 MatIdx = 0; MatIdx < InputSkeletalMesh->Materials.Num(); MatIdx++)
-		{
-			// Try and find the auto-assigned material in the array.
-			int32 LODMatIndex = INDEX_NONE;
-			if (InputSkeletalMesh->Materials[MatIdx].MaterialInterface != NULL)
-			{
-				LODMatIndex = DestSkeletalMesh->Materials.Find( InputSkeletalMesh->Materials[MatIdx] );
-			}
-
-			// If we didn't just use the index - but make sure its within range of the Materials array.
-			if(LODMatIndex == INDEX_NONE)
-			{
-				LODMatIndex = FMath::Clamp(MatIdx, 0, DestSkeletalMesh->Materials.Num() - 1);
-			}
-
-			LODInfo.LODMaterialMap.Add(LODMatIndex);
-		}
-
-		// if new LOD has more material slot, add the extra to main skeletal
-		if ( DestSkeletalMesh->Materials.Num() < InputSkeletalMesh->Materials.Num() )
-		{
-			DestSkeletalMesh->Materials.AddZeroed(InputSkeletalMesh->Materials.Num()-DestSkeletalMesh->Materials.Num());
-		}
-
-		// Release all resources before replacing the model
-		DestSkeletalMesh->PreEditChange(NULL);
-
-		// Index buffer will be destroyed when we copy the LOD model so we must copy the index buffer and reinitialize it after the copy
-		FMultiSizeIndexContainerData Data;
-		NewLODModel.MultiSizeIndexContainer.GetIndexBufferData( Data );
-
-		// Assign new FStaticLODModel to desired slot in selected skeletal mesh.
-		DestSkelMeshResource->LODModels[DesiredLOD] = NewLODModel;
-
-		DestSkelMeshResource->LODModels[DesiredLOD].MultiSizeIndexContainer.RebuildIndexBuffer( Data );
-
-		// rebuild vertex buffers and reinit RHI resources
-		DestSkeletalMesh->PostEditChange();		
-
-		// ReregisterContexts go out of scope here, reregistering skel components with the scene.
-	}
-
-	return true;
-}
-
 void FLODUtilities::SimplifySkeletalMesh( FSkeletalMeshUpdateContext& UpdateContext, TArray<FSkeletalMeshOptimizationSettings> &InSettings )
 {
 	USkeletalMesh* SkeletalMesh = UpdateContext.SkeletalMesh;
diff --git a/Engine/Source/Editor/UnrealEd/Public/FbxMeshUtils.h b/Engine/Source/Editor/UnrealEd/Public/FbxMeshUtils.h
--- a/Engine/Source/Editor/UnrealEd/Public/FbxMeshUtils.h
+++ b/Engine/Source/Editor/UnrealEd/Public/FbxMeshUtils.h
@@ -75,12 +75,4 @@ namespace FbxMeshUtils
 	 */
 	UNREALED_API void ImportMeshLODDialog( class UObject* Mesh, int32 LODLevel );
 
-	/**
-	 * Imports a temporary skeletal mesh into the destination skeletal mesh at the given LOD
-	 *
-	 * @param TempMesh			The temp skeletal mesh which will be used as a LOD
-	 * @param DestinationMesh	The skeletal mesh which will use the temp as a LOD
-	 * @param LODLevel			The level of the lod to import to
-	 */
-	UNREALED_API bool ConvertSkeletalMeshToLOD( class USkeletalMesh* TempMesh, class USkeletalMesh* DestinationMesh, int32 LODLevel );
 }  //end namespace ExportMeshUtils
diff --git a/Engine/Source/Editor/UnrealEd/Public/LODUtilities.h b/Engine/Source/Editor/UnrealEd/Public/LODUtilities.h
--- a/Engine/Source/Editor/UnrealEd/Public/LODUtilities.h
+++ b/Engine/Source/Editor/UnrealEd/Public/LODUtilities.h
@@ -37,16 +37,6 @@ public:
 	*/
 	static void ImportMeshLOD( FSkeletalMeshUpdateContext& UpdateContext, int32 DesiredLOD);
 
-	/**
-	* Import the temporary skeletal mesh into the specified LOD of the destination skeletal mesh
-	*
-	* @param UpdateContext - The skeletal mesh and actor components to operate on.
-	* @param InputSkeletalMesh - newly created mesh used as LOD
-	* @param DesiredLOD - the LOD index to import into. A new LOD entry is created if one doesn't exist
-	* @return If true, import succeeded
-	*/
-	static bool ImportMeshLOD(FSkeletalMeshUpdateContext& UpdateContext, USkeletalMesh* InputSkeletalMesh, int32 DesiredLOD);
-	
 	/**
 	 *	Simplifies the static mesh based upon various user settings.
 	 *
