Added normal map export during landscape material flattening
World Browser LOD maps generator now creates landscape static meshes with normal maps

diff --git a/Engine/Source/Developer/MeshUtilities/Private/MeshUtilities.cpp b/Engine/Source/Developer/MeshUtilities/Private/MeshUtilities.cpp
--- a/Engine/Source/Developer/MeshUtilities/Private/MeshUtilities.cpp
+++ b/Engine/Source/Developer/MeshUtilities/Private/MeshUtilities.cpp
@@ -3062,6 +3062,8 @@ void FMeshUtilities::CreateProxyMesh(
 			MaterialMap.Add(RawMeshId, RawMeshMaterialMap);
 			// This is texture resolution for a landscape, probably need to be calculated using landscape size
 			UniqueMaterials.Last().DiffuseSize = FIntPoint(1024, 1024);
+			// FIXME: Landscape material exporter currently renders world space normal map, so it can't be merged with other meshes normal maps
+			UniqueMaterials.Last().NormalSize = FIntPoint::ZeroValue;
 			MaterialExportUtils::ExportMaterial(Landscape, UniqueMaterials.Last());
 			
 			//Store the bounds for each component
diff --git a/Engine/Source/Editor/UnrealEd/Private/EditorExporters.cpp b/Engine/Source/Editor/UnrealEd/Private/EditorExporters.cpp
--- a/Engine/Source/Editor/UnrealEd/Private/EditorExporters.cpp
+++ b/Engine/Source/Editor/UnrealEd/Private/EditorExporters.cpp
@@ -2579,14 +2579,9 @@ namespace MaterialExportUtils
 	bool ExportMaterial(UMaterialInterface* InMaterial, FFlattenMaterial& OutFlattenMaterial)
 	{
 		// Render diffuse property
+		if (OutFlattenMaterial.DiffuseSize.X > 0 && 
+			OutFlattenMaterial.DiffuseSize.Y > 0)
 		{
-			// Reset to default in case user specified invalid diffuse texture size
-			if (OutFlattenMaterial.DiffuseSize.X <= 0 || 
-				OutFlattenMaterial.DiffuseSize.Y <= 0)
-			{
-				OutFlattenMaterial.DiffuseSize = FFlattenMaterial().DiffuseSize;
-			}
-			
 			// Create temporary render target
 			UTextureRenderTarget2D* RenderTargetDiffuse = new UTextureRenderTarget2D(FPostConstructInitializeProperties());
 			check(RenderTargetDiffuse);
@@ -2616,15 +2611,10 @@ namespace MaterialExportUtils
 		}
 
 		// Render normal property
-		if (InMaterial->GetMaterial()->HasNormalConnected())
+		if (OutFlattenMaterial.NormalSize.X > 0 && 
+			OutFlattenMaterial.NormalSize.Y > 0 &&
+			InMaterial->GetMaterial()->HasNormalConnected())
 		{
-			// Reset to default in case user specified invalid normal texture size
-			if (OutFlattenMaterial.NormalSize.X <= 0 || 
-				OutFlattenMaterial.NormalSize.Y <= 0)
-			{
-				OutFlattenMaterial.NormalSize = FFlattenMaterial().NormalSize;
-			}
-			
 			// Create temporary render target
 			UTextureRenderTarget2D* RenderTargetNormal = new UTextureRenderTarget2D(FPostConstructInitializeProperties());
 			check(RenderTargetNormal);
@@ -2660,40 +2650,43 @@ namespace MaterialExportUtils
 	bool ExportMaterial(ALandscapeProxy* InLandscape, FFlattenMaterial& OutFlattenMaterial)
 	{
 		check(InLandscape);
-		TArray<ULandscapeComponent*> ComponentsToRender;
-		InLandscape->GetComponents<ULandscapeComponent>(ComponentsToRender);
-		if (ComponentsToRender.Num() == 0)
-		{
-			return false;
-		}
-
-		// Reset to default in case user specified invalid diffuse texture size
-		if (OutFlattenMaterial.DiffuseSize.X <= 0 || 
-			OutFlattenMaterial.DiffuseSize.Y <= 0)
-		{
-			OutFlattenMaterial.DiffuseSize = FFlattenMaterial().DiffuseSize;
-		}
-
-		// Normal map will not be used 
-		OutFlattenMaterial.NormalSamples.Empty();
-		OutFlattenMaterial.NormalSize = FIntPoint::ZeroValue;
 
 		FIntRect LandscapeRect = InLandscape->GetBoundingRect();
 		FVector MidPoint = FVector(LandscapeRect.Min, 0.f) + FVector(LandscapeRect.Size(), 0.f)*0.5f;
 		
 		FVector LandscapeCenter = InLandscape->GetTransform().TransformPosition(MidPoint);
 		FVector LandscapeExtent = FVector(LandscapeRect.Size(), 0.f)*InLandscape->GetActorScale()*0.5f; 
+
+		// We need to hide all primitives except target landscape
+		TSet<FPrimitiveComponentId> PrimitivesToHide;
+		for (TObjectIterator<UPrimitiveComponent> It; It; ++It)
 		{
-			UTextureRenderTarget2D* RenderTargetTexture = new UTextureRenderTarget2D(FPostConstructInitializeProperties());
-			check(RenderTargetTexture);
-			RenderTargetTexture->AddToRoot();
-			RenderTargetTexture->ClearColor = FLinearColor(0.0f, 0.0f, 0.0f, 0.0f);
-			RenderTargetTexture->InitCustomFormat(
-				OutFlattenMaterial.DiffuseSize.X, 
-				OutFlattenMaterial.DiffuseSize.Y, PF_B8G8R8A8, true);
-			FTextureRenderTargetResource* RenderTargetResource = RenderTargetTexture->GameThread_GetRenderTargetResource();
-			FSceneInterface* Scene = InLandscape->GetWorld()->Scene;
+			UPrimitiveComponent* PrimitiveComp = *It;
+
+			const bool bTargetPrim = PrimitiveComp->GetOuter() == InLandscape && 
+										PrimitiveComp->IsA<ULandscapeComponent>();
+					
+			if (!bTargetPrim && PrimitiveComp->IsRegistered() && PrimitiveComp->SceneProxy)
 			{
+				PrimitivesToHide.Add(PrimitiveComp->SceneProxy->GetPrimitiveComponentId());
+			}
+		}
+		
+		FSceneInterface* Scene = InLandscape->GetWorld()->Scene;
+		{
+			// Render diffuse texture
+			if (OutFlattenMaterial.DiffuseSize.X > 0 && 
+				OutFlattenMaterial.DiffuseSize.Y > 0)
+			{
+				UTextureRenderTarget2D* RenderTargetTexture = new UTextureRenderTarget2D(FPostConstructInitializeProperties());
+				check(RenderTargetTexture);
+				RenderTargetTexture->AddToRoot();
+				RenderTargetTexture->ClearColor = FLinearColor(0.0f, 0.0f, 0.0f, 0.0f);
+				RenderTargetTexture->InitCustomFormat(
+					OutFlattenMaterial.DiffuseSize.X, 
+					OutFlattenMaterial.DiffuseSize.Y, PF_B8G8R8A8, true);
+				FTextureRenderTargetResource* RenderTargetResource = RenderTargetTexture->GameThread_GetRenderTargetResource();
+			
 				FSceneViewFamilyContext ViewFamily(
 					FSceneViewFamily::ConstructionValues(RenderTargetResource, Scene, FEngineShowFlags(ESFIM_Game))
 						.SetWorldTimes(FApp::GetCurrentTime() - GStartTime, FApp::GetDeltaTime(), FApp::GetCurrentTime() - GStartTime)
@@ -2710,6 +2703,7 @@ namespace MaterialExportUtils
 				FSceneViewInitOptions ViewInitOptions;
 				ViewInitOptions.SetViewRectangle(FIntRect(0, 0, OutFlattenMaterial.DiffuseSize.X, OutFlattenMaterial.DiffuseSize.Y));
 				ViewInitOptions.ViewFamily = &ViewFamily;
+				ViewInitOptions.HiddenPrimitives = PrimitivesToHide;
 
 				ViewInitOptions.ViewMatrix = FTranslationMatrix(-LandscapeCenter);
 				ViewInitOptions.ViewMatrix*= FInverseRotationMatrix(InLandscape->GetActorRotation());
@@ -2725,27 +2719,8 @@ namespace MaterialExportUtils
 					0.5f / ZOffset,
 					ZOffset);
 
-				FSceneView* NewView = new FSceneView(ViewInitOptions);
-				ViewFamily.Views.Add(NewView);
-
-				// We need to hide all primitives except target landscape
-				for (TActorIterator<AActor> It(InLandscape->GetWorld()); It; ++It)
-				{
-					AActor* Actor = *It;
-					if (Actor && Actor != InLandscape)
-					{
-						TArray<UPrimitiveComponent*> PrimitiveComponents;
-						Actor->GetComponents<UPrimitiveComponent>(PrimitiveComponents);
-						for (UPrimitiveComponent* Component : PrimitiveComponents)
-						{
-							if (Component->IsRegistered() && Component->SceneProxy)
-							{
-								NewView->HiddenPrimitives.Add(Component->SceneProxy->GetPrimitiveComponentId());
-							}
-						}
-					}
-				}
-			
+				ViewFamily.Views.Add(new FSceneView(ViewInitOptions));
+					
 				FCanvas Canvas(RenderTargetResource, NULL, FApp::GetCurrentTime() - GStartTime, FApp::GetDeltaTime(), FApp::GetCurrentTime() - GStartTime);
 				Canvas.Clear(FLinearColor::Black);
 				GetRendererModule().BeginRenderingViewFamily(&Canvas, &ViewFamily);
@@ -2774,10 +2749,94 @@ namespace MaterialExportUtils
 					);
 			
 				FlushRenderingCommands();
+					
+				RenderTargetTexture->RemoveFromRoot();
+				RenderTargetTexture = nullptr;
 			}
 
-			RenderTargetTexture->RemoveFromRoot();
-			RenderTargetTexture = nullptr;
+			// Render normal map using BufferVisualizationMode=WorldNormal
+			// Final material should use world space instead of tangent space for normals
+			if (OutFlattenMaterial.NormalSize.X > 0 && 
+				OutFlattenMaterial.NormalSize.Y > 0)
+			{
+				UTextureRenderTarget2D* RenderTargetTexture = new UTextureRenderTarget2D(FPostConstructInitializeProperties());
+				check(RenderTargetTexture);
+				RenderTargetTexture->AddToRoot();
+				RenderTargetTexture->ClearColor = FLinearColor(0.0f, 0.0f, 0.0f, 0.0f);
+				RenderTargetTexture->InitCustomFormat(
+					OutFlattenMaterial.NormalSize.X, 
+					OutFlattenMaterial.NormalSize.Y, PF_B8G8R8A8, true);
+				FTextureRenderTargetResource* RenderTargetResource = RenderTargetTexture->GameThread_GetRenderTargetResource();
+
+				FSceneViewFamilyContext ViewFamily(
+					FSceneViewFamily::ConstructionValues(RenderTargetResource, Scene, FEngineShowFlags(ESFIM_Game))
+						.SetWorldTimes(FApp::GetCurrentTime() - GStartTime, FApp::GetDeltaTime(), FApp::GetCurrentTime() - GStartTime)
+					);
+
+				ViewFamily.EngineShowFlags.DisableAdvancedFeatures();
+				ViewFamily.EngineShowFlags.MotionBlur = 0;
+				ViewFamily.EngineShowFlags.Lighting = 0;
+				ViewFamily.EngineShowFlags.PostProcessing = 1;
+				ViewFamily.EngineShowFlags.VisualizeBuffer = 1;
+				ViewFamily.EngineShowFlags.LightFunctions = 0;
+				ViewFamily.EngineShowFlags.DynamicShadows = 0;
+				ViewFamily.EngineShowFlags.Atmosphere = 0;
+
+				FSceneViewInitOptions ViewInitOptions;
+				ViewInitOptions.SetViewRectangle(FIntRect(0, 0, OutFlattenMaterial.NormalSize.X, OutFlattenMaterial.NormalSize.Y));
+				ViewInitOptions.ViewFamily = &ViewFamily;
+				ViewInitOptions.HiddenPrimitives = PrimitivesToHide;
+
+				ViewInitOptions.ViewMatrix = FTranslationMatrix(-LandscapeCenter);
+				ViewInitOptions.ViewMatrix*= FInverseRotationMatrix(InLandscape->GetActorRotation());
+				ViewInitOptions.ViewMatrix*= FMatrix(	FPlane(1,	0,	0,	0),
+														FPlane(0,	-1,	0,	0),
+														FPlane(0,	0,	-1,	0),
+														FPlane(0,	0,	0,	1));
+				
+				const float ZOffset = WORLD_MAX;
+				ViewInitOptions.ProjectionMatrix =  FReversedZOrthoMatrix(
+					LandscapeExtent.X,
+					LandscapeExtent.Y,
+					0.5f / ZOffset,
+					ZOffset);
+
+				FSceneView* NewView = new FSceneView(ViewInitOptions);
+				NewView->CurrentBufferVisualizationMode = FName("WorldNormal");
+				ViewFamily.Views.Add(NewView);
+										
+				FCanvas Canvas(RenderTargetResource, NULL, FApp::GetCurrentTime() - GStartTime, FApp::GetDeltaTime(), FApp::GetCurrentTime() - GStartTime);
+				Canvas.Clear(FLinearColor::Black);
+				GetRendererModule().BeginRenderingViewFamily(&Canvas, &ViewFamily);
+
+				ENQUEUE_UNIQUE_RENDER_COMMAND_ONEPARAMETER(
+					UpdateThumbnailRTCommand,
+					FTextureRenderTargetResource*, RenderTargetResource, RenderTargetResource,
+				{
+					// Copy (resolve) the rendered thumbnail from the render target to its texture
+					RHICopyToResolveTarget(
+						RenderTargetResource->GetRenderTargetTexture(),		// Source texture
+						RenderTargetResource->TextureRHI,					// Dest texture
+						false,												// Do we need the source image content again?
+						FResolveParams() );									// Resolve parameters
+				});
+
+				OutFlattenMaterial.NormalSamples.Empty();
+				OutFlattenMaterial.NormalSamples.AddUninitialized(OutFlattenMaterial.NormalSize.X*OutFlattenMaterial.NormalSize.Y);
+
+				// Copy the contents of the remote texture to system memory
+				// NOTE: OutRawImageData must be a preallocated buffer!
+				RenderTargetResource->ReadPixelsPtr(
+					(FColor*)OutFlattenMaterial.NormalSamples.GetTypedData(), 
+					FReadSurfaceDataFlags(), 
+					FIntRect(0, 0, OutFlattenMaterial.NormalSize.X, OutFlattenMaterial.NormalSize.Y)
+					);
+			
+				FlushRenderingCommands();
+
+				RenderTargetTexture->RemoveFromRoot();
+				RenderTargetTexture = nullptr;
+			}
 		}
 
 		OutFlattenMaterial.MaterialId = InLandscape->GetLandscapeGuid();
diff --git a/Engine/Source/Editor/WorldBrowser/Private/Tiles/WorldTileCollectionModel.cpp b/Engine/Source/Editor/WorldBrowser/Private/Tiles/WorldTileCollectionModel.cpp
--- a/Engine/Source/Editor/WorldBrowser/Private/Tiles/WorldTileCollectionModel.cpp
+++ b/Engine/Source/Editor/WorldBrowser/Private/Tiles/WorldTileCollectionModel.cpp
@@ -1838,7 +1838,8 @@ bool FWorldTileCollectionModel::GenerateLODLevel(TSharedPtr<FLevelModel> InLevel
 		}
 					
 		// This is texture resolution for a landscape, probably need to be calculated using landscape size
-		LandscapeFlattenMaterial.DiffuseSize = FIntPoint(1024, 1024);
+		LandscapeFlattenMaterial.DiffuseSize	= FIntPoint(1024, 1024);
+		LandscapeFlattenMaterial.NormalSize		= FIntPoint(1024, 1024);
 		ExportMaterial(Landscape, LandscapeFlattenMaterial);
 		
 		// TODO: Disabled landscape mesh simplification, does not really needed since exporting highest landscape LOD already provide low triangle count
@@ -1869,6 +1870,9 @@ bool FWorldTileCollectionModel::GenerateLODLevel(TSharedPtr<FLevelModel> InLevel
 		// Construct landscape material
 		UMaterial* StaticLandscapeMaterial = MaterialExportUtils::CreateMaterial(
 			LandscapeFlattenMaterial, Package, TEXT("FlattenLandscapeMaterial"), RF_Public|RF_Standalone);
+		// Currently landscape exports world space normal map
+		StaticLandscapeMaterial->bTangentSpaceNormal = false;
+		StaticLandscapeMaterial->PostEditChange();
 	
 		// Construct landscape static mesh
 		FName StaticMeshName = MakeUniqueObjectName(Package, UStaticMesh::StaticClass(), TEXT("LandscapeStaticMesh"));
