Updating the PointLightPlacement test so that it can take a screenshot after each major step.

diff --git a/Engine/Source/Editor/UnrealEd/Private/Tests/EditorAutomationTests.cpp b/Engine/Source/Editor/UnrealEd/Private/Tests/EditorAutomationTests.cpp
--- a/Engine/Source/Editor/UnrealEd/Private/Tests/EditorAutomationTests.cpp
+++ b/Engine/Source/Editor/UnrealEd/Private/Tests/EditorAutomationTests.cpp
@@ -12,6 +12,7 @@
 #include "ScopedTransaction.h"
 #include "LevelEditor.h"
 #include "ModuleManager.h"
+#include "MainFrame.h"
 
 #include "Kismet2/KismetEditorUtilities.h"
 #include "Kismet2/BlueprintEditorUtils.h"
@@ -20,6 +21,7 @@
 #include "SlateWordWrapper.h"
 #include "AutomationCommon.h"
 
+
 //DEFINE_LOG_CATEGORY_STATIC(LogEditorAutomationTests, Log, All);
 
 /**
@@ -80,6 +82,71 @@ bool FCloseAllAssetEditorsCommand::Update()
 	return true;
 }
 
+/**
+* Change the attributes for a point light in the level.
+*/
+struct PointLightParameters
+{
+	APointLight* PointLight;
+	float LightBrightness = 5000.0f;
+	float LightRadius = 1000.0f;
+	FVector LightLocation = FVector(0.0f, 0.0f, 0.0f);
+	FColor LightColor = FColor(255, 255, 255);
+};
+
+//Updates the properties of a specified point light.
+DEFINE_LATENT_AUTOMATION_COMMAND_ONE_PARAMETER(PointLightUpdateCommand, PointLightParameters, PointLightUsing);
+
+bool PointLightUpdateCommand::Update()
+{
+	//Set the point light mobility, brightness, radius, and light color.
+	PointLightUsing.PointLight->SetMobility(EComponentMobility::Movable);
+	PointLightUsing.PointLight->SetBrightness(PointLightUsing.LightBrightness);
+	PointLightUsing.PointLight->SetLightColor(PointLightUsing.LightColor);
+	PointLightUsing.PointLight->TeleportTo(PointLightUsing.LightLocation, FRotator(0, 0, 0));
+	PointLightUsing.PointLight->SetRadius(PointLightUsing.LightRadius);
+	return true;
+}
+
+/**
+* Duplicates a point light.
+*/
+DEFINE_LATENT_AUTOMATION_COMMAND_ONE_PARAMETER(PointLightDuplicationCommand, PointLightParameters, PointLightDuplicating);
+
+bool PointLightDuplicationCommand::Update()
+{
+
+	FScopedTransaction DuplicateLightScope(NSLOCTEXT("UnrealEd.Test", "DuplicateLightScope", "Duplicate Light Scope"));
+
+	//Duplicate the light.
+	bool bOffsetLocations = false;
+	GEditor->edactDuplicateSelected(PointLightDuplicating.PointLight->GetLevel(), bOffsetLocations);
+	for (FSelectionIterator It(GEditor->GetSelectedActorIterator()); It; ++It)
+	{
+		AActor* Actor = static_cast<AActor*>(*It);
+		Actor->TeleportTo(FVector(PointLightDuplicating.LightLocation), FRotator(0, 0, 0));
+	}
+	return true;
+}
+
+/**
+* This will then take a screenshot of the editor, only if it has been enabled.
+*/
+void TakeLatentAutomationScreenshot(struct WindowScreenshotParameters ScreenshotParameters, FString BaseFileName, FString ScreenshotTitle, FString ScreenshotFolderName, bool bIncludeHardwareDetails)
+{
+	//Update the screenshot name, then take a screenshot.
+	if (FAutomationTestFramework::GetInstance().IsScreenshotAllowed())
+	{
+		//Update the screenshot name and get the location of where it will be saved.
+		ScreenshotParameters.ScreenshotName = TEXT("ScreenshotTitle");
+		FString TestName = FString::Printf(TEXT("%s/%s"),*BaseFileName, *ScreenshotFolderName);
+		AutomationCommon::GetScreenshotPath(TestName, ScreenshotParameters.ScreenshotName, bIncludeHardwareDetails);
+
+		//Take a screenshot
+		ADD_LATENT_AUTOMATION_COMMAND(FTakeEditorScreenshotCommand(ScreenshotParameters));
+	}
+}
+
 //////////////////////////////////////////////////////////////////////////
 // FEditorAutomationTestUtilities
 
@@ -1420,8 +1487,6 @@ bool FStaticMeshPlacement::RunTest(const FString& Parameters)
 }
 
 
-
-
 //////////////////////////////////////////////////////////////////////////
 /**
  * QA Light Placement Regression Testing
@@ -1430,9 +1495,24 @@ IMPLEMENT_SIMPLE_AUTOMATION_TEST(FLightPlacement, "QA.Point Light Placement", EA
 
 bool FLightPlacement::RunTest(const FString& Parameters)
 {
-	//Open a new empty map
-	UWorld* World = GEditor->NewMap();
+	//Initialize the parameters for taking a screenshot as well as update the placed point light.
+	WindowScreenshotParameters PointLightPlacementWindowParameters;
+	PointLightParameters LightParameters;
+	bool bUndo = true;
+
+		//Set the CurrentWindow to the Mainframe.  This information is used for taking a screenshot later.
+	if (FModuleManager::Get().IsModuleLoaded("MainFrame"))
+	{
+		IMainFrameModule& MainFrame = FModuleManager::LoadModuleChecked<IMainFrameModule>("MainFrame");
+		PointLightPlacementWindowParameters.CurrentWindow = MainFrame.GetParentWindow();
+	}
 
+	//Set the Test Name which is used later for getting the directory to store the screenshots.
+	const FString BaseFileName = TEXT("PointLightPlacementTest");
+	
+	//Open a new blank map.
+	UWorld* World = GEditor->NewMap();
+			
 	//Move the perspective viewport view to show the test.
 	for (int32 i = 0; i < GEditor->LevelViewportClients.Num(); i++)
 	{
@@ -1445,7 +1525,7 @@ bool FLightPlacement::RunTest(const FString& Parameters)
 	}
 
 	{
-		//Gather assets
+		//Gather assets.
 		UObject* Astroid = (UStaticMesh*)StaticLoadObject(UStaticMesh::StaticClass(), NULL, TEXT("/Engine/Content/EditorAutomation/Astroid.Astroid"), NULL, LOAD_None, NULL);
 		//Add Astroid mesh to the world
 		AStaticMeshActor* StaticMesh = Cast<AStaticMeshActor>(FActorFactoryAssetProxy::AddActorForAsset(Astroid));
@@ -1453,35 +1533,54 @@ bool FLightPlacement::RunTest(const FString& Parameters)
 		StaticMesh->SetActorRelativeScale3D(FVector(1.0f, 1.0f, 1.0f));
 	}
 
-	//Point Light
+	//Create the point light and set it's mobility, brightness, and light color.
 	APointLight* PointLight = Cast<APointLight>(GEditor->AddActor(World->GetCurrentLevel(), APointLight::StaticClass(), FVector(0.0f, 0.0f, 400.0f)));
-	PointLight->SetMobility(EComponentMobility::Movable);
-	PointLight->SetBrightness(5000.0f);
-	PointLight->SetLightColor(FColor(255, 0, 0));
+	LightParameters.PointLight = PointLight;
+	LightParameters.LightColor = FColor(255, 0, 0);
+	LightParameters.LightLocation = FVector(0.0f, 0.0f, 400.0f);
+	ADD_LATENT_AUTOMATION_COMMAND(PointLightUpdateCommand(LightParameters));
+	
+	//Wait
+	ADD_LATENT_AUTOMATION_COMMAND(FWaitLatentCommand(0.1f));
+	
+	//Take a screenshot of the newly placed point light.
+	TakeLatentAutomationScreenshot(PointLightPlacementWindowParameters, BaseFileName, FString::Printf(TEXT("PlacedPointLight")), FString::Printf(TEXT("01_Placed")), true);
+	
+	//Duplicate the point light.
+	LightParameters.LightLocation = FVector(10.0f, 10.0f, 400.0f);
+	ADD_LATENT_AUTOMATION_COMMAND(PointLightDuplicationCommand(LightParameters));
 
-	//Duplicate the point light
-	{
-		FScopedTransaction DuplicateLightScope( NSLOCTEXT("UnrealEd.Test", "DuplicateLightScope", "Duplicate Light Scope") );
+	//Wait
+	ADD_LATENT_AUTOMATION_COMMAND(FWaitLatentCommand(0.1f));
 
-		//Duplicate the light
-		bool bOffsetLocations = false;
-		GEditor->edactDuplicateSelected(PointLight->GetLevel(), bOffsetLocations);
-		for ( FSelectionIterator It( GEditor->GetSelectedActorIterator() ) ; It ; ++It )
-		{
-			AActor* Actor = static_cast<AActor*>( *It );
-			Actor->TeleportTo(FVector(10.0f,10.0f,400.0f), FRotator(0, 0, 0));
-		}
-	}
+	//Take a screenshot of the duplicated point light.
+	TakeLatentAutomationScreenshot(PointLightPlacementWindowParameters, BaseFileName, FString::Printf(TEXT("DuplicatedLight")), FString::Printf(TEXT("02_Dupe")), true);
 
-	//Undo the duplication then redo it
-	GEditor->UndoTransaction();
-	GEditor->RedoTransaction();
+	//Undo the duplication.
+	ADD_LATENT_AUTOMATION_COMMAND(UndoRedoCommand(true));
+
+	//Wait
+	ADD_LATENT_AUTOMATION_COMMAND(FWaitLatentCommand(0.1f));
 
-	//Update the duplicated asset
-	PointLight->TeleportTo(FVector(500.0f,300.0f, 300.0f), FRotator(0, 0, 0));
-	PointLight->SetRadius( 500.0f );
-	PointLight->SetLightColor(FColor(255, 255, 255));
+	//Take a screenshot of the scene after the duplication has been undone.
+	TakeLatentAutomationScreenshot(PointLightPlacementWindowParameters, BaseFileName, FString::Printf(TEXT("UndoDuplicationPointLight")), FString::Printf(TEXT("03_Undo")), true);
+	
+	//Redo the duplication.
+	ADD_LATENT_AUTOMATION_COMMAND(UndoRedoCommand(false));
+
+	//Update the original point light actor.
+	LightParameters.LightRadius = 500.0f;
+	LightParameters.LightLocation = FVector(500.0f, 300.0f, 500.0f);
+	LightParameters.LightColor = FColor(255, 255, 255);
+	ADD_LATENT_AUTOMATION_COMMAND(PointLightUpdateCommand(LightParameters));
+	
+	//Wait
+	ADD_LATENT_AUTOMATION_COMMAND(FWaitLatentCommand(0.1f));
 
+	//Take a screenshot of the final scene.
+	//It is expected to show two moveable point lights (one red, one white) and a static mesh.
+	TakeLatentAutomationScreenshot(PointLightPlacementWindowParameters, BaseFileName, FString::Printf(TEXT("FinalPointLight")), FString::Printf(TEXT("04_Final")), true);
+	
 	return true;
 }
 
@@ -2047,4 +2146,4 @@ bool FStaticMeshUVsTest::RunTest(const FString& Parameters)
 	TSharedPtr<StaticMeshUVTest::FUVTestHelper> TestHelper = MakeShareable(new StaticMeshUVTest::FUVTestHelper(this));
 	ADD_LATENT_AUTOMATION_COMMAND(FPerformUVTestCommand(TestHelper));
 	return true;
-}
+}
\ No newline at end of file
diff --git a/Engine/Source/Runtime/Engine/Private/Tests/AutomationCommon.cpp b/Engine/Source/Runtime/Engine/Private/Tests/AutomationCommon.cpp
--- a/Engine/Source/Runtime/Engine/Private/Tests/AutomationCommon.cpp
+++ b/Engine/Source/Runtime/Engine/Private/Tests/AutomationCommon.cpp
@@ -4,6 +4,7 @@
 #include "Slate.h"
 #include "AutomationCommon.h"
 #include "ImageUtils.h"
+#include "UnrealEd.h"
 
 ///////////////////////////////////////////////////////////////////////
 // Common Latent commands
@@ -42,5 +43,21 @@ bool FTakeActiveEditorScreenshotCommand::Update()
 bool FTakeEditorScreenshotCommand::Update()
 {
 	AutomationCommon::SaveWindowAsScreenshot(ScreenshotParameters.CurrentWindow.ToSharedRef(), ScreenshotParameters.ScreenshotName);
+	return true;
+}
+
+bool UndoRedoCommand::Update()
+{
+	if (bUndo == true)
+	{
+		//Undo
+		GEditor->UndoTransaction();
+	}
+	else
+	{
+		//Redo
+		GEditor->RedoTransaction();
+	}
+
 	return true;
 }
\ No newline at end of file
diff --git a/Engine/Source/Runtime/Engine/Public/Tests/AutomationCommon.h b/Engine/Source/Runtime/Engine/Public/Tests/AutomationCommon.h
--- a/Engine/Source/Runtime/Engine/Public/Tests/AutomationCommon.h
+++ b/Engine/Source/Runtime/Engine/Public/Tests/AutomationCommon.h
@@ -98,3 +98,9 @@ DEFINE_ENGINE_LATENT_AUTOMATION_COMMAND_ONE_PARAMETER(FTakeActiveEditorScreensho
  * Take a screenshot of the active window
  */
 DEFINE_ENGINE_LATENT_AUTOMATION_COMMAND_ONE_PARAMETER(FTakeEditorScreenshotCommand, WindowScreenshotParameters, ScreenshotParameters);
+
+/**
+* Creates a latent command which the user can either undo or redo an action.
+* True will trigger Undo, False will trigger a redo
+*/
+DEFINE_ENGINE_LATENT_AUTOMATION_COMMAND_ONE_PARAMETER(UndoRedoCommand, bool, bUndo);
\ No newline at end of file
