#ue4 Fix for stale OwnedComponents array after compiling on load

diff --git a/Engine/Source/Editor/UnrealEd/Private/Kismet2/BlueprintEditorUtils.cpp b/Engine/Source/Editor/UnrealEd/Private/Kismet2/BlueprintEditorUtils.cpp
--- a/Engine/Source/Editor/UnrealEd/Private/Kismet2/BlueprintEditorUtils.cpp
+++ b/Engine/Source/Editor/UnrealEd/Private/Kismet2/BlueprintEditorUtils.cpp
@@ -1202,10 +1202,16 @@ UClass* FBlueprintEditorUtils::RegenerateBlueprintClass(UBlueprint* Blueprint, U
 			CopyDetails.bAggressiveDefaultSubobjectReplacement = false;
 			CopyDetails.bDoDelta = false;
 			UEditorEngine::CopyPropertiesForUnrelatedObjects(SkeletonCDO, GeneratedCDO, CopyDetails);
-			
+
 			Blueprint->SetLegacyGeneratedClassIsAuthoritative();
 		}
 
+		// Now that the CDO is valid, update the OwnedComponents, in case we've added or removed native components
+		if (AActor* MyActor = Cast<AActor>(Blueprint->GeneratedClass->GetDefaultObject()))
+		{
+			MyActor->ResetOwnedComponents();
+		}
+
 		Blueprint->ClearFlags(RF_BeingRegenerated);
 		bNeedsSkelRefRemoval = true;
 	}
diff --git a/Engine/Source/Runtime/Engine/Classes/GameFramework/Actor.h b/Engine/Source/Runtime/Engine/Classes/GameFramework/Actor.h
--- a/Engine/Source/Runtime/Engine/Classes/GameFramework/Actor.h
+++ b/Engine/Source/Runtime/Engine/Classes/GameFramework/Actor.h
@@ -2045,6 +2045,7 @@ private:
 	UPROPERTY(transient, duplicatetransient)
 	TArray<UActorComponent*> OwnedComponents;
 
+	friend class FBlueprintEditorUtils;
 	void ResetOwnedComponents();
 
 public:
