Making sure cook commandlet and cook-on-the-fly server support sandbox paths with spaces

diff --git a/Engine/Source/Editor/UnrealEd/Private/Commandlets/CookCommandlet.cpp b/Engine/Source/Editor/UnrealEd/Private/Commandlets/CookCommandlet.cpp
--- a/Engine/Source/Editor/UnrealEd/Private/Commandlets/CookCommandlet.cpp
+++ b/Engine/Source/Editor/UnrealEd/Private/Commandlets/CookCommandlet.cpp
@@ -539,7 +539,7 @@ int32 UCookCommandlet::Main(const FString& CmdLineParams)
 		FString OutputDirectory = GetOutputDirectoryOverride();
 
 		// Use SandboxFile to do path conversion to properly handle sandbox paths (outside of standard paths in particular).
-		SandboxFile->Initialize(&FPlatformFileManager::Get().GetPlatformFile(), *FString::Printf(TEXT("-sandbox=%s"), *OutputDirectory));
+		SandboxFile->Initialize(&FPlatformFileManager::Get().GetPlatformFile(), *FString::Printf(TEXT("-sandbox=\"%s\""), *OutputDirectory));
 
 		CleanSandbox(Platforms);
 
diff --git a/Engine/Source/Editor/UnrealEd/Private/CookOnTheFlyServer.cpp b/Engine/Source/Editor/UnrealEd/Private/CookOnTheFlyServer.cpp
--- a/Engine/Source/Editor/UnrealEd/Private/CookOnTheFlyServer.cpp
+++ b/Engine/Source/Editor/UnrealEd/Private/CookOnTheFlyServer.cpp
@@ -697,7 +697,7 @@ void UCookOnTheFlyServer::Initialize( bool inCompressed, bool inIterativeCooking
 	FString OutputDirectory = GetOutputDirectoryOverride(OutputDirectoryOverride);
 
 	// Use SandboxFile to do path conversion to properly handle sandbox paths (outside of standard paths in particular).
-	SandboxFile->Initialize(&FPlatformFileManager::Get().GetPlatformFile(), *FString::Printf(TEXT("-sandbox=%s"), *OutputDirectory));
+	SandboxFile->Initialize(&FPlatformFileManager::Get().GetPlatformFile(), *FString::Printf(TEXT("-sandbox=\"%s\""), *OutputDirectory));
 
 	CleanSandbox(Platforms);
 
