Blueprints: Fixed crash on load when refreshing an interface during compile-on-load

diff --git a/Engine/Source/Editor/BlueprintGraph/Private/K2Node_CustomEvent.cpp b/Engine/Source/Editor/BlueprintGraph/Private/K2Node_CustomEvent.cpp
--- a/Engine/Source/Editor/BlueprintGraph/Private/K2Node_CustomEvent.cpp
+++ b/Engine/Source/Editor/BlueprintGraph/Private/K2Node_CustomEvent.cpp
@@ -164,8 +164,9 @@ void UK2Node_CustomEvent::RenameCustomEventCloseToName(int32 StartIndex)
 		const FString NewName = FString::Printf(TEXT("%s_%d"), *BaseName, NameIndex);
 		if (Rename(*NewName, GetOuter(), REN_Test))
 		{
+			UBlueprint* Blueprint = GetBlueprint();
 			CustomFunctionName = FName(NewName.GetCharArray().GetData());
-			Rename(*NewName);
+			Rename(*NewName, GetOuter(), (Blueprint->bIsRegeneratingOnLoad ? REN_ForceNoResetLoaders : 0) | REN_DontCreateRedirectors);
 			bFoundName = true;
 		}
 	}
diff --git a/Engine/Source/Editor/UnrealEd/Private/Kismet2/Kismet2.cpp b/Engine/Source/Editor/UnrealEd/Private/Kismet2/Kismet2.cpp
--- a/Engine/Source/Editor/UnrealEd/Private/Kismet2/Kismet2.cpp
+++ b/Engine/Source/Editor/UnrealEd/Private/Kismet2/Kismet2.cpp
@@ -392,9 +392,12 @@ void FKismetEditorUtilities::CompileBlueprint(UBlueprint* BlueprintObj, bool bIs
 		FBlueprintEditorUtils::GetDependentBlueprints(BlueprintObj, DependentBPs);
 
 		// refresh each dependent blueprint
-		for ( auto ObjIt = DependentBPs.CreateConstIterator(); ObjIt; ++ObjIt )
+		for (UBlueprint* Dependent : DependentBPs)
 		{
-			FBlueprintEditorUtils::RefreshAllNodes(*ObjIt);
+			bool bPreviousRegenValue = Dependent->bIsRegeneratingOnLoad;
+			Dependent->bIsRegeneratingOnLoad = Dependent->bIsRegeneratingOnLoad || BlueprintObj->bIsRegeneratingOnLoad;
+			FBlueprintEditorUtils::RefreshAllNodes(Dependent);
+			Dependent->bIsRegeneratingOnLoad = bPreviousRegenValue;
 		}
 	}
 
