UE4: Optimization for an unnecessary TObjectIterator when reinstancing blueprints. GetObjectsOfClass already excludes CDOs by default.

diff --git a/Engine/Source/Editor/UnrealEd/Private/Kismet2/KismetReinstanceUtilities.cpp b/Engine/Source/Editor/UnrealEd/Private/Kismet2/KismetReinstanceUtilities.cpp
--- a/Engine/Source/Editor/UnrealEd/Private/Kismet2/KismetReinstanceUtilities.cpp
+++ b/Engine/Source/Editor/UnrealEd/Private/Kismet2/KismetReinstanceUtilities.cpp
@@ -224,14 +224,8 @@ void FBlueprintCompileReinstancer::ReplaceInstancesOfClass(UClass* OldClass, UCl
 	TGuardValue<bool> GuardTemplateNameFlag(GIsReconstructingBlueprintInstances, true);
 
 	TArray<UObject*> ObjectsToReplace;
-	for( TObjectIterator<UObject> It; It; ++It )
-	{
-		UObject* Obj = *It;
-		if( Obj && (Obj->GetClass() == OldClass) && (Obj != OldClass->ClassDefaultObject) )
-		{
-			ObjectsToReplace.Add(Obj);
-		}
-	}
+	const bool bIncludeDerivedClasses = false;
+	GetObjectsOfClass(OldClass, ObjectsToReplace, bIncludeDerivedClasses);
 
 	bool bSelectionChanged = false;
 	
