UE4: You can now load maps that are only created in memory into the editor. They no longer have to be saved first. TTP#337966

diff --git a/Engine/Source/Editor/UnrealEd/Private/EditorServer.cpp b/Engine/Source/Editor/UnrealEd/Private/EditorServer.cpp
--- a/Engine/Source/Editor/UnrealEd/Private/EditorServer.cpp
+++ b/Engine/Source/Editor/UnrealEd/Private/EditorServer.cpp
@@ -1934,15 +1934,23 @@ bool UEditorEngine::Map_Load(const TCHAR* Str, FOutputDevice& Ar)
 		{
  			if ( Context.World() && Context.World()->GetOutermost() && Context.World()->GetOutermost() == FindPackage(NULL, *LongTempFname) )
  			{
- 				// This map is already loaded.
+ 				// This map is already loaded and in the editor.
  				return true;
  			}
 
-			FString AlteredPath;
-			if ( FPackageName::DoesPackageExist(LongTempFname, NULL, &AlteredPath) )
+			// Is the new world already loaded?
+			UPackage* ExistingPackage = FindPackage(NULL, *LongTempFname);
+			UWorld* ExistingWorld = NULL;
+			if (ExistingPackage)
+			{
+				ExistingWorld = UWorld::FindWorldInPackage(ExistingPackage);
+			}
+
+			FString UnusedAlteredPath;
+			if ( ExistingWorld || FPackageName::DoesPackageExist(LongTempFname, NULL, &UnusedAlteredPath) )
 			{
 				FText NotMapReason;
-				if( !PackageIsAMapFile( TempFname, NotMapReason ) )
+				if( !ExistingWorld && !PackageIsAMapFile( TempFname, NotMapReason ) )
 				{
 					// Map load failed 
 					FMessageDialog::Open( EAppMsgType::Ok, NotMapReason );
@@ -1959,14 +1967,6 @@ bool UEditorEngine::Map_Load(const TCHAR* Str, FOutputDevice& Ar)
 				// Should we display progress while loading?
 				int32 bShowProgress = 1;
 				FParse::Value(Str, TEXT("SHOWPROGRESS="), bShowProgress);
-
-				// Is the new world already loaded?
-				UPackage* ExistingPackage = FindPackage(NULL, *LongTempFname);
-				UWorld* ExistingWorld = NULL;
-				if ( ExistingPackage )
-				{
-					ExistingWorld = UWorld::FindWorldInPackage(ExistingPackage);
-				}
 			
 				UObject* OldOuter = NULL;
 
