Editor: Made component visualizers hierarchical

diff --git a/Engine/Source/Editor/ComponentVisualizers/Public/ComponentVisualizers.h b/Engine/Source/Editor/ComponentVisualizers/Public/ComponentVisualizers.h
--- a/Engine/Source/Editor/ComponentVisualizers/Public/ComponentVisualizers.h
+++ b/Engine/Source/Editor/ComponentVisualizers/Public/ComponentVisualizers.h
@@ -32,6 +32,6 @@ private:
 	/** Register a visualizer for a particular componen class */
 	void RegisterComponentVisualizer(FName ComponentClassName, TSharedPtr<FComponentVisualizer> Visualizer);
 
-	/** Array of component class names we have registered, so we know what to unregister aferwards */
+	/** Array of component class names we have registered, so we know what to unregister afterwards */
 	TArray<FName> RegisteredComponentClassNames;
 };
\ No newline at end of file
diff --git a/Engine/Source/Editor/Kismet/Private/SCSEditorViewportClient.cpp b/Engine/Source/Editor/Kismet/Private/SCSEditorViewportClient.cpp
--- a/Engine/Source/Editor/Kismet/Private/SCSEditorViewportClient.cpp
+++ b/Engine/Source/Editor/Kismet/Private/SCSEditorViewportClient.cpp
@@ -259,7 +259,7 @@ void FSCSEditorViewportClient::Draw(const FSceneView* View, FPrimitiveDrawInterf
 				if(Comp != NULL && Comp->IsRegistered())
 				{
 					// Try and find a visualizer
-					TSharedPtr<FComponentVisualizer> Visualizer = GUnrealEd->FindComponentVisualizer(Comp->GetClass()->GetFName());
+					TSharedPtr<FComponentVisualizer> Visualizer = GUnrealEd->FindComponentVisualizer(Comp->GetClass());
 					if (Visualizer.IsValid())
 					{
 						Visualizer->DrawVisualization(Comp, View, PDI);
diff --git a/Engine/Source/Editor/UnrealEd/Classes/Editor/UnrealEdEngine.h b/Engine/Source/Editor/UnrealEd/Classes/Editor/UnrealEdEngine.h
--- a/Engine/Source/Editor/UnrealEd/Classes/Editor/UnrealEdEngine.h
+++ b/Engine/Source/Editor/UnrealEd/Classes/Editor/UnrealEdEngine.h
@@ -225,7 +225,10 @@ class UNREALED_API UUnrealEdEngine : public UEditorEngine, public FNotifyHook
 	/** Unregister component visualizer function */
 	void UnregisterComponentVisualizer(FName ComponentClassName);
 	/** Find a component visualizer for the given component class name */
-	TSharedPtr<class FComponentVisualizer> FindComponentVisualizer(FName ComponentClassName);
+	TSharedPtr<class FComponentVisualizer> FindComponentVisualizer(FName ComponentClassName) const;
+
+	/** Find a component visualizer for the given component class (checking parent classes too) */
+	TSharedPtr<class FComponentVisualizer> FindComponentVisualizer(UClass* ComponentClass) const;
 
 	/** Draw component visualizers for components for selected actors */
 	void DrawComponentVisualizers(const FSceneView* View, FPrimitiveDrawInterface* PDI);
diff --git a/Engine/Source/Editor/UnrealEd/Private/ComponentVisualizerManager.cpp b/Engine/Source/Editor/UnrealEd/Private/ComponentVisualizerManager.cpp
--- a/Engine/Source/Editor/UnrealEd/Private/ComponentVisualizerManager.cpp
+++ b/Engine/Source/Editor/UnrealEd/Private/ComponentVisualizerManager.cpp
@@ -11,7 +11,7 @@ bool FComponentVisualizerManager::HandleProxyForComponentVis(HHitProxy *HitProxy
 		const UActorComponent* ClickedComponent = VisProxy->Component.Get();
 		if (ClickedComponent != NULL)
 		{
-			TSharedPtr<FComponentVisualizer> Visualizer = GUnrealEd->FindComponentVisualizer(ClickedComponent->GetClass()->GetFName());
+			TSharedPtr<FComponentVisualizer> Visualizer = GUnrealEd->FindComponentVisualizer(ClickedComponent->GetClass());
 			if (Visualizer.IsValid())
 			{
 				bool bIsActive = Visualizer->VisProxyHandleClick(VisProxy);
diff --git a/Engine/Source/Editor/UnrealEd/Private/UnrealEdEngine.cpp b/Engine/Source/Editor/UnrealEd/Private/UnrealEdEngine.cpp
--- a/Engine/Source/Editor/UnrealEd/Private/UnrealEdEngine.cpp
+++ b/Engine/Source/Editor/UnrealEd/Private/UnrealEdEngine.cpp
@@ -1022,11 +1022,11 @@ void UUnrealEdEngine::UnregisterComponentVisualizer(FName ComponentClassName)
 	ComponentVisualizerMap.Remove(ComponentClassName);
 }
 
-TSharedPtr<FComponentVisualizer> UUnrealEdEngine::FindComponentVisualizer(FName ComponentClassName)
+TSharedPtr<FComponentVisualizer> UUnrealEdEngine::FindComponentVisualizer(FName ComponentClassName) const
 {
 	TSharedPtr<FComponentVisualizer> Visualizer = NULL;
 
-	TSharedPtr<FComponentVisualizer>* VisualizerPtr = ComponentVisualizerMap.Find(ComponentClassName);
+	const TSharedPtr<FComponentVisualizer>* VisualizerPtr = ComponentVisualizerMap.Find(ComponentClassName);
 	if(VisualizerPtr != NULL)
 	{
 		Visualizer = *VisualizerPtr;
@@ -1035,6 +1035,18 @@ TSharedPtr<FComponentVisualizer> UUnrealEdEngine::FindComponentVisualizer(FName
 	return Visualizer;
 }
 
+/** Find a component visualizer for the given component class (checking parent classes too) */
+TSharedPtr<class FComponentVisualizer> UUnrealEdEngine::FindComponentVisualizer(UClass* ComponentClass) const
+{
+	TSharedPtr<FComponentVisualizer> Visualizer;
+	while (!Visualizer.IsValid() && (ComponentClass != nullptr) && (ComponentClass != UActorComponent::StaticClass()))
+	{
+		Visualizer = FindComponentVisualizer(ComponentClass->GetFName());
+		ComponentClass = ComponentClass->GetSuperClass();
+	}
+
+	return Visualizer;
+}
 
 
 
@@ -1056,7 +1068,8 @@ void UUnrealEdEngine::DrawComponentVisualizers(const FSceneView* View, FPrimitiv
 				if(Comp->IsRegistered())
 				{
 					// Try and find a visualizer
-					TSharedPtr<FComponentVisualizer> Visualizer = FindComponentVisualizer(Comp->GetClass()->GetFName());
+					
+					TSharedPtr<FComponentVisualizer> Visualizer = FindComponentVisualizer(Comp->GetClass());
 					if(Visualizer.IsValid())
 					{
 						Visualizer->DrawVisualization(Comp, View, PDI);
