Fixed TTP #337574 temporarily. It crashed when failed to import a destructible asset.

diff --git a/Engine/Source/Editor/UnrealEd/Private/Factories/ApexDestructibleAssetImport.cpp b/Engine/Source/Editor/UnrealEd/Private/Factories/ApexDestructibleAssetImport.cpp
--- a/Engine/Source/Editor/UnrealEd/Private/Factories/ApexDestructibleAssetImport.cpp
+++ b/Engine/Source/Editor/UnrealEd/Private/Factories/ApexDestructibleAssetImport.cpp
@@ -1029,6 +1029,10 @@ UDestructibleMesh* ImportDestructibleMeshFromApexDestructibleAsset(UObject* InPa
 
 	if (!SetApexDestructibleAsset(*DestructibleMesh, ApexDestructibleAsset, OutData, Options))
 	{
+		// should remove this destructible mesh. if not, this object causes a crash when ticking because it doesn't have proper rendering resources
+		// @TODO : creates this destructible mesh after loading data completely
+		DestructibleMesh->PostEditChange();
+		DestructibleMesh->ConditionalBeginDestroy();
 		return NULL;
 	}
 
