UE4: The ActiveGroupActors list is now on the World instead of the EditorEngine. This is so references to these actors to not prevent GC at certain times (like delete) and worlds can be loaded into memory without being referenced forever. TTP#337972

diff --git a/Engine/Source/Editor/UnrealEd/Classes/Editor/EditorEngine.h b/Engine/Source/Editor/UnrealEd/Classes/Editor/EditorEngine.h
--- a/Engine/Source/Editor/UnrealEd/Classes/Editor/EditorEngine.h
+++ b/Engine/Source/Editor/UnrealEd/Classes/Editor/EditorEngine.h
@@ -473,9 +473,6 @@ class UNREALED_API UEditorEngine : public UEngine
 	UPROPERTY(transient)
 	uint32 bDrawParticleHelpers:1;
 
-	UPROPERTY()
-	TArray<class AGroupActor*> ActiveGroupActors;
-
 	/** Brush builders that have been created in the editor */
 	UPROPERTY(transient)
 	TArray<class UBrushBuilder*> BrushBuilders;	
@@ -1261,7 +1258,7 @@ public:
 	// Selection.
 	virtual void SelectActor(AActor* Actor, bool bInSelected, bool bNotify, bool bSelectEvenIfHidden=false) {}
 	virtual bool CanSelectActor(AActor* Actor, bool bInSelected, bool bSelectEvenIfHidden=false, bool bWarnIfLevelLocked=false) const { return true; }
-	virtual void SelectGroup(AGroupActor* InGroupActor, bool bForceSelection=false, bool bInSelected=true, bool bNotify=true) {}
+	virtual void SelectGroup(class AGroupActor* InGroupActor, bool bForceSelection=false, bool bInSelected=true, bool bNotify=true) {}
 
 	/**
 	 * Replaces the components in ActorsToReplace with an primitive component in Replacement
diff --git a/Engine/Source/Editor/UnrealEd/Private/EditorLevelUtils.cpp b/Engine/Source/Editor/UnrealEd/Private/EditorLevelUtils.cpp
--- a/Engine/Source/Editor/UnrealEd/Private/EditorLevelUtils.cpp
+++ b/Engine/Source/Editor/UnrealEd/Private/EditorLevelUtils.cpp
@@ -334,18 +334,22 @@ namespace EditorLevelUtils
 
 		GEditor->CloseEditedWorldAssets(CastChecked<UWorld>(InLevel->GetOuter()));
 
+		UWorld* OwningWorld = InLevel->OwningWorld;
 		const bool bRemovingCurrentLevel	= InLevel && InLevel->IsCurrentLevel();
 		const bool bRemoveSuccessful		= PrivateRemoveLevelFromWorld( InLevel );
 		if ( bRemoveSuccessful )
 		{
-			// remove all group actors from GEditor in the level we are removing
-			// otherwise, this will cause group actors to not be garbage collected
-			for (int32 GroupIndex = GEditor->ActiveGroupActors.Num()-1; GroupIndex >= 0; --GroupIndex)
+			if ( OwningWorld )
 			{
-				AGroupActor* GroupActor = GEditor->ActiveGroupActors[GroupIndex];
-				if (GroupActor && GroupActor->IsInLevel(InLevel))
+				// remove all group actors from the world in the level we are removing
+				// otherwise, this will cause group actors to not be garbage collected
+				for (int32 GroupIndex = OwningWorld->ActiveGroupActors.Num()-1; GroupIndex >= 0; --GroupIndex)
 				{
-					GEditor->ActiveGroupActors.RemoveAt(GroupIndex);
+					AGroupActor* GroupActor = Cast<AGroupActor>(OwningWorld->ActiveGroupActors[GroupIndex]);
+					if (GroupActor && GroupActor->IsInLevel(InLevel))
+					{
+						OwningWorld->ActiveGroupActors.RemoveAt(GroupIndex);
+					}
 				}
 			}
 
diff --git a/Engine/Source/Editor/UnrealEd/Private/EditorServer.cpp b/Engine/Source/Editor/UnrealEd/Private/EditorServer.cpp
--- a/Engine/Source/Editor/UnrealEd/Private/EditorServer.cpp
+++ b/Engine/Source/Editor/UnrealEd/Private/EditorServer.cpp
@@ -1617,7 +1617,7 @@ void UEditorEngine::EditorDestroyWorld( FWorldContext & Context, const FText& Cl
 	EditorClearComponents();
 	ClearPreviewComponents();
 	// Remove all active groups, they belong to a map being unloaded
-	ActiveGroupActors.Empty();
+	Context.World()->ActiveGroupActors.Empty();
 
 	// Make sure we don't have any apps open on for assets owned by the world we are closing
 	CloseEditedWorldAssets(Context.World());
diff --git a/Engine/Source/Editor/UnrealEd/Private/GroupActor.cpp b/Engine/Source/Editor/UnrealEd/Private/GroupActor.cpp
--- a/Engine/Source/Editor/UnrealEd/Private/GroupActor.cpp
+++ b/Engine/Source/Editor/UnrealEd/Private/GroupActor.cpp
@@ -22,7 +22,7 @@ void AGroupActor::PostActorCreated()
 	// Cache our newly created group
 	if( !GetWorld()->IsPlayInEditor() && !IsRunningCommandlet() && GIsEditor )
 	{
-		GEditor->ActiveGroupActors.AddUnique(this);
+		GetWorld()->ActiveGroupActors.AddUnique(this);
 	}
 	Super::PostActorCreated();
 }
@@ -32,7 +32,7 @@ void AGroupActor::PostLoad()
 	if( !GetWorld()->IsPlayInEditor() && !IsRunningCommandlet() && GIsEditor )
 	{
 		// Cache group on de-serialization
-		GEditor->ActiveGroupActors.AddUnique(this);
+		GetWorld()->ActiveGroupActors.AddUnique(this);
 
 		// Fix up references for GetParentForActor()
 		for(int32 i=0; i<GroupActors.Num(); ++i)
@@ -51,11 +51,11 @@ void AGroupActor::PostEditChangeProperty( FPropertyChangedEvent& PropertyChanged
 	// Re-instate group as active if it had children after undo/redo
 	if(GroupActors.Num() || SubGroups.Num())
 	{
-		GEditor->ActiveGroupActors.AddUnique(this);
+		GetWorld()->ActiveGroupActors.AddUnique(this);
 	}
 	else // Otherwise, attempt to remove them
 	{
-		GEditor->ActiveGroupActors.Remove(this);
+		GetWorld()->ActiveGroupActors.Remove(this);
 	}
 	Super::PostEditChangeProperty(PropertyChangedEvent);
 }
@@ -66,7 +66,7 @@ void AGroupActor::PostEditUndo()
 
 	if (IsPendingKill())
 	{
-		GEditor->ActiveGroupActors.RemoveSwap(this);
+		GetWorld()->ActiveGroupActors.RemoveSwap(this);
 	}
 }
 
@@ -352,30 +352,37 @@ void AGroupActor::DrawBracketsForGroups( FPrimitiveDrawInterface* PDI, FViewport
 	if( GUnrealEd->bGroupingActive )
 	{
 		check(PDI);
-	
-		if(bMustBeSelected)
+		
+		UWorld* EditorWorld = GEditor->GetEditorWorldContext().World();
+		if (EditorWorld)
 		{
-			// If we're only drawing for selected group, grab only those that have currently selected actors
-			TArray<AGroupActor*> SelectedGroups;
-			for(int32 GroupIndex=0; GroupIndex < GEditor->ActiveGroupActors.Num(); ++GroupIndex)
+			TArray<AGroupActor*> GroupsToDraw;
+			
+			for(int32 GroupIndex=0; GroupIndex < EditorWorld->ActiveGroupActors.Num(); ++GroupIndex)
 			{
-				AGroupActor* SelectedGroupActor = GEditor->ActiveGroupActors[GroupIndex];		
-				if( SelectedGroupActor != NULL )
+				AGroupActor* GroupActor = Cast<AGroupActor>(EditorWorld->ActiveGroupActors[GroupIndex]);
+				if (GroupActor != NULL)
 				{
-					if(SelectedGroupActor->HasSelectedActors())
+					if (bMustBeSelected)
+					{
+						// If we're only drawing for selected group, grab only those that have currently selected actors
+						if (GroupActor->HasSelectedActors())
+						{
+							// We want to start drawing groups from the highest root level.
+							// Subgroups will be propagated through during the draw code.
+							GroupActor = AGroupActor::GetRootForActor(GroupActor);
+							GroupsToDraw.Add(GroupActor);
+						}
+					}
+					else
 					{
-						// We want to start drawing groups from the highest root level.
-						// Subgroups will be propagated through during the draw code.
-						SelectedGroupActor = AGroupActor::GetRootForActor(SelectedGroupActor);
-						SelectedGroups.Add(SelectedGroupActor);
+						// Otherwise, just add all group actors
+						GroupsToDraw.Add(GroupActor);
 					}
 				}
 			}
-			PrivateDrawBracketsForGroups(PDI, Viewport, SelectedGroups);
-		}
-		else
-		{
-			PrivateDrawBracketsForGroups(PDI, Viewport, GEditor->ActiveGroupActors );
+
+			PrivateDrawBracketsForGroups(PDI, Viewport, GroupsToDraw);
 		}
 	}
 }
@@ -445,94 +452,105 @@ AGroupActor* AGroupActor::GetParentForActor(AActor* InActor)
 
 const int32 AGroupActor::NumActiveGroups( bool bSelected/*=false*/, bool bDeepSearch/*=true*/ )
 {
-	if(!bSelected)
+	UWorld* EditorWorld = GEditor->GetEditorWorldContext().World();
+	if (EditorWorld)
 	{
-		return GEditor->ActiveGroupActors.Num();
-	}
+		if(!bSelected)
+		{
+			return EditorWorld->ActiveGroupActors.Num();
+		}
 
-	int32 ActiveSelectedGroups = 0;
-	for(int32 GroupIdx=0; GroupIdx < GEditor->ActiveGroupActors.Num(); ++GroupIdx )
-	{
-		if( GEditor->ActiveGroupActors[GroupIdx] != NULL )
+		int32 ActiveSelectedGroups = 0;
+		for(int32 GroupIdx=0; GroupIdx < EditorWorld->ActiveGroupActors.Num(); ++GroupIdx )
 		{
-			if(GEditor->ActiveGroupActors[GroupIdx]->HasSelectedActors(bDeepSearch))
+			AGroupActor* GroupActor = Cast<AGroupActor>(EditorWorld->ActiveGroupActors[GroupIdx]);
+			if( GroupActor != NULL )
 			{
-				++ActiveSelectedGroups;
+				if(GroupActor->HasSelectedActors(bDeepSearch))
+				{
+					++ActiveSelectedGroups;
+				}
 			}
 		}
+
+		return ActiveSelectedGroups;
 	}
-	return ActiveSelectedGroups;
+	
+	return 0;
 }
 
 
 void AGroupActor::AddSelectedActorsToSelectedGroup()
 {
-
-	int32 SelectedGroupIndex = -1;
-	for(int32 GroupIdx=0; GroupIdx < GEditor->ActiveGroupActors.Num(); ++GroupIdx )
+	UWorld* EditorWorld = GEditor->GetEditorWorldContext().World();
+	if (EditorWorld)
 	{
-		if( GEditor->ActiveGroupActors[GroupIdx] != NULL )
+		int32 SelectedGroupIndex = -1;
+		for(int32 GroupIdx=0; GroupIdx < EditorWorld->ActiveGroupActors.Num(); ++GroupIdx )
 		{
-			if(GEditor->ActiveGroupActors[GroupIdx]->HasSelectedActors(false))
+			AGroupActor* GroupActor = Cast<AGroupActor>(EditorWorld->ActiveGroupActors[GroupIdx]);
+			if( GroupActor != NULL )
 			{
-				// Assign the index of the selected group.
-				// If this is the second group we find, too many groups are selected, return.
-				if( SelectedGroupIndex == -1 )
+				if(GroupActor->HasSelectedActors(false))
 				{
-					SelectedGroupIndex = GroupIdx;
-				}
-				else
-				{
-					return;
+					// Assign the index of the selected group.
+					// If this is the second group we find, too many groups are selected, return.
+					if( SelectedGroupIndex == -1 )
+					{
+						SelectedGroupIndex = GroupIdx;
+					}
+					else
+					{
+						return;
+					}
 				}
 			}
 		}
-	}
 
-	if( SelectedGroupIndex != -1 && GEditor->ActiveGroupActors[SelectedGroupIndex] != NULL )
-	{
-		AGroupActor* SelectedGroup = GEditor->ActiveGroupActors[SelectedGroupIndex];
-		
-		ULevel* GroupLevel = SelectedGroup->GetLevel();
+		AGroupActor* SelectedGroup = Cast<AGroupActor>(EditorWorld->ActiveGroupActors[SelectedGroupIndex]);
+		if( SelectedGroupIndex != -1 && SelectedGroup != NULL )
+		{
+			ULevel* GroupLevel = SelectedGroup->GetLevel();
 
-		// We've established that only one group is selected, so we can just call Add on all these actors.
-		// Any actors already in the group will be ignored.
+			// We've established that only one group is selected, so we can just call Add on all these actors.
+			// Any actors already in the group will be ignored.
 		
-		TArray<AActor*> ActorsToAdd;
+			TArray<AActor*> ActorsToAdd;
 
-		bool bActorsInSameLevel = true;
-		for ( FSelectionIterator It( GEditor->GetSelectedActorIterator() ) ; It ; ++It )
-		{
-			AActor* Actor = CastChecked<AActor>( *It );
-		
-			if( Actor->GetLevel() == GroupLevel )
+			bool bActorsInSameLevel = true;
+			for ( FSelectionIterator It( GEditor->GetSelectedActorIterator() ) ; It ; ++It )
 			{
-				ActorsToAdd.Add( Actor );
-			}
-			else
-			{
-				bActorsInSameLevel = false;
-				break;
+				AActor* Actor = CastChecked<AActor>( *It );
+		
+				if( Actor->GetLevel() == GroupLevel )
+				{
+					ActorsToAdd.Add( Actor );
+				}
+				else
+				{
+					bActorsInSameLevel = false;
+					break;
+				}
 			}
-		}
 
-		if( bActorsInSameLevel )
-		{
-			if( ActorsToAdd.Num() > 0 )
+			if( bActorsInSameLevel )
 			{
-				const FScopedTransaction Transaction( NSLOCTEXT("UnrealEd", "Group_Add", "Add Actors to Group") );
-				for( int32 ActorIndex = 0; ActorIndex < ActorsToAdd.Num(); ++ActorIndex )
+				if( ActorsToAdd.Num() > 0 )
 				{
-					if ( ActorsToAdd[ActorIndex] != SelectedGroup )
+					const FScopedTransaction Transaction( NSLOCTEXT("UnrealEd", "Group_Add", "Add Actors to Group") );
+					for( int32 ActorIndex = 0; ActorIndex < ActorsToAdd.Num(); ++ActorIndex )
 					{
-						SelectedGroup->Add( *ActorsToAdd[ActorIndex] );
+						if ( ActorsToAdd[ActorIndex] != SelectedGroup )
+						{
+							SelectedGroup->Add( *ActorsToAdd[ActorIndex] );
+						}
 					}
 				}
 			}
-		}
-		else
-		{
-			FMessageDialog::Open( EAppMsgType::Ok, NSLOCTEXT("UnrealEd", "Group_CantCreateGroupMultipleLevels", "Can't group the selected actors because they are in different levels.") );
+			else
+			{
+				FMessageDialog::Open( EAppMsgType::Ok, NSLOCTEXT("UnrealEd", "Group_CantCreateGroupMultipleLevels", "Can't group the selected actors because they are in different levels.") );
+			}
 		}
 	}
 }
@@ -540,73 +558,81 @@ void AGroupActor::AddSelectedActorsToSelectedGroup()
 
 void AGroupActor::LockSelectedGroups()
 {
-	TArray<AGroupActor*> GroupsToLock;
-	for ( int32 GroupIndex=0; GroupIndex<GEditor->ActiveGroupActors.Num(); ++GroupIndex )
+	UWorld* EditorWorld = GEditor->GetEditorWorldContext().World();
+	if (EditorWorld)
 	{
-		AGroupActor* GroupToLock = GEditor->ActiveGroupActors[GroupIndex];
-		if( GroupToLock != NULL )
+		TArray<AGroupActor*> GroupsToLock;
+		for ( int32 GroupIndex=0; GroupIndex<EditorWorld->ActiveGroupActors.Num(); ++GroupIndex )
 		{
-			if( GroupToLock->HasSelectedActors(false) )
+			AGroupActor* GroupToLock = Cast<AGroupActor>(EditorWorld->ActiveGroupActors[GroupIndex]);
+			if( GroupToLock != NULL )
 			{
-				// If our selected group is already locked, move up a level to add it's potential parent for locking
-				if( GroupToLock->IsLocked() )
+				if( GroupToLock->HasSelectedActors(false) )
 				{
-					AGroupActor* GroupParent = AGroupActor::GetParentForActor(GroupToLock);
-					if(GroupParent && !GroupParent->IsLocked())
+					// If our selected group is already locked, move up a level to add it's potential parent for locking
+					if( GroupToLock->IsLocked() )
 					{
-						GroupsToLock.AddUnique(GroupParent);
+						AGroupActor* GroupParent = AGroupActor::GetParentForActor(GroupToLock);
+						if(GroupParent && !GroupParent->IsLocked())
+						{
+							GroupsToLock.AddUnique(GroupParent);
+						}
+					}
+					else // if it's not locked, add it instead!
+					{
+						GroupsToLock.AddUnique(GroupToLock);
 					}
-				}
-				else // if it's not locked, add it instead!
-				{
-					GroupsToLock.AddUnique(GroupToLock);
 				}
 			}
 		}
-	}
 
-	if( GroupsToLock.Num() > 0 )
-	{
-		const FScopedTransaction Transaction( NSLOCTEXT("UnrealEd", "Group_Lock", "Lock Groups") );
-		for ( int32 GroupIndex=0; GroupIndex<GroupsToLock.Num(); ++GroupIndex )
+		if( GroupsToLock.Num() > 0 )
 		{
-			AGroupActor* GroupToLock = GroupsToLock[GroupIndex];
-			GroupToLock->Modify();
-			GroupToLock->Lock();
-			GEditor->SelectGroup(GroupToLock, false );
+			const FScopedTransaction Transaction( NSLOCTEXT("UnrealEd", "Group_Lock", "Lock Groups") );
+			for ( int32 GroupIndex=0; GroupIndex<GroupsToLock.Num(); ++GroupIndex )
+			{
+				AGroupActor* GroupToLock = GroupsToLock[GroupIndex];
+				GroupToLock->Modify();
+				GroupToLock->Lock();
+				GEditor->SelectGroup(GroupToLock, false );
+			}
+			GEditor->NoteSelectionChange();
 		}
-		GEditor->NoteSelectionChange();
 	}
 }
 
 
 void AGroupActor::UnlockSelectedGroups()
 {
-	TArray<AGroupActor*> GroupsToUnlock;
-	for ( int32 GroupIndex=0; GroupIndex<GEditor->ActiveGroupActors.Num(); ++GroupIndex )
+	UWorld* EditorWorld = GEditor->GetEditorWorldContext().World();
+	if (EditorWorld)
 	{
-		AGroupActor* GroupToUnlock = GEditor->ActiveGroupActors[GroupIndex];
-		if( GroupToUnlock != NULL )
+		TArray<AGroupActor*> GroupsToUnlock;
+		for ( int32 GroupIndex=0; GroupIndex<EditorWorld->ActiveGroupActors.Num(); ++GroupIndex )
 		{
-			if( GroupToUnlock->IsSelected() )
+			AGroupActor* GroupToUnlock = Cast<AGroupActor>(EditorWorld->ActiveGroupActors[GroupIndex]);
+			if( GroupToUnlock != NULL )
 			{
-				GroupsToUnlock.Add(GroupToUnlock);
+				if( GroupToUnlock->IsSelected() )
+				{
+					GroupsToUnlock.Add(GroupToUnlock);
+				}
 			}
 		}
-	}
 
-	// Only unlock topmost selected group(s)
-	AGroupActor::RemoveSubGroupsFromArray(GroupsToUnlock);
-	if( GroupsToUnlock.Num() > 0 )
-	{
-		const FScopedTransaction Transaction( NSLOCTEXT("UnrealEd", "Group_Unlock", "Unlock Groups") );
-		for ( int32 GroupIndex=0; GroupIndex<GroupsToUnlock.Num(); ++GroupIndex)
+		// Only unlock topmost selected group(s)
+		AGroupActor::RemoveSubGroupsFromArray(GroupsToUnlock);
+		if( GroupsToUnlock.Num() > 0 )
 		{
-			AGroupActor* GroupToUnlock = GroupsToUnlock[GroupIndex];
-			GroupToUnlock->Modify();
-			GroupToUnlock->Unlock();
+			const FScopedTransaction Transaction( NSLOCTEXT("UnrealEd", "Group_Unlock", "Unlock Groups") );
+			for ( int32 GroupIndex=0; GroupIndex<GroupsToUnlock.Num(); ++GroupIndex)
+			{
+				AGroupActor* GroupToUnlock = GroupsToUnlock[GroupIndex];
+				GroupToUnlock->Modify();
+				GroupToUnlock->Unlock();
+			}
+			GEditor->NoteSelectionChange();
 		}
-		GEditor->NoteSelectionChange();
 	}
 }
 
@@ -733,12 +759,13 @@ void AGroupActor::PostRemove()
 			ParentGroup->Remove(*this);
 		}
 
-		// Group is no longer active
-		GEditor->ActiveGroupActors.Remove(this);
-
-		if( GetWorld() )
+		UWorld* MyWorld = GetWorld();
+		if( MyWorld )
 		{
-			GetWorld()->ModifyLevel(GetLevel());
+			// Group is no longer active
+			MyWorld->ActiveGroupActors.Remove(this);
+
+			MyWorld->ModifyLevel(GetLevel());
 			
 			// Mark the group actor for removal
 			MarkPackageDirty();
@@ -751,7 +778,7 @@ void AGroupActor::PostRemove()
 
 				// Destroy group and clear references.
 				GEditor->Layers->DisassociateActorFromLayers( this );
-				GetWorld()->EditorDestroyActor( this, false );			
+				MyWorld->EditorDestroyActor( this, false );			
 				
 				LevelRefreshAllBrowsers.Request();
 			}
diff --git a/Engine/Source/Runtime/Engine/Classes/Engine/World.h b/Engine/Source/Runtime/Engine/Classes/Engine/World.h
--- a/Engine/Source/Runtime/Engine/Classes/Engine/World.h
+++ b/Engine/Source/Runtime/Engine/Classes/Engine/World.h
@@ -485,6 +485,10 @@ class ENGINE_API UWorld : public UObject, public FNetworkNotify
 	/** List of all the layers referenced by the world's actors */
 	UPROPERTY()
 	TArray< class ULayer* > Layers; 
+
+	// Group actors currently "active"
+	UPROPERTY(transient)
+	TArray<class AActor*> ActiveGroupActors;
 #endif // WITH_EDITORONLY_DATA
 
 	/** Persistent level containing the world info, default brush and actors spawned during gameplay among other things			*/
