Fix auto convex collision never generating more than 4 hulls

diff --git a/Engine/Source/Editor/UnrealEd/Private/ConvexDecompTool.cpp b/Engine/Source/Editor/UnrealEd/Private/ConvexDecompTool.cpp
--- a/Engine/Source/Editor/UnrealEd/Private/ConvexDecompTool.cpp
+++ b/Engine/Source/Editor/UnrealEd/Private/ConvexDecompTool.cpp
@@ -43,7 +43,7 @@ void DecomposeMeshToHulls(UBodySetup* InBodySetup, const TArray<FVector>& InVert
 
 	Desc.mUseFastVersion			= 	false;
 	Desc.mBackFaceDistanceFactor	= 	0.00000000001f;
-	Desc.mDecompositionDepth		= 	2.0f;
+	Desc.mDecompositionDepth		= 	FMath::CeilToInt(FMath::Log2(InMaxHullCount));
 
 	FHACDProgressCallback Callback;
 	Desc.mCallback			= &Callback;
