Fixed an inifinite loop problem in finding cloth section's display mode

diff --git a/Engine/Source/Editor/UnrealEd/Private/Animation/DebugSkelMeshComponent.cpp b/Engine/Source/Editor/UnrealEd/Private/Animation/DebugSkelMeshComponent.cpp
--- a/Engine/Source/Editor/UnrealEd/Private/Animation/DebugSkelMeshComponent.cpp
+++ b/Engine/Source/Editor/UnrealEd/Private/Animation/DebugSkelMeshComponent.cpp
@@ -546,7 +546,7 @@ int32 UDebugSkelMeshComponent::FindCurrentSectionDisplayMode()
 	{
 		int32 LODIndex;
 		int32 NumLODs = SkelMeshResource->LODModels.Num();
-		for (LODIndex = 0; LODIndex < NumLODs; LODIndex)
+		for (LODIndex = 0; LODIndex < NumLODs; LODIndex++)
 		{
 			// if find any LOD model which has cloth data, then break
 			if (SkelMeshResource->LODModels[LODIndex].HasApexClothData())
