Prevent use-after-free bug when the FBX import fails - CleanUp was deleting the object once and then ImportFile would try and dispose of it again.

diff --git a/Engine/Source/Editor/UnrealEd/Private/Fbx/FbxMainImport.cpp b/Engine/Source/Editor/UnrealEd/Private/Fbx/FbxMainImport.cpp
--- a/Engine/Source/Editor/UnrealEd/Private/Fbx/FbxMainImport.cpp
+++ b/Engine/Source/Editor/UnrealEd/Private/Fbx/FbxMainImport.cpp
@@ -647,6 +647,10 @@ bool FFbxImporter::ImportFile(FString Filename)
 	{
 		UE_LOG(LogFbx, Log, TEXT("FBX Scene Loaded Succesfully"));
 		CurPhase = IMPORTED;
+		
+		// Release importer now as it is unneeded and in the failure case is part of CleanUp()
+		Importer->Destroy();
+		Importer = NULL;
 	}
 	else
 	{
@@ -657,9 +661,6 @@ bool FFbxImporter::ImportFile(FString Filename)
 		CurPhase = NOTSTARTED;
 	}
 	
-	Importer->Destroy();
-	Importer = NULL;
-	
 	return Result;
 }
 
