AutomationEditorCommon.cpp and. h are used for common editor latent commands for automation testing.

EditorAutomationTests.cpp -- Point light placement test is now fully functional and has the ability to take screenshots after each major step.

diff --git a/Engine/Source/Editor/UnrealEd/Private/Tests/AutomationEditorCommon.cpp b/Engine/Source/Editor/UnrealEd/Private/Tests/AutomationEditorCommon.cpp
new file mode 100644
--- /dev/null
+++ b/Engine/Source/Editor/UnrealEd/Private/Tests/AutomationEditorCommon.cpp
@@ -0,0 +1,26 @@
+// Copyright 1998-2014 Epic Games, Inc. All Rights Reserved.
+
+#include "UnrealEd.h"
+#include "AutomationEditorCommon.h"
+
+
+///////////////////////////////////////////////////////////////////////
+// Common Latent commands
+
+//Latent Undo and Redo command
+//If bUndo is true then the undo action will occur otherwise a redo will happen.
+bool FUndoRedoCommand::Update()
+{
+	if (bUndo == true)
+	{
+		//Undo
+		GEditor->UndoTransaction();
+	}
+	else
+	{
+		//Redo
+		GEditor->RedoTransaction();
+	}
+
+	return true;
+}
\ No newline at end of file
diff --git a/Engine/Source/Editor/UnrealEd/Private/Tests/EditorAutomationTests.cpp b/Engine/Source/Editor/UnrealEd/Private/Tests/EditorAutomationTests.cpp
--- a/Engine/Source/Editor/UnrealEd/Private/Tests/EditorAutomationTests.cpp
+++ b/Engine/Source/Editor/UnrealEd/Private/Tests/EditorAutomationTests.cpp
@@ -12,6 +12,7 @@
 #include "ScopedTransaction.h"
 #include "LevelEditor.h"
 #include "ModuleManager.h"
+#include "MainFrame.h"
 
 #include "Kismet2/KismetEditorUtilities.h"
 #include "Kismet2/BlueprintEditorUtils.h"
@@ -19,6 +20,8 @@
 
 #include "SlateWordWrapper.h"
 #include "AutomationCommon.h"
+#include "AutomationEditorCommon.h"
+
 
 //DEFINE_LOG_CATEGORY_STATIC(LogEditorAutomationTests, Log, All);
 
@@ -80,6 +83,71 @@ bool FCloseAllAssetEditorsCommand::Update()
 	return true;
 }
 
+/**
+* Change the attributes for a point light in the level.
+*/
+struct PointLightParameters
+{
+	APointLight* PointLight;
+	float LightBrightness = 5000.0f;
+	float LightRadius = 1000.0f;
+	FVector LightLocation = FVector(0.0f, 0.0f, 0.0f);
+	FColor LightColor = FColor(255, 255, 255);
+};
+
+//Updates the properties of a specified point light.
+DEFINE_LATENT_AUTOMATION_COMMAND_ONE_PARAMETER(PointLightUpdateCommand, PointLightParameters, PointLightUsing);
+
+bool PointLightUpdateCommand::Update()
+{
+	//Set the point light mobility, brightness, radius, and light color.
+	PointLightUsing.PointLight->SetMobility(EComponentMobility::Movable);
+	PointLightUsing.PointLight->SetBrightness(PointLightUsing.LightBrightness);
+	PointLightUsing.PointLight->SetLightColor(PointLightUsing.LightColor);
+	PointLightUsing.PointLight->TeleportTo(PointLightUsing.LightLocation, FRotator(0, 0, 0));
+	PointLightUsing.PointLight->SetRadius(PointLightUsing.LightRadius);
+	return true;
+}
+
+/**
+* Duplicates a point light.
+*/
+DEFINE_LATENT_AUTOMATION_COMMAND_ONE_PARAMETER(PointLightDuplicationCommand, PointLightParameters, PointLightDuplicating);
+
+bool PointLightDuplicationCommand::Update()
+{
+
+	FScopedTransaction DuplicateLightScope(NSLOCTEXT("UnrealEd.Test", "DuplicateLightScope", "Duplicate Light Scope"));
+
+	//Duplicate the light.
+	bool bOffsetLocations = false;
+	GEditor->edactDuplicateSelected(PointLightDuplicating.PointLight->GetLevel(), bOffsetLocations);
+	for (FSelectionIterator It(GEditor->GetSelectedActorIterator()); It; ++It)
+	{
+		AActor* Actor = static_cast<AActor*>(*It);
+		Actor->TeleportTo(FVector(PointLightDuplicating.LightLocation), FRotator(0, 0, 0));
+	}
+	return true;
+}
+
+/**
+* This will then take a screenshot of the editor, only if it has been enabled.
+*/
+void TakeLatentAutomationScreenshot(struct WindowScreenshotParameters ScreenshotParameters, FString BaseFileName, FString ScreenshotTitle, FString ScreenshotFolderName, bool bIncludeHardwareDetails)
+{
+	//Update the screenshot name, then take a screenshot.
+	if (FAutomationTestFramework::GetInstance().IsScreenshotAllowed())
+	{
+		//Update the screenshot name and get the location of where it will be saved.
+		ScreenshotParameters.ScreenshotName = TEXT("ScreenshotTitle");
+		FString TestName = FString::Printf(TEXT("%s/%s"), *BaseFileName, *ScreenshotFolderName);
+		AutomationCommon::GetScreenshotPath(TestName, ScreenshotParameters.ScreenshotName, bIncludeHardwareDetails);
+
+		//Take a screenshot
+		ADD_LATENT_AUTOMATION_COMMAND(FTakeEditorScreenshotCommand(ScreenshotParameters));
+	}
+}
+
 //////////////////////////////////////////////////////////////////////////
 // FEditorAutomationTestUtilities
 
@@ -1430,7 +1498,22 @@ IMPLEMENT_SIMPLE_AUTOMATION_TEST(FLightPlacement, "QA.Point Light Placement", EA
 
 bool FLightPlacement::RunTest(const FString& Parameters)
 {
-	//Open a new empty map
+	//Initialize the parameters for taking a screenshot as well as update the placed point light.
+	WindowScreenshotParameters PointLightPlacementWindowParameters;
+	PointLightParameters LightParameters;
+	bool bUndo = true;
+
+	//Set the CurrentWindow to the Mainframe.  This information is used for taking a screenshot later.
+	if (FModuleManager::Get().IsModuleLoaded("MainFrame"))
+	{
+		IMainFrameModule& MainFrame = FModuleManager::LoadModuleChecked<IMainFrameModule>("MainFrame");
+		PointLightPlacementWindowParameters.CurrentWindow = MainFrame.GetParentWindow();
+	}
+
+	//Set the Test Name which is used later for getting the directory to store the screenshots.
+	const FString BaseFileName = TEXT("PointLightPlacementTest");
+
+	//Open a new blank map.
 	UWorld* World = GEditor->NewMap();
 
 	//Move the perspective viewport view to show the test.
@@ -1445,7 +1528,7 @@ bool FLightPlacement::RunTest(const FString& Parameters)
 	}
 
 	{
-		//Gather assets
+		//Gather assets.
 		UObject* Astroid = (UStaticMesh*)StaticLoadObject(UStaticMesh::StaticClass(), NULL, TEXT("/Engine/Content/EditorAutomation/Astroid.Astroid"), NULL, LOAD_None, NULL);
 		//Add Astroid mesh to the world
 		AStaticMeshActor* StaticMesh = Cast<AStaticMeshActor>(FActorFactoryAssetProxy::AddActorForAsset(Astroid));
@@ -1453,34 +1536,53 @@ bool FLightPlacement::RunTest(const FString& Parameters)
 		StaticMesh->SetActorRelativeScale3D(FVector(1.0f, 1.0f, 1.0f));
 	}
 
-	//Point Light
+	//Create the point light and set it's mobility, brightness, and light color.
 	APointLight* PointLight = Cast<APointLight>(GEditor->AddActor(World->GetCurrentLevel(), APointLight::StaticClass(), FVector(0.0f, 0.0f, 400.0f)));
-	PointLight->SetMobility(EComponentMobility::Movable);
-	PointLight->SetBrightness(5000.0f);
-	PointLight->SetLightColor(FColor(255, 0, 0));
+	LightParameters.PointLight = PointLight;
+	LightParameters.LightColor = FColor(255, 0, 0);
+	LightParameters.LightLocation = FVector(0.0f, 0.0f, 400.0f);
+	ADD_LATENT_AUTOMATION_COMMAND(PointLightUpdateCommand(LightParameters));
 
-	//Duplicate the point light
-	{
-		FScopedTransaction DuplicateLightScope( NSLOCTEXT("UnrealEd.Test", "DuplicateLightScope", "Duplicate Light Scope") );
+	//Wait
+	ADD_LATENT_AUTOMATION_COMMAND(FWaitLatentCommand(0.1f));
 
-		//Duplicate the light
-		bool bOffsetLocations = false;
-		GEditor->edactDuplicateSelected(PointLight->GetLevel(), bOffsetLocations);
-		for ( FSelectionIterator It( GEditor->GetSelectedActorIterator() ) ; It ; ++It )
-		{
-			AActor* Actor = static_cast<AActor*>( *It );
-			Actor->TeleportTo(FVector(10.0f,10.0f,400.0f), FRotator(0, 0, 0));
-		}
-	}
+	//Take a screenshot of the newly placed point light.
+	TakeLatentAutomationScreenshot(PointLightPlacementWindowParameters, BaseFileName, FString::Printf(TEXT("PlacedPointLight")), FString::Printf(TEXT("01_Placed")), true);
 
-	//Undo the duplication then redo it
-	GEditor->UndoTransaction();
-	GEditor->RedoTransaction();
+	//Duplicate the point light.
+	LightParameters.LightLocation = FVector(10.0f, 10.0f, 400.0f);
+	ADD_LATENT_AUTOMATION_COMMAND(PointLightDuplicationCommand(LightParameters));
+
+	//Wait
+	ADD_LATENT_AUTOMATION_COMMAND(FWaitLatentCommand(0.1f));
+
+	//Take a screenshot of the duplicated point light.
+	TakeLatentAutomationScreenshot(PointLightPlacementWindowParameters, BaseFileName, FString::Printf(TEXT("DuplicatedLight")), FString::Printf(TEXT("02_Dupe")), true);
+
+	//Undo the duplication.
+	ADD_LATENT_AUTOMATION_COMMAND(FUndoRedoCommand(true));
+
+	//Wait
+	ADD_LATENT_AUTOMATION_COMMAND(FWaitLatentCommand(0.1f));
+
+	//Take a screenshot of the scene after the duplication has been undone.
+	TakeLatentAutomationScreenshot(PointLightPlacementWindowParameters, BaseFileName, FString::Printf(TEXT("UndoDuplicationPointLight")), FString::Printf(TEXT("03_Undo")), true);
+
+	//Redo the duplication.
+	ADD_LATENT_AUTOMATION_COMMAND(FUndoRedoCommand(false));
+
+	//Update the original point light actor.
+	LightParameters.LightRadius = 500.0f;
+	LightParameters.LightLocation = FVector(500.0f, 300.0f, 500.0f);
+	LightParameters.LightColor = FColor(255, 255, 255);
+	ADD_LATENT_AUTOMATION_COMMAND(PointLightUpdateCommand(LightParameters));
+
+	//Wait
+	ADD_LATENT_AUTOMATION_COMMAND(FWaitLatentCommand(0.1f));
 
-	//Update the duplicated asset
-	PointLight->TeleportTo(FVector(500.0f,300.0f, 300.0f), FRotator(0, 0, 0));
-	PointLight->SetRadius( 500.0f );
-	PointLight->SetLightColor(FColor(255, 255, 255));
+	//Take a screenshot of the final scene.
+	//It is expected to show two moveable point lights (one red, one white) and a static mesh.
+	TakeLatentAutomationScreenshot(PointLightPlacementWindowParameters, BaseFileName, FString::Printf(TEXT("FinalPointLight")), FString::Printf(TEXT("04_Final")), true);
 
 	return true;
 }
diff --git a/Engine/Source/Editor/UnrealEd/Public/Tests/AutomationEditorCommon.h b/Engine/Source/Editor/UnrealEd/Public/Tests/AutomationEditorCommon.h
new file mode 100644
--- /dev/null
+++ b/Engine/Source/Editor/UnrealEd/Public/Tests/AutomationEditorCommon.h
@@ -0,0 +1,14 @@
+// Copyright 1998-2014 Epic Games, Inc. All Rights Reserved.
+
+#pragma once
+
+
+//////////////////////////////////////////////////////////////////////////
+//Common latent commands used for automated editor testing.
+
+
+/**
+* Creates a latent command which the user can either undo or redo an action.
+* True will trigger Undo, False will trigger a redo
+*/
+DEFINE_LATENT_AUTOMATION_COMMAND_ONE_PARAMETER(FUndoRedoCommand, bool, bUndo);
\ No newline at end of file
