EditorAutomationTests.cpp - Fixed my struct and added a constructor

diff --git a/Engine/Source/Editor/UnrealEd/Private/Tests/EditorAutomationTests.cpp b/Engine/Source/Editor/UnrealEd/Private/Tests/EditorAutomationTests.cpp
--- a/Engine/Source/Editor/UnrealEd/Private/Tests/EditorAutomationTests.cpp
+++ b/Engine/Source/Editor/UnrealEd/Private/Tests/EditorAutomationTests.cpp
@@ -89,10 +89,18 @@ bool FCloseAllAssetEditorsCommand::Update()
 struct PointLightParameters
 {
 	APointLight* PointLight;
-	float LightBrightness = 5000.0f;
-	float LightRadius = 1000.0f;
-	FVector LightLocation = FVector(0.0f, 0.0f, 0.0f);
-	FColor LightColor = FColor(255, 255, 255);
+	float LightBrightness;
+	float LightRadius;
+	FVector LightLocation;
+	FColor LightColor;
+	PointLightParameters()
+		: PointLight(nullptr)
+		, LightBrightness(5000.0f)
+		, LightRadius(1000.0f)
+		, LightLocation(FVector(0.0f, 0.0f, 0.0f))
+		, LightColor(FColor(255, 255, 255))
+	{
+	}
 };
 
 //Updates the properties of a specified point light.
