UE4 Refactoring. Changed OVERRIDE and FINAL macros to keywords override and final.

diff --git a/Engine/Source/Editor/UMGEditor/Private/SUMGEditorTreeItem.h b/Engine/Source/Editor/UMGEditor/Private/SUMGEditorTreeItem.h
--- a/Engine/Source/Editor/UMGEditor/Private/SUMGEditorTreeItem.h
+++ b/Engine/Source/Editor/UMGEditor/Private/SUMGEditorTreeItem.h
@@ -14,10 +14,10 @@ public:
 
 	void Construct(const FArguments& InArgs, TSharedPtr<FBlueprintEditor> InBlueprintEditor, UWidget* InItem);
 
-	virtual void OnDragEnter(const FGeometry& MyGeometry, const FDragDropEvent& DragDropEvent) OVERRIDE;
-	virtual void OnDragLeave(const FDragDropEvent& DragDropEvent) OVERRIDE;
-	virtual FReply OnDragOver(const FGeometry& MyGeometry, const FDragDropEvent& DragDropEvent) OVERRIDE;
-	virtual FReply OnDrop(const FGeometry& MyGeometry, const FDragDropEvent& DragDropEvent) OVERRIDE;
+	virtual void OnDragEnter(const FGeometry& MyGeometry, const FDragDropEvent& DragDropEvent) override;
+	virtual void OnDragLeave(const FDragDropEvent& DragDropEvent) override;
+	virtual FReply OnDragOver(const FGeometry& MyGeometry, const FDragDropEvent& DragDropEvent) override;
+	virtual FReply OnDrop(const FGeometry& MyGeometry, const FDragDropEvent& DragDropEvent) override;
 	
 private:
 	FText GetItemText() const;
diff --git a/Engine/Source/Editor/UMGEditor/Private/SUMGEditorWidgetTemplates.cpp b/Engine/Source/Editor/UMGEditor/Private/SUMGEditorWidgetTemplates.cpp
--- a/Engine/Source/Editor/UMGEditor/Private/SUMGEditorWidgetTemplates.cpp
+++ b/Engine/Source/Editor/UMGEditor/Private/SUMGEditorWidgetTemplates.cpp
@@ -70,7 +70,7 @@ public:
 		return Template->Name;
 	}
 
-	virtual TSharedRef<ITableRow> BuildRow(const TSharedRef<STableViewBase>& OwnerTable) OVERRIDE
+	virtual TSharedRef<ITableRow> BuildRow(const TSharedRef<STableViewBase>& OwnerTable) override
 	{
 		return SNew(STableRow< TSharedPtr<FWidgetViewModel> >, OwnerTable)
 			.Padding(2.0f)
@@ -96,7 +96,7 @@ public:
 		return GroupName;
 	}
 	
-	virtual TSharedRef<ITableRow> BuildRow(const TSharedRef<STableViewBase>& OwnerTable) OVERRIDE
+	virtual TSharedRef<ITableRow> BuildRow(const TSharedRef<STableViewBase>& OwnerTable) override
 	{
 		return SNew(STableRow< TSharedPtr<FWidgetViewModel> >, OwnerTable)
 			.Padding(2.0f)
@@ -107,7 +107,7 @@ public:
 			];
 	}
 
-	virtual void GetChildren(TArray< TSharedPtr<FWidgetViewModel> >& OutChildren) OVERRIDE
+	virtual void GetChildren(TArray< TSharedPtr<FWidgetViewModel> >& OutChildren) override
 	{
 		for ( TSharedPtr<FWidgetViewModel>& Child : Children )
 		{
diff --git a/Engine/Source/Editor/UMGEditor/Private/Templates/WidgetTemplateButton.h b/Engine/Source/Editor/UMGEditor/Private/Templates/WidgetTemplateButton.h
--- a/Engine/Source/Editor/UMGEditor/Private/Templates/WidgetTemplateButton.h
+++ b/Engine/Source/Editor/UMGEditor/Private/Templates/WidgetTemplateButton.h
@@ -9,7 +9,7 @@ class FWidgetTemplateButton : public FWidgetTemplate
 public:
 	FWidgetTemplateButton();
 
-	virtual FText GetCategory() OVERRIDE;
+	virtual FText GetCategory() override;
 
-	virtual UWidget* Create(UWidgetTree* Tree) OVERRIDE;
+	virtual UWidget* Create(UWidgetTree* Tree) override;
 };
diff --git a/Engine/Source/Editor/UMGEditor/Private/Templates/WidgetTemplateCheckBox.h b/Engine/Source/Editor/UMGEditor/Private/Templates/WidgetTemplateCheckBox.h
--- a/Engine/Source/Editor/UMGEditor/Private/Templates/WidgetTemplateCheckBox.h
+++ b/Engine/Source/Editor/UMGEditor/Private/Templates/WidgetTemplateCheckBox.h
@@ -9,7 +9,7 @@ class FWidgetTemplateCheckBox : public FWidgetTemplate
 public:
 	FWidgetTemplateCheckBox();
 
-	virtual FText GetCategory() OVERRIDE;
+	virtual FText GetCategory() override;
 
-	virtual UWidget* Create(UWidgetTree* Tree) OVERRIDE;
+	virtual UWidget* Create(UWidgetTree* Tree) override;
 };
diff --git a/Engine/Source/Editor/UMGEditor/Private/UMGBlueprintEditorExtensionHook.cpp b/Engine/Source/Editor/UMGEditor/Private/UMGBlueprintEditorExtensionHook.cpp
--- a/Engine/Source/Editor/UMGEditor/Private/UMGBlueprintEditorExtensionHook.cpp
+++ b/Engine/Source/Editor/UMGEditor/Private/UMGBlueprintEditorExtensionHook.cpp
@@ -48,7 +48,7 @@ public:
 		ViewMenuTooltip = LOCTEXT("SlatePreview_ViewMenu_ToolTip", "Show the UI preview");
 	}
 
-	virtual TSharedRef<SWidget> CreateTabBody(const FWorkflowTabSpawnInfo& Info) const OVERRIDE
+	virtual TSharedRef<SWidget> CreateTabBody(const FWorkflowTabSpawnInfo& Info) const override
 	{
 		//TSharedPtr<SWidget> Result;
 		//if ( BlueprintEditorPtr->CanAccessComponentsMode() )
@@ -88,7 +88,7 @@ public:
 		ViewMenuTooltip = LOCTEXT("SlateHierarchy_ViewMenu_ToolTip", "Show the Hierarchy");
 	}
 
-	virtual TSharedRef<SWidget> CreateTabBody(const FWorkflowTabSpawnInfo& Info) const OVERRIDE
+	virtual TSharedRef<SWidget> CreateTabBody(const FWorkflowTabSpawnInfo& Info) const override
 	{
 		TSharedPtr<FWidgetBlueprintEditor> BlueprintEditorPtr = StaticCastSharedPtr<FWidgetBlueprintEditor>(BlueprintEditor.Pin());
 
@@ -120,7 +120,7 @@ public:
 		ViewMenuTooltip = LOCTEXT("WidgetTemplates_ViewMenu_ToolTip", "Show the Templates");
 	}
 
-	virtual TSharedRef<SWidget> CreateTabBody(const FWorkflowTabSpawnInfo& Info) const OVERRIDE
+	virtual TSharedRef<SWidget> CreateTabBody(const FWorkflowTabSpawnInfo& Info) const override
 	{
 		TSharedPtr<FWidgetBlueprintEditor> BlueprintEditorPtr = StaticCastSharedPtr<FWidgetBlueprintEditor>(BlueprintEditor.Pin());
 
@@ -150,7 +150,7 @@ public:
 		ViewMenuTooltip = LOCTEXT("Sequencer_ViewMenu_ToolTip", "Show the Animation editor");
 	}
 
-	virtual TSharedRef<SWidget> CreateTabBody(const FWorkflowTabSpawnInfo& Info) const OVERRIDE
+	virtual TSharedRef<SWidget> CreateTabBody(const FWorkflowTabSpawnInfo& Info) const override
 	{
 		TSharedPtr<FWidgetBlueprintEditor> BlueprintEditorPinned = BlueprintEditor.Pin();
 
@@ -255,7 +255,7 @@ public:
 		);
 	}
 
-	virtual void PreDeactivateMode() OVERRIDE
+	virtual void PreDeactivateMode() override
 	{
 		FBlueprintDefaultsApplicationMode::PreDeactivateMode();
 
@@ -266,7 +266,7 @@ public:
 		PropertyModule.UnregisterCustomPropertyTypeLayout(TEXT("CanvasPanelSlot"));
 	}
 
-	virtual void PostActivateMode() OVERRIDE
+	virtual void PostActivateMode() override
 	{
 		FBlueprintDefaultsApplicationMode::PostActivateMode();
 
@@ -311,12 +311,12 @@ public:
 		: FBlueprintEditorApplicationMode(InBlueprintEditor, InModeName, bRegisterViewport, bRegisterDefaultsTab)
 	{}
 
-	virtual void PreDeactivateMode() OVERRIDE
+	virtual void PreDeactivateMode() override
 	{
 		FBlueprintEditorApplicationMode::PreDeactivateMode();
 	}
 
-	virtual void PostActivateMode() OVERRIDE
+	virtual void PostActivateMode() override
 	{
 		FBlueprintEditorApplicationMode::PostActivateMode();
 	}
diff --git a/Engine/Source/Editor/UMGEditor/Private/WidgetBlueprintCompiler.h b/Engine/Source/Editor/UMGEditor/Private/WidgetBlueprintCompiler.h
--- a/Engine/Source/Editor/UMGEditor/Private/WidgetBlueprintCompiler.h
+++ b/Engine/Source/Editor/UMGEditor/Private/WidgetBlueprintCompiler.h
@@ -17,7 +17,7 @@ public:
 	virtual ~FWidgetBlueprintCompiler();
 
 	// FKismetCompilerContext
-	virtual void Compile() OVERRIDE;
+	virtual void Compile() override;
 	// End FKismetCompilerContext
 
 protected:
@@ -27,14 +27,14 @@ protected:
 
 	// FKismetCompilerContext
 	//virtual UEdGraphSchema_K2* CreateSchema() OVERRIDE;
-	virtual void MergeUbergraphPagesIn(UEdGraph* Ubergraph) OVERRIDE;
-	virtual void SpawnNewClass(const FString& NewClassName) OVERRIDE;
-	virtual void CleanAndSanitizeClass(UBlueprintGeneratedClass* ClassToClean, UObject*& OldCDO) OVERRIDE;
-	virtual void EnsureProperGeneratedClass(UClass*& TargetClass) OVERRIDE;
-	virtual void CreateClassVariablesFromBlueprint() OVERRIDE;
-	virtual void FinishCompilingClass(UClass* Class) OVERRIDE;
-	virtual bool ValidateGeneratedClass(UBlueprintGeneratedClass* Class) OVERRIDE;
-	virtual void CopyTermDefaultsToDefaultObject(UObject* DefaultObject) OVERRIDE;
+	virtual void MergeUbergraphPagesIn(UEdGraph* Ubergraph) override;
+	virtual void SpawnNewClass(const FString& NewClassName) override;
+	virtual void CleanAndSanitizeClass(UBlueprintGeneratedClass* ClassToClean, UObject*& OldCDO) override;
+	virtual void EnsureProperGeneratedClass(UClass*& TargetClass) override;
+	virtual void CreateClassVariablesFromBlueprint() override;
+	virtual void FinishCompilingClass(UClass* Class) override;
+	virtual bool ValidateGeneratedClass(UBlueprintGeneratedClass* Class) override;
+	virtual void CopyTermDefaultsToDefaultObject(UObject* DefaultObject) override;
 	// End FKismetCompilerContext
 
 protected:
diff --git a/Engine/Source/Editor/UMGEditor/Private/WidgetBlueprintEditor.h b/Engine/Source/Editor/UMGEditor/Private/WidgetBlueprintEditor.h
--- a/Engine/Source/Editor/UMGEditor/Private/WidgetBlueprintEditor.h
+++ b/Engine/Source/Editor/UMGEditor/Private/WidgetBlueprintEditor.h
@@ -20,12 +20,12 @@ public:
 
 	void InitWidgetBlueprintEditor(const EToolkitMode::Type Mode, const TSharedPtr< IToolkitHost >& InitToolkitHost, const TArray<UBlueprint*>& InBlueprints, bool bShouldOpenInDefaultsMode);
 
-	virtual void Tick(float DeltaTime) OVERRIDE;
-	virtual void NotifyPreChange(class FEditPropertyChain* PropertyAboutToChange) OVERRIDE;
-	virtual void NotifyPostChange(const FPropertyChangedEvent& PropertyChangedEvent, class FEditPropertyChain* PropertyThatChanged) OVERRIDE;
+	virtual void Tick(float DeltaTime) override;
+	virtual void NotifyPreChange(class FEditPropertyChain* PropertyAboutToChange) override;
+	virtual void NotifyPostChange(const FPropertyChangedEvent& PropertyChangedEvent, class FEditPropertyChain* PropertyThatChanged) override;
 
 	/** FGCObjectInterface */
-	virtual void AddReferencedObjects( FReferenceCollector& Collector ) OVERRIDE;
+	virtual void AddReferencedObjects( FReferenceCollector& Collector ) override;
 
 	class UWidgetBlueprint* GetWidgetBlueprintObj() const;
 
@@ -37,7 +37,7 @@ public:
 	void SelectWidgets(TArray<UWidget*> Widgets);
 
 private:
-	virtual void OnBlueprintChanged(UBlueprint* InBlueprint) OVERRIDE;
+	virtual void OnBlueprintChanged(UBlueprint* InBlueprint) override;
 
 	void DestroyPreview();
 	void UpdatePreview(UBlueprint* InBlueprint, bool bInForceFullUpdate);
diff --git a/Engine/Source/Editor/UMGEditor/Private/WidgetGraphNode_Base.cpp b/Engine/Source/Editor/UMGEditor/Private/WidgetGraphNode_Base.cpp
--- a/Engine/Source/Editor/UMGEditor/Private/WidgetGraphNode_Base.cpp
+++ b/Engine/Source/Editor/UMGEditor/Private/WidgetGraphNode_Base.cpp
@@ -35,7 +35,7 @@ public:
 		}
 	}
 
-	virtual void GetRecordDefaults(UProperty* TestProperty, FOptionalPinFromProperty& Record) const OVERRIDE
+	virtual void GetRecordDefaults(UProperty* TestProperty, FOptionalPinFromProperty& Record) const override
 	{
 		const UWidgetGraphSchema* Schema = GetDefault<UWidgetGraphSchema>();
 
@@ -56,7 +56,7 @@ public:
 		Record.bPropertyIsCustomized = false;
 	}
 
-	virtual void CustomizePinData(UEdGraphPin* Pin, FName SourcePropertyName, int32 ArrayIndex, UProperty* Property) const OVERRIDE
+	virtual void CustomizePinData(UEdGraphPin* Pin, FName SourcePropertyName, int32 ArrayIndex, UProperty* Property) const override
 	{
 		if ( BaseNode != NULL )
 		{
@@ -64,7 +64,7 @@ public:
 		}
 	}
 
-	void PostInitNewPin(UEdGraphPin* Pin, FOptionalPinFromProperty& Record, int32 ArrayIndex, UProperty* Property, uint8* PropertyAddress) const OVERRIDE
+	void PostInitNewPin(UEdGraphPin* Pin, FOptionalPinFromProperty& Record, int32 ArrayIndex, UProperty* Property, uint8* PropertyAddress) const override
 	{
 		check(PropertyAddress != NULL);
 		check(Record.bShowPin);
@@ -89,7 +89,7 @@ public:
 		}
 	}
 
-	void PostRemovedOldPin(FOptionalPinFromProperty& Record, int32 ArrayIndex, UProperty* Property, uint8* PropertyAddress) const OVERRIDE
+	void PostRemovedOldPin(FOptionalPinFromProperty& Record, int32 ArrayIndex, UProperty* Property, uint8* PropertyAddress) const override
 	{
 		check(PropertyAddress != NULL);
 		check(!Record.bShowPin);
diff --git a/Engine/Source/Editor/UndoHistory/Private/UndoHistoryModule.cpp b/Engine/Source/Editor/UndoHistory/Private/UndoHistoryModule.cpp
--- a/Engine/Source/Editor/UndoHistory/Private/UndoHistoryModule.cpp
+++ b/Engine/Source/Editor/UndoHistory/Private/UndoHistoryModule.cpp
@@ -20,7 +20,7 @@ public:
 
 	// Begin IModuleInterface interface
 
-	virtual void StartupModule( ) OVERRIDE
+	virtual void StartupModule( ) override
 	{
 		FGlobalTabmanager::Get()->RegisterNomadTabSpawner(UndoHistoryTabName, FOnSpawnTab::CreateRaw(this, &FUndoHistoryModule::HandleSpawnSettingsTab))
 			.SetDisplayName(NSLOCTEXT("FUndoHistoryModule", "UndoHistoryTabTitle", "Undo History"))
@@ -29,12 +29,12 @@ public:
 			.SetIcon(FSlateIcon(FEditorStyle::GetStyleSetName(), "UndoHistory.TabIcon"));
 	}
 
-	virtual void ShutdownModule( ) OVERRIDE
+	virtual void ShutdownModule( ) override
 	{
 		FGlobalTabmanager::Get()->UnregisterNomadTabSpawner(UndoHistoryTabName);
 	}
 
-	virtual bool SupportsDynamicReloading( ) OVERRIDE
+	virtual bool SupportsDynamicReloading( ) override
 	{
 		return true;
 	}
diff --git a/Engine/Source/Editor/UnrealEd/Classes/ActorFactories/ActorFactoryAmbientSound.h b/Engine/Source/Editor/UnrealEd/Classes/ActorFactories/ActorFactoryAmbientSound.h
--- a/Engine/Source/Editor/UnrealEd/Classes/ActorFactories/ActorFactoryAmbientSound.h
+++ b/Engine/Source/Editor/UnrealEd/Classes/ActorFactories/ActorFactoryAmbientSound.h
@@ -15,10 +15,10 @@ class UActorFactoryAmbientSound : public UActorFactory
 	GENERATED_UCLASS_BODY()
 
 	// Begin UActorFactory Interface
-	virtual void PostSpawnActor( UObject* Asset, AActor* NewActor ) OVERRIDE;
-	virtual void PostCreateBlueprint( UObject* Asset, AActor* CDO ) OVERRIDE;
-	virtual bool CanCreateActorFrom( const FAssetData& AssetData, FText& OutErrorMsg ) OVERRIDE;
-	virtual UObject* GetAssetFromActorInstance(AActor* ActorInstance) OVERRIDE;
+	virtual void PostSpawnActor( UObject* Asset, AActor* NewActor ) override;
+	virtual void PostCreateBlueprint( UObject* Asset, AActor* CDO ) override;
+	virtual bool CanCreateActorFrom( const FAssetData& AssetData, FText& OutErrorMsg ) override;
+	virtual UObject* GetAssetFromActorInstance(AActor* ActorInstance) override;
 	// End UActorFactory Interface
 };
 
diff --git a/Engine/Source/Editor/UnrealEd/Classes/ActorFactories/ActorFactoryAnimationAsset.h b/Engine/Source/Editor/UnrealEd/Classes/ActorFactories/ActorFactoryAnimationAsset.h
--- a/Engine/Source/Editor/UnrealEd/Classes/ActorFactories/ActorFactoryAnimationAsset.h
+++ b/Engine/Source/Editor/UnrealEd/Classes/ActorFactories/ActorFactoryAnimationAsset.h
@@ -11,12 +11,12 @@ class UActorFactoryAnimationAsset : public UActorFactorySkeletalMesh
 
 protected:
 	// Begin UActorFactory Interface
-	virtual void PostSpawnActor( UObject* Asset, AActor* NewActor ) OVERRIDE;
-	virtual void PostCreateBlueprint( UObject* Asset, AActor* CDO ) OVERRIDE;
-	virtual bool CanCreateActorFrom( const FAssetData& AssetData, FText& OutErrorMsg ) OVERRIDE;
+	virtual void PostSpawnActor( UObject* Asset, AActor* NewActor ) override;
+	virtual void PostCreateBlueprint( UObject* Asset, AActor* CDO ) override;
+	virtual bool CanCreateActorFrom( const FAssetData& AssetData, FText& OutErrorMsg ) override;
 	// End UActorFactory Interface
 
-	virtual USkeletalMesh* GetSkeletalMeshFromAsset( UObject* Asset ) const OVERRIDE;
+	virtual USkeletalMesh* GetSkeletalMeshFromAsset( UObject* Asset ) const override;
 };
 
 
diff --git a/Engine/Source/Editor/UnrealEd/Classes/ActorFactories/ActorFactoryBlueprint.h b/Engine/Source/Editor/UnrealEd/Classes/ActorFactories/ActorFactoryBlueprint.h
--- a/Engine/Source/Editor/UnrealEd/Classes/ActorFactories/ActorFactoryBlueprint.h
+++ b/Engine/Source/Editor/UnrealEd/Classes/ActorFactories/ActorFactoryBlueprint.h
@@ -10,10 +10,10 @@ class UActorFactoryBlueprint : public UActorFactory
 	GENERATED_UCLASS_BODY()
 
 	// Begin UActorFactory Interface
-	virtual bool PreSpawnActor( UObject* Asset, FVector& InOutLocation, FRotator& InOutRotation, bool bRotationWasSupplied ) OVERRIDE;
-	virtual void PostSpawnActor( UObject* Asset, AActor* NewActor ) OVERRIDE;
-	virtual bool CanCreateActorFrom( const FAssetData& AssetData, FText& OutErrorMsg ) OVERRIDE;
-	virtual AActor* GetDefaultActor( const FAssetData& AssetData ) OVERRIDE;
+	virtual bool PreSpawnActor( UObject* Asset, FVector& InOutLocation, FRotator& InOutRotation, bool bRotationWasSupplied ) override;
+	virtual void PostSpawnActor( UObject* Asset, AActor* NewActor ) override;
+	virtual bool CanCreateActorFrom( const FAssetData& AssetData, FText& OutErrorMsg ) override;
+	virtual AActor* GetDefaultActor( const FAssetData& AssetData ) override;
 	// End UActorFactory Interface	
 };
 
diff --git a/Engine/Source/Editor/UnrealEd/Classes/ActorFactories/ActorFactoryBoxVolume.h b/Engine/Source/Editor/UnrealEd/Classes/ActorFactories/ActorFactoryBoxVolume.h
--- a/Engine/Source/Editor/UnrealEd/Classes/ActorFactories/ActorFactoryBoxVolume.h
+++ b/Engine/Source/Editor/UnrealEd/Classes/ActorFactories/ActorFactoryBoxVolume.h
@@ -10,7 +10,7 @@ class UActorFactoryBoxVolume : public UActorFactory
 	GENERATED_UCLASS_BODY()
 
 	// Begin UActorFactory Interface
-	virtual bool CanCreateActorFrom( const FAssetData& AssetData, FText& OutErrorMsg ) OVERRIDE;
-	virtual void PostSpawnActor( UObject* Asset, AActor* NewActor ) OVERRIDE;
+	virtual bool CanCreateActorFrom( const FAssetData& AssetData, FText& OutErrorMsg ) override;
+	virtual void PostSpawnActor( UObject* Asset, AActor* NewActor ) override;
 	// End UActorFactory Interface
 };
diff --git a/Engine/Source/Editor/UnrealEd/Classes/ActorFactories/ActorFactoryClass.h b/Engine/Source/Editor/UnrealEd/Classes/ActorFactories/ActorFactoryClass.h
--- a/Engine/Source/Editor/UnrealEd/Classes/ActorFactories/ActorFactoryClass.h
+++ b/Engine/Source/Editor/UnrealEd/Classes/ActorFactories/ActorFactoryClass.h
@@ -10,12 +10,12 @@ class UActorFactoryClass : public UActorFactory
 	GENERATED_UCLASS_BODY()
 
 	// Begin UActorFactory Interface
-	virtual bool PreSpawnActor( UObject* Asset, FVector& InOutLocation, FRotator& InOutRotation, bool bRotationWasSupplied ) OVERRIDE;
-	virtual void PostSpawnActor( UObject* Asset, AActor* NewActor ) OVERRIDE;
-	virtual bool CanCreateActorFrom( const FAssetData& AssetData, FText& OutErrorMsg ) OVERRIDE;
-	virtual AActor* GetDefaultActor( const FAssetData& AssetData ) OVERRIDE;
+	virtual bool PreSpawnActor( UObject* Asset, FVector& InOutLocation, FRotator& InOutRotation, bool bRotationWasSupplied ) override;
+	virtual void PostSpawnActor( UObject* Asset, AActor* NewActor ) override;
+	virtual bool CanCreateActorFrom( const FAssetData& AssetData, FText& OutErrorMsg ) override;
+	virtual AActor* GetDefaultActor( const FAssetData& AssetData ) override;
 	// End UActorFactory Interface	
 
 protected:
-	virtual AActor* SpawnActor( UObject* Asset, ULevel* InLevel, const FVector& Location, const FRotator& Rotation, EObjectFlags ObjectFlags, const FName& Name ) OVERRIDE;
+	virtual AActor* SpawnActor( UObject* Asset, ULevel* InLevel, const FVector& Location, const FRotator& Rotation, EObjectFlags ObjectFlags, const FName& Name ) override;
 };
diff --git a/Engine/Source/Editor/UnrealEd/Classes/ActorFactories/ActorFactoryCylinderVolume.h b/Engine/Source/Editor/UnrealEd/Classes/ActorFactories/ActorFactoryCylinderVolume.h
--- a/Engine/Source/Editor/UnrealEd/Classes/ActorFactories/ActorFactoryCylinderVolume.h
+++ b/Engine/Source/Editor/UnrealEd/Classes/ActorFactories/ActorFactoryCylinderVolume.h
@@ -10,7 +10,7 @@ class UActorFactoryCylinderVolume : public UActorFactory
 	GENERATED_UCLASS_BODY()
 
 	// Begin UActorFactory Interface
-	virtual bool CanCreateActorFrom( const FAssetData& AssetData, FText& OutErrorMsg ) OVERRIDE;
-	virtual void PostSpawnActor( UObject* Asset, AActor* NewActor ) OVERRIDE;
+	virtual bool CanCreateActorFrom( const FAssetData& AssetData, FText& OutErrorMsg ) override;
+	virtual void PostSpawnActor( UObject* Asset, AActor* NewActor ) override;
 	// End UActorFactory Interface
 };
diff --git a/Engine/Source/Editor/UnrealEd/Classes/ActorFactories/ActorFactoryDeferredDecal.h b/Engine/Source/Editor/UnrealEd/Classes/ActorFactories/ActorFactoryDeferredDecal.h
--- a/Engine/Source/Editor/UnrealEd/Classes/ActorFactories/ActorFactoryDeferredDecal.h
+++ b/Engine/Source/Editor/UnrealEd/Classes/ActorFactories/ActorFactoryDeferredDecal.h
@@ -12,10 +12,10 @@ public:
 
 protected:
 	// Begin UActorFactory Interface
-	virtual bool PreSpawnActor( UObject* Asset, FVector& InOutLocation, FRotator& InOutRotation, bool bRotationWasSupplied ) OVERRIDE;
-	virtual void PostSpawnActor( UObject* Asset, AActor* NewActor ) OVERRIDE;
-	virtual void PostCreateBlueprint( UObject* Asset, AActor* CDO ) OVERRIDE;
-	virtual bool CanCreateActorFrom( const FAssetData& AssetData, FText& OutErrorMsg ) OVERRIDE;
+	virtual bool PreSpawnActor( UObject* Asset, FVector& InOutLocation, FRotator& InOutRotation, bool bRotationWasSupplied ) override;
+	virtual void PostSpawnActor( UObject* Asset, AActor* NewActor ) override;
+	virtual void PostCreateBlueprint( UObject* Asset, AActor* CDO ) override;
+	virtual bool CanCreateActorFrom( const FAssetData& AssetData, FText& OutErrorMsg ) override;
 	// End UActorFactory Interface
 
 private:
diff --git a/Engine/Source/Editor/UnrealEd/Classes/ActorFactories/ActorFactoryDestructible.h b/Engine/Source/Editor/UnrealEd/Classes/ActorFactories/ActorFactoryDestructible.h
--- a/Engine/Source/Editor/UnrealEd/Classes/ActorFactories/ActorFactoryDestructible.h
+++ b/Engine/Source/Editor/UnrealEd/Classes/ActorFactories/ActorFactoryDestructible.h
@@ -10,10 +10,10 @@ class UActorFactoryDestructible : public UActorFactory
 	GENERATED_UCLASS_BODY()
 
 	// Begin UActorFactory Interface
-	virtual void PostSpawnActor( UObject* Asset, AActor* NewActor ) OVERRIDE;
-	virtual void PostCreateBlueprint( UObject* Asset, AActor* CDO ) OVERRIDE;
-	virtual bool CanCreateActorFrom( const FAssetData& AssetData, FText& OutErrorMsg ) OVERRIDE;
-	virtual UObject* GetAssetFromActorInstance(AActor* ActorInstance) OVERRIDE;
+	virtual void PostSpawnActor( UObject* Asset, AActor* NewActor ) override;
+	virtual void PostCreateBlueprint( UObject* Asset, AActor* CDO ) override;
+	virtual bool CanCreateActorFrom( const FAssetData& AssetData, FText& OutErrorMsg ) override;
+	virtual UObject* GetAssetFromActorInstance(AActor* ActorInstance) override;
 	// End UActorFactory Interface
 };
 
diff --git a/Engine/Source/Editor/UnrealEd/Classes/ActorFactories/ActorFactoryInteractiveFoliage.h b/Engine/Source/Editor/UnrealEd/Classes/ActorFactories/ActorFactoryInteractiveFoliage.h
--- a/Engine/Source/Editor/UnrealEd/Classes/ActorFactories/ActorFactoryInteractiveFoliage.h
+++ b/Engine/Source/Editor/UnrealEd/Classes/ActorFactories/ActorFactoryInteractiveFoliage.h
@@ -11,7 +11,7 @@ class UActorFactoryInteractiveFoliage : public UActorFactoryStaticMesh
 
 
 	// Begin UActorFactory Interface
-	virtual bool CanCreateActorFrom( const FAssetData& AssetData, FText& OutErrorMsg ) OVERRIDE { return false; };
+	virtual bool CanCreateActorFrom( const FAssetData& AssetData, FText& OutErrorMsg ) override { return false; };
 	// End UActorFactory Interface
 };
 
diff --git a/Engine/Source/Editor/UnrealEd/Classes/ActorFactories/ActorFactoryMatineeActor.h b/Engine/Source/Editor/UnrealEd/Classes/ActorFactories/ActorFactoryMatineeActor.h
--- a/Engine/Source/Editor/UnrealEd/Classes/ActorFactories/ActorFactoryMatineeActor.h
+++ b/Engine/Source/Editor/UnrealEd/Classes/ActorFactories/ActorFactoryMatineeActor.h
@@ -10,9 +10,9 @@ class UActorFactoryMatineeActor : public UActorFactory
 	GENERATED_UCLASS_BODY()
 
 	// Begin UActorFactory Interface
-	virtual void PostSpawnActor( UObject* Asset, AActor* NewActor ) OVERRIDE;
-	virtual void PostCreateBlueprint( UObject* Asset, AActor* CDO ) OVERRIDE;
-	virtual bool CanCreateActorFrom( const FAssetData& AssetData, FText& OutErrorMsg ) OVERRIDE;
+	virtual void PostSpawnActor( UObject* Asset, AActor* NewActor ) override;
+	virtual void PostCreateBlueprint( UObject* Asset, AActor* CDO ) override;
+	virtual bool CanCreateActorFrom( const FAssetData& AssetData, FText& OutErrorMsg ) override;
 	// End UActorFactory Interface
 };
 
diff --git a/Engine/Source/Editor/UnrealEd/Classes/ActorFactories/ActorFactoryPhysicsAsset.h b/Engine/Source/Editor/UnrealEd/Classes/ActorFactories/ActorFactoryPhysicsAsset.h
--- a/Engine/Source/Editor/UnrealEd/Classes/ActorFactories/ActorFactoryPhysicsAsset.h
+++ b/Engine/Source/Editor/UnrealEd/Classes/ActorFactories/ActorFactoryPhysicsAsset.h
@@ -10,10 +10,10 @@ class UActorFactoryPhysicsAsset : public UActorFactory
 	GENERATED_UCLASS_BODY()
 
 	// Begin UActorFactory Interface
-	virtual bool PreSpawnActor( UObject* Asset, FVector& InOutLocation, FRotator& InOutRotation, bool bRotationWasSupplied ) OVERRIDE;
-	virtual void PostSpawnActor( UObject* Asset, AActor* NewActor ) OVERRIDE;
-	virtual void PostCreateBlueprint( UObject* Asset, AActor* CDO ) OVERRIDE;
-	virtual bool CanCreateActorFrom( const FAssetData& AssetData, FText& OutErrorMsg ) OVERRIDE;
+	virtual bool PreSpawnActor( UObject* Asset, FVector& InOutLocation, FRotator& InOutRotation, bool bRotationWasSupplied ) override;
+	virtual void PostSpawnActor( UObject* Asset, AActor* NewActor ) override;
+	virtual void PostCreateBlueprint( UObject* Asset, AActor* CDO ) override;
+	virtual bool CanCreateActorFrom( const FAssetData& AssetData, FText& OutErrorMsg ) override;
 	// End UActorFactory Interface
 };
 
diff --git a/Engine/Source/Editor/UnrealEd/Classes/ActorFactories/ActorFactorySkeletalMesh.h b/Engine/Source/Editor/UnrealEd/Classes/ActorFactories/ActorFactorySkeletalMesh.h
--- a/Engine/Source/Editor/UnrealEd/Classes/ActorFactories/ActorFactorySkeletalMesh.h
+++ b/Engine/Source/Editor/UnrealEd/Classes/ActorFactories/ActorFactorySkeletalMesh.h
@@ -11,9 +11,9 @@ class UActorFactorySkeletalMesh : public UActorFactory
 
 protected:
 	// Begin UActorFactory Interface
-	virtual void PostSpawnActor( UObject* Asset, AActor* NewActor ) OVERRIDE;
-	virtual void PostCreateBlueprint( UObject* Asset, AActor* CDO ) OVERRIDE;
-	virtual bool CanCreateActorFrom( const FAssetData& AssetData, FText& OutErrorMsg ) OVERRIDE;
+	virtual void PostSpawnActor( UObject* Asset, AActor* NewActor ) override;
+	virtual void PostCreateBlueprint( UObject* Asset, AActor* CDO ) override;
+	virtual bool CanCreateActorFrom( const FAssetData& AssetData, FText& OutErrorMsg ) override;
 	// End UActorFactory Interface
 
 	virtual USkeletalMesh* GetSkeletalMeshFromAsset( UObject* Asset ) const;
diff --git a/Engine/Source/Editor/UnrealEd/Classes/ActorFactories/ActorFactoryStaticMesh.h b/Engine/Source/Editor/UnrealEd/Classes/ActorFactories/ActorFactoryStaticMesh.h
--- a/Engine/Source/Editor/UnrealEd/Classes/ActorFactories/ActorFactoryStaticMesh.h
+++ b/Engine/Source/Editor/UnrealEd/Classes/ActorFactories/ActorFactoryStaticMesh.h
@@ -10,10 +10,10 @@ class UActorFactoryStaticMesh : public UActorFactory
 	GENERATED_UCLASS_BODY()
 
 	// Begin UActorFactory Interface
-	virtual bool CanCreateActorFrom( const FAssetData& AssetData, FText& OutErrorMsg ) OVERRIDE;
-	virtual void PostSpawnActor( UObject* Asset, AActor* NewActor) OVERRIDE;
-	virtual void PostCreateBlueprint( UObject* Asset, AActor* CDO ) OVERRIDE;
-	virtual UObject* GetAssetFromActorInstance(AActor* ActorInstance) OVERRIDE;
+	virtual bool CanCreateActorFrom( const FAssetData& AssetData, FText& OutErrorMsg ) override;
+	virtual void PostSpawnActor( UObject* Asset, AActor* NewActor) override;
+	virtual void PostCreateBlueprint( UObject* Asset, AActor* CDO ) override;
+	virtual UObject* GetAssetFromActorInstance(AActor* ActorInstance) override;
 	// End UActorFactory Interface
 };
 
diff --git a/Engine/Source/Editor/UnrealEd/Classes/ActorFactories/ActorFactoryVectorFieldVolume.h b/Engine/Source/Editor/UnrealEd/Classes/ActorFactories/ActorFactoryVectorFieldVolume.h
--- a/Engine/Source/Editor/UnrealEd/Classes/ActorFactories/ActorFactoryVectorFieldVolume.h
+++ b/Engine/Source/Editor/UnrealEd/Classes/ActorFactories/ActorFactoryVectorFieldVolume.h
@@ -10,8 +10,8 @@ class UActorFactoryVectorFieldVolume : public UActorFactory
 	GENERATED_UCLASS_BODY()
 
 	// Begin UActorFactory Interface
-	virtual void PostSpawnActor( UObject* Asset, AActor* NewActor ) OVERRIDE;
-	virtual bool CanCreateActorFrom( const FAssetData& AssetData, FText& OutErrorMsg ) OVERRIDE;
+	virtual void PostSpawnActor( UObject* Asset, AActor* NewActor ) override;
+	virtual bool CanCreateActorFrom( const FAssetData& AssetData, FText& OutErrorMsg ) override;
 	// End UActorFactory Interface
 };
 
diff --git a/Engine/Source/Editor/UnrealEd/Classes/Animation/DebugSkelMeshComponent.h b/Engine/Source/Editor/UnrealEd/Classes/Animation/DebugSkelMeshComponent.h
--- a/Engine/Source/Editor/UnrealEd/Classes/Animation/DebugSkelMeshComponent.h
+++ b/Engine/Source/Editor/UnrealEd/Classes/Animation/DebugSkelMeshComponent.h
@@ -138,31 +138,31 @@ class UDebugSkelMeshComponent : public USkeletalMeshComponent
 	bool bEnableWind;
 
 	// Begin USceneComponent interface.
-	virtual FBoxSphereBounds CalcBounds(const FTransform & LocalToWorld) const OVERRIDE;
+	virtual FBoxSphereBounds CalcBounds(const FTransform & LocalToWorld) const override;
 	// End USceneComponent interface.
 
 	// Begin UPrimitiveComponent interface.
-	virtual FPrimitiveSceneProxy* CreateSceneProxy() OVERRIDE;
+	virtual FPrimitiveSceneProxy* CreateSceneProxy() override;
 	
 	// engine only draw bounds IF selected
 	// @todo fix this properly
 	// this isn't really the best way to do this, but for now
 	// we'll just mark as selected
-	virtual bool ShouldRenderSelected() const OVERRIDE
+	virtual bool ShouldRenderSelected() const override
 	{
 		return bDisplayBound;
 	}
 	// End UPrimitiveComponent interface.
 
 	// Begin SkinnedMeshComponent Interface
-	virtual bool ShouldCPUSkin() OVERRIDE;
-	virtual void PostInitMeshObject(class FSkeletalMeshObject* MeshObject) OVERRIDE;
-	virtual void RefreshBoneTransforms(FActorComponentTickFunction* TickFunction = NULL) OVERRIDE;
+	virtual bool ShouldCPUSkin() override;
+	virtual void PostInitMeshObject(class FSkeletalMeshObject* MeshObject) override;
+	virtual void RefreshBoneTransforms(FActorComponentTickFunction* TickFunction = NULL) override;
 	// End SkinnedMeshComponent Interface
 
 	// Begin SkeletalMeshComponent Interface
-	virtual void InitAnim(bool bForceReinit) OVERRIDE;
-	virtual bool IsWindEnabled() const OVERRIDE;
+	virtual void InitAnim(bool bForceReinit) override;
+	virtual bool IsWindEnabled() const override;
 	// End SkeletalMeshComponent Interface
 	// Preview.
 	// @todo document
@@ -236,7 +236,7 @@ class UDebugSkelMeshComponent : public USkeletalMeshComponent
 	int32 FindCurrentSectionDisplayMode();
 
 	/** to avoid clothing reset while modifying properties in Persona */
-	virtual void CheckClothTeleport(float DeltaTime) OVERRIDE;
+	virtual void CheckClothTeleport(float DeltaTime) override;
 
 #endif //#if WITH_APEX_CLOTHING
 };
diff --git a/Engine/Source/Editor/UnrealEd/Classes/Animation/EditorAnimBaseObj.h b/Engine/Source/Editor/UnrealEd/Classes/Animation/EditorAnimBaseObj.h
--- a/Engine/Source/Editor/UnrealEd/Classes/Animation/EditorAnimBaseObj.h
+++ b/Engine/Source/Editor/UnrealEd/Classes/Animation/EditorAnimBaseObj.h
@@ -24,9 +24,9 @@ public:
 	UAnimSequenceBase* AnimObject;
 	FOnAnimObjectChange OnChange;
 
-	virtual void PreEditChange( class FEditPropertyChain& PropertyAboutToChange ) OVERRIDE;
+	virtual void PreEditChange( class FEditPropertyChain& PropertyAboutToChange ) override;
 
-	virtual void PostEditChangeProperty(FPropertyChangedEvent& PropertyChangedEvent) OVERRIDE;
+	virtual void PostEditChangeProperty(FPropertyChangedEvent& PropertyChangedEvent) override;
 
 	virtual bool PropertyChangeRequiresRebuild(FPropertyChangedEvent& PropertyChangedEvent) { return true;}
 
diff --git a/Engine/Source/Editor/UnrealEd/Classes/Animation/EditorAnimCompositeSegment.h b/Engine/Source/Editor/UnrealEd/Classes/Animation/EditorAnimCompositeSegment.h
--- a/Engine/Source/Editor/UnrealEd/Classes/Animation/EditorAnimCompositeSegment.h
+++ b/Engine/Source/Editor/UnrealEd/Classes/Animation/EditorAnimCompositeSegment.h
@@ -24,9 +24,9 @@ public:
 	int AnimSegmentIndex;
 
 	virtual void InitAnimSegment(int AnimSegmentIndex);
-	virtual bool ApplyChangesToMontage() OVERRIDE;
+	virtual bool ApplyChangesToMontage() override;
 
-	virtual bool PropertyChangeRequiresRebuild(FPropertyChangedEvent& PropertyChangedEvent) OVERRIDE;
+	virtual bool PropertyChangeRequiresRebuild(FPropertyChangedEvent& PropertyChangedEvent) override;
 
 private:
 	//FAnimTrack* GetAnimTrack
diff --git a/Engine/Source/Editor/UnrealEd/Classes/Animation/EditorAnimSegment.h b/Engine/Source/Editor/UnrealEd/Classes/Animation/EditorAnimSegment.h
--- a/Engine/Source/Editor/UnrealEd/Classes/Animation/EditorAnimSegment.h
+++ b/Engine/Source/Editor/UnrealEd/Classes/Animation/EditorAnimSegment.h
@@ -25,7 +25,7 @@ public:
 	int AnimSegmentIndex;
 
 	virtual void InitAnimSegment(int AnimSlotIndex, int AnimSegmentIndex);
-	virtual bool ApplyChangesToMontage() OVERRIDE;
+	virtual bool ApplyChangesToMontage() override;
 
-	virtual bool PropertyChangeRequiresRebuild(FPropertyChangedEvent& PropertyChangedEvent) OVERRIDE;
+	virtual bool PropertyChangeRequiresRebuild(FPropertyChangedEvent& PropertyChangedEvent) override;
 };
diff --git a/Engine/Source/Editor/UnrealEd/Classes/Animation/EditorBranchPoint.h b/Engine/Source/Editor/UnrealEd/Classes/Animation/EditorBranchPoint.h
--- a/Engine/Source/Editor/UnrealEd/Classes/Animation/EditorBranchPoint.h
+++ b/Engine/Source/Editor/UnrealEd/Classes/Animation/EditorBranchPoint.h
@@ -23,5 +23,5 @@ public:
 	int32 BranchIndex;
 
 	virtual void InitBranchPoint(int BranchIndexIn);
-	virtual bool ApplyChangesToMontage() OVERRIDE;
+	virtual bool ApplyChangesToMontage() override;
 };
diff --git a/Engine/Source/Editor/UnrealEd/Classes/Animation/EditorCompositeSection.h b/Engine/Source/Editor/UnrealEd/Classes/Animation/EditorCompositeSection.h
--- a/Engine/Source/Editor/UnrealEd/Classes/Animation/EditorCompositeSection.h
+++ b/Engine/Source/Editor/UnrealEd/Classes/Animation/EditorCompositeSection.h
@@ -21,5 +21,5 @@ public:
 	int32 SectionIndex;
 	
 	virtual void InitSection(int SectionIndex);
-	virtual bool ApplyChangesToMontage() OVERRIDE;
+	virtual bool ApplyChangesToMontage() override;
 };
diff --git a/Engine/Source/Editor/UnrealEd/Classes/AutoReimport/AutoReimportManager.h b/Engine/Source/Editor/UnrealEd/Classes/AutoReimport/AutoReimportManager.h
--- a/Engine/Source/Editor/UnrealEd/Classes/AutoReimport/AutoReimportManager.h
+++ b/Engine/Source/Editor/UnrealEd/Classes/AutoReimport/AutoReimportManager.h
@@ -16,13 +16,13 @@ public:
 
 	// FTickableEditorObject interface
 	/* Tick will reimport files if any are modified and enough delta time has elapsed*/
-	virtual void Tick(float DeltaTime) OVERRIDE;
-	virtual bool IsTickable() const OVERRIDE { return true; }
-	virtual TStatId GetStatId() const OVERRIDE;
+	virtual void Tick(float DeltaTime) override;
+	virtual bool IsTickable() const override { return true; }
+	virtual TStatId GetStatId() const override;
 	// End of FTickableEditorObject interface
 
 	// Begin UObject Interface
-	virtual void BeginDestroy(void) OVERRIDE;
+	virtual void BeginDestroy(void) override;
 	// End UObject Interface
 
 	/** Delegate to process a auto reimport*/
diff --git a/Engine/Source/Editor/UnrealEd/Classes/Builders/ConeBuilder.h b/Engine/Source/Editor/UnrealEd/Classes/Builders/ConeBuilder.h
--- a/Engine/Source/Editor/UnrealEd/Classes/Builders/ConeBuilder.h
+++ b/Engine/Source/Editor/UnrealEd/Classes/Builders/ConeBuilder.h
@@ -45,7 +45,7 @@ class UConeBuilder : public UEditorBrushBuilder
 
 
 	// Begin UBrushBuilder Interface
-	virtual bool Build( UWorld* InWorld, ABrush* InBrush = NULL ) OVERRIDE;
+	virtual bool Build( UWorld* InWorld, ABrush* InBrush = NULL ) override;
 	// End UBrushBuilder Interface
 
 	// @todo document
diff --git a/Engine/Source/Editor/UnrealEd/Classes/Builders/CubeBuilder.h b/Engine/Source/Editor/UnrealEd/Classes/Builders/CubeBuilder.h
--- a/Engine/Source/Editor/UnrealEd/Classes/Builders/CubeBuilder.h
+++ b/Engine/Source/Editor/UnrealEd/Classes/Builders/CubeBuilder.h
@@ -41,7 +41,7 @@ class UCubeBuilder : public UEditorBrushBuilder
 
 
 	// Begin UBrushBuilder Interface
-	virtual bool Build( UWorld* InWorld, ABrush* InBrush = NULL ) OVERRIDE;
+	virtual bool Build( UWorld* InWorld, ABrush* InBrush = NULL ) override;
 	// End UBrushBuilder Interface
 
 	// @todo document
diff --git a/Engine/Source/Editor/UnrealEd/Classes/Builders/EditorBrushBuilder.h b/Engine/Source/Editor/UnrealEd/Classes/Builders/EditorBrushBuilder.h
--- a/Engine/Source/Editor/UnrealEd/Classes/Builders/EditorBrushBuilder.h
+++ b/Engine/Source/Editor/UnrealEd/Classes/Builders/EditorBrushBuilder.h
@@ -12,23 +12,23 @@ class UEditorBrushBuilder : public UBrushBuilder
 public:
 
 	/** UBrushBuilder interface */
-	virtual void BeginBrush( bool InMergeCoplanars, FName InLayer ) OVERRIDE;
-	virtual bool EndBrush( UWorld* InWorld, ABrush* InBrush ) OVERRIDE;
-	virtual int32 GetVertexCount() OVERRIDE;
-	virtual FVector GetVertex( int32 i ) OVERRIDE;
-	virtual int32 GetPolyCount() OVERRIDE;
-	virtual bool BadParameters( const FText& msg ) OVERRIDE;
-	virtual int32 Vertexv( FVector v ) OVERRIDE;
-	virtual int32 Vertex3f( float X, float Y, float Z ) OVERRIDE;
-	virtual void Poly3i( int32 Direction, int32 i, int32 j, int32 k, FName ItemName = NAME_None, bool bIsTwoSidedNonSolid = false ) OVERRIDE;
-	virtual void Poly4i( int32 Direction, int32 i, int32 j, int32 k, int32 l, FName ItemName = NAME_None, bool bIsTwoSidedNonSolid = false ) OVERRIDE;
-	virtual void PolyBegin( int32 Direction, FName ItemName = NAME_None ) OVERRIDE;
-	virtual void Polyi( int32 i ) OVERRIDE;
-	virtual void PolyEnd() OVERRIDE;
-	UNREALED_API virtual bool Build( UWorld* InWorld, ABrush* InBrush = NULL ) OVERRIDE;
+	virtual void BeginBrush( bool InMergeCoplanars, FName InLayer ) override;
+	virtual bool EndBrush( UWorld* InWorld, ABrush* InBrush ) override;
+	virtual int32 GetVertexCount() override;
+	virtual FVector GetVertex( int32 i ) override;
+	virtual int32 GetPolyCount() override;
+	virtual bool BadParameters( const FText& msg ) override;
+	virtual int32 Vertexv( FVector v ) override;
+	virtual int32 Vertex3f( float X, float Y, float Z ) override;
+	virtual void Poly3i( int32 Direction, int32 i, int32 j, int32 k, FName ItemName = NAME_None, bool bIsTwoSidedNonSolid = false ) override;
+	virtual void Poly4i( int32 Direction, int32 i, int32 j, int32 k, int32 l, FName ItemName = NAME_None, bool bIsTwoSidedNonSolid = false ) override;
+	virtual void PolyBegin( int32 Direction, FName ItemName = NAME_None ) override;
+	virtual void Polyi( int32 i ) override;
+	virtual void PolyEnd() override;
+	UNREALED_API virtual bool Build( UWorld* InWorld, ABrush* InBrush = NULL ) override;
 
 	/** UObject interface */
-	virtual void PostEditChangeProperty( struct FPropertyChangedEvent& PropertyChangedEvent) OVERRIDE;
+	virtual void PostEditChangeProperty( struct FPropertyChangedEvent& PropertyChangedEvent) override;
 };
 
 
diff --git a/Engine/Source/Editor/UnrealEd/Classes/Commandlets/AutomatedLauncherCommandlet.h b/Engine/Source/Editor/UnrealEd/Classes/Commandlets/AutomatedLauncherCommandlet.h
--- a/Engine/Source/Editor/UnrealEd/Classes/Commandlets/AutomatedLauncherCommandlet.h
+++ b/Engine/Source/Editor/UnrealEd/Classes/Commandlets/AutomatedLauncherCommandlet.h
@@ -12,6 +12,6 @@ class UAutomatedLauncherCommandlet : public UCommandlet
 	GENERATED_UCLASS_BODY()
 
 	// Begin UCommandlet Interface
-	virtual int32 Main(const FString& Params) OVERRIDE;
+	virtual int32 Main(const FString& Params) override;
 	// End UCommandlet Interface
 };
\ No newline at end of file
diff --git a/Engine/Source/Editor/UnrealEd/Classes/Commandlets/CompressAnimationsCommandlet.h b/Engine/Source/Editor/UnrealEd/Classes/Commandlets/CompressAnimationsCommandlet.h
--- a/Engine/Source/Editor/UnrealEd/Classes/Commandlets/CompressAnimationsCommandlet.h
+++ b/Engine/Source/Editor/UnrealEd/Classes/Commandlets/CompressAnimationsCommandlet.h
@@ -9,7 +9,7 @@ class UCompressAnimationsCommandlet : public UCommandlet
 {
 	GENERATED_UCLASS_BODY()
 	// Begin UCommandlet Interface
-	virtual int32 Main(const FString& Params) OVERRIDE;
+	virtual int32 Main(const FString& Params) override;
 	// End UCommandlet Interface
 };
 
diff --git a/Engine/Source/Editor/UnrealEd/Classes/Commandlets/CookCommandlet.h b/Engine/Source/Editor/UnrealEd/Classes/Commandlets/CookCommandlet.h
--- a/Engine/Source/Editor/UnrealEd/Classes/Commandlets/CookCommandlet.h
+++ b/Engine/Source/Editor/UnrealEd/Classes/Commandlets/CookCommandlet.h
@@ -95,7 +95,7 @@ public:
 
 	// Begin UCommandlet Interface
 
-	virtual int32 Main(const FString& CmdLineParams) OVERRIDE;
+	virtual int32 Main(const FString& CmdLineParams) override;
 	
 	// End UCommandlet Interface
 
diff --git a/Engine/Source/Editor/UnrealEd/Classes/Commandlets/DiffAssetsCommandlet.h b/Engine/Source/Editor/UnrealEd/Classes/Commandlets/DiffAssetsCommandlet.h
--- a/Engine/Source/Editor/UnrealEd/Classes/Commandlets/DiffAssetsCommandlet.h
+++ b/Engine/Source/Editor/UnrealEd/Classes/Commandlets/DiffAssetsCommandlet.h
@@ -13,7 +13,7 @@ class UDiffAssetsCommandlet : public UCommandlet
 {
     GENERATED_UCLASS_BODY()
 	// Begin UCommandlet Interface
-	virtual int32 Main(const FString& Params) OVERRIDE
+	virtual int32 Main(const FString& Params) override
 	{
 		return !ExportFilesToTextAndDiff(Params);
 	}
diff --git a/Engine/Source/Editor/UnrealEd/Classes/Commandlets/DumpBlueprintsInfoCommandlet.h b/Engine/Source/Editor/UnrealEd/Classes/Commandlets/DumpBlueprintsInfoCommandlet.h
--- a/Engine/Source/Editor/UnrealEd/Classes/Commandlets/DumpBlueprintsInfoCommandlet.h
+++ b/Engine/Source/Editor/UnrealEd/Classes/Commandlets/DumpBlueprintsInfoCommandlet.h
@@ -11,6 +11,6 @@ class UDumpBlueprintsInfoCommandlet : public UCommandlet
 
 public:		
 	// Begin UCommandlet Interface
-	virtual int32 Main(FString const& Params) OVERRIDE;
+	virtual int32 Main(FString const& Params) override;
 	// End UCommandlet Interface
 };
diff --git a/Engine/Source/Editor/UnrealEd/Classes/Commandlets/DumpHiddenCategoriesCommandlet.h b/Engine/Source/Editor/UnrealEd/Classes/Commandlets/DumpHiddenCategoriesCommandlet.h
--- a/Engine/Source/Editor/UnrealEd/Classes/Commandlets/DumpHiddenCategoriesCommandlet.h
+++ b/Engine/Source/Editor/UnrealEd/Classes/Commandlets/DumpHiddenCategoriesCommandlet.h
@@ -11,6 +11,6 @@ class UDumpHiddenCategoriesCommandlet : public UCommandlet
 
 public:		
 	// Begin UCommandlet Interface
-	virtual int32 Main(FString const& Params) OVERRIDE;
+	virtual int32 Main(FString const& Params) override;
 	// End UCommandlet Interface
 };
diff --git a/Engine/Source/Editor/UnrealEd/Classes/Commandlets/FileServerCommandlet.h b/Engine/Source/Editor/UnrealEd/Classes/Commandlets/FileServerCommandlet.h
--- a/Engine/Source/Editor/UnrealEd/Classes/Commandlets/FileServerCommandlet.h
+++ b/Engine/Source/Editor/UnrealEd/Classes/Commandlets/FileServerCommandlet.h
@@ -23,7 +23,7 @@ public:
 
 	// Begin UCommandlet Interface
 
-	virtual int32 Main( const FString& Params ) OVERRIDE;
+	virtual int32 Main( const FString& Params ) override;
 	
 	// End UCommandlet Interface
 
diff --git a/Engine/Source/Editor/UnrealEd/Classes/Commandlets/FixupRedirectsCommandlet.h b/Engine/Source/Editor/UnrealEd/Classes/Commandlets/FixupRedirectsCommandlet.h
--- a/Engine/Source/Editor/UnrealEd/Classes/Commandlets/FixupRedirectsCommandlet.h
+++ b/Engine/Source/Editor/UnrealEd/Classes/Commandlets/FixupRedirectsCommandlet.h
@@ -9,8 +9,8 @@ class UFixupRedirectsCommandlet : public UCommandlet
 {
 	GENERATED_UCLASS_BODY()
 	// Begin UCommandlet Interface
-	virtual void CreateCustomEngine(const FString& Params) OVERRIDE;
-	virtual int32 Main(const FString& Params) OVERRIDE;
+	virtual void CreateCustomEngine(const FString& Params) override;
+	virtual int32 Main(const FString& Params) override;
 	// End UCommandlet Interface
 };
 
diff --git a/Engine/Source/Editor/UnrealEd/Classes/Commandlets/GatherTextCommandletBase.h b/Engine/Source/Editor/UnrealEd/Classes/Commandlets/GatherTextCommandletBase.h
--- a/Engine/Source/Editor/UnrealEd/Classes/Commandlets/GatherTextCommandletBase.h
+++ b/Engine/Source/Editor/UnrealEd/Classes/Commandlets/GatherTextCommandletBase.h
@@ -272,6 +272,6 @@ protected:
 
 private:
 
-	virtual void CreateCustomEngine(const FString& Params) OVERRIDE ; //Disallow other text commandlets to make their own engine.
+	virtual void CreateCustomEngine(const FString& Params) override ; //Disallow other text commandlets to make their own engine.
 	
 };
diff --git a/Engine/Source/Editor/UnrealEd/Classes/Commandlets/GenerateGatherArchiveCommandlet.h b/Engine/Source/Editor/UnrealEd/Classes/Commandlets/GenerateGatherArchiveCommandlet.h
--- a/Engine/Source/Editor/UnrealEd/Classes/Commandlets/GenerateGatherArchiveCommandlet.h
+++ b/Engine/Source/Editor/UnrealEd/Classes/Commandlets/GenerateGatherArchiveCommandlet.h
@@ -18,7 +18,7 @@ class UGenerateGatherArchiveCommandlet : public UGatherTextCommandletBase
 #if CPP || UE_BUILD_DOCS
 public:
 	// Begin UCommandlet Interface
-	virtual int32 Main( const FString& Params ) OVERRIDE;
+	virtual int32 Main( const FString& Params ) override;
 	// End UCommandlet Interface
 	
 	bool WriteArchiveToFile( TSharedRef< FJsonObject> ArchiveJSONObject, const FString& OutputDirectoryPath, const TCHAR* Culture, const FString& FileName );
diff --git a/Engine/Source/Editor/UnrealEd/Classes/Commandlets/KismetUpdateCommandlet.h b/Engine/Source/Editor/UnrealEd/Classes/Commandlets/KismetUpdateCommandlet.h
--- a/Engine/Source/Editor/UnrealEd/Classes/Commandlets/KismetUpdateCommandlet.h
+++ b/Engine/Source/Editor/UnrealEd/Classes/Commandlets/KismetUpdateCommandlet.h
@@ -17,9 +17,9 @@ protected:
 	TSet<FString> AlreadyCompiledFullPaths;
 protected:
 	// UResavePackagesCommandlet interface
-	virtual int32 InitializeResaveParameters(const TArray<FString>& Tokens, const TArray<FString>& Switches, TArray<FString>& MapPathNames) OVERRIDE;
-	virtual void PerformAdditionalOperations(class UObject* Object, bool& bSavePackage) OVERRIDE;
-	virtual FText GetChangelistDescription() const OVERRIDE;
+	virtual int32 InitializeResaveParameters(const TArray<FString>& Tokens, const TArray<FString>& Switches, TArray<FString>& MapPathNames) override;
+	virtual void PerformAdditionalOperations(class UObject* Object, bool& bSavePackage) override;
+	virtual FText GetChangelistDescription() const override;
 	// End of UResavePackagesCommandlet interface
 
 	void CompileOneBlueprint(UBlueprint* Blueprint);
diff --git a/Engine/Source/Editor/UnrealEd/Classes/Commandlets/RecompileBlueprintsCommandlet.h b/Engine/Source/Editor/UnrealEd/Classes/Commandlets/RecompileBlueprintsCommandlet.h
--- a/Engine/Source/Editor/UnrealEd/Classes/Commandlets/RecompileBlueprintsCommandlet.h
+++ b/Engine/Source/Editor/UnrealEd/Classes/Commandlets/RecompileBlueprintsCommandlet.h
@@ -11,13 +11,13 @@ class URecompileBlueprintsCommandlet : public UResavePackagesCommandlet
 
 public:		
 	// Begin UCommandlet Interface
-	virtual int32 Main(const FString& Params) OVERRIDE;
+	virtual int32 Main(const FString& Params) override;
 	// End UCommandlet Interface
 
 	// Begin UResavePackagesCommandlet Interface
-	virtual bool ShouldSkipPackage(const FString& Filename) OVERRIDE;
-	virtual bool PerformPreloadOperations( ULinkerLoad* PackageLinker, bool& bSavePackage ) OVERRIDE;
-	virtual void PerformAdditionalOperations( class UObject* Object, bool& bSavePackage ) OVERRIDE;
-	virtual int32 InitializeResaveParameters( const TArray<FString>& Tokens, const TArray<FString>& Switches, TArray<FString>& MapPathNames ) OVERRIDE;
+	virtual bool ShouldSkipPackage(const FString& Filename) override;
+	virtual bool PerformPreloadOperations( ULinkerLoad* PackageLinker, bool& bSavePackage ) override;
+	virtual void PerformAdditionalOperations( class UObject* Object, bool& bSavePackage ) override;
+	virtual int32 InitializeResaveParameters( const TArray<FString>& Tokens, const TArray<FString>& Switches, TArray<FString>& MapPathNames ) override;
 	// End UResavePackagesCommandlet Interface
 };
diff --git a/Engine/Source/Editor/UnrealEd/Classes/CookOnTheSide/CookOnTheFlyServer.h b/Engine/Source/Editor/UnrealEd/Classes/CookOnTheSide/CookOnTheFlyServer.h
--- a/Engine/Source/Editor/UnrealEd/Classes/CookOnTheSide/CookOnTheFlyServer.h
+++ b/Engine/Source/Editor/UnrealEd/Classes/CookOnTheSide/CookOnTheFlyServer.h
@@ -240,9 +240,9 @@ public:
 	/**
 	 * FTickableEditorObject interface used by cook on the side
 	 */
-	TStatId GetStatId() const OVERRIDE;
-	 void Tick(float DeltaTime) OVERRIDE;
-	 bool IsTickable() const OVERRIDE;
+	TStatId GetStatId() const override;
+	 void Tick(float DeltaTime) override;
+	 bool IsTickable() const override;
 
 
 
@@ -299,7 +299,7 @@ public:
 	
 
 
-	virtual void BeginDestroy() OVERRIDE;
+	virtual void BeginDestroy() override;
 	
 
 
diff --git a/Engine/Source/Editor/UnrealEd/Classes/Editor/EditorEngine.h b/Engine/Source/Editor/UnrealEd/Classes/Editor/EditorEngine.h
--- a/Engine/Source/Editor/UnrealEd/Classes/Editor/EditorEngine.h
+++ b/Engine/Source/Editor/UnrealEd/Classes/Editor/EditorEngine.h
@@ -568,18 +568,18 @@ public:
 	void BroadcastObjectReimported(UObject* InObject) { ObjectReimportedEvent.Broadcast(InObject); }
 
 	// Begin UObject interface.
-	virtual void FinishDestroy() OVERRIDE;	
-	virtual void PostEditChangeProperty(FPropertyChangedEvent& PropertyChangedEvent) OVERRIDE;
+	virtual void FinishDestroy() override;	
+	virtual void PostEditChangeProperty(FPropertyChangedEvent& PropertyChangedEvent) override;
 	static void AddReferencedObjects(UObject* InThis, FReferenceCollector& Collector);
 	// End UObject interface.
 
 	// Begin UEngine interface.
-	virtual void Init(IEngineLoop* InEngineLoop) OVERRIDE;
-	virtual void InitializeObjectReferences() OVERRIDE;
+	virtual void Init(IEngineLoop* InEngineLoop) override;
+	virtual void InitializeObjectReferences() override;
 	// End UEngine interface.
 	
 	// Begin FExec Interface
-	virtual bool Exec( UWorld* InWorld, const TCHAR* Cmd, FOutputDevice& Ar=*GLog ) OVERRIDE;
+	virtual bool Exec( UWorld* InWorld, const TCHAR* Cmd, FOutputDevice& Ar=*GLog ) override;
 	// End FExec Interface
 
 	bool	CommandIsDeprecated( const TCHAR* CommandStr, FOutputDevice& Ar );
@@ -628,7 +628,7 @@ public:
 	
 
 	/** Get tick rate limitor. */
-	virtual float GetMaxTickRate( float DeltaTime, bool bAllowFrameRateSmoothing = true ) OVERRIDE;
+	virtual float GetMaxTickRate( float DeltaTime, bool bAllowFrameRateSmoothing = true ) override;
 
 	/**
 	 * Initializes the Editor.
@@ -642,7 +642,7 @@ public:
 	 */
 	void InitBuilderBrush( UWorld* InWorld );
 
-	virtual void Tick( float DeltaSeconds, bool bIdleMode ) OVERRIDE;
+	virtual void Tick( float DeltaSeconds, bool bIdleMode ) override;
 
 	/** Returns the global instance of the editor user settings class. */
 	const UEditorUserSettings& GetEditorUserSettings() const;
@@ -694,7 +694,7 @@ public:
 	/**
 	 * Returns whether or not the map build in progressed was cancelled by the user.
 	 */
-	virtual bool GetMapBuildCancelled() const OVERRIDE
+	virtual bool GetMapBuildCancelled() const override
 	{
 		return false;
 	}
@@ -704,7 +704,7 @@ public:
 	 *
 	 * @param InCancelled	New state for the cancelled flag.
 	 */
-	virtual void SetMapBuildCancelled( bool InCancelled ) OVERRIDE
+	virtual void SetMapBuildCancelled( bool InCancelled ) override
 	{
 		// Intentionally empty.
 	}
@@ -714,7 +714,7 @@ public:
 	 *
 	 * @param InActor	Actor that is being drawn.
 	 */
-	virtual bool ShouldDrawBrushWireframe( AActor* InActor ) OVERRIDE;
+	virtual bool ShouldDrawBrushWireframe( AActor* InActor ) override;
 
 	// Execute a command that is safe for rebuilds.
 	virtual bool SafeExec( UWorld* InWorld, const TCHAR* Cmd, FOutputDevice& Ar=*GLog );
@@ -904,7 +904,7 @@ public:
 	 * Notifies tools that a set of objects is being replaced with a new set of objects, so if any of the old
 	 * objects are being edited, they should be discarded and replaced with the new instance.
 	 */
-	virtual void NotifyToolsOfObjectReplacement(const TMap<UObject*, UObject*>& OldToNewInstanceMap) OVERRIDE;
+	virtual void NotifyToolsOfObjectReplacement(const TMap<UObject*, UObject*>& OldToNewInstanceMap) override;
 
 	//
 	// Pivot handling.
@@ -1585,7 +1585,7 @@ public:
 	 */
 	bool IsAnyViewportRealtime();
 
-	virtual bool ShouldThrottleCPUUsage() const OVERRIDE;
+	virtual bool ShouldThrottleCPUUsage() const override;
 
 	/**
 	 * @return true if all windows are hidden (including minimized)                                                         
@@ -1606,12 +1606,12 @@ public:
 	/**
 	 * Handles freezing/unfreezing of rendering
 	 */
-	virtual void ProcessToggleFreezeCommand( UWorld* InWorld ) OVERRIDE;
+	virtual void ProcessToggleFreezeCommand( UWorld* InWorld ) override;
 
 	/**
 	 * Handles frezing/unfreezing of streaming
 	 */
-	virtual void ProcessToggleFreezeStreamingCommand( UWorld* InWorld ) OVERRIDE;
+	virtual void ProcessToggleFreezeStreamingCommand( UWorld* InWorld ) override;
 
 	// Editor specific
 
@@ -2133,7 +2133,7 @@ public:
 	 * mostly done to check if PIE is being set up, go GWorld is going to change, and it's not really _the_G_World_
 	 * NOTE: hope this goes away once PIE and regular game triggering are not that separate code paths
 	 */
-	virtual bool IsSettingUpPlayWorld() const OVERRIDE { return EditorWorld != NULL && PlayWorld == NULL; }
+	virtual bool IsSettingUpPlayWorld() const override { return EditorWorld != NULL && PlayWorld == NULL; }
 
 	/**
 	 *	Retrieves the active viewport from the editor.
@@ -2234,16 +2234,16 @@ public:
 	/**
 	 * Returns if this world is a PIE world that has its own viewport (is not using the editor viewport)
 	 */
-	virtual bool WorldIsPIEInNewViewport(UWorld *InWorld) OVERRIDE;
+	virtual bool WorldIsPIEInNewViewport(UWorld *InWorld) override;
 
-	virtual void FocusNextPIEWorld(UWorld *CurrentPieWorld, bool previous=false) OVERRIDE;
+	virtual void FocusNextPIEWorld(UWorld *CurrentPieWorld, bool previous=false) override;
 
 	DECLARE_DELEGATE(FPIEInstanceWindowSwitch);
 
 	/** Sets the delegate for when the focused PIE window is changed */
 	void SetPIEInstanceWindowSwitchDelegate(FPIEInstanceWindowSwitch PIEInstanceWindowSwitchDelegate);
 
-	virtual class UGameViewportClient *	GetNextPIEViewport(UGameViewportClient * CurrentViewport) OVERRIDE;
+	virtual class UGameViewportClient *	GetNextPIEViewport(UGameViewportClient * CurrentViewport) override;
 
 private:
 	//
@@ -2293,9 +2293,9 @@ private:
 	/**
 	 * Creates a PIE world by duplicating the editor world	 
 	 */
-	virtual UWorld* CreatePIEWorldByDuplication(FWorldContext &WorldContext, UWorld* InWorld, FString &PlayWorldMapName) OVERRIDE;
+	virtual UWorld* CreatePIEWorldByDuplication(FWorldContext &WorldContext, UWorld* InWorld, FString &PlayWorldMapName) override;
 
-	virtual void RemapGamepadControllerIdForPIE(class UGameViewportClient* GameViewport, int32 &ControllerId) OVERRIDE;
+	virtual void RemapGamepadControllerIdForPIE(class UGameViewportClient* GameViewport, int32 &ControllerId) override;
 
 	/** Creates a PIE world by saving to a temp file and then reloading it */
 	UWorld* CreatePIEWorldBySavingToTemp(FWorldContext &WorldContext, UWorld* InWorld, FString &PlayWorldMapName);
@@ -2413,7 +2413,7 @@ private:
 	void EditorDestroyWorld( FWorldContext & Context, const FText& CleanseText, UWorld* NewWorld = nullptr );
 
 	/** Returns the GameViewport widget */
-	virtual TSharedPtr<SViewport> GetGameViewportWidget() const OVERRIDE;
+	virtual TSharedPtr<SViewport> GetGameViewportWidget() const override;
 
 	/**
 	 * Given a label, attempts to split this into its alpha/numeric parts.
@@ -2460,13 +2460,13 @@ private:
 	/** Updates the project file to auto load and initializes the bLoadTheMostRecentlyLoadedProjectAtStartup flag */
 	void UpdateAutoLoadProject();
 
-	virtual void TriggerStreamingDataRebuild() OVERRIDE;
+	virtual void TriggerStreamingDataRebuild() override;
 
 	/** Remaps a network path for PIE Networking under multiple worlds */
-	virtual bool NetworkRemapPath( UWorld *InWorld, FString &Str, bool reading=true) OVERRIDE;
+	virtual bool NetworkRemapPath( UWorld *InWorld, FString &Str, bool reading=true) override;
 
 	/** Remaps a network path for PIE Networking under multiple worlds */
-	virtual bool NetworkRemapPath( UPendingNetGame *PendingNetGame, FString &Str, bool reading=true) OVERRIDE;
+	virtual bool NetworkRemapPath( UPendingNetGame *PendingNetGame, FString &Str, bool reading=true) override;
 
 	/** Handles user setting changes. */
 	void HandleSettingChanged( FName Name );
@@ -2546,7 +2546,7 @@ private:
 	/** List of files we are deferring adding to source control */
 	TArray<FString> DeferredFilesToAddToSourceControl;
 
-	virtual void VerifyLoadMapWorldCleanup() OVERRIDE;
+	virtual void VerifyLoadMapWorldCleanup() override;
 
 	FPIEInstanceWindowSwitch PIEInstanceWindowSwitchDelegate;
 
diff --git a/Engine/Source/Editor/UnrealEd/Classes/Editor/EditorGameAgnosticSettings.h b/Engine/Source/Editor/UnrealEd/Classes/Editor/EditorGameAgnosticSettings.h
--- a/Engine/Source/Editor/UnrealEd/Classes/Editor/EditorGameAgnosticSettings.h
+++ b/Engine/Source/Editor/UnrealEd/Classes/Editor/EditorGameAgnosticSettings.h
@@ -37,6 +37,6 @@ class UEditorGameAgnosticSettings : public UObject
 	TArray<FGuid> InProgressSurveys;
 
 	// Begin UObject Interface
-	virtual void PostEditChangeProperty( struct FPropertyChangedEvent& PropertyChangedEvent) OVERRIDE;
+	virtual void PostEditChangeProperty( struct FPropertyChangedEvent& PropertyChangedEvent) override;
 	// End UObject Interface
 };
\ No newline at end of file
diff --git a/Engine/Source/Editor/UnrealEd/Classes/Editor/EditorUserSettings.h b/Engine/Source/Editor/UnrealEd/Classes/Editor/EditorUserSettings.h
--- a/Engine/Source/Editor/UnrealEd/Classes/Editor/EditorUserSettings.h
+++ b/Engine/Source/Editor/UnrealEd/Classes/Editor/EditorUserSettings.h
@@ -90,8 +90,8 @@ public:
 	FUserSettingChangedEvent& OnUserSettingChanged() { return UserSettingChangedEvent; }
 
 	// Begin UObject Interface
-	virtual void PostEditChangeProperty( struct FPropertyChangedEvent& PropertyChangedEvent ) OVERRIDE;
-	virtual void PostInitProperties() OVERRIDE;
+	virtual void PostEditChangeProperty( struct FPropertyChangedEvent& PropertyChangedEvent ) override;
+	virtual void PostInitProperties() override;
 	// End UObject Interface
 
 private:
diff --git a/Engine/Source/Editor/UnrealEd/Classes/Editor/GroupActor.h b/Engine/Source/Editor/UnrealEd/Classes/Editor/GroupActor.h
--- a/Engine/Source/Editor/UnrealEd/Classes/Editor/GroupActor.h
+++ b/Engine/Source/Editor/UnrealEd/Classes/Editor/GroupActor.h
@@ -26,25 +26,25 @@ class AGroupActor : public AActor
 
 
 	// Begin UObject interface
-	virtual void PostLoad() OVERRIDE;
-	virtual void PostEditChangeProperty(FPropertyChangedEvent& PropertyChangedEvent) OVERRIDE;
-	virtual void PostEditUndo() OVERRIDE;
-	virtual bool Modify( bool bAlwaysMarkDirty=true ) OVERRIDE;
-	virtual bool NeedsLoadForClient() const OVERRIDE
+	virtual void PostLoad() override;
+	virtual void PostEditChangeProperty(FPropertyChangedEvent& PropertyChangedEvent) override;
+	virtual void PostEditUndo() override;
+	virtual bool Modify( bool bAlwaysMarkDirty=true ) override;
+	virtual bool NeedsLoadForClient() const override
 	{ 
 		return false; 
 	}
-	virtual bool NeedsLoadForServer() const OVERRIDE
+	virtual bool NeedsLoadForServer() const override
 	{ 
 		return false;
 	}
 	// End UObject interface
 
 	// Begin AActor interface
-	virtual void PostActorCreated() OVERRIDE;
-	virtual void InvalidateLightingCacheDetailed(bool bTranslationOnly) OVERRIDE;
-	virtual void PostEditMove(bool bFinished) OVERRIDE;
-	virtual void SetIsTemporarilyHiddenInEditor( bool bIsHidden ) OVERRIDE;
+	virtual void PostActorCreated() override;
+	virtual void InvalidateLightingCacheDetailed(bool bTranslationOnly) override;
+	virtual void PostEditMove(bool bFinished) override;
+	virtual void SetIsTemporarilyHiddenInEditor( bool bIsHidden ) override;
 	// End AActor interface
 
 	// @todo document
diff --git a/Engine/Source/Editor/UnrealEd/Classes/Editor/TransBuffer.h b/Engine/Source/Editor/UnrealEd/Classes/Editor/TransBuffer.h
--- a/Engine/Source/Editor/UnrealEd/Classes/Editor/TransBuffer.h
+++ b/Engine/Source/Editor/UnrealEd/Classes/Editor/TransBuffer.h
@@ -48,8 +48,8 @@ public:
 
 	// Begin UObject interface.
 
-	virtual void Serialize( FArchive& Ar ) OVERRIDE;
-	virtual void FinishDestroy() OVERRIDE;
+	virtual void Serialize( FArchive& Ar ) override;
+	virtual void FinishDestroy() override;
 
 	// End UObject interface.
 
@@ -57,26 +57,26 @@ public:
 
 	// Begin UTransactor interface.
 
-	virtual int32 Begin( const TCHAR* SessionContext, const FText& Description ) OVERRIDE;
-	virtual int32 End() OVERRIDE;
-	virtual void Cancel( int32 StartIndex = 0 ) OVERRIDE;
-	virtual void Reset( const FText& Reason ) OVERRIDE;
-	virtual bool CanUndo( FText* Text=NULL ) OVERRIDE;
-	virtual bool CanRedo( FText* Text=NULL ) OVERRIDE;
-	virtual int32 GetQueueLength( ) const OVERRIDE { return UndoBuffer.Num(); }
-	virtual const FTransaction* GetTransaction( int32 QueueIndex ) const OVERRIDE;
-	virtual FUndoSessionContext GetUndoContext( bool bCheckWhetherUndoPossible = true ) OVERRIDE;
-	virtual SIZE_T GetUndoSize() const OVERRIDE;
-	virtual int32 GetUndoCount( ) const OVERRIDE { return UndoCount; }
-	virtual FUndoSessionContext GetRedoContext() OVERRIDE;
-	virtual bool Undo() OVERRIDE;
-	virtual bool Redo() OVERRIDE;
-	virtual bool EnableObjectSerialization() OVERRIDE;
-	virtual bool DisableObjectSerialization() OVERRIDE;
-	virtual bool IsObjectSerializationEnabled() OVERRIDE { return DisallowObjectSerialization == 0; }
-	virtual void SetPrimaryUndoObject( UObject* Object ) OVERRIDE;
-	virtual ITransaction* CreateInternalTransaction() OVERRIDE;
-	virtual bool IsActive() OVERRIDE
+	virtual int32 Begin( const TCHAR* SessionContext, const FText& Description ) override;
+	virtual int32 End() override;
+	virtual void Cancel( int32 StartIndex = 0 ) override;
+	virtual void Reset( const FText& Reason ) override;
+	virtual bool CanUndo( FText* Text=NULL ) override;
+	virtual bool CanRedo( FText* Text=NULL ) override;
+	virtual int32 GetQueueLength( ) const override { return UndoBuffer.Num(); }
+	virtual const FTransaction* GetTransaction( int32 QueueIndex ) const override;
+	virtual FUndoSessionContext GetUndoContext( bool bCheckWhetherUndoPossible = true ) override;
+	virtual SIZE_T GetUndoSize() const override;
+	virtual int32 GetUndoCount( ) const override { return UndoCount; }
+	virtual FUndoSessionContext GetRedoContext() override;
+	virtual bool Undo() override;
+	virtual bool Redo() override;
+	virtual bool EnableObjectSerialization() override;
+	virtual bool DisableObjectSerialization() override;
+	virtual bool IsObjectSerializationEnabled() override { return DisallowObjectSerialization == 0; }
+	virtual void SetPrimaryUndoObject( UObject* Object ) override;
+	virtual ITransaction* CreateInternalTransaction() override;
+	virtual bool IsActive() override
 	{
 		return ActiveCount > 0;
 	}
diff --git a/Engine/Source/Editor/UnrealEd/Classes/Editor/Transactor.h b/Engine/Source/Editor/UnrealEd/Classes/Editor/Transactor.h
--- a/Engine/Source/Editor/UnrealEd/Classes/Editor/Transactor.h
+++ b/Engine/Source/Editor/UnrealEd/Classes/Editor/Transactor.h
@@ -130,7 +130,7 @@ protected:
 				ArIsLoading = ArIsTransacting = 1;
 			}
 
-			virtual int64 Tell() OVERRIDE {return Offset;}
+			virtual int64 Tell() override {return Offset;}
 			virtual void Seek( int64 InPos ) { Offset = InPos; }
 
 		private:
@@ -203,7 +203,7 @@ protected:
 				ArIsSaving = ArIsTransacting = 1;
 			}
 
-			virtual int64 Tell() OVERRIDE {return Offset;}
+			virtual int64 Tell() override {return Offset;}
 			virtual void Seek( int64 InPos ) 
 			{
 				checkSlow(Offset<=Data.Num());
diff --git a/Engine/Source/Editor/UnrealEd/Classes/Editor/UnrealEdEngine.h b/Engine/Source/Editor/UnrealEd/Classes/Editor/UnrealEdEngine.h
--- a/Engine/Source/Editor/UnrealEd/Classes/Editor/UnrealEdEngine.h
+++ b/Engine/Source/Editor/UnrealEd/Classes/Editor/UnrealEdEngine.h
@@ -163,47 +163,47 @@ class UNREALED_API UUnrealEdEngine : public UEditorEngine, public FNotifyHook
 
 	// Begin UObject interface.
 	~UUnrealEdEngine();
-	virtual void FinishDestroy() OVERRIDE;
-	virtual void Serialize( FArchive& Ar ) OVERRIDE;
+	virtual void FinishDestroy() override;
+	virtual void Serialize( FArchive& Ar ) override;
 	// End UObject interface.
 
 	// Begin FNotify interface.
-	virtual void NotifyPreChange( UProperty* PropertyAboutToChange ) OVERRIDE;
-	virtual void NotifyPostChange( const FPropertyChangedEvent& PropertyChangedEvent, UProperty* PropertyThatChanged ) OVERRIDE;
+	virtual void NotifyPreChange( UProperty* PropertyAboutToChange ) override;
+	virtual void NotifyPostChange( const FPropertyChangedEvent& PropertyChangedEvent, UProperty* PropertyThatChanged ) override;
 	// End FNotify interface.
 
 	// Begin UEditorEngine Interface
-	virtual void SelectActor(AActor* Actor, bool InSelected, bool bNotify, bool bSelectEvenIfHidden=false) OVERRIDE;
-	virtual bool CanSelectActor(AActor* Actor, bool InSelected, bool bSelectEvenIfHidden=false, bool bWarnIfLevelLocked=false) const OVERRIDE;
-	virtual void SelectGroup(AGroupActor* InGroupActor, bool bForceSelection=false, bool bInSelected=true, bool bNotify=true) OVERRIDE;
-	virtual void SelectBSPSurf(UModel* InModel, int32 iSurf, bool bSelected, bool bNoteSelectionChange) OVERRIDE;
-	virtual void SelectNone(bool bNoteSelectionChange, bool bDeselectBSPSurfs, bool WarnAboutManyActors=true) OVERRIDE;
-	virtual void NoteSelectionChange() OVERRIDE;
-	virtual void NoteActorMovement() OVERRIDE;
-	virtual void FinishAllSnaps() OVERRIDE;
-	virtual void Cleanse( bool ClearSelection, bool Redraw, const FText& Reason ) OVERRIDE;
-	virtual bool GetMapBuildCancelled() const OVERRIDE;
-	virtual void SetMapBuildCancelled( bool InCancelled ) OVERRIDE;
-	virtual FVector GetPivotLocation() OVERRIDE;
-	virtual void SetPivot(FVector NewPivot, bool bSnapPivotToGrid, bool bIgnoreAxis, bool bAssignPivot=false) OVERRIDE;
-	virtual void ResetPivot() OVERRIDE;
-	virtual void RedrawLevelEditingViewports(bool bInvalidateHitProxies=true) OVERRIDE;
-	virtual void TakeHighResScreenShots() OVERRIDE;
-	virtual void GetPackageList( TArray<UPackage*>* InPackages, UClass* InClass ) OVERRIDE;
-	virtual bool ShouldAbortActorDeletion() const OVERRIDE;
-	virtual void CloseEditor() OVERRIDE;
-	virtual void OnOpenMatinee() OVERRIDE;
+	virtual void SelectActor(AActor* Actor, bool InSelected, bool bNotify, bool bSelectEvenIfHidden=false) override;
+	virtual bool CanSelectActor(AActor* Actor, bool InSelected, bool bSelectEvenIfHidden=false, bool bWarnIfLevelLocked=false) const override;
+	virtual void SelectGroup(AGroupActor* InGroupActor, bool bForceSelection=false, bool bInSelected=true, bool bNotify=true) override;
+	virtual void SelectBSPSurf(UModel* InModel, int32 iSurf, bool bSelected, bool bNoteSelectionChange) override;
+	virtual void SelectNone(bool bNoteSelectionChange, bool bDeselectBSPSurfs, bool WarnAboutManyActors=true) override;
+	virtual void NoteSelectionChange() override;
+	virtual void NoteActorMovement() override;
+	virtual void FinishAllSnaps() override;
+	virtual void Cleanse( bool ClearSelection, bool Redraw, const FText& Reason ) override;
+	virtual bool GetMapBuildCancelled() const override;
+	virtual void SetMapBuildCancelled( bool InCancelled ) override;
+	virtual FVector GetPivotLocation() override;
+	virtual void SetPivot(FVector NewPivot, bool bSnapPivotToGrid, bool bIgnoreAxis, bool bAssignPivot=false) override;
+	virtual void ResetPivot() override;
+	virtual void RedrawLevelEditingViewports(bool bInvalidateHitProxies=true) override;
+	virtual void TakeHighResScreenShots() override;
+	virtual void GetPackageList( TArray<UPackage*>* InPackages, UClass* InClass ) override;
+	virtual bool ShouldAbortActorDeletion() const override;
+	virtual void CloseEditor() override;
+	virtual void OnOpenMatinee() override;
 	// End UEditorEngine Interface 
 	
 	// Begin FExec Interface
-	virtual bool Exec( UWorld* InWorld, const TCHAR* Cmd, FOutputDevice& Ar=*GLog ) OVERRIDE;
+	virtual bool Exec( UWorld* InWorld, const TCHAR* Cmd, FOutputDevice& Ar=*GLog ) override;
 	// End FExec Interface
 
 
 	// Begin UEngine Interface.
-	virtual void Init(IEngineLoop* InEngineLoop) OVERRIDE;
-	virtual void PreExit() OVERRIDE;
-	virtual void Tick(float DeltaSeconds, bool bIdleMode) OVERRIDE;
+	virtual void Init(IEngineLoop* InEngineLoop) override;
+	virtual void PreExit() override;
+	virtual void Tick(float DeltaSeconds, bool bIdleMode) override;
 	// End UEngine Interface.
 
 	/** Builds a list of sprite categories for use in menus */
@@ -377,7 +377,7 @@ class UNREALED_API UUnrealEdEngine : public UEditorEngine, public FNotifyHook
 	 * @param	bVerifyDeletionCanHappen	[opt] If true (default), verify that deletion can be performed.
 	 * @return								true unless the delete operation was aborted.
 	 */
-	virtual bool edactDeleteSelected( UWorld* InWorld, bool bVerifyDeletionCanHappen=true ) OVERRIDE;
+	virtual bool edactDeleteSelected( UWorld* InWorld, bool bVerifyDeletionCanHappen=true ) override;
 
 	/**
 	 * Creates a new group from the current selection removing any existing groups.
@@ -426,7 +426,7 @@ class UNREALED_API UUnrealEdEngine : public UEditorEngine, public FNotifyHook
 	 * @param	InWorld					World context
 	 * @param	DestinationData			If != NULL, additionally copy data to string
 	 */
-	virtual void edactCopySelected(UWorld* InWorld, FString* DestinationData = NULL) OVERRIDE;
+	virtual void edactCopySelected(UWorld* InWorld, FString* DestinationData = NULL) override;
 
 	/**
 	 * Paste selected actors from the clipboard.
@@ -437,7 +437,7 @@ class UNREALED_API UUnrealEdEngine : public UEditorEngine, public FNotifyHook
 	 * @param	bWarnIfHidden		If true displays a warning if the destination level is hidden
 	 * @param	SourceData			If != NULL, use instead of clipboard data
 	 */
-	virtual void edactPasteSelected(UWorld* InWorld, bool bDuplicate, bool bOffsetLocations, bool bWarnIfHidden, FString* SourceData = NULL) OVERRIDE;
+	virtual void edactPasteSelected(UWorld* InWorld, bool bDuplicate, bool bOffsetLocations, bool bWarnIfHidden, FString* SourceData = NULL) override;
 
 	/**
 	 * Duplicates selected actors.  Handles the case where you are trying to duplicate PrefabInstance actors.
@@ -445,7 +445,7 @@ class UNREALED_API UUnrealEdEngine : public UEditorEngine, public FNotifyHook
 	 * @param	InLevel				Level to place duplicate
 	 * @param	bUseOffset			Should the actor locations be offset after they are created?
 	 */
-	virtual void edactDuplicateSelected(ULevel* InLevel, bool bUseOffset) OVERRIDE;
+	virtual void edactDuplicateSelected(ULevel* InLevel, bool bUseOffset) override;
 
 	/**
 	 * Replace all selected brushes with the default brush.
diff --git a/Engine/Source/Editor/UnrealEd/Classes/Exporters/AnimSequenceExporterFBX.h b/Engine/Source/Editor/UnrealEd/Classes/Exporters/AnimSequenceExporterFBX.h
--- a/Engine/Source/Editor/UnrealEd/Classes/Exporters/AnimSequenceExporterFBX.h
+++ b/Engine/Source/Editor/UnrealEd/Classes/Exporters/AnimSequenceExporterFBX.h
@@ -14,7 +14,7 @@ class UAnimSequenceExporterFBX : public UExporter
 
 
 	// Begin UExporter Interface
-	virtual bool ExportBinary( UObject* Object, const TCHAR* Type, FArchive& Ar, FFeedbackContext* Warn, int32 FileIndex = 0, uint32 PortFlags=0 ) OVERRIDE;
+	virtual bool ExportBinary( UObject* Object, const TCHAR* Type, FArchive& Ar, FFeedbackContext* Warn, int32 FileIndex = 0, uint32 PortFlags=0 ) override;
 	// End UExporter Interface
 };
 
diff --git a/Engine/Source/Editor/UnrealEd/Classes/Exporters/LevelExporterFBX.h b/Engine/Source/Editor/UnrealEd/Classes/Exporters/LevelExporterFBX.h
--- a/Engine/Source/Editor/UnrealEd/Classes/Exporters/LevelExporterFBX.h
+++ b/Engine/Source/Editor/UnrealEd/Classes/Exporters/LevelExporterFBX.h
@@ -14,7 +14,7 @@ class ULevelExporterFBX : public UExporter
 
 
 	// Begin UExporter Interface
-	virtual bool ExportBinary( UObject* Object, const TCHAR* Type, FArchive& Ar, FFeedbackContext* Warn, int32 FileIndex = 0, uint32 PortFlags=0 ) OVERRIDE;
+	virtual bool ExportBinary( UObject* Object, const TCHAR* Type, FArchive& Ar, FFeedbackContext* Warn, int32 FileIndex = 0, uint32 PortFlags=0 ) override;
 	// End UExporter Interface
 };
 
diff --git a/Engine/Source/Editor/UnrealEd/Classes/Exporters/LevelExporterLOD.h b/Engine/Source/Editor/UnrealEd/Classes/Exporters/LevelExporterLOD.h
--- a/Engine/Source/Editor/UnrealEd/Classes/Exporters/LevelExporterLOD.h
+++ b/Engine/Source/Editor/UnrealEd/Classes/Exporters/LevelExporterLOD.h
@@ -14,7 +14,7 @@ class ULevelExporterLOD : public UExporter
 
 
 	// Begin UExporter Interface
-	virtual bool ExportText( const FExportObjectInnerContext* Context, UObject* Object, const TCHAR* Type, FOutputDevice& Ar, FFeedbackContext* Warn, uint32 PortFlags=0 ) OVERRIDE;
+	virtual bool ExportText( const FExportObjectInnerContext* Context, UObject* Object, const TCHAR* Type, FOutputDevice& Ar, FFeedbackContext* Warn, uint32 PortFlags=0 ) override;
 	// End UExporter Interface
 };
 
diff --git a/Engine/Source/Editor/UnrealEd/Classes/Exporters/LevelExporterOBJ.h b/Engine/Source/Editor/UnrealEd/Classes/Exporters/LevelExporterOBJ.h
--- a/Engine/Source/Editor/UnrealEd/Classes/Exporters/LevelExporterOBJ.h
+++ b/Engine/Source/Editor/UnrealEd/Classes/Exporters/LevelExporterOBJ.h
@@ -14,7 +14,7 @@ class ULevelExporterOBJ : public UExporter
 
 
 	// Begin UExporter Interface
-	virtual bool ExportText( const FExportObjectInnerContext* Context, UObject* Object, const TCHAR* Type, FOutputDevice& Ar, FFeedbackContext* Warn, uint32 PortFlags=0 ) OVERRIDE;
+	virtual bool ExportText( const FExportObjectInnerContext* Context, UObject* Object, const TCHAR* Type, FOutputDevice& Ar, FFeedbackContext* Warn, uint32 PortFlags=0 ) override;
 	// End UExporter Interface
 };
 
diff --git a/Engine/Source/Editor/UnrealEd/Classes/Exporters/LevelExporterT3D.h b/Engine/Source/Editor/UnrealEd/Classes/Exporters/LevelExporterT3D.h
--- a/Engine/Source/Editor/UnrealEd/Classes/Exporters/LevelExporterT3D.h
+++ b/Engine/Source/Editor/UnrealEd/Classes/Exporters/LevelExporterT3D.h
@@ -14,10 +14,10 @@ class ULevelExporterT3D : public UExporter
 
 
 	// Begin UExporter Interface
-	virtual bool ExportText( const FExportObjectInnerContext* Context, UObject* Object, const TCHAR* Type, FOutputDevice& Ar, FFeedbackContext* Warn, uint32 PortFlags=0 ) OVERRIDE;
-	virtual void ExportPackageObject(FExportPackageParams& ExpPackageParams) OVERRIDE;
-	virtual void ExportPackageInners(FExportPackageParams& ExpPackageParams) OVERRIDE;
-	virtual void ExportComponentExtra(const FExportObjectInnerContext* Context, const TArray<UActorComponent*>& Components, FOutputDevice& Ar, uint32 PortFlags) OVERRIDE;
+	virtual bool ExportText( const FExportObjectInnerContext* Context, UObject* Object, const TCHAR* Type, FOutputDevice& Ar, FFeedbackContext* Warn, uint32 PortFlags=0 ) override;
+	virtual void ExportPackageObject(FExportPackageParams& ExpPackageParams) override;
+	virtual void ExportPackageInners(FExportPackageParams& ExpPackageParams) override;
+	virtual void ExportComponentExtra(const FExportObjectInnerContext* Context, const TArray<UActorComponent*>& Components, FOutputDevice& Ar, uint32 PortFlags) override;
 	// End UExporter Interface
 };
 
diff --git a/Engine/Source/Editor/UnrealEd/Classes/Exporters/SoundExporterOGG.h b/Engine/Source/Editor/UnrealEd/Classes/Exporters/SoundExporterOGG.h
--- a/Engine/Source/Editor/UnrealEd/Classes/Exporters/SoundExporterOGG.h
+++ b/Engine/Source/Editor/UnrealEd/Classes/Exporters/SoundExporterOGG.h
@@ -8,8 +8,8 @@ class USoundExporterOGG : public UExporter
 {
     GENERATED_UCLASS_BODY()
 	// Begin UExporter Interface
-	virtual bool ExportBinary( UObject* Object, const TCHAR* Type, FArchive& Ar, FFeedbackContext* Warn, int32 FileIndex = 0, uint32 PortFlags=0 ) OVERRIDE;
-	virtual bool SupportsObject(UObject* Object) const OVERRIDE;
+	virtual bool ExportBinary( UObject* Object, const TCHAR* Type, FArchive& Ar, FFeedbackContext* Warn, int32 FileIndex = 0, uint32 PortFlags=0 ) override;
+	virtual bool SupportsObject(UObject* Object) const override;
 	// End UExporter Interface
 };
 
diff --git a/Engine/Source/Editor/UnrealEd/Classes/Exporters/SoundExporterWAV.h b/Engine/Source/Editor/UnrealEd/Classes/Exporters/SoundExporterWAV.h
--- a/Engine/Source/Editor/UnrealEd/Classes/Exporters/SoundExporterWAV.h
+++ b/Engine/Source/Editor/UnrealEd/Classes/Exporters/SoundExporterWAV.h
@@ -8,8 +8,8 @@ class USoundExporterWAV : public UExporter
 {
     GENERATED_UCLASS_BODY()
 	// Begin UExporter Interface
-	virtual bool ExportBinary( UObject* Object, const TCHAR* Type, FArchive& Ar, FFeedbackContext* Warn, int32 FileIndex = 0, uint32 PortFlags=0 ) OVERRIDE;
-	virtual bool SupportsObject(UObject* Object) const OVERRIDE;
+	virtual bool ExportBinary( UObject* Object, const TCHAR* Type, FArchive& Ar, FFeedbackContext* Warn, int32 FileIndex = 0, uint32 PortFlags=0 ) override;
+	virtual bool SupportsObject(UObject* Object) const override;
 	// End UExporter Interface
 };
 
diff --git a/Engine/Source/Editor/UnrealEd/Classes/Exporters/SoundSurroundExporterWAV.h b/Engine/Source/Editor/UnrealEd/Classes/Exporters/SoundSurroundExporterWAV.h
--- a/Engine/Source/Editor/UnrealEd/Classes/Exporters/SoundSurroundExporterWAV.h
+++ b/Engine/Source/Editor/UnrealEd/Classes/Exporters/SoundSurroundExporterWAV.h
@@ -14,10 +14,10 @@ class USoundSurroundExporterWAV : public UExporter
 
 
 	// Begin UExporter Interface
-	virtual int32 GetFileCount( void ) const OVERRIDE;
-	virtual FString GetUniqueFilename( const TCHAR* Filename, int32 FileIndex ) OVERRIDE;
-	virtual bool ExportBinary( UObject* Object, const TCHAR* Type, FArchive& Ar, FFeedbackContext* Warn, int32 FileIndex = 0, uint32 PortFlags=0 ) OVERRIDE;
-	virtual bool SupportsObject(UObject* Object) const OVERRIDE;
+	virtual int32 GetFileCount( void ) const override;
+	virtual FString GetUniqueFilename( const TCHAR* Filename, int32 FileIndex ) override;
+	virtual bool ExportBinary( UObject* Object, const TCHAR* Type, FArchive& Ar, FFeedbackContext* Warn, int32 FileIndex = 0, uint32 PortFlags=0 ) override;
+	virtual bool SupportsObject(UObject* Object) const override;
 	// End UExporter Interface
 };
 
diff --git a/Engine/Source/Editor/UnrealEd/Classes/Exporters/TextureCubeExporterHDR.h b/Engine/Source/Editor/UnrealEd/Classes/Exporters/TextureCubeExporterHDR.h
--- a/Engine/Source/Editor/UnrealEd/Classes/Exporters/TextureCubeExporterHDR.h
+++ b/Engine/Source/Editor/UnrealEd/Classes/Exporters/TextureCubeExporterHDR.h
@@ -13,7 +13,7 @@ class UNREALED_API UTextureCubeExporterHDR : public UExporter
 	GENERATED_UCLASS_BODY()
 
 	// Begin UExporter Interface
-	virtual bool ExportBinary(UObject* Object, const TCHAR* Type, FArchive& Ar, FFeedbackContext* Warn, int32 FileIndex = 0, uint32 PortFlags = 0) OVERRIDE;
+	virtual bool ExportBinary(UObject* Object, const TCHAR* Type, FArchive& Ar, FFeedbackContext* Warn, int32 FileIndex = 0, uint32 PortFlags = 0) override;
 	// End UExporter Interface
 };
 
diff --git a/Engine/Source/Editor/UnrealEd/Classes/Exporters/TextureExporterHDR.h b/Engine/Source/Editor/UnrealEd/Classes/Exporters/TextureExporterHDR.h
--- a/Engine/Source/Editor/UnrealEd/Classes/Exporters/TextureExporterHDR.h
+++ b/Engine/Source/Editor/UnrealEd/Classes/Exporters/TextureExporterHDR.h
@@ -13,7 +13,7 @@ class UNREALED_API UTextureExporterHDR : public UExporter
 	GENERATED_UCLASS_BODY()
 
 	// Begin UExporter Interface
-	virtual bool ExportBinary(UObject* Object, const TCHAR* Type, FArchive& Ar, FFeedbackContext* Warn, int32 FileIndex = 0, uint32 PortFlags = 0) OVERRIDE;
+	virtual bool ExportBinary(UObject* Object, const TCHAR* Type, FArchive& Ar, FFeedbackContext* Warn, int32 FileIndex = 0, uint32 PortFlags = 0) override;
 	// End UExporter Interface
 };
 
diff --git a/Engine/Source/Editor/UnrealEd/Classes/Factories/AnimBlueprintFactory.h b/Engine/Source/Editor/UnrealEd/Classes/Factories/AnimBlueprintFactory.h
--- a/Engine/Source/Editor/UnrealEd/Classes/Factories/AnimBlueprintFactory.h
+++ b/Engine/Source/Editor/UnrealEd/Classes/Factories/AnimBlueprintFactory.h
@@ -25,9 +25,9 @@ class UAnimBlueprintFactory : public UFactory
 	class USkeleton* TargetSkeleton;
 
 	// Begin UFactory Interface
-	virtual bool ConfigureProperties() OVERRIDE;
-	virtual UObject* FactoryCreateNew(UClass* Class, UObject* InParent, FName Name, EObjectFlags Flags, UObject* Context, FFeedbackContext* Warn, FName CallingContext) OVERRIDE;
-	virtual UObject* FactoryCreateNew(UClass* Class,UObject* InParent,FName Name,EObjectFlags Flags,UObject* Context,FFeedbackContext* Warn) OVERRIDE;
+	virtual bool ConfigureProperties() override;
+	virtual UObject* FactoryCreateNew(UClass* Class, UObject* InParent, FName Name, EObjectFlags Flags, UObject* Context, FFeedbackContext* Warn, FName CallingContext) override;
+	virtual UObject* FactoryCreateNew(UClass* Class,UObject* InParent,FName Name,EObjectFlags Flags,UObject* Context,FFeedbackContext* Warn) override;
 	// Begin UFactory Interface	
 };
 
diff --git a/Engine/Source/Editor/UnrealEd/Classes/Factories/AnimCompositeFactory.h b/Engine/Source/Editor/UnrealEd/Classes/Factories/AnimCompositeFactory.h
--- a/Engine/Source/Editor/UnrealEd/Classes/Factories/AnimCompositeFactory.h
+++ b/Engine/Source/Editor/UnrealEd/Classes/Factories/AnimCompositeFactory.h
@@ -20,8 +20,8 @@ class UAnimCompositeFactory : public UFactory
 	class UAnimSequence* SourceAnimation;
 
 	// Begin UFactory Interface
-	virtual bool ConfigureProperties() OVERRIDE;
-	virtual UObject* FactoryCreateNew(UClass* Class,UObject* InParent,FName Name,EObjectFlags Flags,UObject* Context,FFeedbackContext* Warn) OVERRIDE;
+	virtual bool ConfigureProperties() override;
+	virtual UObject* FactoryCreateNew(UClass* Class,UObject* InParent,FName Name,EObjectFlags Flags,UObject* Context,FFeedbackContext* Warn) override;
 	// Begin UFactory Interface	
 
 private:
diff --git a/Engine/Source/Editor/UnrealEd/Classes/Factories/AnimMontageFactory.h b/Engine/Source/Editor/UnrealEd/Classes/Factories/AnimMontageFactory.h
--- a/Engine/Source/Editor/UnrealEd/Classes/Factories/AnimMontageFactory.h
+++ b/Engine/Source/Editor/UnrealEd/Classes/Factories/AnimMontageFactory.h
@@ -20,8 +20,8 @@ class UAnimMontageFactory : public UFactory
 	class UAnimSequence* SourceAnimation;
 
 	// Begin UFactory Interface
-	virtual bool ConfigureProperties() OVERRIDE;
-	virtual UObject* FactoryCreateNew(UClass* Class,UObject* InParent,FName Name,EObjectFlags Flags,UObject* Context,FFeedbackContext* Warn) OVERRIDE;
+	virtual bool ConfigureProperties() override;
+	virtual UObject* FactoryCreateNew(UClass* Class,UObject* InParent,FName Name,EObjectFlags Flags,UObject* Context,FFeedbackContext* Warn) override;
 	// Begin UFactory Interface	
 
 private:
diff --git a/Engine/Source/Editor/UnrealEd/Classes/Factories/BlueprintFactory.h b/Engine/Source/Editor/UnrealEd/Classes/Factories/BlueprintFactory.h
--- a/Engine/Source/Editor/UnrealEd/Classes/Factories/BlueprintFactory.h
+++ b/Engine/Source/Editor/UnrealEd/Classes/Factories/BlueprintFactory.h
@@ -17,9 +17,9 @@ class UNREALED_API UBlueprintFactory : public UFactory
 	TSubclassOf<class UObject> ParentClass;
 
 	// UFactory Interface
-	virtual bool ConfigureProperties() OVERRIDE;
-	virtual UObject* FactoryCreateNew(UClass* Class, UObject* InParent, FName Name, EObjectFlags Flags, UObject* Context, FFeedbackContext* Warn, FName CallingContext) OVERRIDE;
-	virtual UObject* FactoryCreateNew(UClass* Class, UObject* InParent, FName Name, EObjectFlags Flags, UObject* Context, FFeedbackContext* Warn) OVERRIDE;
+	virtual bool ConfigureProperties() override;
+	virtual UObject* FactoryCreateNew(UClass* Class, UObject* InParent, FName Name, EObjectFlags Flags, UObject* Context, FFeedbackContext* Warn, FName CallingContext) override;
+	virtual UObject* FactoryCreateNew(UClass* Class, UObject* InParent, FName Name, EObjectFlags Flags, UObject* Context, FFeedbackContext* Warn) override;
 	// End of UFactory Interface
 
 protected:
diff --git a/Engine/Source/Editor/UnrealEd/Classes/Factories/BlueprintFunctionLibraryFactory.h b/Engine/Source/Editor/UnrealEd/Classes/Factories/BlueprintFunctionLibraryFactory.h
--- a/Engine/Source/Editor/UnrealEd/Classes/Factories/BlueprintFunctionLibraryFactory.h
+++ b/Engine/Source/Editor/UnrealEd/Classes/Factories/BlueprintFunctionLibraryFactory.h
@@ -13,12 +13,12 @@ class UBlueprintFunctionLibraryFactory : public UBlueprintFactory
 	GENERATED_UCLASS_BODY()
 
 	// UFactory interface
-	virtual FText GetDisplayName() const OVERRIDE;
-	virtual FName GetNewAssetThumbnailOverride() const OVERRIDE;
-	virtual uint32 GetMenuCategories() const OVERRIDE;
-	virtual UObject* FactoryCreateNew(UClass* Class, UObject* InParent, FName Name, EObjectFlags Flags, UObject* Context, FFeedbackContext* Warn, FName CallingContext) OVERRIDE;
-	virtual bool ConfigureProperties() OVERRIDE;
-	virtual FString GetDefaultNewAssetName() const OVERRIDE;
+	virtual FText GetDisplayName() const override;
+	virtual FName GetNewAssetThumbnailOverride() const override;
+	virtual uint32 GetMenuCategories() const override;
+	virtual UObject* FactoryCreateNew(UClass* Class, UObject* InParent, FName Name, EObjectFlags Flags, UObject* Context, FFeedbackContext* Warn, FName CallingContext) override;
+	virtual bool ConfigureProperties() override;
+	virtual FString GetDefaultNewAssetName() const override;
 	// End of UFactory interface
 };
 
diff --git a/Engine/Source/Editor/UnrealEd/Classes/Factories/BlueprintInterfaceFactory.h b/Engine/Source/Editor/UnrealEd/Classes/Factories/BlueprintInterfaceFactory.h
--- a/Engine/Source/Editor/UnrealEd/Classes/Factories/BlueprintInterfaceFactory.h
+++ b/Engine/Source/Editor/UnrealEd/Classes/Factories/BlueprintInterfaceFactory.h
@@ -13,12 +13,12 @@ class UBlueprintInterfaceFactory : public UFactory
 	GENERATED_UCLASS_BODY()
 
 	// UFactory interface
-	virtual FText GetDisplayName() const OVERRIDE;
-	virtual FName GetNewAssetThumbnailOverride() const OVERRIDE;
-	virtual uint32 GetMenuCategories() const OVERRIDE;
-	virtual UObject* FactoryCreateNew(UClass* Class, UObject* InParent, FName Name, EObjectFlags Flags, UObject* Context, FFeedbackContext* Warn, FName CallingContext) OVERRIDE;
-	virtual UObject* FactoryCreateNew(UClass* Class, UObject* InParent, FName Name, EObjectFlags Flags, UObject* Context, FFeedbackContext* Warn) OVERRIDE;
-	virtual FString GetDefaultNewAssetName() const OVERRIDE;
+	virtual FText GetDisplayName() const override;
+	virtual FName GetNewAssetThumbnailOverride() const override;
+	virtual uint32 GetMenuCategories() const override;
+	virtual UObject* FactoryCreateNew(UClass* Class, UObject* InParent, FName Name, EObjectFlags Flags, UObject* Context, FFeedbackContext* Warn, FName CallingContext) override;
+	virtual UObject* FactoryCreateNew(UClass* Class, UObject* InParent, FName Name, EObjectFlags Flags, UObject* Context, FFeedbackContext* Warn) override;
+	virtual FString GetDefaultNewAssetName() const override;
 	// End of UFactory interface
 };
 
diff --git a/Engine/Source/Editor/UnrealEd/Classes/Factories/BlueprintMacroFactory.h b/Engine/Source/Editor/UnrealEd/Classes/Factories/BlueprintMacroFactory.h
--- a/Engine/Source/Editor/UnrealEd/Classes/Factories/BlueprintMacroFactory.h
+++ b/Engine/Source/Editor/UnrealEd/Classes/Factories/BlueprintMacroFactory.h
@@ -13,16 +13,16 @@ class UBlueprintMacroFactory : public UBlueprintFactory
 	GENERATED_UCLASS_BODY()
 
 	// UFactory interface
-	virtual FText GetDisplayName() const OVERRIDE;
-	virtual FName GetNewAssetThumbnailOverride() const OVERRIDE;
-	virtual uint32 GetMenuCategories() const OVERRIDE;
-	virtual UObject* FactoryCreateNew(UClass* Class, UObject* InParent, FName Name, EObjectFlags Flags, UObject* Context, FFeedbackContext* Warn, FName CallingContext) OVERRIDE;
-	virtual FString GetDefaultNewAssetName() const OVERRIDE;
+	virtual FText GetDisplayName() const override;
+	virtual FName GetNewAssetThumbnailOverride() const override;
+	virtual uint32 GetMenuCategories() const override;
+	virtual UObject* FactoryCreateNew(UClass* Class, UObject* InParent, FName Name, EObjectFlags Flags, UObject* Context, FFeedbackContext* Warn, FName CallingContext) override;
+	virtual FString GetDefaultNewAssetName() const override;
 	// End of UFactory interface
 
 protected:
 	// UBlueprintFactory interface
-	virtual bool IsMacroFactory() const OVERRIDE { return true; }
+	virtual bool IsMacroFactory() const override { return true; }
 	// End of UBlueprintFactory interface
 };
 
diff --git a/Engine/Source/Editor/UnrealEd/Classes/Factories/CSVImportFactory.h b/Engine/Source/Editor/UnrealEd/Classes/Factories/CSVImportFactory.h
--- a/Engine/Source/Editor/UnrealEd/Classes/Factories/CSVImportFactory.h
+++ b/Engine/Source/Editor/UnrealEd/Classes/Factories/CSVImportFactory.h
@@ -15,9 +15,9 @@ class UCSVImportFactory : public UFactory
 	GENERATED_UCLASS_BODY()
 
 	// Begin UFactory Interface
-	virtual FText GetDisplayName() const OVERRIDE;
-	virtual UObject* FactoryCreateText( UClass* InClass, UObject* InParent, FName InName, EObjectFlags Flags, UObject* Context, const TCHAR* Type, const TCHAR*& Buffer, const TCHAR* BufferEnd, FFeedbackContext* Warn ) OVERRIDE;
-	virtual bool DoesSupportClass(UClass * Class) OVERRIDE;
+	virtual FText GetDisplayName() const override;
+	virtual UObject* FactoryCreateText( UClass* InClass, UObject* InParent, FName InName, EObjectFlags Flags, UObject* Context, const TCHAR* Type, const TCHAR*& Buffer, const TCHAR* BufferEnd, FFeedbackContext* Warn ) override;
+	virtual bool DoesSupportClass(UClass * Class) override;
 	
 	// Begin UFactory Interface
 
diff --git a/Engine/Source/Editor/UnrealEd/Classes/Factories/CameraAnimFactory.h b/Engine/Source/Editor/UnrealEd/Classes/Factories/CameraAnimFactory.h
--- a/Engine/Source/Editor/UnrealEd/Classes/Factories/CameraAnimFactory.h
+++ b/Engine/Source/Editor/UnrealEd/Classes/Factories/CameraAnimFactory.h
@@ -10,10 +10,10 @@ class UCameraAnimFactory : public UFactory
 	GENERATED_UCLASS_BODY()
 
 	// Begin UFactory Interface
-	virtual FText GetDisplayName() const OVERRIDE;
-	virtual FName GetNewAssetThumbnailOverride() const OVERRIDE;
-	virtual uint32 GetMenuCategories() const OVERRIDE;
-	virtual UObject* FactoryCreateNew(UClass* Class,UObject* InParent,FName Name,EObjectFlags Flags,UObject* Context,FFeedbackContext* Warn) OVERRIDE;
+	virtual FText GetDisplayName() const override;
+	virtual FName GetNewAssetThumbnailOverride() const override;
+	virtual uint32 GetMenuCategories() const override;
+	virtual UObject* FactoryCreateNew(UClass* Class,UObject* InParent,FName Name,EObjectFlags Flags,UObject* Context,FFeedbackContext* Warn) override;
 	// Begin UFactory Interface
 };
 
diff --git a/Engine/Source/Editor/UnrealEd/Classes/Factories/CurveImportFactory.h b/Engine/Source/Editor/UnrealEd/Classes/Factories/CurveImportFactory.h
--- a/Engine/Source/Editor/UnrealEd/Classes/Factories/CurveImportFactory.h
+++ b/Engine/Source/Editor/UnrealEd/Classes/Factories/CurveImportFactory.h
@@ -13,7 +13,7 @@ class UCurveImportFactory : public UFactory
 	GENERATED_UCLASS_BODY()
 
 	// Begin UFactory Interface
-	virtual UObject* FactoryCreateText( UClass* InClass, UObject* InParent, FName InName, EObjectFlags Flags, UObject* Context, const TCHAR* Type, const TCHAR*& Buffer, const TCHAR* BufferEnd, FFeedbackContext* Warn ) OVERRIDE;
+	virtual UObject* FactoryCreateText( UClass* InClass, UObject* InParent, FName InName, EObjectFlags Flags, UObject* Context, const TCHAR* Type, const TCHAR*& Buffer, const TCHAR* BufferEnd, FFeedbackContext* Warn ) override;
 	// Begin UFactory Interface
 };
 
diff --git a/Engine/Source/Editor/UnrealEd/Classes/Factories/DestructibleMeshFactory.h b/Engine/Source/Editor/UnrealEd/Classes/Factories/DestructibleMeshFactory.h
--- a/Engine/Source/Editor/UnrealEd/Classes/Factories/DestructibleMeshFactory.h
+++ b/Engine/Source/Editor/UnrealEd/Classes/Factories/DestructibleMeshFactory.h
@@ -10,9 +10,9 @@ class UDestructibleMeshFactory : public UFactory
 {
     GENERATED_UCLASS_BODY()
 	// Begin UFactory Interface
-	virtual FText GetDisplayName() const OVERRIDE;
+	virtual FText GetDisplayName() const override;
 #if WITH_APEX
-	virtual UObject* FactoryCreateBinary( UClass* Class, UObject* InParent, FName Name, EObjectFlags Flags, UObject* Context, const TCHAR* Type, const uint8*& Buffer, const uint8* BufferEnd, FFeedbackContext* Warn ) OVERRIDE;
+	virtual UObject* FactoryCreateBinary( UClass* Class, UObject* InParent, FName Name, EObjectFlags Flags, UObject* Context, const TCHAR* Type, const uint8*& Buffer, const uint8* BufferEnd, FFeedbackContext* Warn ) override;
 #endif // WITH_APEX
 	// Begin UFactory Interface	
 };
diff --git a/Engine/Source/Editor/UnrealEd/Classes/Factories/FbxFactory.h b/Engine/Source/Editor/UnrealEd/Classes/Factories/FbxFactory.h
--- a/Engine/Source/Editor/UnrealEd/Classes/Factories/FbxFactory.h
+++ b/Engine/Source/Editor/UnrealEd/Classes/Factories/FbxFactory.h
@@ -17,15 +17,15 @@ class UFbxFactory : public UFactory
 
 
 	// Begin UObject Interface
-	virtual void CleanUp() OVERRIDE;
-	virtual bool ConfigureProperties() OVERRIDE;
-	virtual void PostInitProperties() OVERRIDE;
+	virtual void CleanUp() override;
+	virtual bool ConfigureProperties() override;
+	virtual void PostInitProperties() override;
 	// End UObject Interface
 
 	// Begin UFactory Interface
-	virtual bool DoesSupportClass(UClass * Class) OVERRIDE;
-	virtual UClass* ResolveSupportedClass() OVERRIDE;
-	virtual UObject* FactoryCreateBinary(UClass* InClass, UObject* InParent, FName InName, EObjectFlags Flags, UObject* Context, const TCHAR* Type, const uint8*& Buffer, const uint8* BufferEnd, FFeedbackContext* Warn, bool& bOutOperationCanceled) OVERRIDE;
+	virtual bool DoesSupportClass(UClass * Class) override;
+	virtual UClass* ResolveSupportedClass() override;
+	virtual UObject* FactoryCreateBinary(UClass* InClass, UObject* InParent, FName InName, EObjectFlags Flags, UObject* Context, const TCHAR* Type, const uint8*& Buffer, const uint8* BufferEnd, FFeedbackContext* Warn, bool& bOutOperationCanceled) override;
 	// End UFactory Interface
 	
 	/**
diff --git a/Engine/Source/Editor/UnrealEd/Classes/Factories/FbxImportUI.h b/Engine/Source/Editor/UnrealEd/Classes/Factories/FbxImportUI.h
--- a/Engine/Source/Editor/UnrealEd/Classes/Factories/FbxImportUI.h
+++ b/Engine/Source/Editor/UnrealEd/Classes/Factories/FbxImportUI.h
@@ -96,7 +96,7 @@ class UFbxImportUI : public UObject
 	class UFbxTextureImportData* TextureImportData;
 
 	// Begin UObject Interface
-	virtual bool CanEditChange( const UProperty* InProperty ) const OVERRIDE;
+	virtual bool CanEditChange( const UProperty* InProperty ) const override;
 	// End UObject Interface
 };
 
diff --git a/Engine/Source/Editor/UnrealEd/Classes/Factories/FontFactory.h b/Engine/Source/Editor/UnrealEd/Classes/Factories/FontFactory.h
--- a/Engine/Source/Editor/UnrealEd/Classes/Factories/FontFactory.h
+++ b/Engine/Source/Editor/UnrealEd/Classes/Factories/FontFactory.h
@@ -23,7 +23,7 @@ class UFontFactory : public UTextureFactory
 
 
 	// Begin UFactory Interface
-	virtual UObject* FactoryCreateBinary( UClass* Class, UObject* InParent, FName Name, EObjectFlags Flags, UObject* Context, const TCHAR* Type, const uint8*& Buffer, const uint8* BufferEnd, FFeedbackContext* Warn ) OVERRIDE;
+	virtual UObject* FactoryCreateBinary( UClass* Class, UObject* InParent, FName Name, EObjectFlags Flags, UObject* Context, const TCHAR* Type, const uint8*& Buffer, const uint8* BufferEnd, FFeedbackContext* Warn ) override;
 	// Begin UFactory Interface	
 };
 
diff --git a/Engine/Source/Editor/UnrealEd/Classes/Factories/LevelFactory.h b/Engine/Source/Editor/UnrealEd/Classes/Factories/LevelFactory.h
--- a/Engine/Source/Editor/UnrealEd/Classes/Factories/LevelFactory.h
+++ b/Engine/Source/Editor/UnrealEd/Classes/Factories/LevelFactory.h
@@ -14,7 +14,7 @@ class ULevelFactory : public UFactory
 
 
 	// Begin UFactory Interface
-	virtual UObject* FactoryCreateText( UClass* Class, UObject* InParent, FName Name, EObjectFlags Flags, UObject* Context, const TCHAR* Type, const TCHAR*& Buffer, const TCHAR* BufferEnd, FFeedbackContext* Warn ) OVERRIDE;
+	virtual UObject* FactoryCreateText( UClass* Class, UObject* InParent, FName Name, EObjectFlags Flags, UObject* Context, const TCHAR* Type, const TCHAR*& Buffer, const TCHAR* BufferEnd, FFeedbackContext* Warn ) override;
 	// Begin UFactory Interface	
 };
 
diff --git a/Engine/Source/Editor/UnrealEd/Classes/Factories/MaterialInstanceConstantFactoryNew.h b/Engine/Source/Editor/UnrealEd/Classes/Factories/MaterialInstanceConstantFactoryNew.h
--- a/Engine/Source/Editor/UnrealEd/Classes/Factories/MaterialInstanceConstantFactoryNew.h
+++ b/Engine/Source/Editor/UnrealEd/Classes/Factories/MaterialInstanceConstantFactoryNew.h
@@ -16,7 +16,7 @@ class UMaterialInstanceConstantFactoryNew : public UFactory
 	class UMaterialInterface* InitialParent;
 
 	// Begin UFactory Interface
-	UNREALED_API virtual UObject* FactoryCreateNew(UClass* Class,UObject* InParent,FName Name,EObjectFlags Flags,UObject* Context,FFeedbackContext* Warn) OVERRIDE;
+	UNREALED_API virtual UObject* FactoryCreateNew(UClass* Class,UObject* InParent,FName Name,EObjectFlags Flags,UObject* Context,FFeedbackContext* Warn) override;
 	// Begin UFactory Interface	
 };
 
diff --git a/Engine/Source/Editor/UnrealEd/Classes/Factories/ModelFactory.h b/Engine/Source/Editor/UnrealEd/Classes/Factories/ModelFactory.h
--- a/Engine/Source/Editor/UnrealEd/Classes/Factories/ModelFactory.h
+++ b/Engine/Source/Editor/UnrealEd/Classes/Factories/ModelFactory.h
@@ -14,7 +14,7 @@ class UModelFactory : public UFactory
 
 
 	// Begin UFactory Interface
-	virtual UObject* FactoryCreateText( UClass* Class, UObject* InParent, FName Name, EObjectFlags Flags, UObject* Context, const TCHAR* Type, const TCHAR*& Buffer, const TCHAR* BufferEnd, FFeedbackContext* Warn ) OVERRIDE;
+	virtual UObject* FactoryCreateText( UClass* Class, UObject* InParent, FName Name, EObjectFlags Flags, UObject* Context, const TCHAR* Type, const TCHAR*& Buffer, const TCHAR* BufferEnd, FFeedbackContext* Warn ) override;
 	// End UFactory Interface
 };
 
diff --git a/Engine/Source/Editor/UnrealEd/Classes/Factories/ReimportCurveFactory.h b/Engine/Source/Editor/UnrealEd/Classes/Factories/ReimportCurveFactory.h
--- a/Engine/Source/Editor/UnrealEd/Classes/Factories/ReimportCurveFactory.h
+++ b/Engine/Source/Editor/UnrealEd/Classes/Factories/ReimportCurveFactory.h
@@ -10,9 +10,9 @@ class UReimportCurveFactory : public UCSVImportFactory, public FReimportHandler
 
 
 	// Begin FReimportHandler interface
-	virtual bool CanReimport( UObject* Obj, TArray<FString>& OutFilenames ) OVERRIDE;
-	virtual void SetReimportPaths( UObject* Obj, const TArray<FString>& NewReimportPaths ) OVERRIDE;
-	virtual EReimportResult::Type Reimport( UObject* Obj ) OVERRIDE;
+	virtual bool CanReimport( UObject* Obj, TArray<FString>& OutFilenames ) override;
+	virtual void SetReimportPaths( UObject* Obj, const TArray<FString>& NewReimportPaths ) override;
+	virtual EReimportResult::Type Reimport( UObject* Obj ) override;
 	// End FReimportHandler interface
 };
 
