UE4 Refactoring. Changed OVERRIDE and FINAL macros to keywords override and final.

diff --git a/Engine/Source/Editor/UnrealEd/Classes/Factories/ReimportCurveTableFactory.h b/Engine/Source/Editor/UnrealEd/Classes/Factories/ReimportCurveTableFactory.h
--- a/Engine/Source/Editor/UnrealEd/Classes/Factories/ReimportCurveTableFactory.h
+++ b/Engine/Source/Editor/UnrealEd/Classes/Factories/ReimportCurveTableFactory.h
@@ -9,9 +9,9 @@ class UReimportCurveTableFactory : public UCSVImportFactory, public FReimportHan
 	GENERATED_UCLASS_BODY()
 
 	// Begin FReimportHandler interface
-	virtual bool CanReimport( UObject* Obj, TArray<FString>& OutFilenames ) OVERRIDE;
-	virtual void SetReimportPaths( UObject* Obj, const TArray<FString>& NewReimportPaths ) OVERRIDE;
-	virtual EReimportResult::Type Reimport( UObject* Obj ) OVERRIDE;
+	virtual bool CanReimport( UObject* Obj, TArray<FString>& OutFilenames ) override;
+	virtual void SetReimportPaths( UObject* Obj, const TArray<FString>& NewReimportPaths ) override;
+	virtual EReimportResult::Type Reimport( UObject* Obj ) override;
 	// End FReimportHandler interface
 };
 
diff --git a/Engine/Source/Editor/UnrealEd/Classes/Factories/ReimportDestructibleMeshFactory.h b/Engine/Source/Editor/UnrealEd/Classes/Factories/ReimportDestructibleMeshFactory.h
--- a/Engine/Source/Editor/UnrealEd/Classes/Factories/ReimportDestructibleMeshFactory.h
+++ b/Engine/Source/Editor/UnrealEd/Classes/Factories/ReimportDestructibleMeshFactory.h
@@ -10,16 +10,16 @@ class UReimportDestructibleMeshFactory : public UDestructibleMeshFactory, public
 {
     GENERATED_UCLASS_BODY()
 public:
-	virtual FText GetDisplayName() const OVERRIDE;
+	virtual FText GetDisplayName() const override;
 #if WITH_APEX
 	// Begin FReimportHandler interface
-	virtual bool CanReimport( UObject* Obj, TArray<FString>& OutFilenames ) OVERRIDE;
-	virtual void SetReimportPaths( UObject* Obj, const TArray<FString>& NewReimportPaths ) OVERRIDE;
-	virtual EReimportResult::Type Reimport( UObject* Obj ) OVERRIDE;
+	virtual bool CanReimport( UObject* Obj, TArray<FString>& OutFilenames ) override;
+	virtual void SetReimportPaths( UObject* Obj, const TArray<FString>& NewReimportPaths ) override;
+	virtual EReimportResult::Type Reimport( UObject* Obj ) override;
 #else
 	virtual bool CanReimport( UObject* Obj, TArray<FString>& OutFilenames ) { return false; }
 	virtual void SetReimportPaths( UObject* Obj, const TArray<FString>& NewReimportPaths ) {}
-	virtual EReimportResult::Type Reimport( UObject* Obj ) OVERRIDE { return EReimportResult::Failed; }
+	virtual EReimportResult::Type Reimport( UObject* Obj ) override { return EReimportResult::Failed; }
 	// End FReimportHandler interface
 #endif // WITH_APEX
 };
diff --git a/Engine/Source/Editor/UnrealEd/Classes/Factories/ReimportTextureFactory.h b/Engine/Source/Editor/UnrealEd/Classes/Factories/ReimportTextureFactory.h
--- a/Engine/Source/Editor/UnrealEd/Classes/Factories/ReimportTextureFactory.h
+++ b/Engine/Source/Editor/UnrealEd/Classes/Factories/ReimportTextureFactory.h
@@ -17,15 +17,15 @@ class UReimportTextureFactory : public UTextureFactory, public FReimportHandler
 
 
 	// Begin FReimportHandler interface
-	virtual bool CanReimport( UObject* Obj, TArray<FString>& OutFilenames ) OVERRIDE;
-	virtual void SetReimportPaths( UObject* Obj, const TArray<FString>& NewReimportPaths ) OVERRIDE;
-	virtual EReimportResult::Type Reimport( UObject* Obj ) OVERRIDE;
+	virtual bool CanReimport( UObject* Obj, TArray<FString>& OutFilenames ) override;
+	virtual void SetReimportPaths( UObject* Obj, const TArray<FString>& NewReimportPaths ) override;
+	virtual EReimportResult::Type Reimport( UObject* Obj ) override;
 	// End FReimportHandler interface
 
 private:
 	// Begin UTextureFactory Interface
-	virtual UTexture2D* CreateTexture2D( UObject* InParent, FName Name, EObjectFlags Flags ) OVERRIDE;
-	virtual UTextureCube* CreateTextureCube( UObject* InParent, FName Name, EObjectFlags Flags ) OVERRIDE;
+	virtual UTexture2D* CreateTexture2D( UObject* InParent, FName Name, EObjectFlags Flags ) override;
+	virtual UTextureCube* CreateTextureCube( UObject* InParent, FName Name, EObjectFlags Flags ) override;
 	// End UTextureFactory Interface
 };
 
diff --git a/Engine/Source/Editor/UnrealEd/Classes/Factories/SlateBrushAssetFactory.h b/Engine/Source/Editor/UnrealEd/Classes/Factories/SlateBrushAssetFactory.h
--- a/Engine/Source/Editor/UnrealEd/Classes/Factories/SlateBrushAssetFactory.h
+++ b/Engine/Source/Editor/UnrealEd/Classes/Factories/SlateBrushAssetFactory.h
@@ -15,9 +15,9 @@ class USlateBrushAssetFactory : public UFactory
 	class UTexture2D* InitialTexture;
 
 	// Begin UFactory Interface
-	virtual FText GetDisplayName() const OVERRIDE;
-	virtual bool ConfigureProperties() OVERRIDE;
-	virtual UObject* FactoryCreateNew(UClass* Class,UObject* InParent,FName Name,EObjectFlags Flags,UObject* Context,FFeedbackContext* Warn) OVERRIDE;
+	virtual FText GetDisplayName() const override;
+	virtual bool ConfigureProperties() override;
+	virtual UObject* FactoryCreateNew(UClass* Class,UObject* InParent,FName Name,EObjectFlags Flags,UObject* Context,FFeedbackContext* Warn) override;
 	// Begin UFactory Interface
 };
 
diff --git a/Engine/Source/Editor/UnrealEd/Classes/Factories/SlateWidgetStyleAssetFactory.h b/Engine/Source/Editor/UnrealEd/Classes/Factories/SlateWidgetStyleAssetFactory.h
--- a/Engine/Source/Editor/UnrealEd/Classes/Factories/SlateWidgetStyleAssetFactory.h
+++ b/Engine/Source/Editor/UnrealEd/Classes/Factories/SlateWidgetStyleAssetFactory.h
@@ -14,9 +14,9 @@ class USlateWidgetStyleAssetFactory : public UFactory
 	TSubclassOf< USlateWidgetStyleContainerBase > StyleType;
 
 	// Begin UFactory Interface
-	virtual FText GetDisplayName() const OVERRIDE;
-	virtual bool ConfigureProperties() OVERRIDE;
-	virtual UObject* FactoryCreateNew(UClass* Class,UObject* InParent,FName Name,EObjectFlags Flags,UObject* Context,FFeedbackContext* Warn) OVERRIDE;
+	virtual FText GetDisplayName() const override;
+	virtual bool ConfigureProperties() override;
+	virtual UObject* FactoryCreateNew(UClass* Class,UObject* InParent,FName Name,EObjectFlags Flags,UObject* Context,FFeedbackContext* Warn) override;
 	// Begin UFactory Interface
 };
 
diff --git a/Engine/Source/Editor/UnrealEd/Classes/Factories/SoundFactory.h b/Engine/Source/Editor/UnrealEd/Classes/Factories/SoundFactory.h
--- a/Engine/Source/Editor/UnrealEd/Classes/Factories/SoundFactory.h
+++ b/Engine/Source/Editor/UnrealEd/Classes/Factories/SoundFactory.h
@@ -38,7 +38,7 @@ class USoundFactory : public UFactory
 
 
 	// Begin UFactory Interface
-	virtual UObject* FactoryCreateBinary( UClass* Class, UObject* InParent, FName Name, EObjectFlags Flags, UObject* Context, const TCHAR* Type, const uint8*& Buffer, const uint8* BufferEnd, FFeedbackContext* Warn ) OVERRIDE;
+	virtual UObject* FactoryCreateBinary( UClass* Class, UObject* InParent, FName Name, EObjectFlags Flags, UObject* Context, const TCHAR* Type, const uint8*& Buffer, const uint8* BufferEnd, FFeedbackContext* Warn ) override;
 	// End UFactory Interface
 	
 	/** Suppresses the import overwrite dialog until one iteration of FactoryCreateBinary completes; this is primarily used for reimporting sounds */
diff --git a/Engine/Source/Editor/UnrealEd/Classes/Factories/SoundSurroundFactory.h b/Engine/Source/Editor/UnrealEd/Classes/Factories/SoundSurroundFactory.h
--- a/Engine/Source/Editor/UnrealEd/Classes/Factories/SoundSurroundFactory.h
+++ b/Engine/Source/Editor/UnrealEd/Classes/Factories/SoundSurroundFactory.h
@@ -17,8 +17,8 @@ class USoundSurroundFactory : public UFactory
 
 
 	// Begin UFactory Interface
-	virtual UObject* FactoryCreateBinary( UClass* Class, UObject* InParent, FName Name, EObjectFlags Flags, UObject* Context, const TCHAR* Type, const uint8*& Buffer, const uint8* BufferEnd, FFeedbackContext* Warn ) OVERRIDE;
-	virtual bool FactoryCanImport( const FString& Filename ) OVERRIDE;
+	virtual UObject* FactoryCreateBinary( UClass* Class, UObject* InParent, FName Name, EObjectFlags Flags, UObject* Context, const TCHAR* Type, const uint8*& Buffer, const uint8* BufferEnd, FFeedbackContext* Warn ) override;
+	virtual bool FactoryCanImport( const FString& Filename ) override;
 	// End UFactory Interface
 };
 
diff --git a/Engine/Source/Editor/UnrealEd/Classes/Factories/Texture2dFactoryNew.h b/Engine/Source/Editor/UnrealEd/Classes/Factories/Texture2dFactoryNew.h
--- a/Engine/Source/Editor/UnrealEd/Classes/Factories/Texture2dFactoryNew.h
+++ b/Engine/Source/Editor/UnrealEd/Classes/Factories/Texture2dFactoryNew.h
@@ -20,6 +20,6 @@ class UTexture2DFactoryNew : public UFactory
 	UPROPERTY()
 	int32		Height;
 
-	virtual bool ShouldShowInNewMenu() const OVERRIDE;
-	virtual UObject* FactoryCreateNew( UClass* InClass, UObject* InParent, FName InName, EObjectFlags Flags, UObject* Context, FFeedbackContext* Warn ) OVERRIDE;
+	virtual bool ShouldShowInNewMenu() const override;
+	virtual UObject* FactoryCreateNew( UClass* InClass, UObject* InParent, FName InName, EObjectFlags Flags, UObject* Context, FFeedbackContext* Warn ) override;
 };
diff --git a/Engine/Source/Editor/UnrealEd/Classes/Factories/TextureCubeThumbnailRenderer.h b/Engine/Source/Editor/UnrealEd/Classes/Factories/TextureCubeThumbnailRenderer.h
--- a/Engine/Source/Editor/UnrealEd/Classes/Factories/TextureCubeThumbnailRenderer.h
+++ b/Engine/Source/Editor/UnrealEd/Classes/Factories/TextureCubeThumbnailRenderer.h
@@ -15,8 +15,8 @@ class UTextureCubeThumbnailRenderer : public UTextureThumbnailRenderer
 
 
 	// Begin UThumbnailRenderer Interface
-	virtual void GetThumbnailSize(UObject* Object, float Zoom, uint32& OutWidth, uint32& OutHeight) const OVERRIDE;
-	virtual void Draw(UObject* Object, int32 X, int32 Y, uint32 Width, uint32 Height, FRenderTarget*, FCanvas* Canvas) OVERRIDE;
+	virtual void GetThumbnailSize(UObject* Object, float Zoom, uint32& OutWidth, uint32& OutHeight) const override;
+	virtual void Draw(UObject* Object, int32 X, int32 Y, uint32 Width, uint32 Height, FRenderTarget*, FCanvas* Canvas) override;
 	// End UThumbnailRenderer Interface
 };
 
diff --git a/Engine/Source/Editor/UnrealEd/Classes/Factories/TextureFactory.h b/Engine/Source/Editor/UnrealEd/Classes/Factories/TextureFactory.h
--- a/Engine/Source/Editor/UnrealEd/Classes/Factories/TextureFactory.h
+++ b/Engine/Source/Editor/UnrealEd/Classes/Factories/TextureFactory.h
@@ -96,8 +96,8 @@ public:
 	// End UObject Interface
 
 	// Begin UFactory Interface
-	virtual bool DoesSupportClass(UClass* Class) OVERRIDE;
-	virtual UObject* FactoryCreateBinary( UClass* Class, UObject* InParent, FName Name, EObjectFlags Flags, UObject* Context, const TCHAR* Type, const uint8*& Buffer, const uint8* BufferEnd, FFeedbackContext* Warn ) OVERRIDE;
+	virtual bool DoesSupportClass(UClass* Class) override;
+	virtual UObject* FactoryCreateBinary( UClass* Class, UObject* InParent, FName Name, EObjectFlags Flags, UObject* Context, const TCHAR* Type, const uint8*& Buffer, const uint8* BufferEnd, FFeedbackContext* Warn ) override;
 	// End UFactory Interface
 	
 
diff --git a/Engine/Source/Editor/UnrealEd/Classes/Factories/TrueTypeFontFactory.h b/Engine/Source/Editor/UnrealEd/Classes/Factories/TrueTypeFontFactory.h
--- a/Engine/Source/Editor/UnrealEd/Classes/Factories/TrueTypeFontFactory.h
+++ b/Engine/Source/Editor/UnrealEd/Classes/Factories/TrueTypeFontFactory.h
@@ -25,18 +25,18 @@ class UTrueTypeFontFactory : public UFontFactory, public FReimportHandler
 	bool bFontSelected;
 
 	// Begin UObject Interface
-	virtual void PostInitProperties() OVERRIDE;
+	virtual void PostInitProperties() override;
 	// End UObject Interface
 
 	// Begin UFactory Interface
-	virtual bool ConfigureProperties() OVERRIDE;
-	virtual UObject* FactoryCreateNew(UClass* Class,UObject* InParent,FName Name,EObjectFlags Flags,UObject* Context,FFeedbackContext* Warn) OVERRIDE;
+	virtual bool ConfigureProperties() override;
+	virtual UObject* FactoryCreateNew(UClass* Class,UObject* InParent,FName Name,EObjectFlags Flags,UObject* Context,FFeedbackContext* Warn) override;
 	// Begin UFactory Interface	
 
 	// Begin FReimportHandler interface
-	virtual bool CanReimport( UObject* Obj, TArray<FString>& OutFilenames ) OVERRIDE;
-	virtual void SetReimportPaths( UObject* Obj, const TArray<FString>& NewReimportPaths ) OVERRIDE;
-	virtual EReimportResult::Type Reimport( UObject* Obj ) OVERRIDE;
+	virtual bool CanReimport( UObject* Obj, TArray<FString>& OutFilenames ) override;
+	virtual void SetReimportPaths( UObject* Obj, const TArray<FString>& NewReimportPaths ) override;
+	virtual EReimportResult::Type Reimport( UObject* Obj ) override;
 	// End FReimportHandler interface
 
 	/** Creates the import options structure for this font */
diff --git a/Engine/Source/Editor/UnrealEd/Classes/Factories/VectorFieldStaticFactory.h b/Engine/Source/Editor/UnrealEd/Classes/Factories/VectorFieldStaticFactory.h
--- a/Engine/Source/Editor/UnrealEd/Classes/Factories/VectorFieldStaticFactory.h
+++ b/Engine/Source/Editor/UnrealEd/Classes/Factories/VectorFieldStaticFactory.h
@@ -15,8 +15,8 @@ class UVectorFieldStaticFactory : public UFactory
 
 
 	// Begin UFactory Interface
-	virtual UObject* FactoryCreateBinary( UClass* Class, UObject* InParent, FName Name, EObjectFlags Flags, UObject* Context, const TCHAR* Type, const uint8*& Buffer, const uint8* BufferEnd, FFeedbackContext* Warn ) OVERRIDE;
-	virtual bool FactoryCanImport( const FString& Filename ) OVERRIDE;
+	virtual UObject* FactoryCreateBinary( UClass* Class, UObject* InParent, FName Name, EObjectFlags Flags, UObject* Context, const TCHAR* Type, const uint8*& Buffer, const uint8* BufferEnd, FFeedbackContext* Warn ) override;
+	virtual bool FactoryCanImport( const FString& Filename ) override;
 	// End UFactory Interface
 
 };
diff --git a/Engine/Source/Editor/UnrealEd/Classes/MaterialEditor/MaterialEditorInstanceConstant.h b/Engine/Source/Editor/UnrealEd/Classes/MaterialEditor/MaterialEditorInstanceConstant.h
--- a/Engine/Source/Editor/UnrealEd/Classes/MaterialEditor/MaterialEditorInstanceConstant.h
+++ b/Engine/Source/Editor/UnrealEd/Classes/MaterialEditor/MaterialEditorInstanceConstant.h
@@ -238,9 +238,9 @@ class UNREALED_API UMaterialEditorInstanceConstant : public UObject
 	uint32 bUseOldStyleMICEditorGroups:1;
 
 	// Begin UObject interface.
-	virtual void PostEditChangeProperty(FPropertyChangedEvent& PropertyChangedEvent) OVERRIDE;
+	virtual void PostEditChangeProperty(FPropertyChangedEvent& PropertyChangedEvent) override;
 #if WITH_EDITOR
-	virtual void PostEditUndo() OVERRIDE;
+	virtual void PostEditUndo() override;
 #endif
 	// End UObject interface.
 
diff --git a/Engine/Source/Editor/UnrealEd/Classes/MaterialEditor/PreviewMaterial.h b/Engine/Source/Editor/UnrealEd/Classes/MaterialEditor/PreviewMaterial.h
--- a/Engine/Source/Editor/UnrealEd/Classes/MaterialEditor/PreviewMaterial.h
+++ b/Engine/Source/Editor/UnrealEd/Classes/MaterialEditor/PreviewMaterial.h
@@ -11,8 +11,8 @@ class UNREALED_API UPreviewMaterial : public UMaterial
 
 
 	// Begin UMaterial interface.
-	virtual FMaterialResource* AllocateResource() OVERRIDE;
-	virtual bool IsAsset()  const OVERRIDE  { return false; }
+	virtual FMaterialResource* AllocateResource() override;
+	virtual bool IsAsset()  const override  { return false; }
 	// End UMaterial interface.
 };
 
diff --git a/Engine/Source/Editor/UnrealEd/Classes/MaterialGraph/MaterialGraphNode.h b/Engine/Source/Editor/UnrealEd/Classes/MaterialGraph/MaterialGraphNode.h
--- a/Engine/Source/Editor/UnrealEd/Classes/MaterialGraph/MaterialGraphNode.h
+++ b/Engine/Source/Editor/UnrealEd/Classes/MaterialGraph/MaterialGraphNode.h
@@ -44,29 +44,29 @@ public:
 	uint32 GetOutputType(const UEdGraphPin* OutputPin);
 
 	// UObject interface
-	virtual void PostEditChangeProperty(FPropertyChangedEvent& PropertyChangedEvent) OVERRIDE;
-	virtual void PostEditImport() OVERRIDE;
-	virtual void PostDuplicate(bool bDuplicateForPIE) OVERRIDE;
+	virtual void PostEditChangeProperty(FPropertyChangedEvent& PropertyChangedEvent) override;
+	virtual void PostEditImport() override;
+	virtual void PostDuplicate(bool bDuplicateForPIE) override;
 	// End of UObject interface
 
 	// Begin UEdGraphNode interface.
-	virtual bool CanPasteHere(const UEdGraph* TargetGraph, const UEdGraphSchema* Schema) const OVERRIDE;
-	virtual FText GetNodeTitle(ENodeTitleType::Type TitleType) const OVERRIDE;
-	virtual FString GetNodeNativeTitle(ENodeTitleType::Type TitleType) const OVERRIDE;
-	virtual FLinearColor GetNodeTitleColor() const OVERRIDE;
-	virtual FString GetTooltip() const OVERRIDE;
-	virtual void PrepareForCopying() OVERRIDE;
-	virtual void GetContextMenuActions(const FGraphNodeContextMenuBuilder& Context) const OVERRIDE;
-	virtual void PostPlacedNewNode() OVERRIDE;
-	virtual void OnRenameNode(const FString& NewName) OVERRIDE;
-	virtual void GetPinHoverText(const UEdGraphPin& Pin, FString& HoverTextOut) const OVERRIDE;
+	virtual bool CanPasteHere(const UEdGraph* TargetGraph, const UEdGraphSchema* Schema) const override;
+	virtual FText GetNodeTitle(ENodeTitleType::Type TitleType) const override;
+	virtual FString GetNodeNativeTitle(ENodeTitleType::Type TitleType) const override;
+	virtual FLinearColor GetNodeTitleColor() const override;
+	virtual FString GetTooltip() const override;
+	virtual void PrepareForCopying() override;
+	virtual void GetContextMenuActions(const FGraphNodeContextMenuBuilder& Context) const override;
+	virtual void PostPlacedNewNode() override;
+	virtual void OnRenameNode(const FString& NewName) override;
+	virtual void GetPinHoverText(const UEdGraphPin& Pin, FString& HoverTextOut) const override;
 	// End UEdGraphNode interface.
 
 	// UMaterialGraphNode_Base interface
-	virtual void CreateInputPins() OVERRIDE;
-	virtual void CreateOutputPins() OVERRIDE;
-	virtual int32 GetInputIndex(const UEdGraphPin* InputPin) const OVERRIDE;
-	virtual uint32 GetInputType(const UEdGraphPin* InputPin) const OVERRIDE;
+	virtual void CreateInputPins() override;
+	virtual void CreateOutputPins() override;
+	virtual int32 GetInputIndex(const UEdGraphPin* InputPin) const override;
+	virtual uint32 GetInputType(const UEdGraphPin* InputPin) const override;
 	// End of UMaterialGraphNode_Base interface
 
 private:
diff --git a/Engine/Source/Editor/UnrealEd/Classes/MaterialGraph/MaterialGraphNode_Base.h b/Engine/Source/Editor/UnrealEd/Classes/MaterialGraph/MaterialGraphNode_Base.h
--- a/Engine/Source/Editor/UnrealEd/Classes/MaterialGraph/MaterialGraphNode_Base.h
+++ b/Engine/Source/Editor/UnrealEd/Classes/MaterialGraph/MaterialGraphNode_Base.h
@@ -41,10 +41,10 @@ class UMaterialGraphNode_Base : public UEdGraphNode
 	void InsertNewNode(UEdGraphPin* FromPin, UEdGraphPin* NewLinkPin, TSet<UEdGraphNode*>& OutNodeList);
 
 	// Begin UEdGraphNode interface.
-	virtual void AllocateDefaultPins() OVERRIDE;
-	virtual void ReconstructNode() OVERRIDE;
-	virtual void AutowireNewNode(UEdGraphPin* FromPin) OVERRIDE;
-	virtual bool CanCreateUnderSpecifiedSchema(const UEdGraphSchema* Schema) const OVERRIDE;
+	virtual void AllocateDefaultPins() override;
+	virtual void ReconstructNode() override;
+	virtual void AutowireNewNode(UEdGraphPin* FromPin) override;
+	virtual bool CanCreateUnderSpecifiedSchema(const UEdGraphSchema* Schema) const override;
 	// End UEdGraphNode interface.
 
 protected:
diff --git a/Engine/Source/Editor/UnrealEd/Classes/MaterialGraph/MaterialGraphNode_Comment.h b/Engine/Source/Editor/UnrealEd/Classes/MaterialGraph/MaterialGraphNode_Comment.h
--- a/Engine/Source/Editor/UnrealEd/Classes/MaterialGraph/MaterialGraphNode_Comment.h
+++ b/Engine/Source/Editor/UnrealEd/Classes/MaterialGraph/MaterialGraphNode_Comment.h
@@ -20,18 +20,18 @@ class UMaterialGraphNode_Comment : public UEdGraphNode_Comment
 	UNREALED_API void PostCopyNode();
 
 	// UObject interface
-	virtual void PostEditChangeProperty(FPropertyChangedEvent& PropertyChangedEvent) OVERRIDE;
-	virtual void PostEditImport() OVERRIDE;
-	virtual void PostDuplicate(bool bDuplicateForPIE) OVERRIDE;
+	virtual void PostEditChangeProperty(FPropertyChangedEvent& PropertyChangedEvent) override;
+	virtual void PostEditImport() override;
+	virtual void PostDuplicate(bool bDuplicateForPIE) override;
 	// End of UObject interface
 
 	// Begin UEdGraphNode interface.
-	virtual void PrepareForCopying() OVERRIDE;
-	virtual void GetContextMenuActions(const FGraphNodeContextMenuBuilder& Context) const OVERRIDE;
-	virtual bool CanCreateUnderSpecifiedSchema(const UEdGraphSchema* Schema) const OVERRIDE;
-	virtual void PostPlacedNewNode() OVERRIDE;
-	virtual void OnRenameNode(const FString& NewName) OVERRIDE;
-	virtual void ResizeNode(const FVector2D& NewSize) OVERRIDE;
+	virtual void PrepareForCopying() override;
+	virtual void GetContextMenuActions(const FGraphNodeContextMenuBuilder& Context) const override;
+	virtual bool CanCreateUnderSpecifiedSchema(const UEdGraphSchema* Schema) const override;
+	virtual void PostPlacedNewNode() override;
+	virtual void OnRenameNode(const FString& NewName) override;
+	virtual void ResizeNode(const FVector2D& NewSize) override;
 	// End UEdGraphNode interface.
 
 private:
diff --git a/Engine/Source/Editor/UnrealEd/Classes/MaterialGraph/MaterialGraphNode_Root.h b/Engine/Source/Editor/UnrealEd/Classes/MaterialGraph/MaterialGraphNode_Root.h
--- a/Engine/Source/Editor/UnrealEd/Classes/MaterialGraph/MaterialGraphNode_Root.h
+++ b/Engine/Source/Editor/UnrealEd/Classes/MaterialGraph/MaterialGraphNode_Root.h
@@ -14,18 +14,18 @@ class UMaterialGraphNode_Root : public UMaterialGraphNode_Base
 	class UMaterial* Material;
 
 	// Begin UEdGraphNode interface.
-	virtual FText GetNodeTitle(ENodeTitleType::Type TitleType) const OVERRIDE;
-	virtual FLinearColor GetNodeTitleColor() const OVERRIDE;
-	virtual FString GetTooltip() const OVERRIDE;
-	virtual bool CanUserDeleteNode() const OVERRIDE { return false; }
-	virtual bool CanDuplicateNode() const OVERRIDE { return false; }
-	virtual void PostPlacedNewNode() OVERRIDE;
+	virtual FText GetNodeTitle(ENodeTitleType::Type TitleType) const override;
+	virtual FLinearColor GetNodeTitleColor() const override;
+	virtual FString GetTooltip() const override;
+	virtual bool CanUserDeleteNode() const override { return false; }
+	virtual bool CanDuplicateNode() const override { return false; }
+	virtual void PostPlacedNewNode() override;
 	// End UEdGraphNode interface.
 
 	// UMaterialGraphNode_Base interface
-	virtual void CreateInputPins() OVERRIDE;
-	virtual bool IsRootNode() const OVERRIDE {return true;}
-	virtual int32 GetInputIndex(const UEdGraphPin* InputPin) const OVERRIDE;
-	virtual uint32 GetInputType(const UEdGraphPin* InputPin) const OVERRIDE;
+	virtual void CreateInputPins() override;
+	virtual bool IsRootNode() const override {return true;}
+	virtual int32 GetInputIndex(const UEdGraphPin* InputPin) const override;
+	virtual uint32 GetInputType(const UEdGraphPin* InputPin) const override;
 	// End of UMaterialGraphNode_Base interface
 };
diff --git a/Engine/Source/Editor/UnrealEd/Classes/MaterialGraph/MaterialGraphSchema.h b/Engine/Source/Editor/UnrealEd/Classes/MaterialGraph/MaterialGraphSchema.h
--- a/Engine/Source/Editor/UnrealEd/Classes/MaterialGraph/MaterialGraphSchema.h
+++ b/Engine/Source/Editor/UnrealEd/Classes/MaterialGraph/MaterialGraphSchema.h
@@ -16,7 +16,7 @@ struct UNREALED_API FMaterialGraphSchemaAction_NewNode : public FEdGraphSchemaAc
 
 	// Simple type info
 	static FString StaticGetTypeId() {static FString Type = TEXT("FMaterialGraphSchemaAction_NewNode"); return Type;}
-	virtual FString GetTypeId() const OVERRIDE { return StaticGetTypeId(); } 
+	virtual FString GetTypeId() const override { return StaticGetTypeId(); } 
 
 	FMaterialGraphSchemaAction_NewNode() 
 		: FEdGraphSchemaAction()
@@ -29,7 +29,7 @@ struct UNREALED_API FMaterialGraphSchemaAction_NewNode : public FEdGraphSchemaAc
 	{}
 
 	// FEdGraphSchemaAction interface
-	virtual UEdGraphNode* PerformAction(class UEdGraph* ParentGraph, UEdGraphPin* FromPin, const FVector2D Location, bool bSelectNewNode = true) OVERRIDE;
+	virtual UEdGraphNode* PerformAction(class UEdGraph* ParentGraph, UEdGraphPin* FromPin, const FVector2D Location, bool bSelectNewNode = true) override;
 	// End of FEdGraphSchemaAction interface
 
 	/**
@@ -53,7 +53,7 @@ struct UNREALED_API FMaterialGraphSchemaAction_NewFunctionCall : public FEdGraph
 
 	// Simple type info
 	static FString StaticGetTypeId() {static FString Type = TEXT("FMaterialGraphSchemaAction_NewFunctionCall"); return Type;}
-	virtual FString GetTypeId() const OVERRIDE { return StaticGetTypeId(); } 
+	virtual FString GetTypeId() const override { return StaticGetTypeId(); } 
 
 	FMaterialGraphSchemaAction_NewFunctionCall() 
 		: FEdGraphSchemaAction()
@@ -64,7 +64,7 @@ struct UNREALED_API FMaterialGraphSchemaAction_NewFunctionCall : public FEdGraph
 	{}
 
 	// FEdGraphSchemaAction interface
-	virtual UEdGraphNode* PerformAction(class UEdGraph* ParentGraph, UEdGraphPin* FromPin, const FVector2D Location, bool bSelectNewNode = true) OVERRIDE;
+	virtual UEdGraphNode* PerformAction(class UEdGraph* ParentGraph, UEdGraphPin* FromPin, const FVector2D Location, bool bSelectNewNode = true) override;
 	// End of FEdGraphSchemaAction interface
 };
 
@@ -76,7 +76,7 @@ struct UNREALED_API FMaterialGraphSchemaAction_NewComment : public FEdGraphSchem
 
 	// Simple type info
 	static FString StaticGetTypeId() {static FString Type = TEXT("FMaterialGraphSchemaAction_NewComment"); return Type;}
-	virtual FString GetTypeId() const OVERRIDE { return StaticGetTypeId(); } 
+	virtual FString GetTypeId() const override { return StaticGetTypeId(); } 
 
 	FMaterialGraphSchemaAction_NewComment() 
 		: FEdGraphSchemaAction()
@@ -87,7 +87,7 @@ struct UNREALED_API FMaterialGraphSchemaAction_NewComment : public FEdGraphSchem
 	{}
 
 	// FEdGraphSchemaAction interface
-	virtual UEdGraphNode* PerformAction(class UEdGraph* ParentGraph, UEdGraphPin* FromPin, const FVector2D Location, bool bSelectNewNode = true) OVERRIDE;
+	virtual UEdGraphNode* PerformAction(class UEdGraph* ParentGraph, UEdGraphPin* FromPin, const FVector2D Location, bool bSelectNewNode = true) override;
 	// End of FEdGraphSchemaAction interface
 };
 
@@ -99,7 +99,7 @@ struct UNREALED_API FMaterialGraphSchemaAction_Paste : public FEdGraphSchemaActi
 
 	// Simple type info
 	static FString StaticGetTypeId() {static FString Type = TEXT("FMaterialGraphSchemaAction_Paste"); return Type;}
-	virtual FString GetTypeId() const OVERRIDE { return StaticGetTypeId(); } 
+	virtual FString GetTypeId() const override { return StaticGetTypeId(); } 
 
 	FMaterialGraphSchemaAction_Paste() 
 		: FEdGraphSchemaAction()
@@ -110,7 +110,7 @@ struct UNREALED_API FMaterialGraphSchemaAction_Paste : public FEdGraphSchemaActi
 	{}
 
 	// FEdGraphSchemaAction interface
-	virtual UEdGraphNode* PerformAction(class UEdGraph* ParentGraph, UEdGraphPin* FromPin, const FVector2D Location, bool bSelectNewNode = true) OVERRIDE;
+	virtual UEdGraphNode* PerformAction(class UEdGraph* ParentGraph, UEdGraphPin* FromPin, const FVector2D Location, bool bSelectNewNode = true) override;
 	// End of FEdGraphSchemaAction interface
 };
 
@@ -191,18 +191,18 @@ class UMaterialGraphSchema : public UEdGraphSchema
 	UNREALED_API static uint32 GetMaterialValueType(const UEdGraphPin* MaterialPin);
 
 	// Begin UEdGraphSchema interface
-	virtual void GetGraphContextActions(FGraphContextMenuBuilder& ContextMenuBuilder) const OVERRIDE;
-	virtual void GetContextMenuActions(const UEdGraph* CurrentGraph, const UEdGraphNode* InGraphNode, const UEdGraphPin* InGraphPin, class FMenuBuilder* MenuBuilder, bool bIsDebugging) const OVERRIDE;
-	virtual const FPinConnectionResponse CanCreateConnection(const UEdGraphPin* A, const UEdGraphPin* B) const OVERRIDE;
-	virtual bool TryCreateConnection(UEdGraphPin* A, UEdGraphPin* B) const OVERRIDE;
-	virtual bool ShouldHidePinDefaultValue(UEdGraphPin* Pin) const OVERRIDE { return true; }
-	virtual FLinearColor GetPinTypeColor(const FEdGraphPinType& PinType) const OVERRIDE;
-	virtual void BreakNodeLinks(UEdGraphNode& TargetNode) const OVERRIDE;
-	virtual void BreakPinLinks(UEdGraphPin& TargetPin, bool bSendsNodeNotifcation) const OVERRIDE;
-	virtual void BreakSinglePinLink(UEdGraphPin* SourcePin, UEdGraphPin* TargetPin) OVERRIDE;
-	virtual void DroppedAssetsOnGraph(const TArray<class FAssetData>& Assets, const FVector2D& GraphPosition, UEdGraph* Graph) const OVERRIDE;
-	virtual int32 GetNodeSelectionCount(const UEdGraph* Graph) const OVERRIDE;
-	virtual TSharedPtr<FEdGraphSchemaAction> GetCreateCommentAction() const OVERRIDE;
+	virtual void GetGraphContextActions(FGraphContextMenuBuilder& ContextMenuBuilder) const override;
+	virtual void GetContextMenuActions(const UEdGraph* CurrentGraph, const UEdGraphNode* InGraphNode, const UEdGraphPin* InGraphPin, class FMenuBuilder* MenuBuilder, bool bIsDebugging) const override;
+	virtual const FPinConnectionResponse CanCreateConnection(const UEdGraphPin* A, const UEdGraphPin* B) const override;
+	virtual bool TryCreateConnection(UEdGraphPin* A, UEdGraphPin* B) const override;
+	virtual bool ShouldHidePinDefaultValue(UEdGraphPin* Pin) const override { return true; }
+	virtual FLinearColor GetPinTypeColor(const FEdGraphPinType& PinType) const override;
+	virtual void BreakNodeLinks(UEdGraphNode& TargetNode) const override;
+	virtual void BreakPinLinks(UEdGraphPin& TargetPin, bool bSendsNodeNotifcation) const override;
+	virtual void BreakSinglePinLink(UEdGraphPin* SourcePin, UEdGraphPin* TargetPin) override;
+	virtual void DroppedAssetsOnGraph(const TArray<class FAssetData>& Assets, const FVector2D& GraphPosition, UEdGraph* Graph) const override;
+	virtual int32 GetNodeSelectionCount(const UEdGraph* Graph) const override;
+	virtual TSharedPtr<FEdGraphSchemaAction> GetCreateCommentAction() const override;
 	// End UEdGraphSchema interface
 
 private:
diff --git a/Engine/Source/Editor/UnrealEd/Classes/Niagara/EdGraphSchema_Niagara.h b/Engine/Source/Editor/UnrealEd/Classes/Niagara/EdGraphSchema_Niagara.h
--- a/Engine/Source/Editor/UnrealEd/Classes/Niagara/EdGraphSchema_Niagara.h
+++ b/Engine/Source/Editor/UnrealEd/Classes/Niagara/EdGraphSchema_Niagara.h
@@ -26,9 +26,9 @@ struct UNREALED_API FNiagaraSchemaAction_NewNode : public FEdGraphSchemaAction
 	{}
 
 	// FEdGraphSchemaAction interface
-	virtual UEdGraphNode* PerformAction(class UEdGraph* ParentGraph, UEdGraphPin* FromPin, const FVector2D Location, bool bSelectNewNode = true) OVERRIDE;
-	virtual UEdGraphNode* PerformAction(class UEdGraph* ParentGraph, TArray<UEdGraphPin*>& FromPins, const FVector2D Location, bool bSelectNewNode = true) OVERRIDE;
-	virtual void AddReferencedObjects( FReferenceCollector& Collector ) OVERRIDE;
+	virtual UEdGraphNode* PerformAction(class UEdGraph* ParentGraph, UEdGraphPin* FromPin, const FVector2D Location, bool bSelectNewNode = true) override;
+	virtual UEdGraphNode* PerformAction(class UEdGraph* ParentGraph, TArray<UEdGraphPin*>& FromPins, const FVector2D Location, bool bSelectNewNode = true) override;
+	virtual void AddReferencedObjects( FReferenceCollector& Collector ) override;
 	// End of FEdGraphSchemaAction interface
 
 	template <typename NodeType>
@@ -51,10 +51,10 @@ class UEdGraphSchema_Niagara : public UEdGraphSchema
 	FString PC_Float;
 
 	// Begin EdGraphSchema interface
-	virtual void GetGraphContextActions(FGraphContextMenuBuilder& ContextMenuBuilder) const OVERRIDE;
-	virtual const FPinConnectionResponse CanCreateConnection(const UEdGraphPin* A, const UEdGraphPin* B) const OVERRIDE;
-	virtual FLinearColor GetPinTypeColor(const FEdGraphPinType& PinType) const OVERRIDE;
-	virtual bool ShouldHidePinDefaultValue(UEdGraphPin* Pin) const OVERRIDE;
+	virtual void GetGraphContextActions(FGraphContextMenuBuilder& ContextMenuBuilder) const override;
+	virtual const FPinConnectionResponse CanCreateConnection(const UEdGraphPin* A, const UEdGraphPin* B) const override;
+	virtual FLinearColor GetPinTypeColor(const FEdGraphPinType& PinType) const override;
+	virtual bool ShouldHidePinDefaultValue(UEdGraphPin* Pin) const override;
 	// End EdGraphSchema interface
 };
 
diff --git a/Engine/Source/Editor/UnrealEd/Classes/Niagara/NiagaraNodeGetAttr.h b/Engine/Source/Editor/UnrealEd/Classes/Niagara/NiagaraNodeGetAttr.h
--- a/Engine/Source/Editor/UnrealEd/Classes/Niagara/NiagaraNodeGetAttr.h
+++ b/Engine/Source/Editor/UnrealEd/Classes/Niagara/NiagaraNodeGetAttr.h
@@ -13,10 +13,10 @@ class UNiagaraNodeGetAttr : public UNiagaraNode
 	FName	AttrName;
 
 	// Begin EdGraphNode interface
-	virtual void AllocateDefaultPins() OVERRIDE;
-	virtual FText GetNodeTitle(ENodeTitleType::Type TitleType) const OVERRIDE;
-	virtual FString GetNodeNativeTitle(ENodeTitleType::Type TitleType) const OVERRIDE;
-	virtual FLinearColor GetNodeTitleColor() const OVERRIDE;
+	virtual void AllocateDefaultPins() override;
+	virtual FText GetNodeTitle(ENodeTitleType::Type TitleType) const override;
+	virtual FString GetNodeNativeTitle(ENodeTitleType::Type TitleType) const override;
+	virtual FLinearColor GetNodeTitleColor() const override;
 	// End EdGraphNode interface
 };
 
diff --git a/Engine/Source/Editor/UnrealEd/Classes/Niagara/NiagaraNodeOp.h b/Engine/Source/Editor/UnrealEd/Classes/Niagara/NiagaraNodeOp.h
--- a/Engine/Source/Editor/UnrealEd/Classes/Niagara/NiagaraNodeOp.h
+++ b/Engine/Source/Editor/UnrealEd/Classes/Niagara/NiagaraNodeOp.h
@@ -13,9 +13,9 @@ class UNiagaraNodeOp : public UNiagaraNode
 	uint8 OpIndex;
 
 	// Begin EdGraphNode interface
-	virtual void AllocateDefaultPins() OVERRIDE;
-	virtual FText GetNodeTitle(ENodeTitleType::Type TitleType) const OVERRIDE;
-	virtual FLinearColor GetNodeTitleColor() const OVERRIDE;
+	virtual void AllocateDefaultPins() override;
+	virtual FText GetNodeTitle(ENodeTitleType::Type TitleType) const override;
+	virtual FLinearColor GetNodeTitleColor() const override;
 	// End EdGraphNode interface
 
 	/** Names to use for each input pin. */
diff --git a/Engine/Source/Editor/UnrealEd/Classes/Niagara/NiagaraNodeOutputUpdate.h b/Engine/Source/Editor/UnrealEd/Classes/Niagara/NiagaraNodeOutputUpdate.h
--- a/Engine/Source/Editor/UnrealEd/Classes/Niagara/NiagaraNodeOutputUpdate.h
+++ b/Engine/Source/Editor/UnrealEd/Classes/Niagara/NiagaraNodeOutputUpdate.h
@@ -9,10 +9,10 @@ class UNiagaraNodeOutputUpdate : public UNiagaraNode
 	GENERATED_UCLASS_BODY()
 
 	// Begin EdGraphNode interface
-	virtual void AllocateDefaultPins() OVERRIDE;
-	virtual bool CanUserDeleteNode() const OVERRIDE { return false; }
-	virtual FText GetNodeTitle(ENodeTitleType::Type TitleType) const OVERRIDE;
-	virtual FLinearColor GetNodeTitleColor() const OVERRIDE;
+	virtual void AllocateDefaultPins() override;
+	virtual bool CanUserDeleteNode() const override { return false; }
+	virtual FText GetNodeTitle(ENodeTitleType::Type TitleType) const override;
+	virtual FLinearColor GetNodeTitleColor() const override;
 	// End EdGraphNode interface
 };
 
diff --git a/Engine/Source/Editor/UnrealEd/Classes/Niagara/NiagaraScriptSource.h b/Engine/Source/Editor/UnrealEd/Classes/Niagara/NiagaraScriptSource.h
--- a/Engine/Source/Editor/UnrealEd/Classes/Niagara/NiagaraScriptSource.h
+++ b/Engine/Source/Editor/UnrealEd/Classes/Niagara/NiagaraScriptSource.h
@@ -15,7 +15,7 @@ class UNiagaraScriptSource : public UNiagaraScriptSourceBase
 	class UNiagaraGraph*	UpdateGraph;
 
 	// UObject interface.
-	virtual void PostLoad() OVERRIDE;
+	virtual void PostLoad() override;
 
 	/** Get the set of outputs from the Update script */
 	UNREALED_API void GetUpdateOutputs(TArray<FName>& ScalarOutputs);
diff --git a/Engine/Source/Editor/UnrealEd/Classes/Preferences/UnrealEdOptions.h b/Engine/Source/Editor/UnrealEd/Classes/Preferences/UnrealEdOptions.h
--- a/Engine/Source/Editor/UnrealEd/Classes/Preferences/UnrealEdOptions.h
+++ b/Engine/Source/Editor/UnrealEd/Classes/Preferences/UnrealEdOptions.h
@@ -106,7 +106,7 @@ public:
 	TMap<FName, int32>	CommandMap;
 
 	// Begin UObject Interface
-	virtual void PostInitProperties() OVERRIDE;
+	virtual void PostInitProperties() override;
 	// End UObject Interface
 
 	/**
diff --git a/Engine/Source/Editor/UnrealEd/Classes/Settings/ContentBrowserSettings.h b/Engine/Source/Editor/UnrealEd/Classes/Settings/ContentBrowserSettings.h
--- a/Engine/Source/Editor/UnrealEd/Classes/Settings/ContentBrowserSettings.h
+++ b/Engine/Source/Editor/UnrealEd/Classes/Settings/ContentBrowserSettings.h
@@ -74,7 +74,7 @@ protected:
 
 	// Begin UObject overrides
 
-	virtual void PostEditChangeProperty( struct FPropertyChangedEvent& PropertyChangedEvent ) OVERRIDE;
+	virtual void PostEditChangeProperty( struct FPropertyChangedEvent& PropertyChangedEvent ) override;
 
 	// End UObject overrides
 
diff --git a/Engine/Source/Editor/UnrealEd/Classes/Settings/EditorExperimentalSettings.h b/Engine/Source/Editor/UnrealEd/Classes/Settings/EditorExperimentalSettings.h
--- a/Engine/Source/Editor/UnrealEd/Classes/Settings/EditorExperimentalSettings.h
+++ b/Engine/Source/Editor/UnrealEd/Classes/Settings/EditorExperimentalSettings.h
@@ -76,7 +76,7 @@ protected:
 
 	// Begin UObject overrides
 
-	virtual void PostEditChangeProperty( struct FPropertyChangedEvent& PropertyChangedEvent ) OVERRIDE;
+	virtual void PostEditChangeProperty( struct FPropertyChangedEvent& PropertyChangedEvent ) override;
 
 	// End UObject overrides
 
diff --git a/Engine/Source/Editor/UnrealEd/Classes/SoundClassGraph/SoundClassGraphNode.h b/Engine/Source/Editor/UnrealEd/Classes/SoundClassGraph/SoundClassGraphNode.h
--- a/Engine/Source/Editor/UnrealEd/Classes/SoundClassGraph/SoundClassGraphNode.h
+++ b/Engine/Source/Editor/UnrealEd/Classes/SoundClassGraph/SoundClassGraphNode.h
@@ -20,12 +20,12 @@ class USoundClassGraphNode : public UEdGraphNode
 	bool CheckRepresentsSoundClass();
 
 	// Begin UEdGraphNode interface.
-	virtual FLinearColor GetNodeTitleColor() const OVERRIDE;
-	virtual void AllocateDefaultPins() OVERRIDE;
-	virtual void AutowireNewNode(UEdGraphPin* FromPin) OVERRIDE;
-	virtual bool CanCreateUnderSpecifiedSchema(const UEdGraphSchema* Schema) const OVERRIDE;
-	virtual FText GetNodeTitle(ENodeTitleType::Type TitleType) const OVERRIDE;
-	virtual bool CanUserDeleteNode() const OVERRIDE;
+	virtual FLinearColor GetNodeTitleColor() const override;
+	virtual void AllocateDefaultPins() override;
+	virtual void AutowireNewNode(UEdGraphPin* FromPin) override;
+	virtual bool CanCreateUnderSpecifiedSchema(const UEdGraphSchema* Schema) const override;
+	virtual FText GetNodeTitle(ENodeTitleType::Type TitleType) const override;
+	virtual bool CanUserDeleteNode() const override;
 	// End UEdGraphNode interface.
 
 private:
diff --git a/Engine/Source/Editor/UnrealEd/Classes/SoundClassGraph/SoundClassGraphSchema.h b/Engine/Source/Editor/UnrealEd/Classes/SoundClassGraph/SoundClassGraphSchema.h
--- a/Engine/Source/Editor/UnrealEd/Classes/SoundClassGraph/SoundClassGraphSchema.h
+++ b/Engine/Source/Editor/UnrealEd/Classes/SoundClassGraph/SoundClassGraphSchema.h
@@ -23,8 +23,8 @@ struct UNREALED_API FSoundClassGraphSchemaAction_NewNode : public FEdGraphSchema
 	{}
 
 	// FEdGraphSchemaAction interface
-	virtual FString GetTypeId() const OVERRIDE { return StaticGetTypeId(); } 
-	virtual UEdGraphNode* PerformAction(class UEdGraph* ParentGraph, UEdGraphPin* FromPin, const FVector2D Location, bool bSelectNewNode = true) OVERRIDE;
+	virtual FString GetTypeId() const override { return StaticGetTypeId(); } 
+	virtual UEdGraphNode* PerformAction(class UEdGraph* ParentGraph, UEdGraphPin* FromPin, const FVector2D Location, bool bSelectNewNode = true) override;
 	// End of FEdGraphSchemaAction interface
 
 	/** Name for the new SoundClass */
@@ -42,16 +42,16 @@ class USoundClassGraphSchema : public UEdGraphSchema
 	void GetBreakLinkToSubMenuActions(class FMenuBuilder& MenuBuilder, class UEdGraphPin* InGraphPin);
 
 	// Begin EdGraphSchema interface
-	virtual void GetGraphContextActions(FGraphContextMenuBuilder& ContextMenuBuilder) const OVERRIDE;
-	virtual void GetContextMenuActions(const UEdGraph* CurrentGraph, const UEdGraphNode* InGraphNode, const UEdGraphPin* InGraphPin, class FMenuBuilder* MenuBuilder, bool bIsDebugging) const OVERRIDE;
-	virtual const FPinConnectionResponse CanCreateConnection(const UEdGraphPin* PinA, const UEdGraphPin* PinB) const OVERRIDE;
-	virtual bool TryCreateConnection(UEdGraphPin* PinA, UEdGraphPin* PinB) const OVERRIDE;
-	virtual bool ShouldHidePinDefaultValue(UEdGraphPin* Pin) const OVERRIDE;
-	virtual FLinearColor GetPinTypeColor(const FEdGraphPinType& PinType) const OVERRIDE;
-	virtual void BreakNodeLinks(UEdGraphNode& TargetNode) const OVERRIDE;
-	virtual void BreakPinLinks(UEdGraphPin& TargetPin, bool bSendsNodeNotifcation) const OVERRIDE;
-	virtual void BreakSinglePinLink(UEdGraphPin* SourcePin, UEdGraphPin* TargetPin) OVERRIDE;
-	virtual void DroppedAssetsOnGraph(const TArray<class FAssetData>& Assets, const FVector2D& GraphPosition, UEdGraph* Graph) const OVERRIDE;
+	virtual void GetGraphContextActions(FGraphContextMenuBuilder& ContextMenuBuilder) const override;
+	virtual void GetContextMenuActions(const UEdGraph* CurrentGraph, const UEdGraphNode* InGraphNode, const UEdGraphPin* InGraphPin, class FMenuBuilder* MenuBuilder, bool bIsDebugging) const override;
+	virtual const FPinConnectionResponse CanCreateConnection(const UEdGraphPin* PinA, const UEdGraphPin* PinB) const override;
+	virtual bool TryCreateConnection(UEdGraphPin* PinA, UEdGraphPin* PinB) const override;
+	virtual bool ShouldHidePinDefaultValue(UEdGraphPin* Pin) const override;
+	virtual FLinearColor GetPinTypeColor(const FEdGraphPinType& PinType) const override;
+	virtual void BreakNodeLinks(UEdGraphNode& TargetNode) const override;
+	virtual void BreakPinLinks(UEdGraphPin& TargetPin, bool bSendsNodeNotifcation) const override;
+	virtual void BreakSinglePinLink(UEdGraphPin* SourcePin, UEdGraphPin* TargetPin) override;
+	virtual void DroppedAssetsOnGraph(const TArray<class FAssetData>& Assets, const FVector2D& GraphPosition, UEdGraph* Graph) const override;
 	// End EdGraphSchema interface
 };
 
diff --git a/Engine/Source/Editor/UnrealEd/Classes/SoundCueGraph/SoundCueGraphNode.h b/Engine/Source/Editor/UnrealEd/Classes/SoundCueGraph/SoundCueGraphNode.h
--- a/Engine/Source/Editor/UnrealEd/Classes/SoundCueGraph/SoundCueGraphNode.h
+++ b/Engine/Source/Editor/UnrealEd/Classes/SoundCueGraph/SoundCueGraphNode.h
@@ -29,21 +29,21 @@ class USoundCueGraphNode : public USoundCueGraphNode_Base
 	UNREALED_API bool CanAddInputPin() const;
 
 	// USoundCueGraphNode_Base interface
-	virtual void CreateInputPins() OVERRIDE;
+	virtual void CreateInputPins() override;
 	// End of USoundCueGraphNode_Base interface
 
 	// UEdGraphNode interface
-	virtual FText GetNodeTitle(ENodeTitleType::Type TitleType) const OVERRIDE;
-	virtual void PrepareForCopying() OVERRIDE;
-	virtual void GetContextMenuActions(const FGraphNodeContextMenuBuilder& Context) const OVERRIDE;
-	virtual FString GetTooltip() const OVERRIDE;
-	virtual FString GetDocumentationExcerptName() const OVERRIDE;
+	virtual FText GetNodeTitle(ENodeTitleType::Type TitleType) const override;
+	virtual void PrepareForCopying() override;
+	virtual void GetContextMenuActions(const FGraphNodeContextMenuBuilder& Context) const override;
+	virtual FString GetTooltip() const override;
+	virtual FString GetDocumentationExcerptName() const override;
 	// End of UEdGraphNode interface
 
 	// UObject interface
-	virtual void PostLoad() OVERRIDE;
-	virtual void PostEditImport() OVERRIDE;
-	virtual void PostDuplicate(bool bDuplicateForPIE) OVERRIDE;
+	virtual void PostLoad() override;
+	virtual void PostEditImport() override;
+	virtual void PostDuplicate(bool bDuplicateForPIE) override;
 	// End of UObject interface
 
 private:
diff --git a/Engine/Source/Editor/UnrealEd/Classes/SoundCueGraph/SoundCueGraphNode_Base.h b/Engine/Source/Editor/UnrealEd/Classes/SoundCueGraph/SoundCueGraphNode_Base.h
--- a/Engine/Source/Editor/UnrealEd/Classes/SoundCueGraph/SoundCueGraphNode_Base.h
+++ b/Engine/Source/Editor/UnrealEd/Classes/SoundCueGraph/SoundCueGraphNode_Base.h
@@ -33,10 +33,10 @@ class USoundCueGraphNode_Base : public UEdGraphNode
 	void InsertNewNode(UEdGraphPin* FromPin, UEdGraphPin* NewLinkPin, TSet<UEdGraphNode*>& OutNodeList);
 
 	// UEdGraphNode interface.
-	virtual void AllocateDefaultPins() OVERRIDE;
-	virtual void ReconstructNode() OVERRIDE;
-	virtual void AutowireNewNode(UEdGraphPin* FromPin) OVERRIDE;
-	virtual bool CanCreateUnderSpecifiedSchema(const UEdGraphSchema* Schema) const OVERRIDE;
-	virtual FString GetDocumentationLink() const OVERRIDE;
+	virtual void AllocateDefaultPins() override;
+	virtual void ReconstructNode() override;
+	virtual void AutowireNewNode(UEdGraphPin* FromPin) override;
+	virtual bool CanCreateUnderSpecifiedSchema(const UEdGraphSchema* Schema) const override;
+	virtual FString GetDocumentationLink() const override;
 	// End of UEdGraphNode interface.
 };
diff --git a/Engine/Source/Editor/UnrealEd/Classes/SoundCueGraph/SoundCueGraphNode_Root.h b/Engine/Source/Editor/UnrealEd/Classes/SoundCueGraph/SoundCueGraphNode_Root.h
--- a/Engine/Source/Editor/UnrealEd/Classes/SoundCueGraph/SoundCueGraphNode_Root.h
+++ b/Engine/Source/Editor/UnrealEd/Classes/SoundCueGraph/SoundCueGraphNode_Root.h
@@ -9,16 +9,16 @@ class USoundCueGraphNode_Root : public USoundCueGraphNode_Base
 	GENERATED_UCLASS_BODY()
 
 	// UEdGraphNode interface
-	virtual FLinearColor GetNodeTitleColor() const OVERRIDE;
-	virtual FString GetTooltip() const OVERRIDE;
-	virtual FText GetNodeTitle(ENodeTitleType::Type TitleType) const OVERRIDE;
-	virtual bool CanUserDeleteNode() const OVERRIDE { return false; }
-	virtual bool CanDuplicateNode() const OVERRIDE { return false; }
-	virtual void GetContextMenuActions(const FGraphNodeContextMenuBuilder& Context) const OVERRIDE;
+	virtual FLinearColor GetNodeTitleColor() const override;
+	virtual FString GetTooltip() const override;
+	virtual FText GetNodeTitle(ENodeTitleType::Type TitleType) const override;
+	virtual bool CanUserDeleteNode() const override { return false; }
+	virtual bool CanDuplicateNode() const override { return false; }
+	virtual void GetContextMenuActions(const FGraphNodeContextMenuBuilder& Context) const override;
 	// End of UEdGraphNode interface
 
 	// USoundCueGraphNode_Base interface
-	virtual void CreateInputPins() OVERRIDE;
-	virtual bool IsRootNode() const OVERRIDE {return true;}
+	virtual void CreateInputPins() override;
+	virtual bool IsRootNode() const override {return true;}
 	// End of USoundCueGraphNode_Base interface
 };
diff --git a/Engine/Source/Editor/UnrealEd/Classes/SoundCueGraph/SoundCueGraphSchema.h b/Engine/Source/Editor/UnrealEd/Classes/SoundCueGraph/SoundCueGraphSchema.h
--- a/Engine/Source/Editor/UnrealEd/Classes/SoundCueGraph/SoundCueGraphSchema.h
+++ b/Engine/Source/Editor/UnrealEd/Classes/SoundCueGraph/SoundCueGraphSchema.h
@@ -28,7 +28,7 @@ struct UNREALED_API FSoundCueGraphSchemaAction_NewNode : public FEdGraphSchemaAc
 	{}
 
 	// FEdGraphSchemaAction interface
-	virtual UEdGraphNode* PerformAction(class UEdGraph* ParentGraph, UEdGraphPin* FromPin, const FVector2D Location, bool bSelectNewNode = true) OVERRIDE;
+	virtual UEdGraphNode* PerformAction(class UEdGraph* ParentGraph, UEdGraphPin* FromPin, const FVector2D Location, bool bSelectNewNode = true) override;
 	// End of FEdGraphSchemaAction interface
 
 private:
@@ -51,7 +51,7 @@ struct UNREALED_API FSoundCueGraphSchemaAction_NewFromSelected : public FSoundCu
 	{}
 
 	// FEdGraphSchemaAction interface
-	virtual UEdGraphNode* PerformAction(class UEdGraph* ParentGraph, UEdGraphPin* FromPin, const FVector2D Location, bool bSelectNewNode = true) OVERRIDE;
+	virtual UEdGraphNode* PerformAction(class UEdGraph* ParentGraph, UEdGraphPin* FromPin, const FVector2D Location, bool bSelectNewNode = true) override;
 	// End of FEdGraphSchemaAction interface
 };
 
@@ -70,7 +70,7 @@ struct UNREALED_API FSoundCueGraphSchemaAction_NewComment : public FEdGraphSchem
 	{}
 
 	// FEdGraphSchemaAction interface
-	virtual UEdGraphNode* PerformAction(class UEdGraph* ParentGraph, UEdGraphPin* FromPin, const FVector2D Location, bool bSelectNewNode = true) OVERRIDE;
+	virtual UEdGraphNode* PerformAction(class UEdGraph* ParentGraph, UEdGraphPin* FromPin, const FVector2D Location, bool bSelectNewNode = true) override;
 	// End of FEdGraphSchemaAction interface
 };
 
@@ -89,7 +89,7 @@ struct UNREALED_API FSoundCueGraphSchemaAction_Paste : public FEdGraphSchemaActi
 	{}
 
 	// FEdGraphSchemaAction interface
-	virtual UEdGraphNode* PerformAction(class UEdGraph* ParentGraph, UEdGraphPin* FromPin, const FVector2D Location, bool bSelectNewNode = true) OVERRIDE;
+	virtual UEdGraphNode* PerformAction(class UEdGraph* ParentGraph, UEdGraphPin* FromPin, const FVector2D Location, bool bSelectNewNode = true) override;
 	// End of FEdGraphSchemaAction interface
 };
 
@@ -108,18 +108,18 @@ class USoundCueGraphSchema : public UEdGraphSchema
 	void TryConnectNodes(const TArray<USoundNode*>& OutputNodes, USoundNode* InputNode) const;
 
 	// Begin EdGraphSchema interface
-	virtual void GetGraphContextActions(FGraphContextMenuBuilder& ContextMenuBuilder) const OVERRIDE;
-	virtual void GetContextMenuActions(const UEdGraph* CurrentGraph, const UEdGraphNode* InGraphNode, const UEdGraphPin* InGraphPin, class FMenuBuilder* MenuBuilder, bool bIsDebugging) const OVERRIDE;
-	virtual void CreateDefaultNodesForGraph(UEdGraph& Graph) const OVERRIDE;
-	virtual const FPinConnectionResponse CanCreateConnection(const UEdGraphPin* A, const UEdGraphPin* B) const OVERRIDE;
-	virtual bool TryCreateConnection(UEdGraphPin* A, UEdGraphPin* B) const OVERRIDE;
-	virtual bool ShouldHidePinDefaultValue(UEdGraphPin* Pin) const OVERRIDE;
-	virtual FLinearColor GetPinTypeColor(const FEdGraphPinType& PinType) const OVERRIDE;
-	virtual void BreakNodeLinks(UEdGraphNode& TargetNode) const OVERRIDE;
-	virtual void BreakPinLinks(UEdGraphPin& TargetPin, bool bSendsNodeNotifcation) const OVERRIDE;
-	virtual void DroppedAssetsOnGraph(const TArray<class FAssetData>& Assets, const FVector2D& GraphPosition, UEdGraph* Graph) const OVERRIDE;
-	virtual int32 GetNodeSelectionCount(const UEdGraph* Graph) const OVERRIDE;
-	virtual TSharedPtr<FEdGraphSchemaAction> GetCreateCommentAction() const OVERRIDE;
+	virtual void GetGraphContextActions(FGraphContextMenuBuilder& ContextMenuBuilder) const override;
+	virtual void GetContextMenuActions(const UEdGraph* CurrentGraph, const UEdGraphNode* InGraphNode, const UEdGraphPin* InGraphPin, class FMenuBuilder* MenuBuilder, bool bIsDebugging) const override;
+	virtual void CreateDefaultNodesForGraph(UEdGraph& Graph) const override;
+	virtual const FPinConnectionResponse CanCreateConnection(const UEdGraphPin* A, const UEdGraphPin* B) const override;
+	virtual bool TryCreateConnection(UEdGraphPin* A, UEdGraphPin* B) const override;
+	virtual bool ShouldHidePinDefaultValue(UEdGraphPin* Pin) const override;
+	virtual FLinearColor GetPinTypeColor(const FEdGraphPinType& PinType) const override;
+	virtual void BreakNodeLinks(UEdGraphNode& TargetNode) const override;
+	virtual void BreakPinLinks(UEdGraphPin& TargetPin, bool bSendsNodeNotifcation) const override;
+	virtual void DroppedAssetsOnGraph(const TArray<class FAssetData>& Assets, const FVector2D& GraphPosition, UEdGraph* Graph) const override;
+	virtual int32 GetNodeSelectionCount(const UEdGraph* Graph) const override;
+	virtual TSharedPtr<FEdGraphSchemaAction> GetCreateCommentAction() const override;
 	// End EdGraphSchema interface
 
 private:
diff --git a/Engine/Source/Editor/UnrealEd/Classes/TexAligner/TexAligner.h b/Engine/Source/Editor/UnrealEd/Classes/TexAligner/TexAligner.h
--- a/Engine/Source/Editor/UnrealEd/Classes/TexAligner/TexAligner.h
+++ b/Engine/Source/Editor/UnrealEd/Classes/TexAligner/TexAligner.h
@@ -57,7 +57,7 @@ class UTexAligner : public UObject
 
 
 	// Begin UObject Interface
-	virtual void PostInitProperties() OVERRIDE;
+	virtual void PostInitProperties() override;
 	// End UObject Interface
 
 	// @todo document
diff --git a/Engine/Source/Editor/UnrealEd/Classes/TexAligner/TexAlignerBox.h b/Engine/Source/Editor/UnrealEd/Classes/TexAligner/TexAlignerBox.h
--- a/Engine/Source/Editor/UnrealEd/Classes/TexAligner/TexAlignerBox.h
+++ b/Engine/Source/Editor/UnrealEd/Classes/TexAligner/TexAlignerBox.h
@@ -16,11 +16,11 @@ class UTexAlignerBox : public UTexAligner
 
 
 	// Begin UObject Interface
-	virtual void PostInitProperties() OVERRIDE;
+	virtual void PostInitProperties() override;
 	// End UObject Interface
 
 	// Begin UTexAligner Interface
-	virtual void AlignSurf( ETexAlign InTexAlignType, UModel* InModel, FBspSurfIdx* InSurfIdx, FPoly* InPoly, FVector* InNormal ) OVERRIDE;
+	virtual void AlignSurf( ETexAlign InTexAlignType, UModel* InModel, FBspSurfIdx* InSurfIdx, FPoly* InPoly, FVector* InNormal ) override;
 	// End UTexAligner Interface
 };
 
diff --git a/Engine/Source/Editor/UnrealEd/Classes/TexAligner/TexAlignerDefault.h b/Engine/Source/Editor/UnrealEd/Classes/TexAligner/TexAlignerDefault.h
--- a/Engine/Source/Editor/UnrealEd/Classes/TexAligner/TexAlignerDefault.h
+++ b/Engine/Source/Editor/UnrealEd/Classes/TexAligner/TexAlignerDefault.h
@@ -16,11 +16,11 @@ class UTexAlignerDefault : public UTexAligner
 
 
 	// Begin UObject Interface
-	virtual void PostInitProperties() OVERRIDE;
+	virtual void PostInitProperties() override;
 	// End UObject Interface
 
 	// Begin UTexAligner Interface
-	virtual void AlignSurf( ETexAlign InTexAlignType, UModel* InModel, FBspSurfIdx* InSurfIdx, FPoly* InPoly, FVector* InNormal ) OVERRIDE;
+	virtual void AlignSurf( ETexAlign InTexAlignType, UModel* InModel, FBspSurfIdx* InSurfIdx, FPoly* InPoly, FVector* InNormal ) override;
 	// End UTexAligner Interface
 };
 
diff --git a/Engine/Source/Editor/UnrealEd/Classes/ThumbnailRendering/BlueprintThumbnailRenderer.h b/Engine/Source/Editor/UnrealEd/Classes/ThumbnailRendering/BlueprintThumbnailRenderer.h
--- a/Engine/Source/Editor/UnrealEd/Classes/ThumbnailRendering/BlueprintThumbnailRenderer.h
+++ b/Engine/Source/Editor/UnrealEd/Classes/ThumbnailRendering/BlueprintThumbnailRenderer.h
@@ -15,11 +15,11 @@ class UBlueprintThumbnailRenderer : public UDefaultSizedThumbnailRenderer
 
 
 	// Begin UThumbnailRenderer Object
-	UNREALED_API virtual void Draw(UObject* Object, int32 X, int32 Y, uint32 Width, uint32 Height, FRenderTarget*, FCanvas* Canvas) OVERRIDE;
+	UNREALED_API virtual void Draw(UObject* Object, int32 X, int32 Y, uint32 Width, uint32 Height, FRenderTarget*, FCanvas* Canvas) override;
 	// End UThumbnailRenderer Object
 
 	// UObject implementation
-	UNREALED_API virtual void BeginDestroy() OVERRIDE;
+	UNREALED_API virtual void BeginDestroy() override;
 	// End UObject implementation
 
 	/** Returns true if this scene is capable of creating a thumbnail with the specified blueprint */
diff --git a/Engine/Source/Editor/UnrealEd/Classes/ThumbnailRendering/DefaultSizedThumbnailRenderer.h b/Engine/Source/Editor/UnrealEd/Classes/ThumbnailRendering/DefaultSizedThumbnailRenderer.h
--- a/Engine/Source/Editor/UnrealEd/Classes/ThumbnailRendering/DefaultSizedThumbnailRenderer.h
+++ b/Engine/Source/Editor/UnrealEd/Classes/ThumbnailRendering/DefaultSizedThumbnailRenderer.h
@@ -27,7 +27,7 @@ class UDefaultSizedThumbnailRenderer : public UThumbnailRenderer
 
 
 	// Begin UThumbnailRenderer Object
-	UNREALED_API virtual void GetThumbnailSize(UObject* Object, float Zoom, uint32& OutWidth, uint32& OutHeight) const OVERRIDE;
+	UNREALED_API virtual void GetThumbnailSize(UObject* Object, float Zoom, uint32& OutWidth, uint32& OutHeight) const override;
 	// End UThumbnailRenderer Object
 };
 
diff --git a/Engine/Source/Editor/UnrealEd/Classes/ThumbnailRendering/FontThumbnailRenderer.h b/Engine/Source/Editor/UnrealEd/Classes/ThumbnailRendering/FontThumbnailRenderer.h
--- a/Engine/Source/Editor/UnrealEd/Classes/ThumbnailRendering/FontThumbnailRenderer.h
+++ b/Engine/Source/Editor/UnrealEd/Classes/ThumbnailRendering/FontThumbnailRenderer.h
@@ -15,8 +15,8 @@ class UFontThumbnailRenderer : public UTextureThumbnailRenderer
 
 
 	// Begin UThumbnailRenderer Object
-	virtual void GetThumbnailSize(UObject* Object, float Zoom, uint32& OutWidth, uint32& OutHeight) const OVERRIDE;
-	virtual void Draw(UObject* Object, int32 X, int32 Y, uint32 Width, uint32 Height, FRenderTarget*, FCanvas* Canvas) OVERRIDE;
+	virtual void GetThumbnailSize(UObject* Object, float Zoom, uint32& OutWidth, uint32& OutHeight) const override;
+	virtual void Draw(UObject* Object, int32 X, int32 Y, uint32 Width, uint32 Height, FRenderTarget*, FCanvas* Canvas) override;
 	// End UThumbnailRenderer Object
 };
 
diff --git a/Engine/Source/Editor/UnrealEd/Classes/ThumbnailRendering/MaterialFunctionThumbnailRenderer.h b/Engine/Source/Editor/UnrealEd/Classes/ThumbnailRendering/MaterialFunctionThumbnailRenderer.h
--- a/Engine/Source/Editor/UnrealEd/Classes/ThumbnailRendering/MaterialFunctionThumbnailRenderer.h
+++ b/Engine/Source/Editor/UnrealEd/Classes/ThumbnailRendering/MaterialFunctionThumbnailRenderer.h
@@ -15,11 +15,11 @@ class UMaterialFunctionThumbnailRenderer : public UDefaultSizedThumbnailRenderer
 
 
 	// Begin UThumbnailRenderer Object
-	virtual void Draw(UObject* Object, int32 X, int32 Y, uint32 Width, uint32 Height, FRenderTarget*, FCanvas* Canvas) OVERRIDE;
+	virtual void Draw(UObject* Object, int32 X, int32 Y, uint32 Width, uint32 Height, FRenderTarget*, FCanvas* Canvas) override;
 	// End UThumbnailRenderer Object
 
 	// UObject implementation
-	UNREALED_API virtual void BeginDestroy() OVERRIDE;
+	UNREALED_API virtual void BeginDestroy() override;
 
 private:
 	class FMaterialThumbnailScene* ThumbnailScene;
diff --git a/Engine/Source/Editor/UnrealEd/Classes/ThumbnailRendering/MaterialInstanceThumbnailRenderer.h b/Engine/Source/Editor/UnrealEd/Classes/ThumbnailRendering/MaterialInstanceThumbnailRenderer.h
--- a/Engine/Source/Editor/UnrealEd/Classes/ThumbnailRendering/MaterialInstanceThumbnailRenderer.h
+++ b/Engine/Source/Editor/UnrealEd/Classes/ThumbnailRendering/MaterialInstanceThumbnailRenderer.h
@@ -15,11 +15,11 @@ class UMaterialInstanceThumbnailRenderer : public UDefaultSizedThumbnailRenderer
 
 
 	// Begin UThumbnailRenderer Object
-	virtual void Draw(UObject* Object, int32 X, int32 Y, uint32 Width, uint32 Height, FRenderTarget*, FCanvas* Canvas) OVERRIDE;
+	virtual void Draw(UObject* Object, int32 X, int32 Y, uint32 Width, uint32 Height, FRenderTarget*, FCanvas* Canvas) override;
 	// End UThumbnailRenderer Object
 
 	// UObject implementation
-	UNREALED_API virtual void BeginDestroy() OVERRIDE;
+	UNREALED_API virtual void BeginDestroy() override;
 
 private:
 	class FMaterialThumbnailScene* ThumbnailScene;
diff --git a/Engine/Source/Editor/UnrealEd/Classes/ThumbnailRendering/ParticleSystemThumbnailRenderer.h b/Engine/Source/Editor/UnrealEd/Classes/ThumbnailRendering/ParticleSystemThumbnailRenderer.h
--- a/Engine/Source/Editor/UnrealEd/Classes/ThumbnailRendering/ParticleSystemThumbnailRenderer.h
+++ b/Engine/Source/Editor/UnrealEd/Classes/ThumbnailRendering/ParticleSystemThumbnailRenderer.h
@@ -21,12 +21,12 @@ class UParticleSystemThumbnailRenderer : public UTextureThumbnailRenderer
 
 
 	// Begin UThumbnailRenderer Object
-	virtual void GetThumbnailSize(UObject* Object, float Zoom, uint32& OutWidth, uint32& OutHeight) const OVERRIDE;
-	virtual void Draw(UObject* Object, int32 X, int32 Y, uint32 Width, uint32 Height, FRenderTarget*, FCanvas* Canvas) OVERRIDE;
+	virtual void GetThumbnailSize(UObject* Object, float Zoom, uint32& OutWidth, uint32& OutHeight) const override;
+	virtual void Draw(UObject* Object, int32 X, int32 Y, uint32 Width, uint32 Height, FRenderTarget*, FCanvas* Canvas) override;
 	// End UThumbnailRenderer Object
 
 	// UObject implementation
-	UNREALED_API virtual void BeginDestroy() OVERRIDE;
+	UNREALED_API virtual void BeginDestroy() override;
 
 private:
 	class FParticleSystemThumbnailScene* ThumbnailScene;
diff --git a/Engine/Source/Editor/UnrealEd/Classes/ThumbnailRendering/SkeletalMeshThumbnailRenderer.h b/Engine/Source/Editor/UnrealEd/Classes/ThumbnailRendering/SkeletalMeshThumbnailRenderer.h
--- a/Engine/Source/Editor/UnrealEd/Classes/ThumbnailRendering/SkeletalMeshThumbnailRenderer.h
+++ b/Engine/Source/Editor/UnrealEd/Classes/ThumbnailRendering/SkeletalMeshThumbnailRenderer.h
@@ -15,11 +15,11 @@ class USkeletalMeshThumbnailRenderer : public UDefaultSizedThumbnailRenderer
 
 
 	// Begin UThumbnailRenderer Object
-	UNREALED_API virtual void Draw(UObject* Object, int32 X, int32 Y, uint32 Width, uint32 Height, FRenderTarget*, FCanvas* Canvas) OVERRIDE;
+	UNREALED_API virtual void Draw(UObject* Object, int32 X, int32 Y, uint32 Width, uint32 Height, FRenderTarget*, FCanvas* Canvas) override;
 	// End UThumbnailRenderer Object
 
 	// UObject implementation
-	UNREALED_API virtual void BeginDestroy() OVERRIDE;
+	UNREALED_API virtual void BeginDestroy() override;
 
 private:
 	class FSkeletalMeshThumbnailScene* ThumbnailScene;
diff --git a/Engine/Source/Editor/UnrealEd/Classes/ThumbnailRendering/SlateBrushThumbnailRenderer.h b/Engine/Source/Editor/UnrealEd/Classes/ThumbnailRendering/SlateBrushThumbnailRenderer.h
--- a/Engine/Source/Editor/UnrealEd/Classes/ThumbnailRendering/SlateBrushThumbnailRenderer.h
+++ b/Engine/Source/Editor/UnrealEd/Classes/ThumbnailRendering/SlateBrushThumbnailRenderer.h
@@ -16,7 +16,7 @@ class USlateBrushThumbnailRenderer : public UDefaultSizedThumbnailRenderer
 public:
 
 	// Begin UThumbnailRenderer Object
-	UNREALED_API virtual void Draw(UObject* Object, int32 X, int32 Y, uint32 Width, uint32 Height, FRenderTarget*, FCanvas* Canvas) OVERRIDE;
+	UNREALED_API virtual void Draw(UObject* Object, int32 X, int32 Y, uint32 Width, uint32 Height, FRenderTarget*, FCanvas* Canvas) override;
 	// End UThumbnailRenderer Object
 
 };
diff --git a/Engine/Source/Editor/UnrealEd/Classes/ThumbnailRendering/TextureThumbnailRenderer.h b/Engine/Source/Editor/UnrealEd/Classes/ThumbnailRendering/TextureThumbnailRenderer.h
--- a/Engine/Source/Editor/UnrealEd/Classes/ThumbnailRendering/TextureThumbnailRenderer.h
+++ b/Engine/Source/Editor/UnrealEd/Classes/ThumbnailRendering/TextureThumbnailRenderer.h
@@ -15,8 +15,8 @@ class UTextureThumbnailRenderer : public UThumbnailRenderer
 
 
 	// Begin UThumbnailRenderer Object
-	UNREALED_API virtual void GetThumbnailSize(UObject* Object, float Zoom, uint32& OutWidth, uint32& OutHeight) const OVERRIDE;
-	UNREALED_API virtual void Draw(UObject* Object, int32 X, int32 Y, uint32 Width, uint32 Height, FRenderTarget*, FCanvas* Canvas) OVERRIDE;
+	UNREALED_API virtual void GetThumbnailSize(UObject* Object, float Zoom, uint32& OutWidth, uint32& OutHeight) const override;
+	UNREALED_API virtual void Draw(UObject* Object, int32 X, int32 Y, uint32 Width, uint32 Height, FRenderTarget*, FCanvas* Canvas) override;
 	// End UThumbnailRenderer Object
 };
 
diff --git a/Engine/Source/Editor/UnrealEd/Classes/ThumbnailRendering/ThumbnailManager.h b/Engine/Source/Editor/UnrealEd/Classes/ThumbnailRendering/ThumbnailManager.h
--- a/Engine/Source/Editor/UnrealEd/Classes/ThumbnailRendering/ThumbnailManager.h
+++ b/Engine/Source/Editor/UnrealEd/Classes/ThumbnailRendering/ThumbnailManager.h
@@ -135,7 +135,7 @@ public:
 
 public:
 	// Begin UObject Interface
-	virtual void Serialize(FArchive& Ar) OVERRIDE;
+	virtual void Serialize(FArchive& Ar) override;
 	// End UObject Interface
 
 	/**
diff --git a/Engine/Source/Editor/UnrealEd/Classes/UserDefinedStructure/UserDefinedStructEditorData.h b/Engine/Source/Editor/UnrealEd/Classes/UserDefinedStructure/UserDefinedStructEditorData.h
--- a/Engine/Source/Editor/UnrealEd/Classes/UserDefinedStructure/UserDefinedStructEditorData.h
+++ b/Engine/Source/Editor/UnrealEd/Classes/UserDefinedStructure/UserDefinedStructEditorData.h
@@ -80,8 +80,8 @@ public:
 	FString ToolTip;
 
 	// UObject interface.
-	virtual void PostEditUndo() OVERRIDE;
-	virtual void PostLoadSubobjects(struct FObjectInstancingGraph* OuterInstanceGraph) OVERRIDE;
+	virtual void PostEditUndo() override;
+	virtual void PostLoadSubobjects(struct FObjectInstancingGraph* OuterInstanceGraph) override;
 	// End of UObject interface.
 
 	uint32 GenerateUniqueNameIdForMemberVariable();
diff --git a/Engine/Source/Editor/UnrealEd/Private/Animation/DebugSkelMeshComponent.cpp b/Engine/Source/Editor/UnrealEd/Private/Animation/DebugSkelMeshComponent.cpp
--- a/Engine/Source/Editor/UnrealEd/Private/Animation/DebugSkelMeshComponent.cpp
+++ b/Engine/Source/Editor/UnrealEd/Private/Animation/DebugSkelMeshComponent.cpp
@@ -67,7 +67,7 @@ public:
 		return FSkeletalMeshSceneProxy::GetAllocatedSize();
 	}
 
-	virtual uint32 GetMemoryFootprint() const OVERRIDE
+	virtual uint32 GetMemoryFootprint() const override
 	{
 		return sizeof(*this) + GetAllocatedSize();
 	}
diff --git a/Engine/Source/Editor/UnrealEd/Private/AssetThumbnail.cpp b/Engine/Source/Editor/UnrealEd/Private/AssetThumbnail.cpp
--- a/Engine/Source/Editor/UnrealEd/Private/AssetThumbnail.cpp
+++ b/Engine/Source/Editor/UnrealEd/Private/AssetThumbnail.cpp
@@ -202,7 +202,7 @@ public:
 	}
 
 	// SWidget implementation
-	virtual void OnMouseEnter( const FGeometry& MyGeometry, const FPointerEvent& MouseEvent ) OVERRIDE
+	virtual void OnMouseEnter( const FGeometry& MyGeometry, const FPointerEvent& MouseEvent ) override
 	{
 		SCompoundWidget::OnMouseEnter(MyGeometry, MouseEvent);
 
@@ -213,7 +213,7 @@ public:
 		}
 	}
 
-	virtual void Tick( const FGeometry& AllottedGeometry, const double InCurrentTime, const float InDeltaTime ) OVERRIDE
+	virtual void Tick( const FGeometry& AllottedGeometry, const double InCurrentTime, const float InDeltaTime ) override
 	{
 		SCompoundWidget::Tick(AllottedGeometry, InCurrentTime, InDeltaTime);
 
diff --git a/Engine/Source/Editor/UnrealEd/Private/ComponentAssetBroker.cpp b/Engine/Source/Editor/UnrealEd/Private/ComponentAssetBroker.cpp
--- a/Engine/Source/Editor/UnrealEd/Private/ComponentAssetBroker.cpp
+++ b/Engine/Source/Editor/UnrealEd/Private/ComponentAssetBroker.cpp
@@ -14,12 +14,12 @@
 class FStaticMeshComponentBroker : public IComponentAssetBroker
 {
 public:
-	UClass* GetSupportedAssetClass() OVERRIDE
+	UClass* GetSupportedAssetClass() override
 	{
 		return UStaticMesh::StaticClass();
 	}
 
-	virtual bool AssignAssetToComponent(UActorComponent* InComponent, UObject* InAsset) OVERRIDE
+	virtual bool AssignAssetToComponent(UActorComponent* InComponent, UObject* InAsset) override
 	{
 		if (UStaticMeshComponent* StaticMeshComp = Cast<UStaticMeshComponent>(InComponent))
 		{
@@ -35,7 +35,7 @@ public:
 		return false;
 	}
 
-	virtual UObject* GetAssetFromComponent(UActorComponent* InComponent) OVERRIDE
+	virtual UObject* GetAssetFromComponent(UActorComponent* InComponent) override
 	{
 		if (UStaticMeshComponent* StaticMeshComp = Cast<UStaticMeshComponent>(InComponent))
 		{
@@ -51,12 +51,12 @@ public:
 class FDestructableMeshComponentBroker : public IComponentAssetBroker
 {
 public:
-	UClass* GetSupportedAssetClass() OVERRIDE
+	UClass* GetSupportedAssetClass() override
 	{
 		return UDestructibleMesh::StaticClass();
 	}
 
-	virtual bool AssignAssetToComponent(UActorComponent* InComponent, UObject* InAsset) OVERRIDE
+	virtual bool AssignAssetToComponent(UActorComponent* InComponent, UObject* InAsset) override
 	{
 		if (UDestructibleComponent* DestMeshComp = Cast<UDestructibleComponent>(InComponent))
 		{
@@ -72,7 +72,7 @@ public:
 		return false;
 	}
 
-	virtual UObject* GetAssetFromComponent(UActorComponent* InComponent) OVERRIDE
+	virtual UObject* GetAssetFromComponent(UActorComponent* InComponent) override
 	{
 		if (UDestructibleComponent* DestMeshComp = Cast<UDestructibleComponent>(InComponent))
 		{
@@ -88,12 +88,12 @@ public:
 class FSkeletalMeshComponentBroker : public IComponentAssetBroker
 {
 public:
-	UClass* GetSupportedAssetClass() OVERRIDE
+	UClass* GetSupportedAssetClass() override
 	{
 		return USkeletalMesh::StaticClass();
 	}
 
-	virtual bool AssignAssetToComponent(UActorComponent* InComponent, UObject* InAsset) OVERRIDE
+	virtual bool AssignAssetToComponent(UActorComponent* InComponent, UObject* InAsset) override
 	{
 		if (USkeletalMeshComponent* SkeletalComp = Cast<USkeletalMeshComponent>(InComponent))
 		{
@@ -109,7 +109,7 @@ public:
 		return false;
 	}
 
-	virtual UObject* GetAssetFromComponent(UActorComponent* InComponent) OVERRIDE
+	virtual UObject* GetAssetFromComponent(UActorComponent* InComponent) override
 	{
 		if (USkeletalMeshComponent* SkelMeshComp = Cast<USkeletalMeshComponent>(InComponent))
 		{
@@ -125,12 +125,12 @@ public:
 class FParticleSystemComponentBroker : public IComponentAssetBroker
 {
 public:
-	UClass* GetSupportedAssetClass() OVERRIDE
+	UClass* GetSupportedAssetClass() override
 	{
 		return UParticleSystem::StaticClass();
 	}
 
-	virtual bool AssignAssetToComponent(UActorComponent* InComponent, UObject* InAsset) OVERRIDE
+	virtual bool AssignAssetToComponent(UActorComponent* InComponent, UObject* InAsset) override
 	{
 		if (UParticleSystemComponent* ParticleComp = Cast<UParticleSystemComponent>(InComponent))
 		{
@@ -146,7 +146,7 @@ public:
 		return false;
 	}
 
-	virtual UObject* GetAssetFromComponent(UActorComponent* InComponent) OVERRIDE
+	virtual UObject* GetAssetFromComponent(UActorComponent* InComponent) override
 	{
 		if (UParticleSystemComponent* ParticleComp = Cast<UParticleSystemComponent>(InComponent))
 		{
@@ -162,12 +162,12 @@ public:
 class FAudioComponentBroker : public IComponentAssetBroker
 {
 public:
-	UClass* GetSupportedAssetClass() OVERRIDE
+	UClass* GetSupportedAssetClass() override
 	{
 		return USoundBase::StaticClass();
 	}
 
-	virtual bool AssignAssetToComponent(UActorComponent* InComponent, UObject* InAsset) OVERRIDE
+	virtual bool AssignAssetToComponent(UActorComponent* InComponent, UObject* InAsset) override
 	{
 		if (UAudioComponent* AudioComp = Cast<UAudioComponent>(InComponent))
 		{
@@ -183,7 +183,7 @@ public:
 		return false;
 	}
 
-	virtual UObject* GetAssetFromComponent(UActorComponent* InComponent) OVERRIDE
+	virtual UObject* GetAssetFromComponent(UActorComponent* InComponent) override
 	{
 		if (UAudioComponent* AudioComp = Cast<UAudioComponent>(InComponent))
 		{
@@ -199,12 +199,12 @@ public:
 class FChildActorComponentBroker : public IComponentAssetBroker
 {
 public:
-	UClass* GetSupportedAssetClass() OVERRIDE
+	UClass* GetSupportedAssetClass() override
 	{
 		return UBlueprint::StaticClass();
 	}
 
-	virtual bool AssignAssetToComponent(UActorComponent* InComponent, UObject* InAsset) OVERRIDE
+	virtual bool AssignAssetToComponent(UActorComponent* InComponent, UObject* InAsset) override
 	{
 		if (UChildActorComponent* ChildActorComp = Cast<UChildActorComponent>(InComponent))
 		{
@@ -219,7 +219,7 @@ public:
 		return false;
 	}
 
-	virtual UObject* GetAssetFromComponent(UActorComponent* InComponent) OVERRIDE
+	virtual UObject* GetAssetFromComponent(UActorComponent* InComponent) override
 	{
 		if (UChildActorComponent* ChildActorComp = Cast<UChildActorComponent>(InComponent))
 		{
diff --git a/Engine/Source/Editor/UnrealEd/Private/ConsolidateWindow.cpp b/Engine/Source/Editor/UnrealEd/Private/ConsolidateWindow.cpp
--- a/Engine/Source/Editor/UnrealEd/Private/ConsolidateWindow.cpp
+++ b/Engine/Source/Editor/UnrealEd/Private/ConsolidateWindow.cpp
@@ -194,21 +194,21 @@ private:
 	// Drag-drop Support
 
 	/** Called in response to the user beginning to drag something over the consolidation panel; parses the drop data into dropped assets, if possible */
-	void OnDragEnter( const FGeometry& MyGeometry, const FDragDropEvent& DragDropEvent ) OVERRIDE;
+	void OnDragEnter( const FGeometry& MyGeometry, const FDragDropEvent& DragDropEvent ) override;
 
 	/** Called in response to the user's drag operation exiting the consolidation panel; deletes any dropped asset data */
-	void OnDragLeave( const FDragDropEvent& DragDropEvent ) OVERRIDE;
+	void OnDragLeave( const FDragDropEvent& DragDropEvent ) override;
 	
 	/** Called in response to the user performing a drop operation in the consolidation panel; adds the dropped objects to the panel */
-	FReply OnDrop( const FGeometry& MyGeometry, const FDragDropEvent& DragDropEvent ) OVERRIDE;
+	FReply OnDrop( const FGeometry& MyGeometry, const FDragDropEvent& DragDropEvent ) override;
 
 	/** Called while the user is dragging something over the consolidation panel; provides visual feedback on whether a drop is allowed or not */
-	FReply OnDragOver( const FGeometry& MyGeometry, const FDragDropEvent& DragDropEvent ) OVERRIDE;
+	FReply OnDragOver( const FGeometry& MyGeometry, const FDragDropEvent& DragDropEvent ) override;
 
 	// Input Responses
 
 	/** Called in response to the user releasing a keyboard key while the consolidation panel has keyboard focus */
-	FReply OnKeyUp( const FGeometry& MyGeometry, const FKeyboardEvent& InKeyboardEvent ) OVERRIDE;
+	FReply OnKeyUp( const FGeometry& MyGeometry, const FKeyboardEvent& InKeyboardEvent ) override;
 
 	/** Track if the panel has already warned the user about consolidating assets with different types, so as not to repeatedly (and annoyingly) warn */
 	bool bAlreadyWarnedAboutTypes;
diff --git a/Engine/Source/Editor/UnrealEd/Private/Dialogs/Dialogs.cpp b/Engine/Source/Editor/UnrealEd/Private/Dialogs/Dialogs.cpp
--- a/Engine/Source/Editor/UnrealEd/Private/Dialogs/Dialogs.cpp
+++ b/Engine/Source/Editor/UnrealEd/Private/Dialogs/Dialogs.cpp
@@ -336,9 +336,9 @@ public:
 	}
 
 	bool GetResponse() const { return bUserResponse; }
-	virtual bool SupportsKeyboardFocus() const OVERRIDE { return true; }
+	virtual bool SupportsKeyboardFocus() const override { return true; }
 
-	virtual FReply OnKeyDown( const FGeometry& MyGeometry, const FKeyboardEvent& InKeyboardEvent ) OVERRIDE
+	virtual FReply OnKeyDown( const FGeometry& MyGeometry, const FKeyboardEvent& InKeyboardEvent ) override
 	{
 		if (InKeyboardEvent.GetKey() == EKeys::C && InKeyboardEvent.IsControlDown())
 		{
diff --git a/Engine/Source/Editor/UnrealEd/Private/Dialogs/DlgReferenceTree.cpp b/Engine/Source/Editor/UnrealEd/Private/Dialogs/DlgReferenceTree.cpp
--- a/Engine/Source/Editor/UnrealEd/Private/Dialogs/DlgReferenceTree.cpp
+++ b/Engine/Source/Editor/UnrealEd/Private/Dialogs/DlgReferenceTree.cpp
@@ -184,7 +184,7 @@ public:
 
 private:
 	/** Called when a tree item is double clicked. */
-	virtual FReply OnMouseButtonDoubleClick( const FGeometry& InMyGeometry, const FPointerEvent& InMouseEvent ) OVERRIDE
+	virtual FReply OnMouseButtonDoubleClick( const FGeometry& InMyGeometry, const FPointerEvent& InMouseEvent ) override
 	{
 		// Show the object in the editor.  I.E show the object in the level if its an actor, or content browser otherwise.
 		if( Item->Object )
diff --git a/Engine/Source/Editor/UnrealEd/Private/Dialogs/SMeshProxyDialog.cpp b/Engine/Source/Editor/UnrealEd/Private/Dialogs/SMeshProxyDialog.cpp
--- a/Engine/Source/Editor/UnrealEd/Private/Dialogs/SMeshProxyDialog.cpp
+++ b/Engine/Source/Editor/UnrealEd/Private/Dialogs/SMeshProxyDialog.cpp
@@ -28,11 +28,11 @@ public:
 	virtual ~SMeshProxyDialog();
 
 	/** SMeshProxyDialog functions */
-	virtual void AssignWindowContent(TSharedPtr<SWindow> Window) OVERRIDE;
-	virtual TSharedPtr<SWindow> GetParentWindow() OVERRIDE;
-	virtual void SetParentWindow(TSharedPtr<SWindow> Window) OVERRIDE;
-	virtual void OnWindowClosed(const TSharedRef<SWindow>&) OVERRIDE;
-	virtual void MarkDirty() OVERRIDE;
+	virtual void AssignWindowContent(TSharedPtr<SWindow> Window) override;
+	virtual TSharedPtr<SWindow> GetParentWindow() override;
+	virtual void SetParentWindow(TSharedPtr<SWindow> Window) override;
+	virtual void OnWindowClosed(const TSharedRef<SWindow>&) override;
+	virtual void MarkDirty() override;
 
 	/** SWidget functions */
 	void Construct(const FArguments& InArgs);
diff --git a/Engine/Source/Editor/UnrealEd/Private/EdGraphUtilities.cpp b/Engine/Source/Editor/UnrealEd/Private/EdGraphUtilities.cpp
--- a/Engine/Source/Editor/UnrealEd/Private/EdGraphUtilities.cpp
+++ b/Engine/Source/Editor/UnrealEd/Private/EdGraphUtilities.cpp
@@ -29,7 +29,7 @@ public:
 	}
 
 protected:
-	virtual bool CanCreateClass(UClass* ObjectClass) const OVERRIDE
+	virtual bool CanCreateClass(UClass* ObjectClass) const override
 	{
 		if (const UEdGraphNode* DefaultNode = Cast<UEdGraphNode>(ObjectClass->GetDefaultObject()))
 		{
@@ -52,7 +52,7 @@ protected:
 		return false;
 	}
 
-	virtual void ProcessConstructedObject(UObject* CreatedObject) OVERRIDE
+	virtual void ProcessConstructedObject(UObject* CreatedObject) override
 	{
 		if (UEdGraphNode* Node = Cast<UEdGraphNode>(CreatedObject))
 		{
diff --git a/Engine/Source/Editor/UnrealEd/Private/EditorExporters.cpp b/Engine/Source/Editor/UnrealEd/Private/EditorExporters.cpp
--- a/Engine/Source/Editor/UnrealEd/Private/EditorExporters.cpp
+++ b/Engine/Source/Editor/UnrealEd/Private/EditorExporters.cpp
@@ -715,7 +715,7 @@ public:
 	}
 
 	/** This override is required otherwise the shaders aren't ready for use when the surface is rendered resulting in a blank image */
-	virtual bool RequiresSynchronousCompilation() const OVERRIDE { return true; };
+	virtual bool RequiresSynchronousCompilation() const override { return true; };
 
 	/**
 	* Should the shader for this material with the given platform, shader type and vertex 
@@ -733,7 +733,7 @@ public:
 		return true;
 	}
 
-	virtual const TArray<UTexture*>& GetReferencedTextures() const OVERRIDE
+	virtual const TArray<UTexture*>& GetReferencedTextures() const override
 	{
 		return ReferencedTextures;
 	}
@@ -881,7 +881,7 @@ public:
 	}
 
 	virtual FString GetMaterialUsageDescription() const { return FString::Printf(TEXT("FExportMaterialRenderer %s"), MaterialInterface ? *MaterialInterface->GetName() : TEXT("NULL")); }
-	virtual int32 GetMaterialDomain() const OVERRIDE
+	virtual int32 GetMaterialDomain() const override
 	{
 		if (Material)
 		{
@@ -935,7 +935,7 @@ public:
 	* Should shaders compiled for this material be saved to disk?
 	*/
 	virtual bool IsPersistent() const { return false; }
-	virtual FGuid GetMaterialId() const OVERRIDE { return Id; }
+	virtual FGuid GetMaterialId() const override { return Id; }
 
 	const UMaterialInterface* GetMaterialInterface() const
 	{
diff --git a/Engine/Source/Editor/UnrealEd/Private/Factories/EditorFactories.cpp b/Engine/Source/Editor/UnrealEd/Private/Factories/EditorFactories.cpp
--- a/Engine/Source/Editor/UnrealEd/Private/Factories/EditorFactories.cpp
+++ b/Engine/Source/Editor/UnrealEd/Private/Factories/EditorFactories.cpp
@@ -6043,12 +6043,12 @@ public:
 	/** Classes to not allow any children of into the Class Viewer/Picker. */
 	TSet< const UClass* > DisallowedChildrenOfClasses;
 
-	virtual bool IsClassAllowed(const FClassViewerInitializationOptions& InInitOptions, const UClass* InClass, TSharedRef< FClassViewerFilterFuncs > InFilterFuncs ) OVERRIDE
+	virtual bool IsClassAllowed(const FClassViewerInitializationOptions& InInitOptions, const UClass* InClass, TSharedRef< FClassViewerFilterFuncs > InFilterFuncs ) override
 	{
 		return InFilterFuncs->IfInChildOfClassesSet(DisallowedChildrenOfClasses, InClass) != EFilterReturn::Passed;
 	}
 
-	virtual bool IsUnloadedClassAllowed(const FClassViewerInitializationOptions& InInitOptions, const TSharedRef< const IUnloadedBlueprintData > InUnloadedClassData, TSharedRef< FClassViewerFilterFuncs > InFilterFuncs) OVERRIDE
+	virtual bool IsUnloadedClassAllowed(const FClassViewerInitializationOptions& InInitOptions, const TSharedRef< const IUnloadedBlueprintData > InUnloadedClassData, TSharedRef< FClassViewerFilterFuncs > InFilterFuncs) override
 	{
 		return InFilterFuncs->IfInChildOfClassesSet(DisallowedChildrenOfClasses, InUnloadedClassData) != EFilterReturn::Passed;
 	}
@@ -6435,13 +6435,13 @@ public:
 	/** Disallowed class flags. */
 	uint32 DisallowedClassFlags;
 
-	virtual bool IsClassAllowed(const FClassViewerInitializationOptions& InInitOptions, const UClass* InClass, TSharedRef< FClassViewerFilterFuncs > InFilterFuncs ) OVERRIDE
+	virtual bool IsClassAllowed(const FClassViewerInitializationOptions& InInitOptions, const UClass* InClass, TSharedRef< FClassViewerFilterFuncs > InFilterFuncs ) override
 	{
 		return !InClass->HasAnyClassFlags(DisallowedClassFlags)
 				&& InFilterFuncs->IfInChildOfClassesSet( AllowedChildrenOfClasses, InClass) != EFilterReturn::Failed;
 	}
 
-	virtual bool IsUnloadedClassAllowed(const FClassViewerInitializationOptions& InInitOptions, const TSharedRef< const IUnloadedBlueprintData > InUnloadedClassData, TSharedRef< FClassViewerFilterFuncs > InFilterFuncs) OVERRIDE
+	virtual bool IsUnloadedClassAllowed(const FClassViewerInitializationOptions& InInitOptions, const TSharedRef< const IUnloadedBlueprintData > InUnloadedClassData, TSharedRef< FClassViewerFilterFuncs > InFilterFuncs) override
 	{
 		return !InUnloadedClassData->HasAnyClassFlags(DisallowedClassFlags)
 			&& InFilterFuncs->IfInChildOfClassesSet( AllowedChildrenOfClasses, InUnloadedClassData) != EFilterReturn::Failed;
diff --git a/Engine/Source/Editor/UnrealEd/Private/FbxOptionWindow.h b/Engine/Source/Editor/UnrealEd/Private/FbxOptionWindow.h
--- a/Engine/Source/Editor/UnrealEd/Private/FbxOptionWindow.h
+++ b/Engine/Source/Editor/UnrealEd/Private/FbxOptionWindow.h
@@ -25,7 +25,7 @@ public:
 
 public:
 	void Construct(const FArguments& InArgs);
-	virtual bool SupportsKeyboardFocus() const OVERRIDE { return true; }
+	virtual bool SupportsKeyboardFocus() const override { return true; }
 
 	FReply OnImport()
 	{
@@ -54,7 +54,7 @@ public:
 		return FReply::Handled();
 	}
 
-	virtual FReply OnKeyDown( const FGeometry& MyGeometry, const FKeyboardEvent& InKeyboardEvent ) OVERRIDE
+	virtual FReply OnKeyDown( const FGeometry& MyGeometry, const FKeyboardEvent& InKeyboardEvent ) override
 	{
 		if( InKeyboardEvent.GetKey() == EKeys::Escape )
 		{
diff --git a/Engine/Source/Editor/UnrealEd/Private/FileHelpers.cpp b/Engine/Source/Editor/UnrealEd/Private/FileHelpers.cpp
--- a/Engine/Source/Editor/UnrealEd/Private/FileHelpers.cpp
+++ b/Engine/Source/Editor/UnrealEd/Private/FileHelpers.cpp
@@ -43,7 +43,7 @@ static const size_t NumInvalidNames = sizeof(InvalidFilenames) / sizeof(FString)
 class FSaveErrorOutputDevice : public FOutputDevice
 {
 public:
-	virtual void Serialize( const TCHAR* InData, ELogVerbosity::Type Verbosity, const class FName& Category ) OVERRIDE
+	virtual void Serialize( const TCHAR* InData, ELogVerbosity::Type Verbosity, const class FName& Category ) override
 	{
 		if ( Verbosity == ELogVerbosity::Error || Verbosity == ELogVerbosity::Warning )
 		{
@@ -64,7 +64,7 @@ public:
 		}
 	}
 
-	virtual void Flush() OVERRIDE
+	virtual void Flush() override
 	{
 		if ( ErrorMessages.Num() > 0 )
 		{
diff --git a/Engine/Source/Editor/UnrealEd/Private/Kismet2/EnumEditorUtils.cpp b/Engine/Source/Editor/UnrealEd/Private/Kismet2/EnumEditorUtils.cpp
--- a/Engine/Source/Editor/UnrealEd/Private/Kismet2/EnumEditorUtils.cpp
+++ b/Engine/Source/Editor/UnrealEd/Private/Kismet2/EnumEditorUtils.cpp
@@ -184,7 +184,7 @@ public:
 	{
 	}
 
-	virtual bool UseToResolveEnumerators() const OVERRIDE
+	virtual bool UseToResolveEnumerators() const override
 	{
 		return true;
 	}
diff --git a/Engine/Source/Editor/UnrealEd/Private/Layers/Layers.h b/Engine/Source/Editor/UnrealEd/Private/Layers/Layers.h
--- a/Engine/Source/Editor/UnrealEd/Private/Layers/Layers.h
+++ b/Engine/Source/Editor/UnrealEd/Private/Layers/Layers.h
@@ -28,98 +28,98 @@ public:
 	//	Begin ILayer Implementation -- See ILayers for documentation
 
 	DECLARE_DERIVED_EVENT( FLayers, ILayers::FOnLayersChanged, FOnLayersChanged );
-	virtual FOnLayersChanged& OnLayersChanged() OVERRIDE { return LayersChanged; }
+	virtual FOnLayersChanged& OnLayersChanged() override { return LayersChanged; }
 
 	///////////////////////////////////////////////////////////////////////////////////////////////////////////////////////
 	// Operations on Levels
 
-	virtual void AddLevelLayerInformation( const TWeakObjectPtr< ULevel >& Level ) OVERRIDE;
-	virtual void RemoveLevelLayerInformation( const TWeakObjectPtr< ULevel >& Level ) OVERRIDE;
+	virtual void AddLevelLayerInformation( const TWeakObjectPtr< ULevel >& Level ) override;
+	virtual void RemoveLevelLayerInformation( const TWeakObjectPtr< ULevel >& Level ) override;
 
 	///////////////////////////////////////////////////////////////////////////////////////////////////////////////////////
 	// Operations on an individual actor.
 
-	virtual bool IsActorValidForLayer( const TWeakObjectPtr< AActor >& Actor ) OVERRIDE;
+	virtual bool IsActorValidForLayer( const TWeakObjectPtr< AActor >& Actor ) override;
 	
-	virtual bool InitializeNewActorLayers( const TWeakObjectPtr< AActor >& Actor ) OVERRIDE;
-	virtual bool DisassociateActorFromLayers( const TWeakObjectPtr< AActor >& Actor ) OVERRIDE;
+	virtual bool InitializeNewActorLayers( const TWeakObjectPtr< AActor >& Actor ) override;
+	virtual bool DisassociateActorFromLayers( const TWeakObjectPtr< AActor >& Actor ) override;
 
-	virtual bool AddActorToLayer( const TWeakObjectPtr< AActor >& Actor, const FName& LayerName ) OVERRIDE;
-	virtual bool AddActorToLayers( const TWeakObjectPtr< AActor >& Actor, const TArray< FName >& LayerNames ) OVERRIDE;
+	virtual bool AddActorToLayer( const TWeakObjectPtr< AActor >& Actor, const FName& LayerName ) override;
+	virtual bool AddActorToLayers( const TWeakObjectPtr< AActor >& Actor, const TArray< FName >& LayerNames ) override;
 	
-	virtual bool RemoveActorFromLayer( const TWeakObjectPtr< AActor >& Actor, const FName& LayerToRemove, const bool bUpdateStats = true ) OVERRIDE;
-	virtual bool RemoveActorFromLayers( const TWeakObjectPtr< AActor >& Actor, const TArray< FName >& LayerNames, const bool bUpdateStats = true  ) OVERRIDE;
+	virtual bool RemoveActorFromLayer( const TWeakObjectPtr< AActor >& Actor, const FName& LayerToRemove, const bool bUpdateStats = true ) override;
+	virtual bool RemoveActorFromLayers( const TWeakObjectPtr< AActor >& Actor, const TArray< FName >& LayerNames, const bool bUpdateStats = true  ) override;
 
 
 	/////////////////////////////////////////////////
 	// Operations on multiple actors.
 	
-	virtual bool AddActorsToLayer( const TArray< TWeakObjectPtr< AActor > >& Actors, const FName& LayerName ) OVERRIDE;
-	virtual bool AddActorsToLayers( const TArray< TWeakObjectPtr< AActor > >& Actors, const TArray< FName >& LayerNames ) OVERRIDE;
+	virtual bool AddActorsToLayer( const TArray< TWeakObjectPtr< AActor > >& Actors, const FName& LayerName ) override;
+	virtual bool AddActorsToLayers( const TArray< TWeakObjectPtr< AActor > >& Actors, const TArray< FName >& LayerNames ) override;
 	
-	virtual bool RemoveActorsFromLayer( const TArray< TWeakObjectPtr< AActor > >& Actors, const FName& LayerName, const bool bUpdateStats = true ) OVERRIDE;
-	virtual bool RemoveActorsFromLayers( const TArray< TWeakObjectPtr< AActor > >& Actors, const TArray< FName >& LayerNames, const bool bUpdateStats = true ) OVERRIDE;
+	virtual bool RemoveActorsFromLayer( const TArray< TWeakObjectPtr< AActor > >& Actors, const FName& LayerName, const bool bUpdateStats = true ) override;
+	virtual bool RemoveActorsFromLayers( const TArray< TWeakObjectPtr< AActor > >& Actors, const TArray< FName >& LayerNames, const bool bUpdateStats = true ) override;
 
 
 	/////////////////////////////////////////////////
 	// Operations on selected actors.
 	TArray< TWeakObjectPtr< AActor > > GetSelectedActors() const;
 
-	virtual bool AddSelectedActorsToLayer( const FName& LayerName ) OVERRIDE;
-	virtual bool AddSelectedActorsToLayers( const TArray< FName >& LayerNames ) OVERRIDE;
+	virtual bool AddSelectedActorsToLayer( const FName& LayerName ) override;
+	virtual bool AddSelectedActorsToLayers( const TArray< FName >& LayerNames ) override;
 	
-	virtual bool RemoveSelectedActorsFromLayer( const FName& LayerName ) OVERRIDE;
-	virtual bool RemoveSelectedActorsFromLayers( const TArray< FName >& LayerNames ) OVERRIDE;
+	virtual bool RemoveSelectedActorsFromLayer( const FName& LayerName ) override;
+	virtual bool RemoveSelectedActorsFromLayers( const TArray< FName >& LayerNames ) override;
 
 
 	/////////////////////////////////////////////////
 	// Operations on actors in layers
 
-	virtual bool SelectActorsInLayers( const TArray< FName >& LayerNames, bool bSelect, bool bNotify, bool bSelectEvenIfHidden = false, const TSharedPtr< ActorFilter >& Filter = TSharedPtr< ActorFilter >( NULL ) ) OVERRIDE;
-	virtual bool SelectActorsInLayer( const FName& LayerName, bool bSelect, bool bNotify, bool bSelectEvenIfHidden = false, const TSharedPtr< ActorFilter >& Filter = TSharedPtr< ActorFilter >( NULL ) ) OVERRIDE;
+	virtual bool SelectActorsInLayers( const TArray< FName >& LayerNames, bool bSelect, bool bNotify, bool bSelectEvenIfHidden = false, const TSharedPtr< ActorFilter >& Filter = TSharedPtr< ActorFilter >( NULL ) ) override;
+	virtual bool SelectActorsInLayer( const FName& LayerName, bool bSelect, bool bNotify, bool bSelectEvenIfHidden = false, const TSharedPtr< ActorFilter >& Filter = TSharedPtr< ActorFilter >( NULL ) ) override;
 
 
 	/////////////////////////////////////////////////
 	// Operations on actor viewport visibility regarding layers
 
-	virtual void UpdateAllViewVisibility( const FName& LayerThatChanged ) OVERRIDE;
-	virtual void UpdatePerViewVisibility( FLevelEditorViewportClient* ViewportClient, const FName& LayerThatChanged=NAME_Skip ) OVERRIDE;
+	virtual void UpdateAllViewVisibility( const FName& LayerThatChanged ) override;
+	virtual void UpdatePerViewVisibility( FLevelEditorViewportClient* ViewportClient, const FName& LayerThatChanged=NAME_Skip ) override;
 	
-	virtual void UpdateActorViewVisibility( FLevelEditorViewportClient* ViewportClient, const TWeakObjectPtr< AActor >& Actor, bool bReregisterIfDirty=true ) OVERRIDE;
-	virtual void UpdateActorAllViewsVisibility( const TWeakObjectPtr< AActor >& Actor ) OVERRIDE;
+	virtual void UpdateActorViewVisibility( FLevelEditorViewportClient* ViewportClient, const TWeakObjectPtr< AActor >& Actor, bool bReregisterIfDirty=true ) override;
+	virtual void UpdateActorAllViewsVisibility( const TWeakObjectPtr< AActor >& Actor ) override;
 	
-	virtual void RemoveViewFromActorViewVisibility( FLevelEditorViewportClient* ViewportClient ) OVERRIDE;
+	virtual void RemoveViewFromActorViewVisibility( FLevelEditorViewportClient* ViewportClient ) override;
 	
-	virtual bool UpdateActorVisibility( const TWeakObjectPtr< AActor >& Actor, bool& bOutSelectionChanged, bool& bOutActorModified, bool bNotifySelectionChange, bool bRedrawViewports ) OVERRIDE;
-	virtual bool UpdateAllActorsVisibility( bool bNotifySelectionChange, bool bRedrawViewports ) OVERRIDE;
+	virtual bool UpdateActorVisibility( const TWeakObjectPtr< AActor >& Actor, bool& bOutSelectionChanged, bool& bOutActorModified, bool bNotifySelectionChange, bool bRedrawViewports ) override;
+	virtual bool UpdateAllActorsVisibility( bool bNotifySelectionChange, bool bRedrawViewports ) override;
 
 
 	/////////////////////////////////////////////////
 	// Operations on layers
 
-	virtual void AppendActorsForLayer( const FName& LayerName, TArray< TWeakObjectPtr< AActor > >& InActors, const TSharedPtr< ActorFilter >& Filter = TSharedPtr< ActorFilter >( NULL )  ) const OVERRIDE;
-	virtual void AppendActorsForLayers( const TArray< FName >& LayerNames, TArray< TWeakObjectPtr< AActor > >& OutActors, const TSharedPtr< ActorFilter >& Filter = TSharedPtr< ActorFilter >( NULL )  ) const OVERRIDE;
+	virtual void AppendActorsForLayer( const FName& LayerName, TArray< TWeakObjectPtr< AActor > >& InActors, const TSharedPtr< ActorFilter >& Filter = TSharedPtr< ActorFilter >( NULL )  ) const override;
+	virtual void AppendActorsForLayers( const TArray< FName >& LayerNames, TArray< TWeakObjectPtr< AActor > >& OutActors, const TSharedPtr< ActorFilter >& Filter = TSharedPtr< ActorFilter >( NULL )  ) const override;
 
-	virtual void SetLayerVisibility( const FName& LayerName, bool bIsVisible ) OVERRIDE;
-	virtual void SetLayersVisibility( const TArray< FName >& LayerNames, bool bIsVisible ) OVERRIDE;
+	virtual void SetLayerVisibility( const FName& LayerName, bool bIsVisible ) override;
+	virtual void SetLayersVisibility( const TArray< FName >& LayerNames, bool bIsVisible ) override;
 
-	virtual void ToggleLayerVisibility( const FName& LayerName ) OVERRIDE;
-	virtual void ToggleLayersVisibility( const TArray< FName >& LayerNames ) OVERRIDE;
+	virtual void ToggleLayerVisibility( const FName& LayerName ) override;
+	virtual void ToggleLayersVisibility( const TArray< FName >& LayerNames ) override;
 
-	virtual void MakeAllLayersVisible() OVERRIDE;
+	virtual void MakeAllLayersVisible() override;
 
-	virtual TWeakObjectPtr< ULayer > GetLayer( const FName& LayerName ) const OVERRIDE;
-	virtual bool TryGetLayer( const FName& LayerName, TWeakObjectPtr< ULayer >& OutLayer ) OVERRIDE;
+	virtual TWeakObjectPtr< ULayer > GetLayer( const FName& LayerName ) const override;
+	virtual bool TryGetLayer( const FName& LayerName, TWeakObjectPtr< ULayer >& OutLayer ) override;
 
-	virtual void AddAllLayerNamesTo( TArray< FName >& OutLayerNames ) const OVERRIDE;
-	virtual void AddAllLayersTo( TArray< TWeakObjectPtr< ULayer > >& OutLayers ) const OVERRIDE;
+	virtual void AddAllLayerNamesTo( TArray< FName >& OutLayerNames ) const override;
+	virtual void AddAllLayersTo( TArray< TWeakObjectPtr< ULayer > >& OutLayers ) const override;
 
-	virtual TWeakObjectPtr< ULayer > CreateLayer( const FName& LayerName ) OVERRIDE;
+	virtual TWeakObjectPtr< ULayer > CreateLayer( const FName& LayerName ) override;
 
-	virtual void DeleteLayers( const TArray< FName >& LayersToDelete ) OVERRIDE;
-	virtual void DeleteLayer( const FName& LayerToDelete ) OVERRIDE;
+	virtual void DeleteLayers( const TArray< FName >& LayersToDelete ) override;
+	virtual void DeleteLayer( const FName& LayerToDelete ) override;
 
-	virtual bool RenameLayer( const FName OriginalLayerName, const FName& NewLayerName ) OVERRIDE;
+	virtual bool RenameLayer( const FName OriginalLayerName, const FName& NewLayerName ) override;
 
 
 	UWorld* GetWorld() const { return GWorld;} // Fallback to GWorld
diff --git a/Engine/Source/Editor/UnrealEd/Private/LevelViewportClickHandlers.h b/Engine/Source/Editor/UnrealEd/Private/LevelViewportClickHandlers.h
--- a/Engine/Source/Editor/UnrealEd/Private/LevelViewportClickHandlers.h
+++ b/Engine/Source/Editor/UnrealEd/Private/LevelViewportClickHandlers.h
@@ -25,7 +25,7 @@ struct HLevelSocketProxy : public HHitProxy
 		,   SceneComponent( InSceneComponent )
 		,	SocketName( InSocketName )
 	{}
-	virtual void AddReferencedObjects( FReferenceCollector& Collector ) OVERRIDE
+	virtual void AddReferencedObjects( FReferenceCollector& Collector ) override
 	{
 		Collector.AddReferencedObject( Actor );
 		Collector.AddReferencedObject( SceneComponent );
diff --git a/Engine/Source/Editor/UnrealEd/Private/Lightmass/LightmassRender.cpp b/Engine/Source/Editor/UnrealEd/Private/Lightmass/LightmassRender.cpp
--- a/Engine/Source/Editor/UnrealEd/Private/Lightmass/LightmassRender.cpp
+++ b/Engine/Source/Editor/UnrealEd/Private/Lightmass/LightmassRender.cpp
@@ -218,7 +218,7 @@ public:
 		}
 	}
 
-	virtual const TArray<UTexture*>& GetReferencedTextures() const OVERRIDE
+	virtual const TArray<UTexture*>& GetReferencedTextures() const override
 	{
 		return ReferencedTextures;
 	}
@@ -397,11 +397,11 @@ public:
 	 * This mechanism allows derived material classes to create different DDC keys with the same base material.
 	 * For example lightmass exports diffuse and emissive, each of which requires a material resource with the same base material.
 	 */
-	virtual EMaterialShaderMapUsage::Type GetShaderMapUsage() const OVERRIDE { return Usage; }
+	virtual EMaterialShaderMapUsage::Type GetShaderMapUsage() const override { return Usage; }
 
-	virtual FString GetMaterialUsageDescription() const OVERRIDE { return FString::Printf(TEXT("%s FLightmassMaterialRenderer"), MaterialInterface ? *MaterialInterface->GetName() : TEXT("NULL")); }
+	virtual FString GetMaterialUsageDescription() const override { return FString::Printf(TEXT("%s FLightmassMaterialRenderer"), MaterialInterface ? *MaterialInterface->GetName() : TEXT("NULL")); }
 	
-	virtual int32 GetMaterialDomain() const OVERRIDE
+	virtual int32 GetMaterialDomain() const override
 	{
 		if (Material)
 		{
@@ -456,7 +456,7 @@ public:
 	 */
 	virtual bool IsPersistent() const { return true; }
 
-	virtual FGuid GetMaterialId() const OVERRIDE 
+	virtual FGuid GetMaterialId() const override 
 	{ 
 		// Reuse the base material's Id
 		// Normally this would cause a bug as the shader map would try to be shared by both, 
diff --git a/Engine/Source/Editor/UnrealEd/Private/NormalMapPreview.cpp b/Engine/Source/Editor/UnrealEd/Private/NormalMapPreview.cpp
--- a/Engine/Source/Editor/UnrealEd/Private/NormalMapPreview.cpp
+++ b/Engine/Source/Editor/UnrealEd/Private/NormalMapPreview.cpp
@@ -58,7 +58,7 @@ public:
 	 * Serialization.
 	 * @param Ar - The archive with which to serialize.
 	 */
-	virtual bool Serialize(FArchive& Ar) OVERRIDE
+	virtual bool Serialize(FArchive& Ar) override
 	{
 		bool bShaderHasOutdatedParameters = FGlobalShader::Serialize(Ar);
 		Ar << Texture;
diff --git a/Engine/Source/Editor/UnrealEd/Private/PackageAutoSaver.h b/Engine/Source/Editor/UnrealEd/Private/PackageAutoSaver.h
--- a/Engine/Source/Editor/UnrealEd/Private/PackageAutoSaver.h
+++ b/Engine/Source/Editor/UnrealEd/Private/PackageAutoSaver.h
@@ -15,16 +15,16 @@ public:
 	virtual ~FPackageAutoSaver();
 
 	/** IPackageAutoSaver */
-	virtual void UpdateAutoSaveCount(const float DeltaSeconds) OVERRIDE;
-	virtual void ResetAutoSaveTimer() OVERRIDE;
-	virtual void ForceAutoSaveTimer() OVERRIDE;
-	virtual void ForceMinimumTimeTillAutoSave(const float TimeTillAutoSave = 10.0f) OVERRIDE;
-	virtual void AttemptAutoSave() OVERRIDE;
-	virtual void LoadRestoreFile() OVERRIDE;
-	virtual void UpdateRestoreFile(const bool bRestoreEnabled) const OVERRIDE;
-	virtual bool HasPackagesToRestore() const OVERRIDE;
-	virtual void OfferToRestorePackages() OVERRIDE;
-	virtual bool IsAutoSaving() const OVERRIDE
+	virtual void UpdateAutoSaveCount(const float DeltaSeconds) override;
+	virtual void ResetAutoSaveTimer() override;
+	virtual void ForceAutoSaveTimer() override;
+	virtual void ForceMinimumTimeTillAutoSave(const float TimeTillAutoSave = 10.0f) override;
+	virtual void AttemptAutoSave() override;
+	virtual void LoadRestoreFile() override;
+	virtual void UpdateRestoreFile(const bool bRestoreEnabled) const override;
+	virtual bool HasPackagesToRestore() const override;
+	virtual void OfferToRestorePackages() override;
+	virtual bool IsAutoSaving() const override
 	{
 		return bIsAutoSaving;
 	}
diff --git a/Engine/Source/Editor/UnrealEd/Private/SSocketManager.h b/Engine/Source/Editor/UnrealEd/Private/SSocketManager.h
--- a/Engine/Source/Editor/UnrealEd/Private/SSocketManager.h
+++ b/Engine/Source/Editor/UnrealEd/Private/SSocketManager.h
@@ -19,11 +19,11 @@ public:
 
 	// ISocketManager interface
 	virtual UStaticMeshSocket* GetSelectedSocket() const;
-	virtual void SetSelectedSocket(UStaticMeshSocket* InSelectedSocket) OVERRIDE;
-	virtual void DeleteSelectedSocket() OVERRIDE;
-	virtual void DuplicateSelectedSocket() OVERRIDE;
-	virtual void RequestRenameSelectedSocket() OVERRIDE;
-	virtual void UpdateStaticMesh() OVERRIDE;
+	virtual void SetSelectedSocket(UStaticMeshSocket* InSelectedSocket) override;
+	virtual void DeleteSelectedSocket() override;
+	virtual void DuplicateSelectedSocket() override;
+	virtual void RequestRenameSelectedSocket() override;
+	virtual void UpdateStaticMesh() override;
 	// End of ISocketManager
 
 		/**
@@ -91,7 +91,7 @@ private:
 	TSharedPtr<SWidget> OnContextMenuOpening();
 
 	/** FNotifyHook interface */
-	virtual void NotifyPostChange( const FPropertyChangedEvent& PropertyChangedEvent, UProperty* PropertyThatChanged ) OVERRIDE;
+	virtual void NotifyPostChange( const FPropertyChangedEvent& PropertyChangedEvent, UProperty* PropertyThatChanged ) override;
 
 	/** Post undo */
 	void PostUndo();
