UE4 Refactoring. Changed OVERRIDE and FINAL macros to keywords override and final.

diff --git a/Engine/Source/Editor/UnrealEd/Private/Settings/EditorLoadingSavingSettingsCustomization.h b/Engine/Source/Editor/UnrealEd/Private/Settings/EditorLoadingSavingSettingsCustomization.h
--- a/Engine/Source/Editor/UnrealEd/Private/Settings/EditorLoadingSavingSettingsCustomization.h
+++ b/Engine/Source/Editor/UnrealEd/Private/Settings/EditorLoadingSavingSettingsCustomization.h
@@ -24,7 +24,7 @@ public:
 
 	// Begin IDetailCustomization interface
 
-	virtual void CustomizeDetails( IDetailLayoutBuilder& LayoutBuilder ) OVERRIDE
+	virtual void CustomizeDetails( IDetailLayoutBuilder& LayoutBuilder ) override
 	{
 		CustomizeStartupCategory(LayoutBuilder);
 	}
diff --git a/Engine/Source/Editor/UnrealEd/Private/Settings/GameMapsSettingsCustomization.h b/Engine/Source/Editor/UnrealEd/Private/Settings/GameMapsSettingsCustomization.h
--- a/Engine/Source/Editor/UnrealEd/Private/Settings/GameMapsSettingsCustomization.h
+++ b/Engine/Source/Editor/UnrealEd/Private/Settings/GameMapsSettingsCustomization.h
@@ -24,7 +24,7 @@ public:
 
 	// Begin IDetailCustomization interface
 
-	virtual void CustomizeDetails( IDetailLayoutBuilder& LayoutBuilder ) OVERRIDE
+	virtual void CustomizeDetails( IDetailLayoutBuilder& LayoutBuilder ) override
 	{
 		IDetailCategoryBuilder& DefaultMapsCategory = LayoutBuilder.EditCategory("DefaultMaps");
 		{
diff --git a/Engine/Source/Editor/UnrealEd/Private/SnappingUtils.cpp b/Engine/Source/Editor/UnrealEd/Private/SnappingUtils.cpp
--- a/Engine/Source/Editor/UnrealEd/Private/SnappingUtils.cpp
+++ b/Engine/Source/Editor/UnrealEd/Private/SnappingUtils.cpp
@@ -15,11 +15,11 @@ class FEditorViewportSnapping : public ISnappingPolicy
 {
 public:
 	// FEditorViewportSnapping interface
-	virtual void SnapScale(FVector& Point, const FVector& GridBase) OVERRIDE;
-	virtual void SnapPointToGrid(FVector& Point, const FVector& GridBase) OVERRIDE;
-	virtual void SnapRotatorToGrid(FRotator& Rotation) OVERRIDE;
-	virtual void ClearSnappingHelpers(bool bClearImmediately = false) OVERRIDE;
-	virtual void DrawSnappingHelpers(const FSceneView* View, FPrimitiveDrawInterface* PDI) OVERRIDE;
+	virtual void SnapScale(FVector& Point, const FVector& GridBase) override;
+	virtual void SnapPointToGrid(FVector& Point, const FVector& GridBase) override;
+	virtual void SnapRotatorToGrid(FRotator& Rotation) override;
+	virtual void ClearSnappingHelpers(bool bClearImmediately = false) override;
+	virtual void DrawSnappingHelpers(const FSceneView* View, FPrimitiveDrawInterface* PDI) override;
 	// End of FEditorViewportSnapping interface
 
 	bool IsSnapToGridEnabled();
diff --git a/Engine/Source/Editor/UnrealEd/Private/Toolkits/GlobalEditorCommonCommands.cpp b/Engine/Source/Editor/UnrealEd/Private/Toolkits/GlobalEditorCommonCommands.cpp
--- a/Engine/Source/Editor/UnrealEd/Private/Toolkits/GlobalEditorCommonCommands.cpp
+++ b/Engine/Source/Editor/UnrealEd/Private/Toolkits/GlobalEditorCommonCommands.cpp
@@ -48,7 +48,7 @@ public:
 	}
 
 	// SWidget interface
-	virtual FReply OnPreviewKeyDown(const FGeometry& MyGeometry, const FKeyboardEvent& InKeyboardEvent) OVERRIDE;
+	virtual FReply OnPreviewKeyDown(const FGeometry& MyGeometry, const FKeyboardEvent& InKeyboardEvent) override;
 	// End of SWidget interface
 
 protected:
diff --git a/Engine/Source/Editor/UnrealEd/Private/Toolkits/SStandaloneAssetEditorToolkitHost.h b/Engine/Source/Editor/UnrealEd/Private/Toolkits/SStandaloneAssetEditorToolkitHost.h
--- a/Engine/Source/Editor/UnrealEd/Private/Toolkits/SStandaloneAssetEditorToolkitHost.h
+++ b/Engine/Source/Editor/UnrealEd/Private/Toolkits/SStandaloneAssetEditorToolkitHost.h
@@ -44,19 +44,19 @@ public:
 	virtual ~SStandaloneAssetEditorToolkitHost();
 
 	/** IToolkitHost interface */
-	virtual TSharedRef< class SWidget > GetParentWidget() OVERRIDE;
-	virtual void BringToFront() OVERRIDE;
-	virtual TSharedRef< class SDockTabStack > GetTabSpot( const EToolkitTabSpot::Type TabSpot ) OVERRIDE;
-	virtual void OnToolkitHostingStarted( const TSharedRef< class IToolkit >& Toolkit ) OVERRIDE;
-	virtual void OnToolkitHostingFinished( const TSharedRef< class IToolkit >& Toolkit ) OVERRIDE;
-	virtual UWorld* GetWorld() const OVERRIDE;
+	virtual TSharedRef< class SWidget > GetParentWidget() override;
+	virtual void BringToFront() override;
+	virtual TSharedRef< class SDockTabStack > GetTabSpot( const EToolkitTabSpot::Type TabSpot ) override;
+	virtual void OnToolkitHostingStarted( const TSharedRef< class IToolkit >& Toolkit ) override;
+	virtual void OnToolkitHostingFinished( const TSharedRef< class IToolkit >& Toolkit ) override;
+	virtual UWorld* GetWorld() const override;
 
 	/** SWidget overrides */
-	virtual bool SupportsKeyboardFocus() const OVERRIDE
+	virtual bool SupportsKeyboardFocus() const override
 	{
 		return true;
 	}
-	virtual FReply OnKeyDown( const FGeometry& MyGeometry, const FKeyboardEvent& InKeyboardEvent ) OVERRIDE;
+	virtual FReply OnKeyDown( const FGeometry& MyGeometry, const FKeyboardEvent& InKeyboardEvent ) override;
 
 	/** Fills in the content by loading the associated layout or using the defaults provided.  Must be called after
 	  * the widget is constructed.*/
diff --git a/Engine/Source/Editor/UnrealEd/Private/VertexSnapping.cpp b/Engine/Source/Editor/UnrealEd/Private/VertexSnapping.cpp
--- a/Engine/Source/Editor/UnrealEd/Private/VertexSnapping.cpp
+++ b/Engine/Source/Editor/UnrealEd/Private/VertexSnapping.cpp
@@ -71,23 +71,23 @@ public:
 	}
 
 	/** FVertexIterator interface */
-	virtual FVector Position() OVERRIDE
+	virtual FVector Position() override
 	{
 		return StaticMeshComponent->ComponentToWorld.TransformPosition( PositionBuffer.VertexPosition( CurrentVertexIndex ) );	
 	}
 
-	virtual FVector Normal() OVERRIDE
+	virtual FVector Normal() override
 	{
 		return ComponentToWorldIT.TransformVector( VertexBuffer.VertexTangentZ( CurrentVertexIndex ) );
 	}
 
 protected:
-	virtual void Advance() OVERRIDE
+	virtual void Advance() override
 	{
 		++CurrentVertexIndex;
 	}
 
-	virtual bool HasMoreVertices() OVERRIDE
+	virtual bool HasMoreVertices() override
 	{
 		return CurrentVertexIndex < PositionBuffer.GetNumVertices();
 	}
@@ -127,24 +127,24 @@ public:
 	}
 
 	/** FVertexIterator interface */
-	virtual FVector Position() OVERRIDE
+	virtual FVector Position() override
 	{
 		return BrushComponent->ComponentToWorld.TransformPosition( Vertices[CurrentVertexIndex] );
 	}
 
 	/** FVertexIterator interface */
-	virtual FVector Normal() OVERRIDE
+	virtual FVector Normal() override
 	{
 		return FVector::ZeroVector;
 	}
 
 protected:
-	virtual void Advance() OVERRIDE
+	virtual void Advance() override
 	{
 		++CurrentVertexIndex;
 	}
 
-	virtual bool HasMoreVertices() OVERRIDE
+	virtual bool HasMoreVertices() override
 	{
 		return Vertices.IsValidIndex( CurrentVertexIndex );
 	}
@@ -177,7 +177,7 @@ public:
 	}
 	
 	/** FVertexIterator interface */
-	virtual FVector Position() OVERRIDE
+	virtual FVector Position() override
 	{
 		const FSkelMeshChunk& Chunk = LODModel.Chunks[CurrentChunkIndex];
 
@@ -191,7 +191,7 @@ public:
 		}
 	}
 
-	virtual FVector Normal() OVERRIDE
+	virtual FVector Normal() override
 	{
 		const FSkelMeshChunk& Chunk = LODModel.Chunks[CurrentChunkIndex];
 
@@ -206,7 +206,7 @@ public:
 	}
 
 protected:
-	virtual void Advance() OVERRIDE
+	virtual void Advance() override
 	{
 		// First advance the rigid vertex in the current chunk
 		const FSkelMeshChunk& Chunk = LODModel.Chunks[CurrentChunkIndex];
@@ -228,7 +228,7 @@ protected:
 		}
 	}
 
-	virtual bool HasMoreVertices() OVERRIDE
+	virtual bool HasMoreVertices() override
 	{
 		bool bHasMoreVerts = false;
 		if( LODModel.Chunks.IsValidIndex(CurrentChunkIndex) )
diff --git a/Engine/Source/Editor/UnrealEd/Public/AssetThumbnail.h b/Engine/Source/Editor/UnrealEd/Public/AssetThumbnail.h
--- a/Engine/Source/Editor/UnrealEd/Public/AssetThumbnail.h
+++ b/Engine/Source/Editor/UnrealEd/Public/AssetThumbnail.h
@@ -36,17 +36,17 @@ public:
 	/**
 	 * @return	The size of the viewport (thumbnail size)                   
 	 */
-	virtual FIntPoint GetSize() const OVERRIDE;
+	virtual FIntPoint GetSize() const override;
 
 	/**
 	 * @return The texture used to display the viewports content                  
 	 */
-	virtual FSlateShaderResource* GetViewportRenderTargetTexture() const OVERRIDE;
+	virtual FSlateShaderResource* GetViewportRenderTargetTexture() const override;
 
 	/**
 	 * Returns true if the viewport should be vsynced.
 	 */
-	virtual bool RequiresVsync() const OVERRIDE { return false; }
+	virtual bool RequiresVsync() const override { return false; }
 
 	/**
 	 * @return The object we are rendering the thumbnail for                
diff --git a/Engine/Source/Editor/UnrealEd/Public/CubemapUnwrapUtils.h b/Engine/Source/Editor/UnrealEd/Public/CubemapUnwrapUtils.h
--- a/Engine/Source/Editor/UnrealEd/Public/CubemapUnwrapUtils.h
+++ b/Engine/Source/Editor/UnrealEd/Public/CubemapUnwrapUtils.h
@@ -119,7 +119,7 @@ public:
 	}
 
 	/** Binds vertex and pixel shaders for this element */
-	virtual void BindShaders_RenderThread(FRHICommandList& RHICmdList, const FMatrix& InTransform, const float InGamma, const FMatrix& ColorWeights, const FTexture* Texture) OVERRIDE;
+	virtual void BindShaders_RenderThread(FRHICommandList& RHICmdList, const FMatrix& InTransform, const float InGamma, const FMatrix& ColorWeights, const FTexture* Texture) override;
 
 private:
 	template<typename TPixelShader> void BindShaders_RenderThread(FRHICommandList& RHICmdList, const FMatrix& InTransform, const float InGamma, const FMatrix& ColorWeights, const FTexture* Texture);
@@ -155,7 +155,7 @@ public:
 
 	void SetParameters(FRHICommandList& RHICmdList, const FTexture* Texture, float InBrightnessInLumens);
 
-	virtual bool Serialize(FArchive& Ar) OVERRIDE
+	virtual bool Serialize(FArchive& Ar) override
 	{
 		bool bShaderHasOutdatedParameters = FGlobalShader::Serialize(Ar);
 		Ar << IESTexture;
@@ -179,7 +179,7 @@ public:
 	}
 
 	/** Binds vertex and pixel shaders for this element */
-	virtual void BindShaders_RenderThread(FRHICommandList& RHICmdList, const FMatrix& InTransform, const float InGamma, const FMatrix& ColorWeights, const FTexture* Texture ) OVERRIDE;
+	virtual void BindShaders_RenderThread(FRHICommandList& RHICmdList, const FMatrix& InTransform, const float InGamma, const FMatrix& ColorWeights, const FTexture* Texture ) override;
 
 private:
 	float BrightnessInLumens;
diff --git a/Engine/Source/Editor/UnrealEd/Public/DebugToolExec.h b/Engine/Source/Editor/UnrealEd/Public/DebugToolExec.h
--- a/Engine/Source/Editor/UnrealEd/Public/DebugToolExec.h
+++ b/Engine/Source/Editor/UnrealEd/Public/DebugToolExec.h
@@ -25,7 +25,7 @@ protected:
 public:
 	
 	// Begin Exec Interface
-	UNREALED_API virtual bool Exec( UWorld* InWorld, const TCHAR* Cmd, FOutputDevice& Ar ) OVERRIDE;
+	UNREALED_API virtual bool Exec( UWorld* InWorld, const TCHAR* Cmd, FOutputDevice& Ar ) override;
 	// End Exec Interface
 
 };
diff --git a/Engine/Source/Editor/UnrealEd/Public/Dialogs/SDeleteAssetsDialog.h b/Engine/Source/Editor/UnrealEd/Public/Dialogs/SDeleteAssetsDialog.h
--- a/Engine/Source/Editor/UnrealEd/Public/Dialogs/SDeleteAssetsDialog.h
+++ b/Engine/Source/Editor/UnrealEd/Public/Dialogs/SDeleteAssetsDialog.h
@@ -38,8 +38,8 @@ public:
 	// Construct a SContentReference
 	void Construct( const FArguments& InArgs, TSharedRef<FAssetDeleteModel> InDeleteModel );
 
-	virtual FReply OnKeyDown( const FGeometry& MyGeometry, const FKeyboardEvent& InKeyboardEvent ) OVERRIDE;
-	virtual void Tick( const FGeometry& AllottedGeometry, const double InCurrentTime, const float InDeltaTime ) OVERRIDE;
+	virtual FReply OnKeyDown( const FGeometry& MyGeometry, const FKeyboardEvent& InKeyboardEvent ) override;
+	virtual void Tick( const FGeometry& AllottedGeometry, const double InCurrentTime, const float InDeltaTime ) override;
 
 private:
 	TSharedRef<SWidget> CreateThumbnailWidget();
diff --git a/Engine/Source/Editor/UnrealEd/Public/DragAndDrop/AssetDragDropOp.h b/Engine/Source/Editor/UnrealEd/Public/DragAndDrop/AssetDragDropOp.h
--- a/Engine/Source/Editor/UnrealEd/Public/DragAndDrop/AssetDragDropOp.h
+++ b/Engine/Source/Editor/UnrealEd/Public/DragAndDrop/AssetDragDropOp.h
@@ -31,7 +31,7 @@ public:
 public:
 	UNREALED_API virtual ~FAssetDragDropOp();
 
-	UNREALED_API virtual TSharedPtr<SWidget> GetDefaultDecorator() const OVERRIDE;
+	UNREALED_API virtual TSharedPtr<SWidget> GetDefaultDecorator() const override;
 
 	UNREALED_API void Init();
 
diff --git a/Engine/Source/Editor/UnrealEd/Public/DragAndDrop/BrushBuilderDragDropOp.h b/Engine/Source/Editor/UnrealEd/Public/DragAndDrop/BrushBuilderDragDropOp.h
--- a/Engine/Source/Editor/UnrealEd/Public/DragAndDrop/BrushBuilderDragDropOp.h
+++ b/Engine/Source/Editor/UnrealEd/Public/DragAndDrop/BrushBuilderDragDropOp.h
@@ -29,7 +29,7 @@ public:
 		}
 	}
 
-	virtual void OnDrop( bool bDropWasHandled, const FPointerEvent& MouseEvent ) OVERRIDE
+	virtual void OnDrop( bool bDropWasHandled, const FPointerEvent& MouseEvent ) override
 	{
 		UWorld* World = GEditor->GetEditorWorldContext().World();
 		if ( (World != nullptr) && (World->GetDefaultBrush() != nullptr))
@@ -54,7 +54,7 @@ public:
 		FDragDropOperation::OnDrop( bDropWasHandled, MouseEvent );
 	}
 
-	virtual void OnDragged( const FDragDropEvent& DragDropEvent ) OVERRIDE
+	virtual void OnDragged( const FDragDropEvent& DragDropEvent ) override
 	{
 		UWorld* World = GEditor->GetEditorWorldContext().World();
 		if(World != nullptr)
@@ -70,7 +70,7 @@ public:
 		FDragDropOperation::OnDragged(DragDropEvent);
 	}
 
-	virtual TSharedPtr<class SWidget> GetDefaultDecorator() const OVERRIDE
+	virtual TSharedPtr<class SWidget> GetDefaultDecorator() const override
 	{
 		return SNew(SBox)
 			.WidthOverride(100.0f)
diff --git a/Engine/Source/Editor/UnrealEd/Public/DragAndDrop/ClassDragDropOp.h b/Engine/Source/Editor/UnrealEd/Public/DragAndDrop/ClassDragDropOp.h
--- a/Engine/Source/Editor/UnrealEd/Public/DragAndDrop/ClassDragDropOp.h
+++ b/Engine/Source/Editor/UnrealEd/Public/DragAndDrop/ClassDragDropOp.h
@@ -10,7 +10,7 @@ public:
 	TArray< TWeakObjectPtr<UClass> > ClassesToDrop;
 	
 	/** The widget decorator to use */
-	virtual TSharedPtr<SWidget> GetDefaultDecorator() const OVERRIDE
+	virtual TSharedPtr<SWidget> GetDefaultDecorator() const override
 	{
 		// Just use the first class for the cursor decorator.
 		const FSlateBrush* ClassIcon = FEditorStyle::GetOptionalBrush(*FString::Printf( TEXT( "ClassIcon.%s" ), *ClassesToDrop[0]->GetName() ), nullptr, nullptr );
@@ -88,7 +88,7 @@ public:
 	TSharedPtr< TArray< FClassPackageData > >	AssetsToDrop;
 	
 	/** The widget decorator to use */
-	virtual TSharedPtr<SWidget> GetDefaultDecorator() const OVERRIDE
+	virtual TSharedPtr<SWidget> GetDefaultDecorator() const override
 	{
 		// Create hover widget
 		return SNew(SBorder)
diff --git a/Engine/Source/Editor/UnrealEd/Public/DragAndDrop/DecoratedDragDropOp.h b/Engine/Source/Editor/UnrealEd/Public/DragAndDrop/DecoratedDragDropOp.h
--- a/Engine/Source/Editor/UnrealEd/Public/DragAndDrop/DecoratedDragDropOp.h
+++ b/Engine/Source/Editor/UnrealEd/Public/DragAndDrop/DecoratedDragDropOp.h
@@ -21,7 +21,7 @@ public:
 	{ }
 
 	/** Overridden to provide public access */
-	virtual void Construct() OVERRIDE
+	virtual void Construct() override
 	{
 		FDragDropOperation::Construct();
 	}
@@ -34,7 +34,7 @@ public:
 	}
 
 	/** The widget decorator to use */
-	virtual TSharedPtr<SWidget> GetDefaultDecorator() const OVERRIDE
+	virtual TSharedPtr<SWidget> GetDefaultDecorator() const override
 	{
 		// Create hover widget
 		return SNew(SBorder)
diff --git a/Engine/Source/Editor/UnrealEd/Public/DragTool_BoxSelect.h b/Engine/Source/Editor/UnrealEd/Public/DragTool_BoxSelect.h
--- a/Engine/Source/Editor/UnrealEd/Public/DragTool_BoxSelect.h
+++ b/Engine/Source/Editor/UnrealEd/Public/DragTool_BoxSelect.h
@@ -23,20 +23,20 @@ public:
 	 * @param	InViewportClient	The viewport client in which the drag event occurred.
 	 * @param	InStart				Where the mouse was when the drag started.
 	 */
-	virtual void StartDrag(FEditorViewportClient* InViewportClient, const FVector& InStart, const FVector2D& InStartScreen) OVERRIDE;
+	virtual void StartDrag(FEditorViewportClient* InViewportClient, const FVector& InStart, const FVector2D& InStartScreen) override;
 
 	/* Updates the drag tool's end location with the specified delta.  The end location is
 	 * snapped to the editor constraints if bUseSnapping is true.
 	 *
 	 * @param	InDelta		A delta of mouse movement.
 	 */
-	virtual void AddDelta(const FVector& InDelta) OVERRIDE;
+	virtual void AddDelta(const FVector& InDelta) override;
 
 	/**
 	* Ends a mouse drag behavior (the user has let go of the mouse button).
 	*/
-	virtual void EndDrag() OVERRIDE;
-	virtual void Render(const FSceneView* View, FCanvas* Canvas) OVERRIDE;
+	virtual void EndDrag() override;
+	virtual void Render(const FSceneView* View, FCanvas* Canvas) override;
 
 private:
 	/** 
diff --git a/Engine/Source/Editor/UnrealEd/Public/DragTool_FrustumSelect.h b/Engine/Source/Editor/UnrealEd/Public/DragTool_FrustumSelect.h
--- a/Engine/Source/Editor/UnrealEd/Public/DragTool_FrustumSelect.h
+++ b/Engine/Source/Editor/UnrealEd/Public/DragTool_FrustumSelect.h
@@ -29,7 +29,7 @@ public:
 	 * @param	InViewportClient	The viewport client in which the drag event occurred.
 	 * @param	InStart				Where the mouse was when the drag started.
 	 */
-	virtual void StartDrag(FEditorViewportClient* InViewportClient, const FVector& InStart, const FVector2D& InStartScreen) OVERRIDE;
+	virtual void StartDrag(FEditorViewportClient* InViewportClient, const FVector& InStart, const FVector2D& InStartScreen) override;
 	
 	/**
 	 * Ends a mouse drag behavior (the user has let go of the mouse button).
diff --git a/Engine/Source/Editor/UnrealEd/Public/DragTool_Measure.h b/Engine/Source/Editor/UnrealEd/Public/DragTool_Measure.h
--- a/Engine/Source/Editor/UnrealEd/Public/DragTool_Measure.h
+++ b/Engine/Source/Editor/UnrealEd/Public/DragTool_Measure.h
@@ -22,16 +22,16 @@ public:
 	 * @param	InStart				Where the mouse was when the drag started (world space).
 	 * @param	InStartScreen		Where the mouse was when the drag started (screen space).
 	 */
-	virtual void StartDrag(FEditorViewportClient* InViewportClient, const FVector& InStart, const FVector2D& InStartScreen) OVERRIDE;
+	virtual void StartDrag(FEditorViewportClient* InViewportClient, const FVector& InStart, const FVector2D& InStartScreen) override;
 
 	/* Updates the drag tool's end location with the specified delta.  The end location is
 	 * snapped to the editor constraints if bUseSnapping is true.
 	 *
 	 * @param	InDelta		A delta of mouse movement.
 	 */
-	virtual void AddDelta(const FVector& InDelta) OVERRIDE;
+	virtual void AddDelta(const FVector& InDelta) override;
 
-	virtual void Render(const FSceneView* View,FCanvas* Canvas) OVERRIDE;
+	virtual void Render(const FSceneView* View,FCanvas* Canvas) override;
 
 private:
 	/**
diff --git a/Engine/Source/Editor/UnrealEd/Public/EditorActorFolders.h b/Engine/Source/Editor/UnrealEd/Public/EditorActorFolders.h
--- a/Engine/Source/Editor/UnrealEd/Public/EditorActorFolders.h
+++ b/Engine/Source/Editor/UnrealEd/Public/EditorActorFolders.h
@@ -30,7 +30,7 @@ class UEditorActorFolders : public UObject
 {
 	GENERATED_UCLASS_BODY()
 public:
-	virtual void Serialize(FArchive& Ar) OVERRIDE;
+	virtual void Serialize(FArchive& Ar) override;
 
 	TMap<FName, FActorFolderProps> Folders;
 };
@@ -42,7 +42,7 @@ struct UNREALED_API FActorFolders : public FGCObject
 	~FActorFolders();
 
 	// FGCObject Interface
-	virtual void AddReferencedObjects(FReferenceCollector& Collector) OVERRIDE;
+	virtual void AddReferencedObjects(FReferenceCollector& Collector) override;
 	// End FGCObject Interface
 
 	/** Check whether the singleton is valid */
diff --git a/Engine/Source/Editor/UnrealEd/Public/EditorModeInterpolation.h b/Engine/Source/Editor/UnrealEd/Public/EditorModeInterpolation.h
--- a/Engine/Source/Editor/UnrealEd/Public/EditorModeInterpolation.h
+++ b/Engine/Source/Editor/UnrealEd/Public/EditorModeInterpolation.h
@@ -17,13 +17,13 @@ public:
 	FEdModeInterpEdit();
 	~FEdModeInterpEdit();
 
-	virtual bool InputKey( FLevelEditorViewportClient* ViewportClient, FViewport* Viewport, FKey Key, EInputEvent Event ) OVERRIDE;
-	virtual void Enter() OVERRIDE;
-	virtual void Exit() OVERRIDE;
-	virtual void ActorMoveNotify() OVERRIDE;
-	virtual void ActorPropChangeNotify() OVERRIDE;
-	virtual bool AllowWidgetMove() OVERRIDE;
-	virtual void ActorSelectionChangeNotify() OVERRIDE;
+	virtual bool InputKey( FLevelEditorViewportClient* ViewportClient, FViewport* Viewport, FKey Key, EInputEvent Event ) override;
+	virtual void Enter() override;
+	virtual void Exit() override;
+	virtual void ActorMoveNotify() override;
+	virtual void ActorPropChangeNotify() override;
+	virtual bool AllowWidgetMove() override;
+	virtual void ActorSelectionChangeNotify() override;
 
 	virtual void Render(const FSceneView* View,FViewport* Viewport,FPrimitiveDrawInterface* PDI);
 	virtual void DrawHUD(FLevelEditorViewportClient* ViewportClient,FViewport* Viewport,const FSceneView* View,FCanvas* Canvas);
@@ -57,15 +57,15 @@ public:
 	/**
 	 * @return		true if the key was handled by this editor mode tool.
 	 */
-	virtual bool InputKey(FLevelEditorViewportClient* ViewportClient, FViewport* Viewport, FKey Key, EInputEvent Event) OVERRIDE;
-	virtual bool InputAxis(FLevelEditorViewportClient* InViewportClient, FViewport* Viewport, int32 ControllerId, FKey Key, float Delta, float DeltaTime) OVERRIDE;
-	virtual bool MouseMove(FLevelEditorViewportClient* ViewportClient, FViewport* Viewport, int32 x, int32 y) OVERRIDE;
+	virtual bool InputKey(FLevelEditorViewportClient* ViewportClient, FViewport* Viewport, FKey Key, EInputEvent Event) override;
+	virtual bool InputAxis(FLevelEditorViewportClient* InViewportClient, FViewport* Viewport, int32 ControllerId, FKey Key, float Delta, float DeltaTime) override;
+	virtual bool MouseMove(FLevelEditorViewportClient* ViewportClient, FViewport* Viewport, int32 x, int32 y) override;
 
 	/**
 	 * @return		true if the delta was handled by this editor mode tool.
 	 */
-	virtual bool InputDelta(FLevelEditorViewportClient* InViewportClient,FViewport* InViewport,FVector& InDrag,FRotator& InRot,FVector& InScale) OVERRIDE;
-	virtual void SelectNone() OVERRIDE;
+	virtual bool InputDelta(FLevelEditorViewportClient* InViewportClient,FViewport* InViewport,FVector& InDrag,FRotator& InRot,FVector& InScale) override;
+	virtual void SelectNone() override;
 
 
 	bool bMovingHandle;
diff --git a/Engine/Source/Editor/UnrealEd/Public/EditorModes.h b/Engine/Source/Editor/UnrealEd/Public/EditorModes.h
--- a/Engine/Source/Editor/UnrealEd/Public/EditorModes.h
+++ b/Engine/Source/Editor/UnrealEd/Public/EditorModes.h
@@ -196,7 +196,7 @@ public:
 
 	virtual FVector GetWidgetNormalFromCurrentAxis( void* InData );
 
-	virtual void AddReferencedObjects( FReferenceCollector& Collector ) OVERRIDE {}
+	virtual void AddReferencedObjects( FReferenceCollector& Collector ) override {}
 
 	virtual void Enter();
 	virtual void Exit();
@@ -424,7 +424,7 @@ public:
 	FEdModeDefault();
 
 	// FEdMode interface
-	virtual bool UsesPropertyWidgets() const OVERRIDE { return true; }
+	virtual bool UsesPropertyWidgets() const override { return true; }
 	// End of FEdMode interface
 };
 
@@ -716,7 +716,7 @@ public:
 	UNREALED_API void ClearAllBookmarks( FLevelEditorViewportClient* InViewportClient );
 
 	// FGCObject interface
-	virtual void AddReferencedObjects( FReferenceCollector& Collector ) OVERRIDE;
+	virtual void AddReferencedObjects( FReferenceCollector& Collector ) override;
 	// End of FGCObject interface
 
 	/**
diff --git a/Engine/Source/Editor/UnrealEd/Public/EditorViewportClient.h b/Engine/Source/Editor/UnrealEd/Public/EditorViewportClient.h
--- a/Engine/Source/Editor/UnrealEd/Public/EditorViewportClient.h
+++ b/Engine/Source/Editor/UnrealEd/Public/EditorViewportClient.h
@@ -43,7 +43,7 @@ public:
 	/**
 	* Initialize commands
 	*/
-	virtual void RegisterCommands() OVERRIDE;
+	virtual void RegisterCommands() override;
 
 
 	TSharedPtr< FUICommandInfo > Forward;
@@ -394,7 +394,7 @@ public:
 
 	void SetInitialViewTransform( const FVector& ViewLocation, const FRotator& ViewRotation, float InOrthoZoom );
 
-	virtual void ProcessScreenShots(FViewport* Viewport) OVERRIDE;
+	virtual void ProcessScreenShots(FViewport* Viewport) override;
 
 	void TakeHighResScreenShot();
 
@@ -402,28 +402,28 @@ public:
 	void OnEditorModeChanged(FEdMode* EditorMode, bool bIsEntering);
 
 	/** FViewElementDrawer interface */
-	virtual void Draw(const FSceneView* View,FPrimitiveDrawInterface* PDI) OVERRIDE;
-	virtual void Draw(FViewport* Viewport,FCanvas* Canvas) OVERRIDE;
+	virtual void Draw(const FSceneView* View,FPrimitiveDrawInterface* PDI) override;
+	virtual void Draw(FViewport* Viewport,FCanvas* Canvas) override;
 
 	/** FViewportClient interface */
-	virtual void RedrawRequested(FViewport* Viewport) OVERRIDE;
-	virtual void RequestInvalidateHitProxy(FViewport* Viewport) OVERRIDE;
-	virtual bool InputKey(FViewport* Viewport, int32 ControllerId, FKey Key, EInputEvent Event, float AmountDepressed = 1.f, bool bGamepad=false) OVERRIDE;
-	virtual bool InputAxis(FViewport* Viewport, int32 ControllerId, FKey Key, float Delta, float DeltaTime, int32 NumSamples=1, bool bGamepad=false) OVERRIDE;
-	virtual bool InputGesture(FViewport* Viewport, EGestureEvent::Type GestureType, const FVector2D& GestureDelta) OVERRIDE;
-	virtual void ReceivedFocus(FViewport* Viewport) OVERRIDE;
-	virtual void OnJoystickPlugged(const uint32 InControllerID, const uint32 InType, const uint32 bInConnected)  OVERRIDE;
-	virtual void MouseEnter(FViewport* Viewport,int32 x, int32 y) OVERRIDE;
-	virtual void MouseMove(FViewport* Viewport,int32 x, int32 y) OVERRIDE;
-	virtual void MouseLeave( FViewport* Viewport ) OVERRIDE;
-	virtual EMouseCursor::Type GetCursor(FViewport* Viewport,int32 X,int32 Y) OVERRIDE;
-	virtual void CapturedMouseMove( FViewport* InViewport, int32 InMouseX, int32 InMouseY ) OVERRIDE;
-	virtual bool IsOrtho() const OVERRIDE;
-	virtual void LostFocus(FViewport* Viewport) OVERRIDE;
+	virtual void RedrawRequested(FViewport* Viewport) override;
+	virtual void RequestInvalidateHitProxy(FViewport* Viewport) override;
+	virtual bool InputKey(FViewport* Viewport, int32 ControllerId, FKey Key, EInputEvent Event, float AmountDepressed = 1.f, bool bGamepad=false) override;
+	virtual bool InputAxis(FViewport* Viewport, int32 ControllerId, FKey Key, float Delta, float DeltaTime, int32 NumSamples=1, bool bGamepad=false) override;
+	virtual bool InputGesture(FViewport* Viewport, EGestureEvent::Type GestureType, const FVector2D& GestureDelta) override;
+	virtual void ReceivedFocus(FViewport* Viewport) override;
+	virtual void OnJoystickPlugged(const uint32 InControllerID, const uint32 InType, const uint32 bInConnected)  override;
+	virtual void MouseEnter(FViewport* Viewport,int32 x, int32 y) override;
+	virtual void MouseMove(FViewport* Viewport,int32 x, int32 y) override;
+	virtual void MouseLeave( FViewport* Viewport ) override;
+	virtual EMouseCursor::Type GetCursor(FViewport* Viewport,int32 X,int32 Y) override;
+	virtual void CapturedMouseMove( FViewport* InViewport, int32 InMouseX, int32 InMouseY ) override;
+	virtual bool IsOrtho() const override;
+	virtual void LostFocus(FViewport* Viewport) override;
 
 
 	/** FGCObject interface */
-	virtual void AddReferencedObjects( FReferenceCollector& Collector ) OVERRIDE;
+	virtual void AddReferencedObjects( FReferenceCollector& Collector ) override;
 
 	/**
 	 * Called when the user clicks in the viewport
@@ -576,7 +576,7 @@ public:
 	*
 	* @return 		A valid pointer to the viewports world scene.
 	*/
-	virtual UWorld* GetWorld() const OVERRIDE;
+	virtual UWorld* GetWorld() const override;
 
 	/** If true, this is a level editor viewport */
 	virtual bool IsLevelEditorClient() const { return false; }
@@ -839,7 +839,7 @@ public:
 	/** 
 	 * Returns true if this viewport is excluding non-game elements from its display
 	 */
-	virtual bool IsInGameView() const OVERRIDE { return bInGameViewMode; }
+	virtual bool IsInGameView() const override { return bInGameViewMode; }
 
 	/**
 	 * Aspect ratio bar display settings
diff --git a/Engine/Source/Editor/UnrealEd/Public/FeedbackContextEditor.h b/Engine/Source/Editor/UnrealEd/Public/FeedbackContextEditor.h
--- a/Engine/Source/Editor/UnrealEd/Public/FeedbackContextEditor.h
+++ b/Engine/Source/Editor/UnrealEd/Public/FeedbackContextEditor.h
@@ -93,7 +93,7 @@ public:
 
 	UNREALED_API FFeedbackContextEditor();
 
-	virtual void Serialize( const TCHAR* V, ELogVerbosity::Type Verbosity, const class FName& Category ) OVERRIDE;
+	virtual void Serialize( const TCHAR* V, ELogVerbosity::Type Verbosity, const class FName& Category ) override;
 
 	void BeginSlowTask( const FText& Task, bool bShouldShowProgressDialog, bool bShowCancelButton=false );
 	void EndSlowTask();
@@ -104,7 +104,7 @@ public:
 
 	void OnUserCancel();
 
-	virtual bool YesNof( const FText& Question ) OVERRIDE
+	virtual bool YesNof( const FText& Question ) override
 	{
 		return EAppReturnType::Yes == FMessageDialog::Open( EAppMsgType::YesNo, Question );
 	}
@@ -117,7 +117,7 @@ public:
 	 * 
 	 * @return true by default and false if status update is not possible
 	 */
-	virtual bool StatusUpdate( int32 Numerator, int32 Denominator, const FText& StatusText ) OVERRIDE;
+	virtual bool StatusUpdate( int32 Numerator, int32 Denominator, const FText& StatusText ) override;
 
 	/** Force updating feedback status
 	 *
@@ -127,7 +127,7 @@ public:
 	 * @return true by default and false if status update is not possible
 	 *
 	 */
-	virtual bool StatusForceUpdate( int32 Numerator, int32 Denominator, const FText& StatusText ) OVERRIDE;
+	virtual bool StatusForceUpdate( int32 Numerator, int32 Denominator, const FText& StatusText ) override;
 
 	/**
 	 * Updates the progress amount without changing the status message text
diff --git a/Engine/Source/Editor/UnrealEd/Public/Kismet2/Kismet2NameValidators.h b/Engine/Source/Editor/UnrealEd/Public/Kismet2/Kismet2NameValidators.h
--- a/Engine/Source/Editor/UnrealEd/Public/Kismet2/Kismet2NameValidators.h
+++ b/Engine/Source/Editor/UnrealEd/Public/Kismet2/Kismet2NameValidators.h
@@ -71,8 +71,8 @@ public:
 	static int32 GetMaximumNameLength();
 
 	// Begin FNameValidatorInterface
-	virtual EValidatorResult IsValid( const FString& Name, bool bOriginal = false) OVERRIDE;
-	virtual EValidatorResult IsValid( const FName& Name, bool bOriginal = false) OVERRIDE;
+	virtual EValidatorResult IsValid( const FString& Name, bool bOriginal = false) override;
+	virtual EValidatorResult IsValid( const FName& Name, bool bOriginal = false) override;
 	// End FNameValidatorInterface
 private:
 	/** Name set to validate */
@@ -90,8 +90,8 @@ class UNREALED_API FStringSetNameValidator : public INameValidatorInterface
 {
 public:
 	// Begin FNameValidatorInterface
-	virtual EValidatorResult IsValid(const FString& Name, bool bOriginal) OVERRIDE;
-	virtual EValidatorResult IsValid(const FName& Name, bool bOriginal) OVERRIDE;
+	virtual EValidatorResult IsValid(const FString& Name, bool bOriginal) override;
+	virtual EValidatorResult IsValid(const FName& Name, bool bOriginal) override;
 	// End FNameValidatorInterface
 
 protected:
@@ -135,8 +135,8 @@ public:
 	FDummyNameValidator(EValidatorResult InReturnValue) : ReturnValue(InReturnValue) {}
 
 	// Begin FNameValidatorInterface
-	virtual EValidatorResult IsValid(const FString& Name, bool bOriginal) OVERRIDE { return ReturnValue; }
-	virtual EValidatorResult IsValid(const FName& Name, bool bOriginal) OVERRIDE { return ReturnValue; }
+	virtual EValidatorResult IsValid(const FString& Name, bool bOriginal) override { return ReturnValue; }
+	virtual EValidatorResult IsValid(const FName& Name, bool bOriginal) override { return ReturnValue; }
 	// End FNameValidatorInterface
 
 private:
diff --git a/Engine/Source/Editor/UnrealEd/Public/LevelEditorViewport.h b/Engine/Source/Editor/UnrealEd/Public/LevelEditorViewport.h
--- a/Engine/Source/Editor/UnrealEd/Public/LevelEditorViewport.h
+++ b/Engine/Source/Editor/UnrealEd/Public/LevelEditorViewport.h
@@ -141,44 +141,44 @@ public:
 
 	////////////////////////////
 	// FViewElementDrawer interface
-	virtual void Draw(const FSceneView* View,FPrimitiveDrawInterface* PDI) OVERRIDE;
+	virtual void Draw(const FSceneView* View,FPrimitiveDrawInterface* PDI) override;
 	// End of FViewElementDrawer interface
 	
-	virtual FSceneView* CalcSceneView(FSceneViewFamily* ViewFamily) OVERRIDE;
+	virtual FSceneView* CalcSceneView(FSceneViewFamily* ViewFamily) override;
 
 	////////////////////////////
 	// FEditorViewportClient interface
-	virtual void DrawCanvas( FViewport& InViewport, FSceneView& View, FCanvas& Canvas ) OVERRIDE;
-	virtual bool InputKey(FViewport* Viewport, int32 ControllerId, FKey Key, EInputEvent Event, float AmountDepressed = 1.f, bool bGamepad=false) OVERRIDE;
-	virtual bool InputAxis(FViewport* Viewport, int32 ControllerId, FKey Key, float Delta, float DeltaTime, int32 NumSamples=1, bool bGamepad=false) OVERRIDE;
-	virtual void MouseEnter( FViewport* Viewport,int32 x, int32 y ) OVERRIDE;
-	virtual void MouseLeave( FViewport* Viewport ) OVERRIDE;
-	virtual void MouseMove(FViewport* Viewport,int32 x, int32 y) OVERRIDE;
-	virtual EMouseCursor::Type GetCursor(FViewport* Viewport,int32 X,int32 Y) OVERRIDE;
-	virtual void CapturedMouseMove( FViewport* InViewport, int32 InMouseX, int32 InMouseY ) OVERRIDE;
-	virtual void Tick(float DeltaSeconds) OVERRIDE;
-	virtual bool InputWidgetDelta( FViewport* Viewport, EAxisList::Type CurrentAxis, FVector& Drag, FRotator& Rot, FVector& Scale ) OVERRIDE;
-	virtual TSharedPtr<FDragTool> MakeDragTool( EDragTool::Type DragToolType ) OVERRIDE;
+	virtual void DrawCanvas( FViewport& InViewport, FSceneView& View, FCanvas& Canvas ) override;
+	virtual bool InputKey(FViewport* Viewport, int32 ControllerId, FKey Key, EInputEvent Event, float AmountDepressed = 1.f, bool bGamepad=false) override;
+	virtual bool InputAxis(FViewport* Viewport, int32 ControllerId, FKey Key, float Delta, float DeltaTime, int32 NumSamples=1, bool bGamepad=false) override;
+	virtual void MouseEnter( FViewport* Viewport,int32 x, int32 y ) override;
+	virtual void MouseLeave( FViewport* Viewport ) override;
+	virtual void MouseMove(FViewport* Viewport,int32 x, int32 y) override;
+	virtual EMouseCursor::Type GetCursor(FViewport* Viewport,int32 X,int32 Y) override;
+	virtual void CapturedMouseMove( FViewport* InViewport, int32 InMouseX, int32 InMouseY ) override;
+	virtual void Tick(float DeltaSeconds) override;
+	virtual bool InputWidgetDelta( FViewport* Viewport, EAxisList::Type CurrentAxis, FVector& Drag, FRotator& Rot, FVector& Scale ) override;
+	virtual TSharedPtr<FDragTool> MakeDragTool( EDragTool::Type DragToolType ) override;
 	virtual bool IsLevelEditorClient() const { return ParentLevelEditor.IsValid(); }
-	virtual void TrackingStarted( const struct FInputEventState& InInputState, bool bIsDraggingWidget, bool bNudge ) OVERRIDE;
-	virtual void TrackingStopped() OVERRIDE;
-	virtual void SetWidgetMode( FWidget::EWidgetMode NewMode ) OVERRIDE;
-	virtual bool CanSetWidgetMode( FWidget::EWidgetMode NewMode ) const OVERRIDE;
-	virtual void SetWidgetCoordSystemSpace( ECoordSystem NewCoordSystem ) OVERRIDE;
-	virtual FWidget::EWidgetMode GetWidgetMode() const OVERRIDE;
-	virtual FVector GetWidgetLocation() const OVERRIDE;
+	virtual void TrackingStarted( const struct FInputEventState& InInputState, bool bIsDraggingWidget, bool bNudge ) override;
+	virtual void TrackingStopped() override;
+	virtual void SetWidgetMode( FWidget::EWidgetMode NewMode ) override;
+	virtual bool CanSetWidgetMode( FWidget::EWidgetMode NewMode ) const override;
+	virtual void SetWidgetCoordSystemSpace( ECoordSystem NewCoordSystem ) override;
+	virtual FWidget::EWidgetMode GetWidgetMode() const override;
+	virtual FVector GetWidgetLocation() const override;
 	virtual FMatrix GetWidgetCoordSystem() const;
 	virtual ECoordSystem GetWidgetCoordSystemSpace() const;
-	virtual void SetupViewForRendering( FSceneViewFamily& ViewFamily, FSceneView& View ) OVERRIDE;
-	virtual FSceneInterface* GetScene() const OVERRIDE;
-	virtual FLinearColor GetBackgroundColor() const OVERRIDE;
-	virtual bool IsAspectRatioConstrained() const OVERRIDE;
-	virtual int32 GetCameraSpeedSetting() const OVERRIDE;
-	virtual void SetCameraSpeedSetting(int32 SpeedSetting) OVERRIDE;
-	virtual void ReceivedFocus(FViewport* Viewport) OVERRIDE;
-	virtual void LostFocus(FViewport* Viewport) OVERRIDE;
+	virtual void SetupViewForRendering( FSceneViewFamily& ViewFamily, FSceneView& View ) override;
+	virtual FSceneInterface* GetScene() const override;
+	virtual FLinearColor GetBackgroundColor() const override;
+	virtual bool IsAspectRatioConstrained() const override;
+	virtual int32 GetCameraSpeedSetting() const override;
+	virtual void SetCameraSpeedSetting(int32 SpeedSetting) override;
+	virtual void ReceivedFocus(FViewport* Viewport) override;
+	virtual void LostFocus(FViewport* Viewport) override;
 
-	virtual bool OverrideHighResScreenshotCaptureRegion(FIntRect& OutCaptureRegion) OVERRIDE;
+	virtual bool OverrideHighResScreenshotCaptureRegion(FIntRect& OutCaptureRegion) override;
 
 	void SetIsCameraCut( bool bInIsCameraCut ) { bEditorCameraCut = bInIsCameraCut; }
 
@@ -301,7 +301,7 @@ public:
 	/** Updates the rotate widget with the passed in delta rotation. */
 	void ApplyDeltaToRotateWidget( const FRotator& InRot );
 
-	virtual void SetIsSimulateInEditorViewport( bool bInIsSimulateInEditorViewport ) OVERRIDE;
+	virtual void SetIsSimulateInEditorViewport( bool bInIsSimulateInEditorViewport ) override;
 
 	/**
 	 * Draws a screen space bounding box around the specified actor
@@ -373,10 +373,10 @@ public:
 	void SetAllSpriteCategoryVisibility( bool bVisible );
 
 	/** FEditorViewportClient Interface*/
-	virtual void UpdateMouseDelta() OVERRIDE;
+	virtual void UpdateMouseDelta() override;
 	virtual void ProcessClick(FSceneView& View, HHitProxy* HitProxy, FKey Key, EInputEvent Event, uint32 HitX, uint32 HitY);
-	virtual void SetCurrentWidgetAxis( EAxisList::Type NewAxis ) OVERRIDE;
-	virtual UWorld* GetWorld() const OVERRIDE;
+	virtual void SetCurrentWidgetAxis( EAxisList::Type NewAxis ) override;
+	virtual UWorld* GetWorld() const override;
 
 	void SetReferenceToWorldContext(FWorldContext& WorldContext);
 
@@ -430,14 +430,14 @@ public:
 	 * 
 	 * @return the previous GWorld
 	 */
-	virtual UWorld* ConditionalSetWorld() OVERRIDE;
+	virtual UWorld* ConditionalSetWorld() override;
 
 	/**
 	 * Restores GWorld to InWorld
 	 *
 	 * @param InWorld	The world to restore
 	 */
-	virtual void ConditionalRestoreWorld( UWorld* InWorld  ) OVERRIDE;
+	virtual void ConditionalRestoreWorld( UWorld* InWorld  ) override;
 
 	/**
 	 *	Called to check if a material can be applied to an object, given the hit proxy
@@ -536,7 +536,7 @@ public:
 	/**
 	 * Get a ptr to the stat unit data for this viewport
 	 */
-	virtual FStatUnitData* GetStatUnitData() const OVERRIDE
+	virtual FStatUnitData* GetStatUnitData() const override
 	{
 		return StatUnitData;
 	}
@@ -544,7 +544,7 @@ public:
 	/**
 	 * Get a ptr to the stat unit data for this viewport
 	 */
-	virtual FStatHitchesData* GetStatHitchesData() const OVERRIDE
+	virtual FStatHitchesData* GetStatHitchesData() const override
 	{
 		return StatHitchesData;
 	}
@@ -552,7 +552,7 @@ public:
 	/**
 	 * Get a ptr to the enabled stats list
 	 */
-	virtual const TArray<FString>* GetEnabledStats() const OVERRIDE
+	virtual const TArray<FString>* GetEnabledStats() const override
 	{
 		return &EnabledStats;
 	}
@@ -560,7 +560,7 @@ public:
 	/**
 	 * Sets all the stats that should be enabled for the viewport
 	 */
-	virtual void SetEnabledStats(const TArray<FString>& InEnabledStats) OVERRIDE
+	virtual void SetEnabledStats(const TArray<FString>& InEnabledStats) override
 	{
 		EnabledStats = InEnabledStats;
 	}
@@ -568,7 +568,7 @@ public:
 	/**
 	 * Check whether a specific stat is enabled for this viewport
 	 */
-	virtual bool IsStatEnabled(const TCHAR* InName) const OVERRIDE
+	virtual bool IsStatEnabled(const TCHAR* InName) const override
 	{
 		return EnabledStats.Contains(InName);
 	}
@@ -576,7 +576,7 @@ public:
 	/**
 	 * Get the sound stat flags enabled for this viewport
 	 */
-	virtual ESoundShowFlags::Type GetSoundShowFlags() const OVERRIDE
+	virtual ESoundShowFlags::Type GetSoundShowFlags() const override
 	{ 
 		return SoundShowFlags;
 	}
@@ -584,7 +584,7 @@ public:
 	/**
 	 * Set the sound stat flags enabled for this viewport
 	 */
-	virtual void SetSoundShowFlags(const ESoundShowFlags::Type InSoundShowFlags) OVERRIDE
+	virtual void SetSoundShowFlags(const ESoundShowFlags::Type InSoundShowFlags) override
 	{
 		SoundShowFlags = InSoundShowFlags;
 	}
@@ -633,13 +633,13 @@ protected:
 	void OnActorMoved(AActor* InActor);
 
 	/** FEditorViewportClient Interface*/
-	virtual void UpdateLinkedOrthoViewports( bool bInvalidate = false ) OVERRIDE;
-	virtual ELevelViewportType GetViewportType() const OVERRIDE;
-	virtual void OverridePostProcessSettings( FSceneView& View ) OVERRIDE;
-	virtual void PerspectiveCameraMoved() OVERRIDE;
-	virtual bool ShouldLockPitch() const OVERRIDE;
-	virtual void CheckHoveredHitProxy( HHitProxy* HoveredHitProxy ) OVERRIDE;
-	virtual bool GetActiveSafeFrame(float& OutAspectRatio) const OVERRIDE;
+	virtual void UpdateLinkedOrthoViewports( bool bInvalidate = false ) override;
+	virtual ELevelViewportType GetViewportType() const override;
+	virtual void OverridePostProcessSettings( FSceneView& View ) override;
+	virtual void PerspectiveCameraMoved() override;
+	virtual bool ShouldLockPitch() const override;
+	virtual void CheckHoveredHitProxy( HHitProxy* HoveredHitProxy ) override;
+	virtual bool GetActiveSafeFrame(float& OutAspectRatio) const override;
 
 private:
 	/**
diff --git a/Engine/Source/Editor/UnrealEd/Public/MiniCurveEditor.h b/Engine/Source/Editor/UnrealEd/Public/MiniCurveEditor.h
--- a/Engine/Source/Editor/UnrealEd/Public/MiniCurveEditor.h
+++ b/Engine/Source/Editor/UnrealEd/Public/MiniCurveEditor.h
@@ -20,10 +20,10 @@ public:
 	virtual ~SMiniCurveEditor();
 
 	// IAssetEditorInstance interface
-	virtual FName GetEditorName() const OVERRIDE;
-	virtual void FocusWindow(UObject* ObjectToFocusOn) OVERRIDE;
-	virtual bool CloseWindow() OVERRIDE;
-	virtual bool IsPrimaryEditor() const OVERRIDE { return true; };
+	virtual FName GetEditorName() const override;
+	virtual void FocusWindow(UObject* ObjectToFocusOn) override;
+	virtual bool CloseWindow() override;
+	virtual bool IsPrimaryEditor() const override { return true; };
 private:
 
 	float ViewMinInput;
diff --git a/Engine/Source/Editor/UnrealEd/Public/NormalMapPreview.h b/Engine/Source/Editor/UnrealEd/Public/NormalMapPreview.h
--- a/Engine/Source/Editor/UnrealEd/Public/NormalMapPreview.h
+++ b/Engine/Source/Editor/UnrealEd/Public/NormalMapPreview.h
@@ -19,5 +19,5 @@ class UNREALED_API FNormalMapBatchedElementParameters : public FBatchedElementPa
 		const FMatrix& InTransform,
 		const float InGamma,
 		const FMatrix& ColorWeights,
-		const FTexture* Texture) OVERRIDE;
+		const FTexture* Texture) override;
 };
diff --git a/Engine/Source/Editor/UnrealEd/Public/ObjectTools.h b/Engine/Source/Editor/UnrealEd/Public/ObjectTools.h
--- a/Engine/Source/Editor/UnrealEd/Public/ObjectTools.h
+++ b/Engine/Source/Editor/UnrealEd/Public/ObjectTools.h
@@ -150,7 +150,7 @@ namespace ObjectTools
 	struct UNREALED_API FConsolidationResults : public FGCObject
 	{
 		/** FGCObject interface; Serialize any object references */
-		virtual void AddReferencedObjects( FReferenceCollector& Collector ) OVERRIDE
+		virtual void AddReferencedObjects( FReferenceCollector& Collector ) override
 		{
 			for( int32 Index = 0; Index < DirtiedPackages.Num(); Index++ )
 			{
diff --git a/Engine/Source/Editor/UnrealEd/Public/SColorGradientEditor.h b/Engine/Source/Editor/UnrealEd/Public/SColorGradientEditor.h
--- a/Engine/Source/Editor/UnrealEd/Public/SColorGradientEditor.h
+++ b/Engine/Source/Editor/UnrealEd/Public/SColorGradientEditor.h
@@ -45,12 +45,12 @@ public:
 
 	/** SWidget Interface */
 	virtual bool SupportsKeyboardFocus() const { return true; }
-	virtual int32 OnPaint( const FGeometry& AllottedGeometry, const FSlateRect& MyClippingRect, FSlateWindowElementList& OutDrawElements, int32 LayerId, const FWidgetStyle& InWidgetStyle, bool bParentEnabled ) const OVERRIDE;
-	virtual FReply OnMouseButtonDown( const FGeometry& MyGeometry, const FPointerEvent& MouseEvent ) OVERRIDE;
-	virtual FReply OnMouseButtonDoubleClick( const FGeometry& InMyGeometry, const FPointerEvent& InMouseEvent ) OVERRIDE;
-	virtual FReply OnMouseMove( const FGeometry& MyGeometry, const FPointerEvent& MouseEvent ) OVERRIDE;
-	virtual FReply OnMouseButtonUp( const FGeometry& MyGeometry, const FPointerEvent& MouseEvent ) OVERRIDE;
-	virtual FReply OnKeyDown( const FGeometry& MyGeometry, const FKeyboardEvent& InKeyboardEvent ) OVERRIDE;
+	virtual int32 OnPaint( const FGeometry& AllottedGeometry, const FSlateRect& MyClippingRect, FSlateWindowElementList& OutDrawElements, int32 LayerId, const FWidgetStyle& InWidgetStyle, bool bParentEnabled ) const override;
+	virtual FReply OnMouseButtonDown( const FGeometry& MyGeometry, const FPointerEvent& MouseEvent ) override;
+	virtual FReply OnMouseButtonDoubleClick( const FGeometry& InMyGeometry, const FPointerEvent& InMouseEvent ) override;
+	virtual FReply OnMouseMove( const FGeometry& MyGeometry, const FPointerEvent& MouseEvent ) override;
+	virtual FReply OnMouseButtonUp( const FGeometry& MyGeometry, const FPointerEvent& MouseEvent ) override;
+	virtual FReply OnKeyDown( const FGeometry& MyGeometry, const FKeyboardEvent& InKeyboardEvent ) override;
 	virtual FVector2D ComputeDesiredSize() const;
 
 	/**
diff --git a/Engine/Source/Editor/UnrealEd/Public/SCurveEditor.h b/Engine/Source/Editor/UnrealEd/Public/SCurveEditor.h
--- a/Engine/Source/Editor/UnrealEd/Public/SCurveEditor.h
+++ b/Engine/Source/Editor/UnrealEd/Public/SCurveEditor.h
@@ -289,7 +289,7 @@ private:
 	FSimpleDelegate OnCreateAsset;
 
 	// SWidget interface
-	UNREALED_API virtual FVector2D ComputeDesiredSize() const OVERRIDE;
+	UNREALED_API virtual FVector2D ComputeDesiredSize() const override;
 
 	/** Paint a curve */
 	void PaintCurve(	FRichCurve* Curve, const FGeometry &AllottedGeometry, FTrackScaleInfo &ScaleInfo, FSlateWindowElementList &OutDrawElements, 
@@ -306,14 +306,14 @@ private:
 	void PaintGridLines( const FGeometry &AllottedGeometry, FTrackScaleInfo &ScaleInfo, FSlateWindowElementList &OutDrawElements, int32 LayerId, const FSlateRect& MyClippingRect, ESlateDrawEffect::Type DrawEffects )const;
 
 	// SWidget interface
-	UNREALED_API virtual FReply OnMouseButtonDown( const FGeometry& MyGeometry, const FPointerEvent& MouseEvent ) OVERRIDE;
-	UNREALED_API virtual FReply OnMouseButtonUp( const FGeometry& MyGeometry, const FPointerEvent& MouseEvent ) OVERRIDE;
-	UNREALED_API virtual FReply OnMouseMove( const FGeometry& MyGeometry, const FPointerEvent& MouseEvent ) OVERRIDE;
-	UNREALED_API virtual FReply OnMouseWheel( const FGeometry& MyGeometry, const FPointerEvent& MouseEvent ) OVERRIDE;
-	UNREALED_API virtual FReply OnKeyDown( const FGeometry& MyGeometry, const FKeyboardEvent& InKeyboardEvent ) OVERRIDE;
-	UNREALED_API virtual void   OnMouseCaptureLost() OVERRIDE;
-
-	virtual bool SupportsKeyboardFocus() const OVERRIDE
+	UNREALED_API virtual FReply OnMouseButtonDown( const FGeometry& MyGeometry, const FPointerEvent& MouseEvent ) override;
+	UNREALED_API virtual FReply OnMouseButtonUp( const FGeometry& MyGeometry, const FPointerEvent& MouseEvent ) override;
+	UNREALED_API virtual FReply OnMouseMove( const FGeometry& MyGeometry, const FPointerEvent& MouseEvent ) override;
+	UNREALED_API virtual FReply OnMouseWheel( const FGeometry& MyGeometry, const FPointerEvent& MouseEvent ) override;
+	UNREALED_API virtual FReply OnKeyDown( const FGeometry& MyGeometry, const FKeyboardEvent& InKeyboardEvent ) override;
+	UNREALED_API virtual void   OnMouseCaptureLost() override;
+
+	virtual bool SupportsKeyboardFocus() const override
 	{
 		return true;
 	}
@@ -374,8 +374,8 @@ private:
 	void	RedoAction();
 
 	// Begin FEditorUndoClient Interface
-	UNREALED_API virtual void PostUndo(bool bSuccess) OVERRIDE;
-	UNREALED_API virtual void PostRedo(bool bSuccess) OVERRIDE { PostUndo(bSuccess); }
+	UNREALED_API virtual void PostUndo(bool bSuccess) override;
+	UNREALED_API virtual void PostRedo(bool bSuccess) override { PostUndo(bSuccess); }
 	// End of FEditorUndoClient
 
 	bool AreCurvesVisible() const { return bAlwaysDisplayColorCurves || bAreCurvesVisible; }
diff --git a/Engine/Source/Editor/UnrealEd/Public/SEditorViewport.h b/Engine/Source/Editor/UnrealEd/Public/SEditorViewport.h
--- a/Engine/Source/Editor/UnrealEd/Public/SEditorViewport.h
+++ b/Engine/Source/Editor/UnrealEd/Public/SEditorViewport.h
@@ -14,9 +14,9 @@ public:
 
 	void Construct( const FArguments& InArgs );
 
-	virtual FReply OnKeyDown( const FGeometry& MyGeometry, const FKeyboardEvent& InKeyboardEvent ) OVERRIDE;
-	virtual bool SupportsKeyboardFocus() const OVERRIDE;
-	virtual FReply OnKeyboardFocusReceived( const FGeometry& MyGeometry, const FKeyboardFocusEvent& InKeyboardFocusEvent ) OVERRIDE;
+	virtual FReply OnKeyDown( const FGeometry& MyGeometry, const FKeyboardEvent& InKeyboardEvent ) override;
+	virtual bool SupportsKeyboardFocus() const override;
+	virtual FReply OnKeyboardFocusReceived( const FGeometry& MyGeometry, const FKeyboardFocusEvent& InKeyboardFocusEvent ) override;
 
 	/**
 	 * @return True if the viewport is being updated in realtime
diff --git a/Engine/Source/Editor/UnrealEd/Public/SListViewSelectorDropdownMenu.h b/Engine/Source/Editor/UnrealEd/Public/SListViewSelectorDropdownMenu.h
--- a/Engine/Source/Editor/UnrealEd/Public/SListViewSelectorDropdownMenu.h
+++ b/Engine/Source/Editor/UnrealEd/Public/SListViewSelectorDropdownMenu.h
@@ -32,7 +32,7 @@ public:
 	}
 
 	// SWidget interface
-	virtual FReply OnKeyDown(const FGeometry& MyGeometry, const FKeyboardEvent& KeyboardEvent) OVERRIDE
+	virtual FReply OnKeyDown(const FGeometry& MyGeometry, const FKeyboardEvent& KeyboardEvent) override
 	{
 		TSharedPtr< SListView< ListType > > TargetListViewPtr = TargetListView.Pin();
 
diff --git a/Engine/Source/Editor/UnrealEd/Public/SViewportToolBar.h b/Engine/Source/Editor/UnrealEd/Public/SViewportToolBar.h
--- a/Engine/Source/Editor/UnrealEd/Public/SViewportToolBar.h
+++ b/Engine/Source/Editor/UnrealEd/Public/SViewportToolBar.h
@@ -38,14 +38,14 @@ private:
 	 * @param MyGeometry	Information about the size of the toolbar
 	 * @param MouseEvent	The mouse event that triggered this function
 	 */
-	virtual void OnMouseEnter( const FGeometry& MyGeometry, const FPointerEvent& MouseEvent ) OVERRIDE;
+	virtual void OnMouseEnter( const FGeometry& MyGeometry, const FPointerEvent& MouseEvent ) override;
 
 	/**
 	 * Called when the mouse leaves the toolbar area.  We fade out the toolbar when this happens
 	 *
 	 * @param MouseEvent	The mouse event that triggered this function
 	 */
-	virtual void OnMouseLeave( const FPointerEvent& MouseEvent ) OVERRIDE;
+	virtual void OnMouseLeave( const FPointerEvent& MouseEvent ) override;
 
 private:
 	/** Curve sequence for fading out the toolbar */
diff --git a/Engine/Source/Editor/UnrealEd/Public/SoundCueGraphEditorCommands.h b/Engine/Source/Editor/UnrealEd/Public/SoundCueGraphEditorCommands.h
--- a/Engine/Source/Editor/UnrealEd/Public/SoundCueGraphEditorCommands.h
+++ b/Engine/Source/Editor/UnrealEd/Public/SoundCueGraphEditorCommands.h
@@ -41,5 +41,5 @@ public:
 	TSharedPtr<FUICommandInfo> DeleteInput;
 
 	/** Initialize commands */
-	UNREALED_API virtual void RegisterCommands() OVERRIDE;
+	UNREALED_API virtual void RegisterCommands() override;
 };
diff --git a/Engine/Source/Editor/UnrealEd/Public/Texture2DPreview.h b/Engine/Source/Editor/UnrealEd/Public/Texture2DPreview.h
--- a/Engine/Source/Editor/UnrealEd/Public/Texture2DPreview.h
+++ b/Engine/Source/Editor/UnrealEd/Public/Texture2DPreview.h
@@ -22,7 +22,7 @@ public:
 	}
 
 	/** Binds vertex and pixel shaders for this element */
-	virtual void BindShaders_RenderThread(FRHICommandList& RHICmdList, const FMatrix& InTransform, const float InGamma, const FMatrix& ColorWeights, const FTexture* Texture ) OVERRIDE;
+	virtual void BindShaders_RenderThread(FRHICommandList& RHICmdList, const FMatrix& InTransform, const float InGamma, const FMatrix& ColorWeights, const FTexture* Texture ) override;
 
 private:
 
diff --git a/Engine/Source/Editor/UnrealEd/Public/ThumbnailHelpers.h b/Engine/Source/Editor/UnrealEd/Public/ThumbnailHelpers.h
--- a/Engine/Source/Editor/UnrealEd/Public/ThumbnailHelpers.h
+++ b/Engine/Source/Editor/UnrealEd/Public/ThumbnailHelpers.h
@@ -43,7 +43,7 @@ public:
 
 protected:
 	// FThumbnailPreviewScene implementation
-	virtual void GetViewMatrixParameters(const float InFOVDegrees, FVector& OutOrigin, float& OutOrbitPitch, float& OutOrbitYaw, float& OutOrbitZoom) const OVERRIDE;
+	virtual void GetViewMatrixParameters(const float InFOVDegrees, FVector& OutOrigin, float& OutOrbitPitch, float& OutOrbitYaw, float& OutOrbitZoom) const override;
 
 protected:
 	/** The particle system component used to display all particle system thumbnails */
@@ -64,7 +64,7 @@ public:
 
 protected:
 	// FThumbnailPreviewScene implementation
-	virtual void GetViewMatrixParameters(const float InFOVDegrees, FVector& OutOrigin, float& OutOrbitPitch, float& OutOrbitYaw, float& OutOrbitZoom) const OVERRIDE;
+	virtual void GetViewMatrixParameters(const float InFOVDegrees, FVector& OutOrigin, float& OutOrbitPitch, float& OutOrbitYaw, float& OutOrbitZoom) const override;
 
 protected:
 	/** The static mesh actor used to display all material thumbnails */
@@ -82,7 +82,7 @@ public:
 
 protected:
 	// FThumbnailPreviewScene implementation
-	virtual void GetViewMatrixParameters(const float InFOVDegrees, FVector& OutOrigin, float& OutOrbitPitch, float& OutOrbitYaw, float& OutOrbitZoom) const OVERRIDE;
+	virtual void GetViewMatrixParameters(const float InFOVDegrees, FVector& OutOrigin, float& OutOrbitPitch, float& OutOrbitYaw, float& OutOrbitZoom) const override;
 
 private:
 	/** The skeletal mesh actor used to display all skeletal mesh thumbnails */
@@ -100,7 +100,7 @@ public:
 
 protected:
 	// FThumbnailPreviewScene implementation
-	virtual void GetViewMatrixParameters(const float InFOVDegrees, FVector& OutOrigin, float& OutOrbitPitch, float& OutOrbitYaw, float& OutOrbitZoom) const OVERRIDE;
+	virtual void GetViewMatrixParameters(const float InFOVDegrees, FVector& OutOrigin, float& OutOrbitPitch, float& OutOrbitYaw, float& OutOrbitZoom) const override;
 
 private:
 	/** The static mesh actor used to display all static mesh thumbnails */
@@ -125,11 +125,11 @@ public:
 	/** Refreshes components for the specified blueprint */
 	void BlueprintChanged(class UBlueprint* Blueprint);
 
-	void AddReferencedObjects( FReferenceCollector& Collector ) OVERRIDE;
+	void AddReferencedObjects( FReferenceCollector& Collector ) override;
 
 protected:
 	// FThumbnailPreviewScene implementation
-	virtual void GetViewMatrixParameters(const float InFOVDegrees, FVector& OutOrigin, float& OutOrbitPitch, float& OutOrbitYaw, float& OutOrbitZoom) const OVERRIDE;
+	virtual void GetViewMatrixParameters(const float InFOVDegrees, FVector& OutOrigin, float& OutOrbitPitch, float& OutOrbitYaw, float& OutOrbitZoom) const override;
 
 	/** Returns a duplicate of the specified component whose outer is the transient package. If the component can not be created, a placeholder component is made in its place. */
 	UActorComponent* CreateComponentInstanceFromTemplate(UActorComponent* ComponentTemplate) const;
diff --git a/Engine/Source/Editor/UnrealEd/Public/Toolkits/AssetEditorToolkit.h b/Engine/Source/Editor/UnrealEd/Public/Toolkits/AssetEditorToolkit.h
--- a/Engine/Source/Editor/UnrealEd/Public/Toolkits/AssetEditorToolkit.h
+++ b/Engine/Source/Editor/UnrealEd/Public/Toolkits/AssetEditorToolkit.h
@@ -40,21 +40,21 @@ public:
 	virtual ~FAssetEditorToolkit();
 
 	/** IToolkit interface */
-	virtual void RegisterTabSpawners(const TSharedRef<class FTabManager>& TabManager) OVERRIDE;
-	virtual void UnregisterTabSpawners(const TSharedRef<class FTabManager>& TabManager) OVERRIDE;
-	virtual bool IsAssetEditor() const OVERRIDE;
-	virtual const TArray< UObject* >* GetObjectsCurrentlyBeingEdited() const OVERRIDE;
+	virtual void RegisterTabSpawners(const TSharedRef<class FTabManager>& TabManager) override;
+	virtual void UnregisterTabSpawners(const TSharedRef<class FTabManager>& TabManager) override;
+	virtual bool IsAssetEditor() const override;
+	virtual const TArray< UObject* >* GetObjectsCurrentlyBeingEdited() const override;
 	virtual FName GetToolkitFName() const = 0;					// Must implement in derived class!
 	virtual FText GetBaseToolkitName() const = 0;				// Must implement in derived class!
-	virtual FText GetToolkitName() const OVERRIDE;		
+	virtual FText GetToolkitName() const override;		
 	virtual FString GetWorldCentricTabPrefix() const = 0;		// Must implement in derived class!
-	virtual class FEdMode* GetEditorMode() const OVERRIDE;
+	virtual class FEdMode* GetEditorMode() const override;
 
 	/** IAssetEditorInstance interface */
-	virtual FName GetEditorName() const OVERRIDE;
-	virtual void FocusWindow(UObject* ObjectToFocusOn = NULL) OVERRIDE;
-	virtual bool CloseWindow() OVERRIDE;
-	virtual bool IsPrimaryEditor() const OVERRIDE { return true; };
+	virtual FName GetEditorName() const override;
+	virtual void FocusWindow(UObject* ObjectToFocusOn = NULL) override;
+	virtual bool CloseWindow() override;
+	virtual bool IsPrimaryEditor() const override { return true; };
 
 	/**
 	 * Fills in the supplied menu with commands for working with this asset file
diff --git a/Engine/Source/Editor/UnrealEd/Public/Toolkits/BaseToolkit.h b/Engine/Source/Editor/UnrealEd/Public/Toolkits/BaseToolkit.h
--- a/Engine/Source/Editor/UnrealEd/Public/Toolkits/BaseToolkit.h
+++ b/Engine/Source/Editor/UnrealEd/Public/Toolkits/BaseToolkit.h
@@ -20,13 +20,13 @@ public:
 	virtual ~FBaseToolkit();
 
 	/** IToolkit interface */
-	virtual bool ProcessCommandBindings( const FKeyboardEvent& InKeyboardEvent ) const OVERRIDE;
-	virtual bool IsHosted() const OVERRIDE;
-	virtual const TSharedRef< class IToolkitHost > GetToolkitHost() const OVERRIDE;
-	virtual const TMap< EToolkitTabSpot::Type, TArray< TWeakPtr< SDockableTab > > >& GetToolkitTabsInSpots() const OVERRIDE;
-	virtual void BringToolkitToFront() OVERRIDE;
-	virtual TSharedPtr<class SWidget> GetInlineContent() const OVERRIDE;
-	virtual bool IsBlueprintEditor() const OVERRIDE;
+	virtual bool ProcessCommandBindings( const FKeyboardEvent& InKeyboardEvent ) const override;
+	virtual bool IsHosted() const override;
+	virtual const TSharedRef< class IToolkitHost > GetToolkitHost() const override;
+	virtual const TMap< EToolkitTabSpot::Type, TArray< TWeakPtr< SDockableTab > > >& GetToolkitTabsInSpots() const override;
+	virtual void BringToolkitToFront() override;
+	virtual TSharedPtr<class SWidget> GetInlineContent() const override;
+	virtual bool IsBlueprintEditor() const override;
 
 	/** @return	Returns true if this is a world-centric asset editor.  That is, the user is editing the asset inline in a Level Editor app. */
 	bool IsWorldCentricAssetEditor() const;	
@@ -80,18 +80,18 @@ protected:
 class UNREALED_API FModeToolkit : public FBaseToolkit, public TSharedFromThis<FModeToolkit>
 {
 public:
-	virtual void RegisterTabSpawners(const TSharedRef<class FTabManager>& TabManager) OVERRIDE {}
-	virtual void UnregisterTabSpawners(const TSharedRef<class FTabManager>& TabManager) OVERRIDE {}
+	virtual void RegisterTabSpawners(const TSharedRef<class FTabManager>& TabManager) override {}
+	virtual void UnregisterTabSpawners(const TSharedRef<class FTabManager>& TabManager) override {}
 
 	/** Initializes the mode toolkit */
 	virtual void Init(const TSharedPtr< class IToolkitHost >& InitToolkitHost);
 
 	/** IToolkit interface */
-	virtual FText GetToolkitName() const OVERRIDE { return GetBaseToolkitName(); }
-	virtual FString GetWorldCentricTabPrefix() const OVERRIDE;
-	virtual bool IsAssetEditor() const OVERRIDE;
-	virtual const TArray< UObject* >* GetObjectsCurrentlyBeingEdited() const OVERRIDE;
-	virtual FLinearColor GetWorldCentricTabColorScale() const OVERRIDE;
+	virtual FText GetToolkitName() const override { return GetBaseToolkitName(); }
+	virtual FString GetWorldCentricTabPrefix() const override;
+	virtual bool IsAssetEditor() const override;
+	virtual const TArray< UObject* >* GetObjectsCurrentlyBeingEdited() const override;
+	virtual FLinearColor GetWorldCentricTabColorScale() const override;
 };
 
 
diff --git a/Engine/Source/Editor/UnrealEd/Public/Toolkits/SimpleAssetEditor.h b/Engine/Source/Editor/UnrealEd/Public/Toolkits/SimpleAssetEditor.h
--- a/Engine/Source/Editor/UnrealEd/Public/Toolkits/SimpleAssetEditor.h
+++ b/Engine/Source/Editor/UnrealEd/Public/Toolkits/SimpleAssetEditor.h
@@ -12,8 +12,8 @@ public:
 	/** Delegate that, given an array of assets, returns an array of objects to use in the details view of an FSimpleAssetEditor */
 	DECLARE_DELEGATE_RetVal_OneParam(TArray<UObject*>, FGetDetailsViewObjects, const TArray<UObject*>&);
 
-	virtual void RegisterTabSpawners(const TSharedRef<class FTabManager>& TabManager) OVERRIDE;
-	virtual void UnregisterTabSpawners(const TSharedRef<class FTabManager>& TabManager) OVERRIDE;
+	virtual void RegisterTabSpawners(const TSharedRef<class FTabManager>& TabManager) override;
+	virtual void UnregisterTabSpawners(const TSharedRef<class FTabManager>& TabManager) override;
 
 
 	/**
@@ -30,12 +30,12 @@ public:
 	virtual ~FSimpleAssetEditor();
 
 	/** IToolkit interface */
-	virtual FName GetToolkitFName() const OVERRIDE;
-	virtual FText GetBaseToolkitName() const OVERRIDE;
-	virtual FText GetToolkitName() const OVERRIDE;
-	virtual FString GetWorldCentricTabPrefix() const OVERRIDE;
-	virtual FLinearColor GetWorldCentricTabColorScale() const OVERRIDE;
-	virtual bool IsPrimaryEditor() const OVERRIDE { return false; }
+	virtual FName GetToolkitFName() const override;
+	virtual FText GetBaseToolkitName() const override;
+	virtual FText GetToolkitName() const override;
+	virtual FString GetWorldCentricTabPrefix() const override;
+	virtual FLinearColor GetWorldCentricTabColorScale() const override;
+	virtual bool IsPrimaryEditor() const override { return false; }
 	
 	/** Used to show or hide certain properties */
 	void SetPropertyVisibilityDelegate(FIsPropertyVisible InVisibilityDelegate);
diff --git a/Engine/Source/Editor/UnrealEdMessages/Private/UnrealEdMessagesModule.cpp b/Engine/Source/Editor/UnrealEdMessages/Private/UnrealEdMessagesModule.cpp
--- a/Engine/Source/Editor/UnrealEdMessages/Private/UnrealEdMessagesModule.cpp
+++ b/Engine/Source/Editor/UnrealEdMessages/Private/UnrealEdMessagesModule.cpp
@@ -16,11 +16,11 @@ class FUnrealEdMessagesModule
 {
 public:
 
-	virtual void StartupModule( ) OVERRIDE { }
+	virtual void StartupModule( ) override { }
 
-	virtual void ShutdownModule( ) OVERRIDE { }
+	virtual void ShutdownModule( ) override { }
 
-	virtual bool SupportsDynamicReloading( ) OVERRIDE
+	virtual bool SupportsDynamicReloading( ) override
 	{
 		return true;
 	}
diff --git a/Engine/Source/Editor/UserFeedback/Private/UserFeedback.cpp b/Engine/Source/Editor/UserFeedback/Private/UserFeedback.cpp
--- a/Engine/Source/Editor/UserFeedback/Private/UserFeedback.cpp
+++ b/Engine/Source/Editor/UserFeedback/Private/UserFeedback.cpp
@@ -181,7 +181,7 @@ public:
 	}
 
 	/** Called when a key is pressed on the widget */
-	virtual FReply OnPreviewKeyDown(const FGeometry& MyGeometry, const FKeyboardEvent& InKeyboardEvent) OVERRIDE
+	virtual FReply OnPreviewKeyDown(const FGeometry& MyGeometry, const FKeyboardEvent& InKeyboardEvent) override
 	{
 		if (InKeyboardEvent.GetKey() == EKeys::Escape)
 		{
@@ -474,7 +474,7 @@ public:
 	}
 
 	/** Compute the desired size for this widget */
-	virtual FVector2D ComputeDesiredSize() const OVERRIDE
+	virtual FVector2D ComputeDesiredSize() const override
 	{
 		const float MinWidthVal = MinWidth.Get();
 
@@ -623,7 +623,7 @@ class FUserFeedbackModuleImpl : public IUserFeedbackModule
 public:
 
 	/** Create a widget which allows the user to send positive or negative feedback about a feature */
-	virtual TSharedRef<SWidget> CreateFeedbackWidget(FText Context) const OVERRIDE
+	virtual TSharedRef<SWidget> CreateFeedbackWidget(FText Context) const override
 	{
 		if (FEngineAnalytics::IsAvailable() )
 		{
diff --git a/Engine/Source/Editor/ViewportSnapping/Private/ViewportSnappingModule.cpp b/Engine/Source/Editor/ViewportSnapping/Private/ViewportSnappingModule.cpp
--- a/Engine/Source/Editor/ViewportSnapping/Private/ViewportSnappingModule.cpp
+++ b/Engine/Source/Editor/ViewportSnapping/Private/ViewportSnappingModule.cpp
@@ -11,7 +11,7 @@ class FMergedSnappingPolicy : public ISnappingPolicy
 public:
 	TArray< TSharedPtr<ISnappingPolicy> > PolicyList;
 public:
-	virtual void SnapScale(FVector& Point, const FVector& GridBase) OVERRIDE
+	virtual void SnapScale(FVector& Point, const FVector& GridBase) override
 	{
 		for (auto PolicyIt = PolicyList.CreateConstIterator(); PolicyIt; ++PolicyIt)
 		{
@@ -19,7 +19,7 @@ public:
 		}
 	}
 
-	virtual void SnapPointToGrid(FVector& Point, const FVector& GridBase) OVERRIDE
+	virtual void SnapPointToGrid(FVector& Point, const FVector& GridBase) override
 	{
 		for (auto PolicyIt = PolicyList.CreateConstIterator(); PolicyIt; ++PolicyIt)
 		{
@@ -27,7 +27,7 @@ public:
 		}
 	}
 
-	virtual void SnapRotatorToGrid(FRotator& Rotation) OVERRIDE
+	virtual void SnapRotatorToGrid(FRotator& Rotation) override
 	{
 		for (auto PolicyIt = PolicyList.CreateConstIterator(); PolicyIt; ++PolicyIt)
 		{
@@ -35,7 +35,7 @@ public:
 		}
 	}
 
-	virtual void ClearSnappingHelpers(bool bClearImmediately) OVERRIDE
+	virtual void ClearSnappingHelpers(bool bClearImmediately) override
 	{
 		for (auto PolicyIt = PolicyList.CreateConstIterator(); PolicyIt; ++PolicyIt)
 		{
@@ -65,7 +65,7 @@ public:
 	}
 
 	// IViewportSnappingModule interface
-	virtual void RegisterSnappingPolicy(TSharedPtr<ISnappingPolicy> NewPolicy) OVERRIDE
+	virtual void RegisterSnappingPolicy(TSharedPtr<ISnappingPolicy> NewPolicy) override
 	{
 		MergedPolicy->PolicyList.Add(NewPolicy);
 	}
@@ -75,7 +75,7 @@ public:
 		MergedPolicy->PolicyList.Remove(PolicyToRemove);
 	}
 	
-	virtual TSharedPtr<ISnappingPolicy> GetMergedPolicy() OVERRIDE
+	virtual TSharedPtr<ISnappingPolicy> GetMergedPolicy() override
 	{
 		return MergedPolicy;
 	}
@@ -83,12 +83,12 @@ public:
 	// End of IViewportSnappingModule interface
 
 	// IModuleInterface interface
-	virtual void StartupModule() OVERRIDE
+	virtual void StartupModule() override
 	{
 		MergedPolicy = MakeShareable(new FMergedSnappingPolicy);
 	}
 
-	virtual void ShutdownModule() OVERRIDE
+	virtual void ShutdownModule() override
 	{
 		MergedPolicy.Reset();
 	}
diff --git a/Engine/Source/Editor/WorkspaceMenuStructure/Private/WorkspaceMenuStructureModule.cpp b/Engine/Source/Editor/WorkspaceMenuStructure/Private/WorkspaceMenuStructureModule.cpp
--- a/Engine/Source/Editor/WorkspaceMenuStructure/Private/WorkspaceMenuStructureModule.cpp
+++ b/Engine/Source/Editor/WorkspaceMenuStructure/Private/WorkspaceMenuStructureModule.cpp
@@ -16,47 +16,47 @@ IMPLEMENT_MODULE( FWorkspaceMenuStructureModule, WorkspaceMenuStructure );
 class FWorkspaceMenuStructure : public IWorkspaceMenuStructure
 {
 public:
-	virtual TSharedRef<FWorkspaceItem> GetStructureRoot() const OVERRIDE
+	virtual TSharedRef<FWorkspaceItem> GetStructureRoot() const override
 	{
 		return MenuRoot.ToSharedRef();
 	}
 
-	virtual TSharedRef<FWorkspaceItem> GetAssetEditorCategory() const OVERRIDE
+	virtual TSharedRef<FWorkspaceItem> GetAssetEditorCategory() const override
 	{
 		return AssetEditorCategory.ToSharedRef();
 	}
 
-	virtual TSharedRef<FWorkspaceItem> GetLevelEditorCategory() const OVERRIDE
+	virtual TSharedRef<FWorkspaceItem> GetLevelEditorCategory() const override
 	{
 		return LevelEditorCategory.ToSharedRef();
 	}
 
-	virtual TSharedRef<FWorkspaceItem> GetLevelEditorViewportsCategory() const OVERRIDE
+	virtual TSharedRef<FWorkspaceItem> GetLevelEditorViewportsCategory() const override
 	{
 		return LevelEditorViewportsCategory.ToSharedRef();
 	}
 
-	virtual TSharedRef<FWorkspaceItem> GetLevelEditorDetailsCategory() const OVERRIDE
+	virtual TSharedRef<FWorkspaceItem> GetLevelEditorDetailsCategory() const override
 	{
 		return LevelEditorDetailsCategory.ToSharedRef();
 	}
 
-	virtual TSharedRef<FWorkspaceItem> GetLevelEditorModesCategory() const OVERRIDE
+	virtual TSharedRef<FWorkspaceItem> GetLevelEditorModesCategory() const override
 	{
 		return LevelEditorModesCategory.ToSharedRef();
 	}
 
-	virtual TSharedRef<FWorkspaceItem> GetToolsCategory() const OVERRIDE
+	virtual TSharedRef<FWorkspaceItem> GetToolsCategory() const override
 	{
 		return ToolsCategory.ToSharedRef();
 	}
 
-	virtual TSharedRef<FWorkspaceItem> GetDeveloperToolsCategory() const OVERRIDE
+	virtual TSharedRef<FWorkspaceItem> GetDeveloperToolsCategory() const override
 	{
 		return DeveloperToolsCategory.ToSharedRef();
 	}
 
-	virtual TSharedRef<FWorkspaceItem> GetEditOptions() const OVERRIDE
+	virtual TSharedRef<FWorkspaceItem> GetEditOptions() const override
 	{
 		return EditOptions.ToSharedRef();
 	}
diff --git a/Engine/Source/Editor/WorldBrowser/Private/LevelCollectionModel.cpp b/Engine/Source/Editor/WorldBrowser/Private/LevelCollectionModel.cpp
--- a/Engine/Source/Editor/WorldBrowser/Private/LevelCollectionModel.cpp
+++ b/Engine/Source/Editor/WorldBrowser/Private/LevelCollectionModel.cpp
@@ -1230,7 +1230,7 @@ FLevelModelList FLevelCollectionModel::GetLevelsHierarchy(const FLevelModelList&
 {
 	struct FHierarchyCollector : public FLevelModelVisitor
 	{
-		virtual void Visit(FLevelModel& Item) OVERRIDE
+		virtual void Visit(FLevelModel& Item) override
 		{
 			ResultList.AddUnique(Item.AsShared());
 		}
@@ -1510,7 +1510,7 @@ void FLevelCollectionModel::ExpandSelectedItems_Executed()
 {
 	struct FExpandLevelVisitor : public FLevelModelVisitor
 	{
-		virtual void Visit(FLevelModel& Item) OVERRIDE { Item.SetLevelExpansionFlag(true); }
+		virtual void Visit(FLevelModel& Item) override { Item.SetLevelExpansionFlag(true); }
 	} Expander;
 	
 	for (TSharedPtr<FLevelModel> LevelModel: SelectedLevelsList)
diff --git a/Engine/Source/Editor/WorldBrowser/Private/LevelCollectionModel.h b/Engine/Source/Editor/WorldBrowser/Private/LevelCollectionModel.h
--- a/Engine/Source/Editor/WorldBrowser/Private/LevelCollectionModel.h
+++ b/Engine/Source/Editor/WorldBrowser/Private/LevelCollectionModel.h
@@ -26,9 +26,9 @@ public:
 	virtual ~FLevelCollectionModel();
 
 	/** FTickableEditorObject interface */
-	void Tick( float DeltaTime ) OVERRIDE;
-	bool IsTickable() const OVERRIDE { return true; }
-	TStatId GetStatId() const OVERRIDE;
+	void Tick( float DeltaTime ) override;
+	bool IsTickable() const override { return true; }
+	TStatId GetStatId() const override;
 	/** FTickableEditorObject interface */
 	
 	/**	@return	Whether level collection is read only now */
diff --git a/Engine/Source/Editor/WorldBrowser/Private/SPropertyEditorLevelPackage.cpp b/Engine/Source/Editor/WorldBrowser/Private/SPropertyEditorLevelPackage.cpp
--- a/Engine/Source/Editor/WorldBrowser/Private/SPropertyEditorLevelPackage.cpp
+++ b/Engine/Source/Editor/WorldBrowser/Private/SPropertyEditorLevelPackage.cpp
@@ -205,7 +205,7 @@ void SPropertyEditorLevelPackage::PopulatePackages()
 			, Output(InOutput)
 		{}
 
-		virtual bool Visit(const TCHAR* FilenameOrDirectory, bool bIsDirectory) OVERRIDE
+		virtual bool Visit(const TCHAR* FilenameOrDirectory, bool bIsDirectory) override
 		{
 			FString FullPath = FilenameOrDirectory;
 
diff --git a/Engine/Source/Editor/WorldBrowser/Private/SWorldHierarchy.cpp b/Engine/Source/Editor/WorldBrowser/Private/SWorldHierarchy.cpp
--- a/Engine/Source/Editor/WorldBrowser/Private/SWorldHierarchy.cpp
+++ b/Engine/Source/Editor/WorldBrowser/Private/SWorldHierarchy.cpp
@@ -205,7 +205,7 @@ public:
 		struct FExpander : public FLevelModelVisitor 
 		{
 			TSharedPtr<SLevelsTreeWidget> TreeWidget;
-			virtual void Visit(FLevelModel& Item) OVERRIDE
+			virtual void Visit(FLevelModel& Item) override
 			{
 				TreeWidget->SetItemExpansion(Item.AsShared(), Item.GetLevelExpansionFlag());
 			};
@@ -309,7 +309,7 @@ private:
 	 *
 	 * @return  True if this widget can take keyboard focus
 	 */
-	virtual bool SupportsKeyboardFocus() const OVERRIDE
+	virtual bool SupportsKeyboardFocus() const override
 	{
 		return true;
 	}
@@ -322,7 +322,7 @@ private:
 	 *
 	 * @return  Returns whether the event was handled, along with other possible actions
 	 */
-	virtual FReply OnKeyDown(const FGeometry& MyGeometry, const FKeyboardEvent& InKeyboardEvent) OVERRIDE
+	virtual FReply OnKeyDown(const FGeometry& MyGeometry, const FKeyboardEvent& InKeyboardEvent) override
 	{
 		if (WorldModel->GetCommandList()->ProcessCommandBindings(InKeyboardEvent))
 		{
diff --git a/Engine/Source/Editor/WorldBrowser/Private/SWorldHierarchyItem.h b/Engine/Source/Editor/WorldBrowser/Private/SWorldHierarchyItem.h
--- a/Engine/Source/Editor/WorldBrowser/Private/SWorldHierarchyItem.h
+++ b/Engine/Source/Editor/WorldBrowser/Private/SWorldHierarchyItem.h
@@ -42,12 +42,12 @@ public:
 	/** Constructs this widget with InArgs */
 	void Construct(const FArguments& InArgs, TSharedRef<STableViewBase> OwnerTableView);
 
-	TSharedRef<SWidget> GenerateWidgetForColumn( const FName& ColumnName ) OVERRIDE;
+	TSharedRef<SWidget> GenerateWidgetForColumn( const FName& ColumnName ) override;
 
 	FReply OnItemDragDetected(const FGeometry& MyGeometry, const FPointerEvent& MouseEvent);
-	FReply OnDrop( const FGeometry& MyGeometry, const FDragDropEvent& DragDropEvent ) OVERRIDE;
-	void OnDragEnter( const FGeometry& MyGeometry, const FDragDropEvent& DragDropEvent ) OVERRIDE;
-	void OnDragLeave( const FDragDropEvent& DragDropEvent ) OVERRIDE;
+	FReply OnDrop( const FGeometry& MyGeometry, const FDragDropEvent& DragDropEvent ) override;
+	void OnDragEnter( const FGeometry& MyGeometry, const FDragDropEvent& DragDropEvent ) override;
+	void OnDragLeave( const FDragDropEvent& DragDropEvent ) override;
 
 private:
 	/** Operations buttons enabled/disabled state */
diff --git a/Engine/Source/Editor/WorldBrowser/Private/StreamingLevels/StreamingLevelCollectionModel.h b/Engine/Source/Editor/WorldBrowser/Private/StreamingLevels/StreamingLevelCollectionModel.h
--- a/Engine/Source/Editor/WorldBrowser/Private/StreamingLevels/StreamingLevelCollectionModel.h
+++ b/Engine/Source/Editor/WorldBrowser/Private/StreamingLevels/StreamingLevelCollectionModel.h
@@ -28,18 +28,18 @@ public:
 
 public:
 	/** FLevelCollection interface */
-	virtual void UnloadLevels(const FLevelModelList& InLevelList) OVERRIDE;
-	virtual TSharedPtr<FLevelDragDropOp> CreateDragDropOp() const OVERRIDE;
-	virtual void BuildHierarchyMenu(FMenuBuilder& InMenuBuilder) const OVERRIDE;
-	virtual void CustomizeFileMainMenu(FMenuBuilder& InMenuBuilder) const OVERRIDE;
-	virtual void RegisterDetailsCustomization(class FPropertyEditorModule& PropertyModule, TSharedPtr<class IDetailsView> InDetailsView)  OVERRIDE;
-	virtual void UnregisterDetailsCustomization(class FPropertyEditorModule& PropertyModule, TSharedPtr<class IDetailsView> InDetailsView) OVERRIDE;
+	virtual void UnloadLevels(const FLevelModelList& InLevelList) override;
+	virtual TSharedPtr<FLevelDragDropOp> CreateDragDropOp() const override;
+	virtual void BuildHierarchyMenu(FMenuBuilder& InMenuBuilder) const override;
+	virtual void CustomizeFileMainMenu(FMenuBuilder& InMenuBuilder) const override;
+	virtual void RegisterDetailsCustomization(class FPropertyEditorModule& PropertyModule, TSharedPtr<class IDetailsView> InDetailsView)  override;
+	virtual void UnregisterDetailsCustomization(class FPropertyEditorModule& PropertyModule, TSharedPtr<class IDetailsView> InDetailsView) override;
 
 private:
-	virtual void Initialize() OVERRIDE;
-	virtual void BindCommands() OVERRIDE;
-	virtual void OnLevelsCollectionChanged() OVERRIDE;
-	virtual void OnLevelsSelectionChanged() OVERRIDE;
+	virtual void Initialize() override;
+	virtual void BindCommands() override;
+	virtual void OnLevelsCollectionChanged() override;
+	virtual void OnLevelsSelectionChanged() override;
 	/** FLevelCollection interface end */
 	
 public:
@@ -68,8 +68,8 @@ private:
 	void SortFilteredLevels();
 
 	// Begin FEditorUndoClient Interface
-	virtual void PostUndo(bool bSuccess) OVERRIDE { UpdateAllLevels(); }
-	virtual void PostRedo(bool bSuccess) OVERRIDE { PostUndo(bSuccess); }
+	virtual void PostUndo(bool bSuccess) override { UpdateAllLevels(); }
+	virtual void PostRedo(bool bSuccess) override { PostUndo(bSuccess); }
 	// End of FEditorUndoClient
 
 private:
diff --git a/Engine/Source/Editor/WorldBrowser/Private/StreamingLevels/StreamingLevelEdMode.h b/Engine/Source/Editor/WorldBrowser/Private/StreamingLevels/StreamingLevelEdMode.h
--- a/Engine/Source/Editor/WorldBrowser/Private/StreamingLevels/StreamingLevelEdMode.h
+++ b/Engine/Source/Editor/WorldBrowser/Private/StreamingLevels/StreamingLevelEdMode.h
@@ -14,32 +14,32 @@ public:
 	virtual ~FStreamingLevelEdMode();
 
 	// Begin FEdMode
-	virtual void Enter() OVERRIDE;
-	virtual void Exit() OVERRIDE; 
+	virtual void Enter() override;
+	virtual void Exit() override; 
 
-	virtual void AddReferencedObjects( FReferenceCollector& Collector ) OVERRIDE;
+	virtual void AddReferencedObjects( FReferenceCollector& Collector ) override;
 
-	virtual EAxisList::Type GetWidgetAxisToDraw( FWidget::EWidgetMode InWidgetMode ) const OVERRIDE;
+	virtual EAxisList::Type GetWidgetAxisToDraw( FWidget::EWidgetMode InWidgetMode ) const override;
 
-	virtual bool ShouldDrawWidget() const OVERRIDE;
+	virtual bool ShouldDrawWidget() const override;
 
-	virtual bool UsesTransformWidget(FWidget::EWidgetMode CheckMode) const OVERRIDE;
+	virtual bool UsesTransformWidget(FWidget::EWidgetMode CheckMode) const override;
 
-	virtual FVector GetWidgetLocation() const OVERRIDE;
+	virtual FVector GetWidgetLocation() const override;
 
-	virtual bool AllowWidgetMove() OVERRIDE { return true; }
+	virtual bool AllowWidgetMove() override { return true; }
 
-	virtual bool InputDelta( FLevelEditorViewportClient* InViewportClient, FViewport* InViewport, FVector& InDrag, FRotator& InRot, FVector& InScale ) OVERRIDE;
+	virtual bool InputDelta( FLevelEditorViewportClient* InViewportClient, FViewport* InViewport, FVector& InDrag, FRotator& InRot, FVector& InScale ) override;
 
-	virtual void Render( const FSceneView* View, FViewport* Viewport, FPrimitiveDrawInterface* PDI ) OVERRIDE;
+	virtual void Render( const FSceneView* View, FViewport* Viewport, FPrimitiveDrawInterface* PDI ) override;
 
-	virtual bool StartTracking(FLevelEditorViewportClient* InViewportClient, FViewport* InViewport) OVERRIDE;
+	virtual bool StartTracking(FLevelEditorViewportClient* InViewportClient, FViewport* InViewport) override;
 
-	virtual bool EndTracking(FLevelEditorViewportClient* InViewportClient, FViewport* InViewport) OVERRIDE;
+	virtual bool EndTracking(FLevelEditorViewportClient* InViewportClient, FViewport* InViewport) override;
 
-	virtual bool IsSnapRotationEnabled() OVERRIDE;
+	virtual bool IsSnapRotationEnabled() override;
 
-	virtual bool SnapRotatorToGridOverride(FRotator& Rotation) OVERRIDE;
+	virtual bool SnapRotatorToGridOverride(FRotator& Rotation) override;
 
 	// End FEdMode
 
diff --git a/Engine/Source/Editor/WorldBrowser/Private/StreamingLevels/StreamingLevelModel.h b/Engine/Source/Editor/WorldBrowser/Private/StreamingLevels/StreamingLevelModel.h
--- a/Engine/Source/Editor/WorldBrowser/Private/StreamingLevels/StreamingLevelModel.h
+++ b/Engine/Source/Editor/WorldBrowser/Private/StreamingLevels/StreamingLevelModel.h
@@ -27,15 +27,15 @@ public:
 
 public:
 	// FLevelModel interface
-	virtual bool HasValidPackage() const OVERRIDE;
-	virtual UObject* GetNodeObject() OVERRIDE;
-	virtual ULevel* GetLevelObject() const OVERRIDE;
-	virtual FName GetAssetName() const OVERRIDE;
-	virtual FName GetLongPackageName() const OVERRIDE;
-	virtual void Update() OVERRIDE;
-	virtual void OnDrop(const TSharedPtr<FLevelDragDropOp>& Op) OVERRIDE;
-	virtual bool IsGoodToDrop(const TSharedPtr<FLevelDragDropOp>& Op) const OVERRIDE;
-	virtual UClass* GetStreamingClass() const OVERRIDE;
+	virtual bool HasValidPackage() const override;
+	virtual UObject* GetNodeObject() override;
+	virtual ULevel* GetLevelObject() const override;
+	virtual FName GetAssetName() const override;
+	virtual FName GetLongPackageName() const override;
+	virtual void Update() override;
+	virtual void OnDrop(const TSharedPtr<FLevelDragDropOp>& Op) override;
+	virtual bool IsGoodToDrop(const TSharedPtr<FLevelDragDropOp>& Op) const override;
+	virtual UClass* GetStreamingClass() const override;
 	// FLevelModel interface end
 		
 	/** @return the StreamingLevelIndex */
@@ -55,8 +55,8 @@ private:
 	void UpdatePackageFileAvailability();
 
 	// Begin FEditorUndoClient
-	virtual void PostUndo(bool bSuccess) OVERRIDE { Update(); }
-	virtual void PostRedo(bool bSuccess) OVERRIDE { PostUndo(bSuccess); }
+	virtual void PostUndo(bool bSuccess) override { Update(); }
+	virtual void PostRedo(bool bSuccess) override { PostUndo(bSuccess); }
 	// End of FEditorUndoClient
 
 private:
diff --git a/Engine/Source/Editor/WorldBrowser/Private/Tiles/SWorldComposition.cpp b/Engine/Source/Editor/WorldBrowser/Private/Tiles/SWorldComposition.cpp
--- a/Engine/Source/Editor/WorldBrowser/Private/Tiles/SWorldComposition.cpp
+++ b/Engine/Source/Editor/WorldBrowser/Private/Tiles/SWorldComposition.cpp
@@ -19,12 +19,12 @@
 struct FWorldZoomLevelsContainer 
 	: public FZoomLevelsContainer
 {
-	float	GetZoomAmount(int32 InZoomLevel) const OVERRIDE
+	float	GetZoomAmount(int32 InZoomLevel) const override
 	{
 		return 1.f/FMath::Square(GetNumZoomLevels() - InZoomLevel + 1)*2.f;
 	}
 
-	int32 GetNearestZoomLevel( float InZoomAmount ) const OVERRIDE
+	int32 GetNearestZoomLevel( float InZoomAmount ) const override
 	{
 		for (int32 ZoomLevelIndex=0; ZoomLevelIndex < GetNumZoomLevels(); ++ZoomLevelIndex)
 		{
@@ -37,22 +37,22 @@ struct FWorldZoomLevelsContainer
 		return GetDefaultZoomLevel();
 	}
 
-	FText GetZoomText(int32 InZoomLevel) const OVERRIDE
+	FText GetZoomText(int32 InZoomLevel) const override
 	{
 		return FText::AsNumber(GetZoomAmount(InZoomLevel));
 	}
 
-	int32	GetNumZoomLevels() const OVERRIDE
+	int32	GetNumZoomLevels() const override
 	{
 		return 100;
 	}
 
-	int32	GetDefaultZoomLevel() const OVERRIDE
+	int32	GetDefaultZoomLevel() const override
 	{
 		return GetNumZoomLevels() - 10;
 	}
 
-	EGraphRenderingLOD::Type GetLOD(int32 InZoomLevel) const OVERRIDE
+	EGraphRenderingLOD::Type GetLOD(int32 InZoomLevel) const override
 	{
 		return EGraphRenderingLOD::DefaultDetail;
 	}
@@ -183,7 +183,7 @@ public:
 	}
 		
 	/**  SWidget interface */
-	virtual void Tick(const FGeometry& AllottedGeometry, const double InCurrentTime, const float InDeltaTime) OVERRIDE
+	virtual void Tick(const FGeometry& AllottedGeometry, const double InCurrentTime, const float InDeltaTime) override
 	{
 		SNodePanel::Tick(AllottedGeometry, InCurrentTime, InDeltaTime);
 
@@ -253,7 +253,7 @@ public:
 	}
 	
 	/**  SWidget interface */
-	virtual void OnArrangeChildren( const FGeometry& AllottedGeometry, FArrangedChildren& ArrangedChildren ) const OVERRIDE
+	virtual void OnArrangeChildren( const FGeometry& AllottedGeometry, FArrangedChildren& ArrangedChildren ) const override
 	{
 		for (int32 ChildIndex=0; ChildIndex < VisibleChildren.Num(); ++ChildIndex)
 		{
@@ -276,7 +276,7 @@ public:
 	/**  SWidget interface */
 	virtual int32 OnPaint(const FGeometry& AllottedGeometry, const FSlateRect& MyClippingRect, 
 							FSlateWindowElementList& OutDrawElements, int32 LayerId, 
-							const FWidgetStyle& InWidgetStyle, bool bParentEnabled ) const OVERRIDE
+							const FWidgetStyle& InWidgetStyle, bool bParentEnabled ) const override
 	{
 		// First paint the background
 		{
@@ -375,7 +375,7 @@ public:
 	}
 		
 	/** SWidget interface */
-	virtual FReply OnMouseMove(const FGeometry& MyGeometry, const FPointerEvent& MouseEvent) OVERRIDE
+	virtual FReply OnMouseMove(const FGeometry& MyGeometry, const FPointerEvent& MouseEvent) override
 	{
 		const bool bIsRightMouseButtonDown = MouseEvent.IsMouseButtonDown( EKeys::RightMouseButton );
 		const bool bIsLeftMouseButtonDown = MouseEvent.IsMouseButtonDown( EKeys::LeftMouseButton );
@@ -565,7 +565,7 @@ protected:
 	}
 		
 	/**  SWidget interface */
-	virtual FReply OnKeyDown(const FGeometry& MyGeometry, const FKeyboardEvent& InKeyboardEvent) OVERRIDE
+	virtual FReply OnKeyDown(const FGeometry& MyGeometry, const FKeyboardEvent& InKeyboardEvent) override
 	{
 		if (CommandList->ProcessCommandBindings(InKeyboardEvent))
 		{
@@ -581,7 +581,7 @@ protected:
 	}
 
 	/**  SWidget interface */
-	bool SupportsKeyboardFocus() const OVERRIDE
+	bool SupportsKeyboardFocus() const override
 	{
 		return true;
 	}
@@ -634,7 +634,7 @@ protected:
 	}
 	
 	/**  SNodePanel interface */
-	virtual void PopulateVisibleChildren(const FGeometry& AllottedGeometry) OVERRIDE
+	virtual void PopulateVisibleChildren(const FGeometry& AllottedGeometry) override
 	{
 		VisibleChildren.Empty();
 
@@ -680,14 +680,14 @@ protected:
 	}
 
 	/**  SNodePanel interface */
-	virtual void OnBeginNodeInteraction(const TSharedRef<SNode>& InNodeToDrag, const FVector2D& GrabOffset) OVERRIDE
+	virtual void OnBeginNodeInteraction(const TSharedRef<SNode>& InNodeToDrag, const FVector2D& GrabOffset) override
 	{
 		bHasNodeInteraction = true;
 		SNodePanel::OnBeginNodeInteraction(InNodeToDrag, GrabOffset);
 	}
 	
 	/**  SNodePanel interface */
-	virtual void OnEndNodeInteraction(const TSharedRef<SNode>& InNodeDragged) OVERRIDE
+	virtual void OnEndNodeInteraction(const TSharedRef<SNode>& InNodeDragged) override
 	{
 		const SWorldTileItem& Item = static_cast<const SWorldTileItem&>(InNodeDragged.Get());
 		if (Item.IsItemEditable())
diff --git a/Engine/Source/Editor/WorldBrowser/Private/Tiles/SWorldLayers.cpp b/Engine/Source/Editor/WorldBrowser/Private/Tiles/SWorldLayers.cpp
--- a/Engine/Source/Editor/WorldBrowser/Private/Tiles/SWorldLayers.cpp
+++ b/Engine/Source/Editor/WorldBrowser/Private/Tiles/SWorldLayers.cpp
@@ -125,7 +125,7 @@ public:
 		OnLayerCtrlClicked = NewLayerCtrlClicked;
 	}
 	
-	virtual FReply OnMouseButtonDoubleClick(const FGeometry& InMyGeometry, const FPointerEvent& InMouseEvent) OVERRIDE
+	virtual FReply OnMouseButtonDoubleClick(const FGeometry& InMyGeometry, const FPointerEvent& InMouseEvent) override
 	{
 		if (OnLayerDoubleClicked.IsBound())
 		{
@@ -137,7 +137,7 @@ public:
 		}
 	}
 
-	virtual FReply OnMouseButtonUp(const FGeometry& InMyGeometry, const FPointerEvent& InMouseEvent) OVERRIDE
+	virtual FReply OnMouseButtonUp(const FGeometry& InMyGeometry, const FPointerEvent& InMouseEvent) override
 	{
 		if (!InMouseEvent.IsControlDown())
 		{
diff --git a/Engine/Source/Editor/WorldBrowser/Private/Tiles/SWorldTileItem.h b/Engine/Source/Editor/WorldBrowser/Private/Tiles/SWorldTileItem.h
--- a/Engine/Source/Editor/WorldBrowser/Private/Tiles/SWorldTileItem.h
+++ b/Engine/Source/Editor/WorldBrowser/Private/Tiles/SWorldTileItem.h
@@ -13,9 +13,9 @@ public:
 	~FTileItemThumbnail();
 
 	/* ISlateViewport interface */
-	virtual FIntPoint GetSize() const OVERRIDE;
-	virtual class FSlateShaderResource* GetViewportRenderTargetTexture() const OVERRIDE;
-	virtual bool RequiresVsync() const OVERRIDE;
+	virtual FIntPoint GetSize() const override;
+	virtual class FSlateShaderResource* GetViewportRenderTargetTexture() const override;
+	virtual bool RequiresVsync() const override;
 	
 	/** Request thumbnail redraw*/
 	void UpdateThumbnail();
@@ -54,10 +54,10 @@ public:
 	void Construct(const FArguments& InArgs);
 	
 	// SNodePanel::SNode interface start
-	virtual FVector2D GetDesiredSizeForMarquee() const OVERRIDE;
-	virtual UObject* GetObjectBeingDisplayed() const OVERRIDE;
-	virtual FVector2D GetPosition() const OVERRIDE;
-	virtual const FSlateBrush* GetShadowBrush(bool bSelected) const OVERRIDE;
+	virtual FVector2D GetDesiredSizeForMarquee() const override;
+	virtual UObject* GetObjectBeingDisplayed() const override;
+	virtual FVector2D GetPosition() const override;
+	virtual const FSlateBrush* GetShadowBrush(bool bSelected) const override;
 	// SNodePanel::SNode interface end
 	
 	/** @return Deferred item refresh */
@@ -86,13 +86,13 @@ public:
 
 private:
 	// SWidget interface start
-	virtual TSharedPtr<IToolTip> GetToolTip() OVERRIDE;
-	virtual FVector2D ComputeDesiredSize() const OVERRIDE;
-	virtual bool OnHitTest(const FGeometry& MyGeometry, FVector2D InAbsoluteCursorPosition) OVERRIDE;
+	virtual TSharedPtr<IToolTip> GetToolTip() override;
+	virtual FVector2D ComputeDesiredSize() const override;
+	virtual bool OnHitTest(const FGeometry& MyGeometry, FVector2D InAbsoluteCursorPosition) override;
 	virtual int32 OnPaint(const FGeometry& AllottedGeometry, const FSlateRect& MyClippingRect, 
 					FSlateWindowElementList& OutDrawElements, int32 LayerId, 
-					const FWidgetStyle& InWidgetStyle, bool bParentEnabled) const OVERRIDE;
-	virtual FReply OnMouseButtonDoubleClick(const FGeometry& InMyGeometry, const FPointerEvent& InMouseEvent) OVERRIDE;
+					const FWidgetStyle& InWidgetStyle, bool bParentEnabled) const override;
+	virtual FReply OnMouseButtonDoubleClick(const FGeometry& InMyGeometry, const FPointerEvent& InMouseEvent) override;
 	// SWidget interface end
 
 	TSharedRef<SToolTip> CreateToolTipWidget();
diff --git a/Engine/Source/Editor/WorldBrowser/Private/Tiles/WorldTileCollectionModel.cpp b/Engine/Source/Editor/WorldBrowser/Private/Tiles/WorldTileCollectionModel.cpp
--- a/Engine/Source/Editor/WorldBrowser/Private/Tiles/WorldTileCollectionModel.cpp
+++ b/Engine/Source/Editor/WorldBrowser/Private/Tiles/WorldTileCollectionModel.cpp
@@ -114,7 +114,7 @@ void FWorldTileCollectionModel::UnloadLevels(const FLevelModelList& InLevelList)
 			FLevelModelList& DirtyHierarchy;
 			FHierachyCollector(FLevelModelList& LevelToUnload) : DirtyHierarchy(LevelToUnload) {}
 
-			virtual void Visit(FLevelModel& Item) OVERRIDE
+			virtual void Visit(FLevelModel& Item) override
 			{
 				DirtyHierarchy.AddUnique(Item.AsShared());
 			}
diff --git a/Engine/Source/Editor/WorldBrowser/Private/Tiles/WorldTileCollectionModel.h b/Engine/Source/Editor/WorldBrowser/Private/Tiles/WorldTileCollectionModel.h
--- a/Engine/Source/Editor/WorldBrowser/Private/Tiles/WorldTileCollectionModel.h
+++ b/Engine/Source/Editor/WorldBrowser/Private/Tiles/WorldTileCollectionModel.h
@@ -26,34 +26,34 @@ public:
 	}
 public:
 	/** FLevelCollection interface */
-	virtual void UnloadLevels(const FLevelModelList& InLevelList) OVERRIDE;
-	virtual void TranslateLevels(const FLevelModelList& InList, FVector2D InAbsoluteDelta, bool bSnapDelta = true)  OVERRIDE;
-	virtual FVector2D SnapTranslationDelta(const FLevelModelList& InList, FVector2D InTranslationDelta, bool bBoundsSnapping, float InSnappingValue) OVERRIDE;
-	virtual TSharedPtr<FLevelDragDropOp> CreateDragDropOp() const OVERRIDE;
-	virtual bool PassesAllFilters(TSharedPtr<FLevelModel> InLevelModel) const OVERRIDE;
-	virtual void BuildGridMenu(FMenuBuilder& InMenuBuilder) const OVERRIDE;
-	virtual void BuildHierarchyMenu(FMenuBuilder& InMenuBuilder) const OVERRIDE;
-	virtual void CustomizeFileMainMenu(FMenuBuilder& InMenuBuilder) const OVERRIDE;
-	virtual FVector GetObserverPosition() const OVERRIDE;
-	virtual bool CompareLevelsZOrder(TSharedPtr<FLevelModel> InA, TSharedPtr<FLevelModel> InB) const OVERRIDE;
-	virtual void RegisterDetailsCustomization(class FPropertyEditorModule& PropertyModule, TSharedPtr<class IDetailsView> InDetailsView)  OVERRIDE;
-	virtual void UnregisterDetailsCustomization(class FPropertyEditorModule& PropertyModule, TSharedPtr<class IDetailsView> InDetailsView) OVERRIDE;
-	virtual bool IsTileWorld() const OVERRIDE { return true; };
+	virtual void UnloadLevels(const FLevelModelList& InLevelList) override;
+	virtual void TranslateLevels(const FLevelModelList& InList, FVector2D InAbsoluteDelta, bool bSnapDelta = true)  override;
+	virtual FVector2D SnapTranslationDelta(const FLevelModelList& InList, FVector2D InTranslationDelta, bool bBoundsSnapping, float InSnappingValue) override;
+	virtual TSharedPtr<FLevelDragDropOp> CreateDragDropOp() const override;
+	virtual bool PassesAllFilters(TSharedPtr<FLevelModel> InLevelModel) const override;
+	virtual void BuildGridMenu(FMenuBuilder& InMenuBuilder) const override;
+	virtual void BuildHierarchyMenu(FMenuBuilder& InMenuBuilder) const override;
+	virtual void CustomizeFileMainMenu(FMenuBuilder& InMenuBuilder) const override;
+	virtual FVector GetObserverPosition() const override;
+	virtual bool CompareLevelsZOrder(TSharedPtr<FLevelModel> InA, TSharedPtr<FLevelModel> InB) const override;
+	virtual void RegisterDetailsCustomization(class FPropertyEditorModule& PropertyModule, TSharedPtr<class IDetailsView> InDetailsView)  override;
+	virtual void UnregisterDetailsCustomization(class FPropertyEditorModule& PropertyModule, TSharedPtr<class IDetailsView> InDetailsView) override;
+	virtual bool IsTileWorld() const override { return true; };
 	/** FLevelCollection interface end */
 
 private:
 	/** FTickableEditorObject interface */
-	void Tick( float DeltaTime ) OVERRIDE;
+	void Tick( float DeltaTime ) override;
 	/** FTickableEditorObject interface end */
 
 	/** FLevelCollection interface */
-	virtual void Initialize() OVERRIDE;
-	virtual void BindCommands() OVERRIDE;
-	virtual void OnLevelsCollectionChanged() OVERRIDE;
-	virtual void OnLevelsSelectionChanged() OVERRIDE;
-	virtual void OnLevelsHierarchyChanged() OVERRIDE;
-	virtual void OnPreLoadLevels(const FLevelModelList& InList) OVERRIDE;
-	virtual void OnPreShowLevels(const FLevelModelList& InList) OVERRIDE;
+	virtual void Initialize() override;
+	virtual void BindCommands() override;
+	virtual void OnLevelsCollectionChanged() override;
+	virtual void OnLevelsSelectionChanged() override;
+	virtual void OnLevelsHierarchyChanged() override;
+	virtual void OnPreLoadLevels(const FLevelModelList& InList) override;
+	virtual void OnPreShowLevels(const FLevelModelList& InList) override;
 	/** FLevelCollection interface end */
 	
 public:
@@ -121,8 +121,8 @@ public:
 											float SnappingDistance);
 	
 	// Begin FEditorUndoClient Interface
-	virtual void PostUndo(bool bSuccess) OVERRIDE;
-	virtual void PostRedo(bool bSuccess) OVERRIDE { PostUndo(bSuccess); }
+	virtual void PostUndo(bool bSuccess) override;
+	virtual void PostRedo(bool bSuccess) override { PostUndo(bSuccess); }
 	// End of FEditorUndoClient
 
 	enum FocusStrategy 
diff --git a/Engine/Source/Editor/WorldBrowser/Private/Tiles/WorldTileDetails.h b/Engine/Source/Editor/WorldBrowser/Private/Tiles/WorldTileDetails.h
--- a/Engine/Source/Editor/WorldBrowser/Private/Tiles/WorldTileDetails.h
+++ b/Engine/Source/Editor/WorldBrowser/Private/Tiles/WorldTileDetails.h
@@ -102,7 +102,7 @@ class UWorldTileDetails : public UObject
 	FTilePropertyChanged			ZOrderChangedEvent;
 	
 	// UObject interface
-	virtual void PostEditChangeProperty( struct FPropertyChangedEvent& PropertyChangedEvent) OVERRIDE;
+	virtual void PostEditChangeProperty( struct FPropertyChangedEvent& PropertyChangedEvent) override;
 	// todo: undo operations
 
 public:
diff --git a/Engine/Source/Editor/WorldBrowser/Private/Tiles/WorldTileDetailsCustomization.h b/Engine/Source/Editor/WorldBrowser/Private/Tiles/WorldTileDetailsCustomization.h
--- a/Engine/Source/Editor/WorldBrowser/Private/Tiles/WorldTileDetailsCustomization.h
+++ b/Engine/Source/Editor/WorldBrowser/Private/Tiles/WorldTileDetailsCustomization.h
@@ -12,7 +12,7 @@ public:
 	static TSharedRef<IDetailCustomization> MakeInstance(TSharedRef<FWorldTileCollectionModel> InWorldData);
 
 	// IDetailCustomization interface
-	virtual void CustomizeDetails(IDetailLayoutBuilder& DetailLayout) OVERRIDE;
+	virtual void CustomizeDetails(IDetailLayoutBuilder& DetailLayout) override;
 	// End of IDetailCustomization interface
 private:
 	/** Whether property should be editable (depends on level load status and level lock status) */
@@ -38,8 +38,8 @@ public:
 	static TSharedRef<IPropertyTypeCustomization> MakeInstance(TSharedRef<FWorldTileCollectionModel> InWorldData);
 
 	/** IPropertyTypeCustomization interface */
-	virtual void CustomizeHeader(TSharedRef<class IPropertyHandle> StructPropertyHandle, class FDetailWidgetRow& HeaderRow, IPropertyTypeCustomizationUtils& StructCustomizationUtils ) OVERRIDE;
-	virtual void CustomizeChildren(TSharedRef<class IPropertyHandle> StructPropertyHandle, class IDetailChildrenBuilder& ChildBuilder, IPropertyTypeCustomizationUtils& StructCustomizationUtils ) OVERRIDE;
+	virtual void CustomizeHeader(TSharedRef<class IPropertyHandle> StructPropertyHandle, class FDetailWidgetRow& HeaderRow, IPropertyTypeCustomizationUtils& StructCustomizationUtils ) override;
+	virtual void CustomizeChildren(TSharedRef<class IPropertyHandle> StructPropertyHandle, class IDetailChildrenBuilder& ChildBuilder, IPropertyTypeCustomizationUtils& StructCustomizationUtils ) override;
 
 private:
 	bool OnShouldFilterStreamingPackage(const FString& InPackageName) const;
@@ -56,8 +56,8 @@ public:
 	static TSharedRef<IPropertyTypeCustomization> MakeInstance(TSharedRef<FWorldTileCollectionModel> InWorldData);
 
 	/** IPropertyTypeCustomization interface */
-	virtual void CustomizeHeader( TSharedRef<class IPropertyHandle> StructPropertyHandle, class FDetailWidgetRow& HeaderRow, IPropertyTypeCustomizationUtils& StructCustomizationUtils ) OVERRIDE;
-	virtual void CustomizeChildren( TSharedRef<IPropertyHandle> StructPropertyHandle, class IDetailChildrenBuilder& ChildBuilder, IPropertyTypeCustomizationUtils& StructCustomizationUtils ) OVERRIDE;
+	virtual void CustomizeHeader( TSharedRef<class IPropertyHandle> StructPropertyHandle, class FDetailWidgetRow& HeaderRow, IPropertyTypeCustomizationUtils& StructCustomizationUtils ) override;
+	virtual void CustomizeChildren( TSharedRef<IPropertyHandle> StructPropertyHandle, class IDetailChildrenBuilder& ChildBuilder, IPropertyTypeCustomizationUtils& StructCustomizationUtils ) override;
 
 private:
 	FReply OnGenerateTile();
