Add fallback calculation for OBB in PhAT when multiple children are present

diff --git a/Engine/Source/Editor/UnrealEd/Private/PhysicsAssetUtils.cpp b/Engine/Source/Editor/UnrealEd/Private/PhysicsAssetUtils.cpp
--- a/Engine/Source/Editor/UnrealEd/Private/PhysicsAssetUtils.cpp
+++ b/Engine/Source/Editor/UnrealEd/Private/PhysicsAssetUtils.cpp
@@ -33,13 +33,13 @@ static int32 GetChildIndex(int32 BoneIndex, USkeletalMesh* SkelMesh, const TArra
 	{
 		int32 ParentIndex = SkelMesh->RefSkeleton.GetParentIndex(i);
 
-		if(ParentIndex == BoneIndex)
+		if (ParentIndex == BoneIndex && Infos[i].Positions.Num() > 0)
 		{
 			if(ChildIndex != INDEX_NONE)
 			{
 				return INDEX_NONE; // if we already have a child, this bone has more than one so return INDEX_NONE.
 			}
-			else if (Infos[i].Positions.Num() > 0)
+			else
 			{
 				ChildIndex = i;
 			}
@@ -237,6 +237,69 @@ bool CreateFromSkeletalMesh(UPhysicsAsset* PhysicsAsset, USkeletalMesh* SkelMesh
 	return bSuccess;
 }
 
+FMatrix ComputeCovarianceMatrix(const FBoneVertInfo & VertInfo)
+{
+	if (VertInfo.Positions.Num() == 0)
+	{
+		return FMatrix::Identity;
+	}
+
+	const TArray<FVector> & Positions = VertInfo.Positions;
+
+	//get average
+	const float N = Positions.Num();
+	FVector U = FVector::ZeroVector;
+	for (int32 i = 0; i < N; ++i)
+	{
+		U += Positions[i];
+	}
+
+	U = U / N;
+
+	//compute error terms
+	TArray<FVector> Errors;
+	Errors.AddUninitialized(N);
+
+	for (int32 i = 0; i < N; ++i)
+	{
+		Errors[i] = Positions[i] - U;
+	}
+
+	FMatrix Covariance = FMatrix::Identity;
+	for (int32 j = 0; j < 3; ++j)
+	{
+		FVector Axis = FVector::ZeroVector;
+		float * Cj = &Axis.X;
+		for (int32 k = 0; k < 3; ++k)
+		{
+			float Cjk = 0.f;
+			for (int32 i = 0; i < N; ++i)
+			{
+				const float * error = &Errors[i].X;
+				Cj[k] += error[j] * error[k];
+			}
+			Cj[k] /= N;
+		}
+
+		Covariance.SetAxis(j, Axis);
+	}
+
+	return Covariance;
+}
+
+FVector ComputeEigenVector(const FMatrix & A)
+{
+	//using the power method: this is ok because we only need the dominate eigenvector and speed is not critical: http://en.wikipedia.org/wiki/Power_iteration
+	FVector Bk = FVector(0, 0, 1);
+	for (int32 i = 0; i < 32; ++i)
+	{
+		float Length = Bk.Size();
+		Bk = A.TransformVector(Bk) / Length;
+	}
+
+	return Bk.SafeNormal();
+}
+
 
 void CreateCollisionFromBone( UBodySetup* bs, USkeletalMesh* skelMesh, int32 BoneIndex, FPhysAssetCreateParams& Params, const TArray<FBoneVertInfo>& Infos )
 {
@@ -247,6 +310,7 @@ void CreateCollisionFromBone( UBodySetup* bs, USkeletalMesh* skelMesh, int32 Bon
 
 	// Calculate orientation of to use for collision primitive.
 	FMatrix ElemTM;
+	bool ComputeFromVerts = false;
 
 	if(Params.bAlignDownBone)
 	{
@@ -273,11 +337,13 @@ void CreateCollisionFromBone( UBodySetup* bs, USkeletalMesh* skelMesh, int32 Bon
 			else
 			{
 				ElemTM = FMatrix::Identity;
+				ComputeFromVerts = true;
 			}
 		}
 		else
 		{
 			ElemTM = FMatrix::Identity;
+			ComputeFromVerts = true;
 		}
 	}
 	else
@@ -285,6 +351,20 @@ void CreateCollisionFromBone( UBodySetup* bs, USkeletalMesh* skelMesh, int32 Bon
 		ElemTM = FMatrix::Identity;
 	}
 
+	
+	if (ComputeFromVerts)
+	{
+		// Compute covariance matrix for verts of this bone
+		// Then use axis with largest variance for orienting bone box
+		const FMatrix CovarianceMatrix = ComputeCovarianceMatrix(Infos[BoneIndex]);
+		FVector ZAxis = ComputeEigenVector(CovarianceMatrix);
+		FVector XAxis, YAxis;
+		ZAxis.FindBestAxisVectors(YAxis, XAxis);
+		ElemTM = FMatrix(XAxis, YAxis, ZAxis, FVector::ZeroVector);
+	}
+	
+
+
 	// Get the (Unreal scale) bounding box for this bone using the rotation.
 	const FBoneVertInfo* BoneInfo = &Infos[BoneIndex];
 	FBox BoneBox(0);
