Fixed case when level package was not GCed due to Metadata refs, added checks that level package gets always GCed

diff --git a/Engine/Source/Editor/UnrealEd/Private/EditorLevelUtils.cpp b/Engine/Source/Editor/UnrealEd/Private/EditorLevelUtils.cpp
--- a/Engine/Source/Editor/UnrealEd/Private/EditorLevelUtils.cpp
+++ b/Engine/Source/Editor/UnrealEd/Private/EditorLevelUtils.cpp
@@ -335,6 +335,7 @@ namespace EditorLevelUtils
 		GEditor->CloseEditedWorldAssets(CastChecked<UWorld>(InLevel->GetOuter()));
 
 		UWorld* OwningWorld = InLevel->OwningWorld;
+		const FName LevelPackageName		= InLevel->GetOutermost()->GetFName();
 		const bool bRemovingCurrentLevel	= InLevel && InLevel->IsCurrentLevel();
 		const bool bRemoveSuccessful		= PrivateRemoveLevelFromWorld( InLevel );
 		if ( bRemoveSuccessful )
@@ -369,6 +370,17 @@ namespace EditorLevelUtils
 
 			// Collect garbage to clear out the destroyed level
 			CollectGarbage( GARBAGE_COLLECTION_KEEPFLAGS );
+
+			// Ensure that level package were removed
+			UPackage* LevelPackage = FindObjectFast<UPackage>(NULL, LevelPackageName);
+			if (LevelPackage != nullptr)
+			{
+				StaticExec(NULL, *FString::Printf(TEXT("OBJ REFS CLASS=%s NAME=%s shortest"),*LevelPackage->GetClass()->GetName(), *LevelPackage->GetPathName()));
+				TMap<UObject*,UProperty*>	Route		= FArchiveTraceRoute::FindShortestRootPath( LevelPackage, true, GARBAGE_COLLECTION_KEEPFLAGS );
+				FString						ErrorString	= FArchiveTraceRoute::PrintRootPath( Route, LevelPackage );
+			
+				UE_LOG(LogStreaming, Fatal, TEXT("%s didn't get garbage collected!") LINE_TERMINATOR TEXT("%s"), *LevelPackage->GetFullName(), *ErrorString );
+			}
 		}
 		return bRemoveSuccessful;
 	}
@@ -465,6 +477,11 @@ namespace EditorLevelUtils
 		InLevel->GetOuter()->MarkPendingKill();
 		InLevel->MarkPendingKill();		
 		InLevel->GetOuter()->ClearFlags(RF_Public|RF_Standalone);
+		if (InLevel->GetOutermost()->MetaData)
+		{
+			InLevel->GetOutermost()->MetaData->MarkPendingKill();
+		}
+
 		World->MarkPackageDirty();
 		World->BroadcastLevelsChanged();
 
