#UE4 const fix, staunching some const incorrectness from bleeding...

diff --git a/Engine/Source/Editor/UnrealEd/Private/Toolkits/ToolkitManager.cpp b/Engine/Source/Editor/UnrealEd/Private/Toolkits/ToolkitManager.cpp
--- a/Engine/Source/Editor/UnrealEd/Private/Toolkits/ToolkitManager.cpp
+++ b/Engine/Source/Editor/UnrealEd/Private/Toolkits/ToolkitManager.cpp
@@ -89,7 +89,7 @@ void FToolkitManager::OnToolkitHostDestroyed( IToolkitHost* HostBeingDestroyed )
 }
 
 
-TSharedPtr< IToolkit > FToolkitManager::FindEditorForAsset( UObject* Asset )
+TSharedPtr< IToolkit > FToolkitManager::FindEditorForAsset( const UObject* Asset )
 {
 	for( auto ToolkitIndex = 0; ToolkitIndex < Toolkits.Num(); ++ToolkitIndex )
 	{
@@ -98,7 +98,7 @@ TSharedPtr< IToolkit > FToolkitManager::FindEditorForAsset( UObject* Asset )
 
 		if( Toolkit->IsAssetEditor() )
 		{
-			if( Toolkit->GetObjectsCurrentlyBeingEdited()->Contains( Asset ) )
+			if (Toolkit->GetObjectsCurrentlyBeingEdited()->FindByKey(Asset) != NULL )
 			{
 				return Toolkit;
 			}
diff --git a/Engine/Source/Editor/UnrealEd/Public/Toolkits/ToolkitManager.h b/Engine/Source/Editor/UnrealEd/Public/Toolkits/ToolkitManager.h
--- a/Engine/Source/Editor/UnrealEd/Public/Toolkits/ToolkitManager.h
+++ b/Engine/Source/Editor/UnrealEd/Public/Toolkits/ToolkitManager.h
@@ -23,7 +23,7 @@ public:
 	UNREALED_API void OnToolkitHostDestroyed( class IToolkitHost* HostBeingDestroyed );
 
 	/** Tries to find an open asset editor that's editing the specified asset, and returns the toolkit for that asset editor */
-	UNREALED_API TSharedPtr< IToolkit > FindEditorForAsset( UObject* Asset );
+	UNREALED_API TSharedPtr< IToolkit > FindEditorForAsset( const UObject* Asset );
 
 private:
 	// Buried constructor since the asset editor manager is a singleton
