EditorAutomationTests.cpp
--Updated the ConvertMesh test so that the final scene scene has the converted BSP to blocking volume and the BSP converted to Static Mesh next to each other and makes them easier to see.
--Added the ability to take a screenshot for the final scene.

diff --git a/Engine/Source/Editor/UnrealEd/Private/Tests/EditorAutomationTests.cpp b/Engine/Source/Editor/UnrealEd/Private/Tests/EditorAutomationTests.cpp
--- a/Engine/Source/Editor/UnrealEd/Private/Tests/EditorAutomationTests.cpp
+++ b/Engine/Source/Editor/UnrealEd/Private/Tests/EditorAutomationTests.cpp
@@ -1268,20 +1268,34 @@ bool FConvertToValidation::RunTest(const FString& Parameters)
 {
 	UWorld* World = GEditor->NewMap();
 
+	//Set the Test Name which is used later for getting the directory to store the screenshots.
+	const FString BaseFileName = TEXT("ConvertMeshTest");
+	
+	//Creating the parameters needed for latent screenshot capturing.
+	WindowScreenshotParameters ConvertMeshParameters;
+	
+	//Check if the main frame is loaded.  When using the old main frame it may not be.
+	if (FModuleManager::Get().IsModuleLoaded("MainFrame"))
+	{
+		IMainFrameModule& MainFrame = FModuleManager::LoadModuleChecked<IMainFrameModule>("MainFrame");
+		//Now set the WindowScreenshot struct CurrentWindow name to be the mainframe.
+		ConvertMeshParameters.CurrentWindow = MainFrame.GetParentWindow();
+	}
+
+	//Set the screenshot name.
+	ConvertMeshParameters.ScreenshotName = BaseFileName;
+
 	//Adjust camera in viewports
 	for( int32 i = 0; i < GEditor->LevelViewportClients.Num(); i++ )
 	{
 		FLevelEditorViewportClient* ViewportClient = GEditor->LevelViewportClients[i];
 		if(!ViewportClient->IsOrtho())
 		{
-			ViewportClient->SetViewLocation( FVector(448, 902, 423) );
-			ViewportClient->SetViewRotation( FRotator(0, 270, 0) );
+			ViewportClient->SetViewLocation( FVector(190, 590, 360) );
+			ViewportClient->SetViewRotation( FRotator(0, -90, 0) );
 		}
 	}
 
-	//Gather assets
-	UObject* EditorCubeMesh = (UStaticMesh*)StaticLoadObject(UStaticMesh::StaticClass(),NULL,TEXT("/Engine/EditorMeshes/EditorCube.EditorCube"),NULL,LOAD_None,NULL);
-
 	//BSP TO BLOCKING VOLUME
 	{
 		TArray<ABrush*> PreviousBrushes;
@@ -1293,7 +1307,7 @@ bool FConvertToValidation::RunTest(const FString& Parameters)
 		CubeAdditiveBrushBuilder->Y = 256.0f;
 		CubeAdditiveBrushBuilder->Z = 256.0f;
 		CubeAdditiveBrushBuilder->Build(World);
-		GEditor->Exec( World, TEXT("BRUSH MOVETO X=768 Y=0 Z=768"));
+		GEditor->Exec( World, TEXT("BRUSH MOVETO X=384 Y=0 Z=384"));
 		GEditor->Exec( World, TEXT("BRUSH ADD"));
 
 		//find brush that was just added by finding the brush not in our previous list
@@ -1321,7 +1335,7 @@ bool FConvertToValidation::RunTest(const FString& Parameters)
 		CubeAdditiveBrushBuilder->Y = 256.0f;
 		CubeAdditiveBrushBuilder->Z = 256.0f;
 		CubeAdditiveBrushBuilder->Build(World);
-		GEditor->Exec( World, TEXT("BRUSH MOVETO X=0 Y=0 Z=0"));
+		GEditor->Exec( World, TEXT("BRUSH MOVETO X=0 Y=0 Z=384"));
 		GEditor->Exec( World, TEXT("BRUSH ADD"));
 
 		//find brush that was just added by finding the brush not in our previous list
@@ -1353,13 +1367,10 @@ bool FConvertToValidation::RunTest(const FString& Parameters)
 		}
 	}
 
-
-	//Directional Light
-	ADirectionalLight* DirectionalLight = Cast<ADirectionalLight>(GEditor->AddActor(World->GetCurrentLevel(), ADirectionalLight::StaticClass(), FVector(384, 0, 384)));
-	DirectionalLight->SetMobility(EComponentMobility::Movable);
-	DirectionalLight->SetActorRotation(FRotator(314, 339, 0));
-	DirectionalLight->SetBrightness(3.142f);
-	DirectionalLight->SetLightColor(FColor::White);
+	TakeLatentAutomationScreenshot(ConvertMeshParameters, BaseFileName, FString::Printf(TEXT("FinalConvertMesh")), FString::Printf(TEXT("04_Final")), false);
+	
+	//Wait to give the screenshot capture some time to complete.
+	ADD_LATENT_AUTOMATION_COMMAND(FWaitLatentCommand(0.1f));
 
 	GEditorModeTools().MapChangeNotify();
 
