Performing a save via CTRL+S now performs the correct name validation for maps and packages

This addresses TTP#337708 - UE4: editor: the editor allows you to save a map with spaces in the name (this causes a few other bugs to occur which are detailed in the summary)

diff --git a/Engine/Source/Editor/UnrealEd/Private/FileHelpers.cpp b/Engine/Source/Editor/UnrealEd/Private/FileHelpers.cpp
--- a/Engine/Source/Editor/UnrealEd/Private/FileHelpers.cpp
+++ b/Engine/Source/Editor/UnrealEd/Private/FileHelpers.cpp
@@ -2043,9 +2043,23 @@ static int32 InternalSavePackage( UPackage* PackageToSave, bool& bOutPackageLoca
 				}
 			
 				FText ErrorMessage;
-				if( !FEditorFileUtils::IsFilenameValidForSaving( FinalPackageFilename, ErrorMessage ) )
+				bool bValidFilename = FEditorFileUtils::IsFilenameValidForSaving( FinalPackageFilename, ErrorMessage );
+				if ( bValidFilename )
 				{
-					FMessageDialog::Open( EAppMsgType::Ok, ErrorMessage );
+					bValidFilename = bIsMapPackage ? FEditorFileUtils::IsValidMapFilename( FinalPackageFilename, ErrorMessage ) : FPackageName::IsValidLongPackageName( FinalPackageFilename, false, &ErrorMessage );
+				}
+
+				if ( !bValidFilename )
+				{
+					// Start the loop over, prompting for save again
+					const FText DisplayFilename = FText::FromString( IFileManager::Get().ConvertToAbsolutePathForExternalAppForRead( *FinalPackageFilename ) );
+					FFormatNamedArguments Arguments;
+					Arguments.Add( TEXT("Filename"), DisplayFilename );
+					Arguments.Add( TEXT("LineTerminators"), FText::FromString( LINE_TERMINATOR LINE_TERMINATOR ) );
+					Arguments.Add( TEXT("ErrorMessage"), ErrorMessage );
+					const FText DisplayMessage = FText::Format( LOCTEXT( "InvalidSaveFilename", "Failed to save to {Filename}{LineTerminators}{ErrorMessage}" ), Arguments );
+					FMessageDialog::Open( EAppMsgType::Ok, DisplayMessage );
+
 					// Start the loop over, prompting for save again
 					continue;
 				}
