Fixed incorrect button/heading text shown when converting actors to assets.

This addresses TTP#337700 - EDITOR: When creating a static mesh from selected brush, the path picker displays blueprint text for "create" button and asset name text field

diff --git a/Engine/Source/Editor/UnrealEd/Private/SCreateAssetFromActor.cpp b/Engine/Source/Editor/UnrealEd/Private/SCreateAssetFromActor.cpp
--- a/Engine/Source/Editor/UnrealEd/Private/SCreateAssetFromActor.cpp
+++ b/Engine/Source/Editor/UnrealEd/Private/SCreateAssetFromActor.cpp
@@ -68,7 +68,7 @@ void SCreateAssetFromActor::Construct(const FArguments& InArgs, TSharedPtr<SWind
 			.AutoWidth()
 			[
 				SNew(STextBlock)
-				.Text(LOCTEXT("BlueprintNameLabel", "Blueprint Name:"))
+				.Text(HeadingText)
 			]
 			+ SHorizontalBox::Slot()
 			.FillWidth(1.0f)
@@ -94,7 +94,7 @@ void SCreateAssetFromActor::Construct(const FArguments& InArgs, TSharedPtr<SWind
 				.IsEnabled(this, &SCreateAssetFromActor::IsCreateAssetFromActorEnabled)
 				[
 					SNew(STextBlock)
-					.Text(LOCTEXT("CreateBlueprintFromActorButton", "Create Blueprint"))
+					.Text(CreateButtonText)
 				]
 			]
 			+ SHorizontalBox::Slot()
@@ -108,7 +108,7 @@ void SCreateAssetFromActor::Construct(const FArguments& InArgs, TSharedPtr<SWind
 				.OnClicked(this, &SCreateAssetFromActor::OnCancelCreateAssetFromActor)
 				[
 					SNew(STextBlock)
-					.Text(LOCTEXT("CancelBlueprintFromActorButton", "Cancel"))
+					.Text(LOCTEXT("CancelButtonText", "Cancel"))
 				]
 			]
 		]
