Level Editor export functionality now defaults to FBX rather than OBJ

This addresses TTP#336924 - Editor: Mesh Export: Default Export Selected/All option is .OBJ but should be .FBX

diff --git a/Engine/Source/Editor/UnrealEd/Private/FileHelpers.cpp b/Engine/Source/Editor/UnrealEd/Private/FileHelpers.cpp
--- a/Engine/Source/Editor/UnrealEd/Private/FileHelpers.cpp
+++ b/Engine/Source/Editor/UnrealEd/Private/FileHelpers.cpp
@@ -295,7 +295,7 @@ FString FEditorFileUtils::GetFilterString(EFileInteraction Interaction)
 		break;
 
 	case FI_Export:
-		Result = TEXT("Object (*.obj)|*.obj|Unreal Text (*.t3d)|*.t3d|Stereo Litho (*.stl)|*.stl|LOD Export (*.lod.obj)|*.lod.obj|FBX (*.fbx)|*.fbx|All Files|*.*");
+		Result = TEXT("FBX (*.fbx)|*.fbx|Object (*.obj)|*.obj|Unreal Text (*.t3d)|*.t3d|Stereo Litho (*.stl)|*.stl|LOD Export (*.lod.obj)|*.lod.obj|All Files|*.*");
 		break;
 
 	default:
