Fixed converting actors broadcasting about a newly added actor twice.

The event is broadcast inside SpawnActor, so there's no need for the outer function to broadcast it as well. This addressees TTP#338333 - EDITOR: VOLUMES: Convert Actor: Duplicate actor appears in scene outliner.

diff --git a/Engine/Source/Editor/UnrealEd/Private/Editor.cpp b/Engine/Source/Editor/UnrealEd/Private/Editor.cpp
--- a/Engine/Source/Editor/UnrealEd/Private/Editor.cpp
+++ b/Engine/Source/Editor/UnrealEd/Private/Editor.cpp
@@ -3332,9 +3332,6 @@ void UEditorEngine::ConvertSelectedBrushesToVolumes( UClass* VolumeClass )
 				NewVolume->PostEditMove( true );
 				NewVolume->Modify();
 
-				// Send notification about actors that may have changed
-				GEngine->BroadcastLevelActorAdded(NewVolume);		
-
 				// Destroy the old actor.
 				GEditor->Layers->DisassociateActorFromLayers( CurBrushActor );
 				World->EditorDestroyActor( CurBrushActor, true );
