UD Struct/Enum editor logic update - notification before changes, ready for UDS property editor

diff --git a/Engine/Config/BaseEditor.ini b/Engine/Config/BaseEditor.ini
--- a/Engine/Config/BaseEditor.ini
+++ b/Engine/Config/BaseEditor.ini
@@ -251,6 +251,7 @@ SupportWebsiteURL="http://unrealengine.com/support"
 
 [UserDefinedStructure]
 bUseUserDefinedStructure=true
+bShowDefaultValuePropertyEditor=false
 
 [EpicSurvey]
 Source="CDN"
diff --git a/Engine/Source/Editor/Kismet/Private/UserDefinedEnumEditor.cpp b/Engine/Source/Editor/Kismet/Private/UserDefinedEnumEditor.cpp
--- a/Engine/Source/Editor/Kismet/Private/UserDefinedEnumEditor.cpp
+++ b/Engine/Source/Editor/Kismet/Private/UserDefinedEnumEditor.cpp
@@ -180,7 +180,11 @@ void FEnumDetails::OnForceRefresh()
 	}
 }
 
-void FEnumDetails::OnChanged(const class UUserDefinedEnum* Enum)
+void FEnumDetails::PreChange(const class UUserDefinedEnum* Enum)
+{
+}
+
+void FEnumDetails::PostChange(const class UUserDefinedEnum* Enum)
 {
 	if (Enum && (TargetEnum.Get() == Enum))
 	{
diff --git a/Engine/Source/Editor/Kismet/Private/UserDefinedEnumEditor.h b/Engine/Source/Editor/Kismet/Private/UserDefinedEnumEditor.h
--- a/Engine/Source/Editor/Kismet/Private/UserDefinedEnumEditor.h
+++ b/Engine/Source/Editor/Kismet/Private/UserDefinedEnumEditor.h
@@ -69,7 +69,8 @@ public:
 	void OnForceRefresh();
 
 	/** FEnumEditorUtils::INotifyOnEnumChanged */
-	virtual void OnChanged(const class UUserDefinedEnum* Enum) override;
+	virtual void PreChange(const class UUserDefinedEnum* Struct) override;
+	virtual void PostChange(const class UUserDefinedEnum* Enum) override;
 
 private:
 	/** Handles new enum element request */
diff --git a/Engine/Source/Editor/UnrealEd/Private/Kismet2/BlueprintEditorUtils.cpp b/Engine/Source/Editor/UnrealEd/Private/Kismet2/BlueprintEditorUtils.cpp
--- a/Engine/Source/Editor/UnrealEd/Private/Kismet2/BlueprintEditorUtils.cpp
+++ b/Engine/Source/Editor/UnrealEd/Private/Kismet2/BlueprintEditorUtils.cpp
@@ -5586,7 +5586,8 @@ bool FBlueprintEditorUtils::PropertyValueFromString(const UProperty* Property, c
 		}
 		else if( Property->IsA(UClassProperty::StaticClass()) )
 		{
-			CastChecked<UClassProperty>(Property)->SetPropertyValue_InContainer(DefaultObject, FindObject<UClass>(ANY_PACKAGE, *Value, true));
+			UClass* FoundClass = FindObject<UClass>(ANY_PACKAGE, *Value);
+			CastChecked<UClassProperty>(Property)->SetPropertyValue_InContainer(DefaultObject, FoundClass);
 		}
 		else if( Property->IsA(UObjectPropertyBase::StaticClass()) )
 		{
@@ -5659,6 +5660,54 @@ bool FBlueprintEditorUtils::PropertyValueFromString(const UProperty* Property, c
 	return bParseSuccedded;
 }
 
+bool FBlueprintEditorUtils::PropertyValueToString(const UProperty* Property, const uint8* Container, FString& OutForm)
+{
+	check(Property && Container);
+	OutForm.Empty();
+	if (Property->IsA(UStructProperty::StaticClass()))
+	{
+		static UScriptStruct* VectorStruct = FindObjectChecked<UScriptStruct>(UObject::StaticClass(), TEXT("Vector"));
+		static UScriptStruct* RotatorStruct = FindObjectChecked<UScriptStruct>(UObject::StaticClass(), TEXT("Rotator"));
+		static UScriptStruct* TransformStruct = FindObjectChecked<UScriptStruct>(UObject::StaticClass(), TEXT("Transform"));
+		static UScriptStruct* LinearColorStruct = FindObjectChecked<UScriptStruct>(UObject::StaticClass(), TEXT("LinearColor"));
+
+		const UStructProperty* StructProperty = CastChecked<UStructProperty>(Property);
+
+		// Struct properties must be handled differently, unfortunately.  We only support FVector, FRotator, and FTransform
+		if (StructProperty->Struct == VectorStruct)
+		{
+			FVector Vector;
+			Property->CopyCompleteValue(&Vector, Property->ContainerPtrToValuePtr<uint8>(Container));
+			OutForm = FString::Printf(TEXT("%f,%f,%f"), Vector.X, Vector.Y, Vector.Z);
+		}
+		else if (StructProperty->Struct == RotatorStruct)
+		{
+			FRotator Rotator;
+			Property->CopyCompleteValue(&Rotator, Property->ContainerPtrToValuePtr<uint8>(Container));
+			OutForm = FString::Printf(TEXT("%f,%f,%f"), Rotator.Pitch, Rotator.Yaw, Rotator.Roll);
+		}
+		else if (StructProperty->Struct == TransformStruct)
+		{
+			FTransform Transform;
+			Property->CopyCompleteValue(&Transform, Property->ContainerPtrToValuePtr<uint8>(Container));
+			OutForm = Transform.ToString();
+		}
+		else if (StructProperty->Struct == LinearColorStruct)
+		{
+			FLinearColor Color;
+			Property->CopyCompleteValue(&Color, Property->ContainerPtrToValuePtr<uint8>(Container));
+			OutForm = Color.ToString();
+		}
+	}
+
+	bool bSuccedded = true;
+	if (OutForm.IsEmpty())
+	{
+		bSuccedded = Property->ExportText_InContainer(0, OutForm, Container, Container, NULL, 0);
+	}
+	return bSuccedded;
+}
+
 FName FBlueprintEditorUtils::GenerateUniqueGraphName(UBlueprint* const BlueprintOuter, FString const& ProposedName)
 {
 	FName UniqueGraphName(*ProposedName);
diff --git a/Engine/Source/Editor/UnrealEd/Private/Kismet2/EnumEditorUtils.cpp b/Engine/Source/Editor/UnrealEd/Private/Kismet2/EnumEditorUtils.cpp
--- a/Engine/Source/Editor/UnrealEd/Private/Kismet2/EnumEditorUtils.cpp
+++ b/Engine/Source/Editor/UnrealEd/Private/Kismet2/EnumEditorUtils.cpp
@@ -85,6 +85,8 @@ void FEnumEditorUtils::AddNewEnumeratorForUserDefinedEnum(UUserDefinedEnum* Enum
 {
 	if (Enum)
 	{
+		PrepareForChange(Enum);
+
 		TArray<FName> OldNames, Names;
 		CopyEnumeratorsWithoutMax(Enum, OldNames);
 		Names = OldNames;
@@ -107,6 +109,8 @@ void FEnumEditorUtils::RemoveEnumeratorFromUserDefinedEnum(UUserDefinedEnum* Enu
 {
 	if (Enum && (Enum->GetEnum(EnumeratorIndex) != NAME_None))
 	{
+		PrepareForChange(Enum);
+
 		TArray<FName> OldNames, Names;
 		CopyEnumeratorsWithoutMax(Enum, OldNames);
 		Names = OldNames;
@@ -129,6 +133,8 @@ void FEnumEditorUtils::MoveEnumeratorInUserDefinedEnum(UUserDefinedEnum* Enum, i
 {
 	if (Enum && (Enum->GetEnum(EnumeratorIndex) != NAME_None))
 	{
+		PrepareForChange(Enum);
+
 		TArray<FName> OldNames, Names;
 		CopyEnumeratorsWithoutMax(Enum, OldNames);
 		Names = OldNames;
@@ -190,6 +196,11 @@ public:
 	}
 };
 
+void FEnumEditorUtils::PrepareForChange(const UUserDefinedEnum* Enum)
+{
+	FEnumEditorManager::Get().PreChange(Enum);
+}
+
 void FEnumEditorUtils::BroadcastChanges(const UUserDefinedEnum* Enum, const TArray<FName>& OldNames, bool bResolveData)
 {
 	check(NULL != Enum);
@@ -292,7 +303,7 @@ void FEnumEditorUtils::BroadcastChanges(const UUserDefinedEnum* Enum, const TArr
 		(*It)->BroadcastChanged();
 	}
 
-	FEnumEditorManager::Get().OnChanged(Enum);
+	FEnumEditorManager::Get().PostChange(Enum);
 }
 
 int32 FEnumEditorUtils::ResolveEnumerator(const UEnum* Enum, FArchive& Ar, int32 EnumeratorIndex)
@@ -331,6 +342,7 @@ bool FEnumEditorUtils::SetEnumeratorDisplayName(UUserDefinedEnum* Enum, int32 En
 	{
 		if (IsEnumeratorDisplayNameValid(Enum, NewDisplayName))
 		{
+			PrepareForChange(Enum);
 			Enum->SetMetaData(FEnumEditorUtilsHelper::DisplayName(), *NewDisplayName, EnumeratorIndex);
 			EnsureAllDisplayNamesExist(Enum);
 			BroadcastChanges(Enum, TArray<FName>(), false);
diff --git a/Engine/Source/Editor/UnrealEd/Private/Kismet2/StructureEditorUtils.cpp b/Engine/Source/Editor/UnrealEd/Private/Kismet2/StructureEditorUtils.cpp
--- a/Engine/Source/Editor/UnrealEd/Private/Kismet2/StructureEditorUtils.cpp
+++ b/Engine/Source/Editor/UnrealEd/Private/Kismet2/StructureEditorUtils.cpp
@@ -340,9 +340,8 @@ bool FStructureEditorUtils::ChangeVariableDefaultValue(UUserDefinedStruct* Struc
 {
 	const UEdGraphSchema_K2* K2Schema = GetDefault<UEdGraphSchema_K2>();
 	auto VarDesc = GetVarDescByGuid(Struct, VarGuid);
-	if (VarDesc 
-		&& (NewDefaultValue != VarDesc->DefaultValue) 
-		&& K2Schema->DefaultValueSimpleValidation(VarDesc->ToPinType(), FString(), NewDefaultValue, NULL, FText::GetEmpty()))
+	if (VarDesc && (NewDefaultValue != VarDesc->DefaultValue))
+		//&& K2Schema->DefaultValueSimpleValidation(VarDesc->ToPinType(), FString(), NewDefaultValue, NULL, FText::GetEmpty()))
 	{
 		bool bAdvancedValidation = true;
 		if (!NewDefaultValue.IsEmpty())
@@ -468,9 +467,10 @@ void FStructureEditorUtils::OnStructureChanged(UUserDefinedStruct* Struct)
 {
 	if (Struct)
 	{
+		FStructEditorManager::Get().PreChange(Struct);
 		Struct->Status = EUserDefinedStructureStatus::UDSS_Dirty;
 		CompileStructure(Struct);
-		FStructEditorManager::Get().OnChanged(Struct);
+		FStructEditorManager::Get().PostChange(Struct);
 		Struct->MarkPackageDirty();
 	}
 }
@@ -654,4 +654,11 @@ bool FStructureEditorUtils::Is3dWidgetEnabled(const UUserDefinedStruct* Struct,
 	return VarDesc && VarDesc->bEnable3dWidget && FEdMode::CanCreateWidgetForStructure(PropertyStruct);
 }
 
+FGuid FStructureEditorUtils::GetGuidForProperty(const UProperty* Property)
+{
+	auto UDStruct = Property ? Cast<const UUserDefinedStruct>(Property->GetOwnerStruct()) : NULL;
+	auto VarDesc = UDStruct ? GetVarDesc(UDStruct).FindByPredicate(FFindByNameHelper<FStructVariableDescription>(Property->GetFName())) : NULL;
+	return VarDesc ? VarDesc->VarGuid : FGuid();
+}
+
 #undef LOCTEXT_NAMESPACE
diff --git a/Engine/Source/Editor/UnrealEd/Public/Kismet2/BlueprintEditorUtils.h b/Engine/Source/Editor/UnrealEd/Public/Kismet2/BlueprintEditorUtils.h
--- a/Engine/Source/Editor/UnrealEd/Public/Kismet2/BlueprintEditorUtils.h
+++ b/Engine/Source/Editor/UnrealEd/Public/Kismet2/BlueprintEditorUtils.h
@@ -669,6 +669,8 @@ public:
 
 	static bool PropertyValueFromString(const UProperty* Property, const FString& StrValue, uint8* ContainerMem);
 
+	static bool PropertyValueToString(const UProperty* Property, const uint8* Container, FString& OutForm);
+
 	/** Call PostEditChange() on all Actors based on the given Blueprint */
 	static void PostEditChangeBlueprintActors(UBlueprint* Blueprint);
 
diff --git a/Engine/Source/Editor/UnrealEd/Public/Kismet2/EnumEditorUtils.h b/Engine/Source/Editor/UnrealEd/Public/Kismet2/EnumEditorUtils.h
--- a/Engine/Source/Editor/UnrealEd/Public/Kismet2/EnumEditorUtils.h
+++ b/Engine/Source/Editor/UnrealEd/Public/Kismet2/EnumEditorUtils.h
@@ -9,6 +9,7 @@
 
 class UNREALED_API FEnumEditorUtils
 {
+	static void PrepareForChange(const UUserDefinedEnum* Enum);
 	static void BroadcastChanges(const UUserDefinedEnum* Enum, const TArray<FName>& OldNames, bool bResolveData = true);
 
 	/** copy full enumeratos names from given enum to OutEnumNames, the last '_MAX' enumerator is skipped */
diff --git a/Engine/Source/Editor/UnrealEd/Public/Kismet2/ListenerManager.h b/Engine/Source/Editor/UnrealEd/Public/Kismet2/ListenerManager.h
--- a/Engine/Source/Editor/UnrealEd/Public/Kismet2/ListenerManager.h
+++ b/Engine/Source/Editor/UnrealEd/Public/Kismet2/ListenerManager.h
@@ -6,7 +6,8 @@ template<typename TType>
 class INotifyOnChanged
 {
 public:
-	virtual void OnChanged(const TType* Changed) = 0;
+	virtual void PreChange(const TType* Changed) = 0;
+	virtual void PostChange(const TType* Changed) = 0;
 };
 
 template<typename TType>
@@ -32,11 +33,19 @@ public:
 		Listeners.Remove(Listener);
 	}
 
-	void OnChanged(const TType* Changed)
+	void PreChange(const TType* Changed)
 	{
 		for (auto Listener : Listeners)
 		{
-			Listener->OnChanged(Changed);
+			Listener->PreChange(Changed);
+		}
+	}
+
+	void PostChange(const TType* Changed)
+	{
+		for (auto Listener : Listeners)
+		{
+			Listener->PostChange(Changed);
 		}
 	}
 };
\ No newline at end of file
diff --git a/Engine/Source/Editor/UnrealEd/Public/Kismet2/StructureEditorUtils.h b/Engine/Source/Editor/UnrealEd/Public/Kismet2/StructureEditorUtils.h
--- a/Engine/Source/Editor/UnrealEd/Public/Kismet2/StructureEditorUtils.h
+++ b/Engine/Source/Editor/UnrealEd/Public/Kismet2/StructureEditorUtils.h
@@ -101,6 +101,8 @@ public:
 
 	static void RemoveInvalidStructureMemberVariableFromBlueprint(UBlueprint* Blueprint);
 
+	static FGuid GetGuidForProperty(const UProperty* Property);
+
 	/*
 	 Default values for member variables in User Defined Structure are stored in meta data "MakeStructureDefaultValue"
 	 The following functions are used to fill an instance of user defined struct with those default values.
