add back in call to obtain addresses for the network file server
#ue4

diff --git a/Engine/Source/Editor/UnrealEd/Private/CookOnTheFlyServer.cpp b/Engine/Source/Editor/UnrealEd/Private/CookOnTheFlyServer.cpp
--- a/Engine/Source/Editor/UnrealEd/Private/CookOnTheFlyServer.cpp
+++ b/Engine/Source/Editor/UnrealEd/Private/CookOnTheFlyServer.cpp
@@ -114,7 +114,7 @@ bool UCookOnTheFlyServer::BroadcastFileserverPresence( const FGuid &InstanceId )
 {
 	TArray<TSharedPtr<FInternetAddr> > AddressList;
 
-	if ((NetworkFileServer == NULL || !NetworkFileServer->IsItReadyToAcceptConnections()))
+	if ((NetworkFileServer == NULL || !NetworkFileServer->IsItReadyToAcceptConnections() || !NetworkFileServer->GetAddressList(AddressList)))
 	{
 		UE_LOG(LogCookOnTheFly, Error, TEXT("Failed to create network file server"));
 
