UDS: No default values for objects/classes

diff --git a/Engine/Source/Editor/UnrealEd/Private/Kismet2/StructureEditorUtils.cpp b/Engine/Source/Editor/UnrealEd/Private/Kismet2/StructureEditorUtils.cpp
--- a/Engine/Source/Editor/UnrealEd/Private/Kismet2/StructureEditorUtils.cpp
+++ b/Engine/Source/Editor/UnrealEd/Private/Kismet2/StructureEditorUtils.cpp
@@ -340,15 +340,16 @@ bool FStructureEditorUtils::ChangeVariableDefaultValue(UUserDefinedStruct* Struc
 {
 	const UEdGraphSchema_K2* K2Schema = GetDefault<UEdGraphSchema_K2>();
 	auto VarDesc = GetVarDescByGuid(Struct, VarGuid);
-	if (VarDesc && (NewDefaultValue != VarDesc->DefaultValue))
-		//&& K2Schema->DefaultValueSimpleValidation(VarDesc->ToPinType(), FString(), NewDefaultValue, NULL, FText::GetEmpty()))
+	if (VarDesc 
+		&& (NewDefaultValue != VarDesc->DefaultValue)
+		&& K2Schema->DefaultValueSimpleValidation(VarDesc->ToPinType(), FString(), NewDefaultValue, NULL, FText::GetEmpty()))
 	{
 		bool bAdvancedValidation = true;
 		if (!NewDefaultValue.IsEmpty())
 		{
 			const auto Property = FindField<UProperty>(Struct, VarDesc->VarName);
 			FStructOnScope StructDefaultMem(Struct);
-			bAdvancedValidation = StructDefaultMem.GetStructMemory() && Property &&
+			bAdvancedValidation = StructDefaultMem.IsValid() && Property &&
 				FBlueprintEditorUtils::PropertyValueFromString(Property, NewDefaultValue, StructDefaultMem.GetStructMemory());
 		}
 
