Merged GitHub pull request #202 to main.

refactored tga code to allows use outside of EditorFactories.cpp
added support for grayscale jpeg
fix memory leak in jpeg code
changes in AssetTools to allow to specify precise factory when multiple factories support the same filetype
changes in Plugin.cs to allow binary only plugins
exposed parts of engine API to other modules

diff --git a/Engine/Source/Developer/AssetTools/Private/AssetTools.cpp b/Engine/Source/Developer/AssetTools/Private/AssetTools.cpp
--- a/Engine/Source/Developer/AssetTools/Private/AssetTools.cpp
+++ b/Engine/Source/Developer/AssetTools/Private/AssetTools.cpp
@@ -310,8 +310,10 @@ TArray<UObject*> FAssetTools::ImportAssets(const FString& DestinationPath)
 		}
 	}
 
+	TMultiMap<uint32, UFactory*> FilterIndexToFactory;
+
 	// Generate the file types and extensions represented by the selected factories
-	ObjectTools::GenerateFactoryFileExtensions( Factories, FileTypes, AllExtensions );
+	ObjectTools::GenerateFactoryFileExtensions( Factories, FileTypes, AllExtensions, FilterIndexToFactory );
 
 	FileTypes = FString::Printf(TEXT("All Files (%s)|%s|%s"),*AllExtensions,*AllExtensions,*FileTypes);
 
@@ -319,6 +321,8 @@ TArray<UObject*> FAssetTools::ImportAssets(const FString& DestinationPath)
 	TArray<FString> OpenFilenames;
 	IDesktopPlatform* DesktopPlatform = FDesktopPlatformModule::Get();
 	bool bOpened = false;
+	int32 FilterIndex = -1;
+
 	if ( DesktopPlatform )
 	{
 		void* ParentWindowWindowHandle = NULL;
@@ -337,7 +341,8 @@ TArray<UObject*> FAssetTools::ImportAssets(const FString& DestinationPath)
 			TEXT(""),
 			FileTypes,
 			EFileDialogFlags::Multiple,
-			OpenFilenames
+			OpenFilenames,
+			FilterIndex
 			);
 	}
 
@@ -345,8 +350,15 @@ TArray<UObject*> FAssetTools::ImportAssets(const FString& DestinationPath)
 	{
 		if ( OpenFilenames.Num() > 0 )
 		{
+			UFactory* ChosenFactory = NULL;
+			if (FilterIndex > 0)
+			{
+				ChosenFactory = *FilterIndexToFactory.Find(FilterIndex);
+			}
+
+
 			FEditorDirectories::Get().SetLastDirectory(ELastDirectory::GENERIC_IMPORT, OpenFilenames[0]);
-			ReturnObjects = ImportAssets(OpenFilenames, DestinationPath);
+			ReturnObjects = ImportAssets(OpenFilenames, DestinationPath, ChosenFactory);
 		}
 	}
 
@@ -389,7 +401,7 @@ void FAssetTools::ExpandDirectories(const TArray<FString>& Files, const FString&
 	}
 }
 
-TArray<UObject*> FAssetTools::ImportAssets(const TArray<FString>& Files, const FString& RootDestinationPath) const
+TArray<UObject*> FAssetTools::ImportAssets(const TArray<FString>& Files, const FString& RootDestinationPath, UFactory* SpecifiedFactory) const
 {
 	TArray<UObject*> ReturnObjects;
 	TMap< FString, TArray<UFactory*> > ExtensionToFactoriesMap;
@@ -403,48 +415,94 @@ TArray<UObject*> FAssetTools::ImportAssets(const TArray<FString>& Files, const F
 	TArray<TPair<FString, FString>> FilesAndDestinations;
 	ExpandDirectories(Files, RootDestinationPath, FilesAndDestinations);
 
-	// First instantiate one factory for each file extension encountered that supports the extension
-	for( TObjectIterator<UClass> ClassIt; ClassIt; ++ClassIt )
-	{
-		if( (*ClassIt)->IsChildOf(UFactory::StaticClass()) && !((*ClassIt)->HasAnyClassFlags(CLASS_Abstract)) )
+	if (SpecifiedFactory == NULL)
+	{	
+		// First instantiate one factory for each file extension encountered that supports the extension
+		for( TObjectIterator<UClass> ClassIt; ClassIt; ++ClassIt )
 		{
-			UFactory* Factory = Cast<UFactory>( (*ClassIt)->GetDefaultObject() );
-			if( Factory->bEditorImport && Factory->ValidForCurrentGame() )
+			if( (*ClassIt)->IsChildOf(UFactory::StaticClass()) && !((*ClassIt)->HasAnyClassFlags(CLASS_Abstract)) )
 			{
-				TArray<FString> FactoryExtensions;
-				Factory->GetSupportedFileExtensions(FactoryExtensions);
-				for ( auto& FileDest : FilesAndDestinations )
+				UFactory* Factory = Cast<UFactory>( (*ClassIt)->GetDefaultObject() );
+				if( Factory->bEditorImport && Factory->ValidForCurrentGame() )
 				{
-					const FString FileExtension = FPaths::GetExtension(FileDest.Key);
-
-					// Case insensitive string compare with supported formats of this factory
-					if ( FactoryExtensions.Contains(FileExtension) )
+					TArray<FString> FactoryExtensions;
+					Factory->GetSupportedFileExtensions(FactoryExtensions);
+					for ( auto& FileDest : FilesAndDestinations )
 					{
-						TArray<UFactory*>& ExistingFactories = ExtensionToFactoriesMap.FindOrAdd(FileExtension);
-						
-						// Do not remap extensions, just reuse the existing UFactory.
-						// There may be multiple UFactories, so we will keep track of all of them
-						bool bFactoryAlreadyInMap = false;
-						for ( auto FoundFactoryIt = ExistingFactories.CreateConstIterator(); FoundFactoryIt; ++FoundFactoryIt )
+						const FString FileExtension = FPaths::GetExtension(FileDest.Key);
+	
+						// Case insensitive string compare with supported formats of this factory
+						if ( FactoryExtensions.Contains(FileExtension) )
 						{
-							if ( (*FoundFactoryIt)->GetClass() == Factory->GetClass() )
+							TArray<UFactory*>& ExistingFactories = ExtensionToFactoriesMap.FindOrAdd(FileExtension);
+							
+							// Do not remap extensions, just reuse the existing UFactory.
+							// There may be multiple UFactories, so we will keep track of all of them
+							bool bFactoryAlreadyInMap = false;
+							for ( auto FoundFactoryIt = ExistingFactories.CreateConstIterator(); FoundFactoryIt; ++FoundFactoryIt )
+							{
+								if ( (*FoundFactoryIt)->GetClass() == Factory->GetClass() )
+								{
+									bFactoryAlreadyInMap = true;
+									break;
+								}
+							}
+	
+							if ( !bFactoryAlreadyInMap )
 							{
-								bFactoryAlreadyInMap = true;
-								break;
+								// We found a factory for this file, it can be imported!
+								// Create a new factory of the same class and make sure it doesn't get GCed.
+								// The object will be removed from the root set at the end of this function.
+								UFactory* NewFactory = ConstructObject<UFactory>( Factory->GetClass() );
+								if ( NewFactory->ConfigureProperties() )
+								{
+									NewFactory->AddToRoot();
+									ExistingFactories.Add(NewFactory);
+								}
 							}
 						}
+					}
+				}
+			}
+		}
+	}
+	else
+	{
+		if( SpecifiedFactory->bEditorImport && SpecifiedFactory->ValidForCurrentGame() )
+		{
+			TArray<FString> FactoryExtensions;
+			SpecifiedFactory->GetSupportedFileExtensions(FactoryExtensions);
+			for ( auto FileIt = Files.CreateConstIterator(); FileIt; ++FileIt )
+			{
+				const FString FileExtension = FPaths::GetExtension(*FileIt);
+
+				// Case insensitive string compare with supported formats of this factory
+				if ( FactoryExtensions.Contains(FileExtension) )
+				{
+					TArray<UFactory*>& ExistingFactories = ExtensionToFactoriesMap.FindOrAdd(FileExtension);
+
+					// Do not remap extensions, just reuse the existing UFactory.
+					// There may be multiple UFactories, so we will keep track of all of them
+					bool bFactoryAlreadyInMap = false;
+					for ( auto FoundFactoryIt = ExistingFactories.CreateConstIterator(); FoundFactoryIt; ++FoundFactoryIt )
+					{
+						if ( (*FoundFactoryIt)->GetClass() == SpecifiedFactory->GetClass() )
+						{
+							bFactoryAlreadyInMap = true;
+							break;
+						}
+					}
 
-						if ( !bFactoryAlreadyInMap )
+					if ( !bFactoryAlreadyInMap )
+					{
+						// We found a factory for this file, it can be imported!
+						// Create a new factory of the same class and make sure it doesnt get GCed.
+						// The object will be removed from the root set at the end of this function.
+						UFactory* NewFactory = ConstructObject<UFactory>( SpecifiedFactory->GetClass() );
+						if ( NewFactory->ConfigureProperties() )
 						{
-							// We found a factory for this file, it can be imported!
-							// Create a new factory of the same class and make sure it doesn't get GCed.
-							// The object will be removed from the root set at the end of this function.
-							UFactory* NewFactory = ConstructObject<UFactory>( Factory->GetClass() );
-							if ( NewFactory->ConfigureProperties() )
-							{
-								NewFactory->AddToRoot();
-								ExistingFactories.Add(NewFactory);
-							}
+							NewFactory->AddToRoot();
+							ExistingFactories.Add(NewFactory);
 						}
 					}
 				}
@@ -452,6 +510,7 @@ TArray<UObject*> FAssetTools::ImportAssets(const TArray<FString>& Files, const F
 		}
 	}
 
+
 	// Some flags to keep track of what the user decided when asked about overwriting or replacing
 	bool bOverwriteAll = false;
 	bool bReplaceAll = false;
diff --git a/Engine/Source/Developer/AssetTools/Private/AssetTools.h b/Engine/Source/Developer/AssetTools/Private/AssetTools.h
--- a/Engine/Source/Developer/AssetTools/Private/AssetTools.h
+++ b/Engine/Source/Developer/AssetTools/Private/AssetTools.h
@@ -18,7 +18,7 @@ public:
 	virtual UObject* DuplicateAsset(const FString& AssetName, const FString& PackagePath, UObject* OriginalObject) override;
 	virtual void RenameAssets(const TArray<FAssetRenameData>& AssetsAndNames) const override;
 	virtual TArray<UObject*> ImportAssets(const FString& DestinationPath) override;
-	virtual TArray<UObject*> ImportAssets(const TArray<FString>& Files, const FString& DestinationPath) const override;
+	virtual TArray<UObject*> ImportAssets(const TArray<FString>& Files, const FString& DestinationPath, UFactory* ChosenFactory) const override;
 	virtual void CreateUniqueAssetName(const FString& InBasePackageName, const FString& InSuffix, FString& OutPackageName, FString& OutAssetName) const override;
 	virtual bool AssetUsesGenericThumbnail( const FAssetData& AssetData ) const override;
 	virtual void DiffAgainstDepot(UObject* InObject, const FString& InPackagePath, const FString& InPackageName) const override;
diff --git a/Engine/Source/Developer/AssetTools/Public/IAssetTools.h b/Engine/Source/Developer/AssetTools/Public/IAssetTools.h
--- a/Engine/Source/Developer/AssetTools/Public/IAssetTools.h
+++ b/Engine/Source/Developer/AssetTools/Public/IAssetTools.h
@@ -69,7 +69,7 @@ public:
 	virtual TArray<UObject*> ImportAssets(const FString& DestinationPath) = 0;
 
 	/** Imports the specified files to the destination path. */
-	virtual TArray<UObject*> ImportAssets(const TArray<FString>& Files, const FString& DestinationPath) const = 0;
+	virtual TArray<UObject*> ImportAssets(const TArray<FString>& Files, const FString& DestinationPath, UFactory* ChosenFactory = NULL) const = 0;
 
 	/** Creates a unique package and asset name taking the form InBasePackageName+InSuffix */
 	virtual void CreateUniqueAssetName(const FString& InBasePackageName, const FString& InSuffix, FString& OutPackageName, FString& OutAssetName) const = 0;
diff --git a/Engine/Source/Developer/DesktopPlatform/Private/Mac/DesktopPlatformMac.cpp b/Engine/Source/Developer/DesktopPlatform/Private/Mac/DesktopPlatformMac.cpp
--- a/Engine/Source/Developer/DesktopPlatform/Private/Mac/DesktopPlatformMac.cpp
+++ b/Engine/Source/Developer/DesktopPlatform/Private/Mac/DesktopPlatformMac.cpp
@@ -211,6 +211,11 @@ bool FDesktopPlatformMac::OpenFileDialog(const void* ParentWindowHandle, const F
 	return FileDialogShared(false, ParentWindowHandle, DialogTitle, DefaultPath, DefaultFile, FileTypes, Flags, OutFilenames);
 }
 
+bool FDesktopPlatformMac::OpenFileDialog(const void* ParentWindowHandle, const FString& DialogTitle, const FString& DefaultPath, const FString& DefaultFile, const FString& FileTypes, uint32 Flags, TArray<FString>& OutFilenames, int32& OutFilterIndex)
+{
+	return FileDialogShared(false, ParentWindowHandle, DialogTitle, DefaultPath, DefaultFile, FileTypes, Flags, OutFilenames);
+}
+
 bool FDesktopPlatformMac::SaveFileDialog(const void* ParentWindowHandle, const FString& DialogTitle, const FString& DefaultPath, const FString& DefaultFile, const FString& FileTypes, uint32 Flags, TArray<FString>& OutFilenames)
 {
 	return FileDialogShared(true, ParentWindowHandle, DialogTitle, DefaultPath, DefaultFile, FileTypes, Flags, OutFilenames);
diff --git a/Engine/Source/Developer/DesktopPlatform/Private/Mac/DesktopPlatformMac.h b/Engine/Source/Developer/DesktopPlatform/Private/Mac/DesktopPlatformMac.h
--- a/Engine/Source/Developer/DesktopPlatform/Private/Mac/DesktopPlatformMac.h
+++ b/Engine/Source/Developer/DesktopPlatform/Private/Mac/DesktopPlatformMac.h
@@ -9,6 +9,7 @@ class FDesktopPlatformMac : public FDesktopPlatformBase
 public:
 	// IDesktopPlatform Implementation
 	virtual bool OpenFileDialog(const void* ParentWindowHandle, const FString& DialogTitle, const FString& DefaultPath, const FString& DefaultFile, const FString& FileTypes, uint32 Flags, TArray<FString>& OutFilenames) override;
+	virtual bool OpenFileDialog(const void* ParentWindowHandle, const FString& DialogTitle, const FString& DefaultPath, const FString& DefaultFile, const FString& FileTypes, uint32 Flags, TArray<FString>& OutFilenames, int32& OutFilterIndex) override;
 	virtual bool SaveFileDialog(const void* ParentWindowHandle, const FString& DialogTitle, const FString& DefaultPath, const FString& DefaultFile, const FString& FileTypes, uint32 Flags, TArray<FString>& OutFilenames) override;
 	virtual bool OpenDirectoryDialog(const void* ParentWindowHandle, const FString& DialogTitle, const FString& DefaultPath, FString& OutFolderName) override;
 	virtual bool OpenFontDialog(const void* ParentWindowHandle, FString& OutFontName, float& OutHeight, EFontImportFlags::Type& OutFlags) override;
@@ -28,4 +29,4 @@ private:
 	bool FileDialogShared(bool bSave, const void* ParentWindowHandle, const FString& DialogTitle, const FString& DefaultPath, const FString& DefaultFile, const FString& FileTypes, uint32 Flags, TArray<FString>& OutFilenames);
 };
 
-typedef FDesktopPlatformMac FDesktopPlatform;
\ No newline at end of file
+typedef FDesktopPlatformMac FDesktopPlatform;
diff --git a/Engine/Source/Developer/DesktopPlatform/Private/Windows/DesktopPlatformWindows.cpp b/Engine/Source/Developer/DesktopPlatform/Private/Windows/DesktopPlatformWindows.cpp
--- a/Engine/Source/Developer/DesktopPlatform/Private/Windows/DesktopPlatformWindows.cpp
+++ b/Engine/Source/Developer/DesktopPlatform/Private/Windows/DesktopPlatformWindows.cpp
@@ -18,14 +18,21 @@
 
 static const TCHAR *InstallationsSubKey = TEXT("SOFTWARE\\Epic Games\\Unreal Engine\\Builds");
 
-bool FDesktopPlatformWindows::OpenFileDialog(const void* ParentWindowHandle, const FString& DialogTitle, const FString& DefaultPath, const FString& DefaultFile, const FString& FileTypes, uint32 Flags, TArray<FString>& OutFilenames)
+bool FDesktopPlatformWindows::OpenFileDialog( const void* ParentWindowHandle, const FString& DialogTitle, const FString& DefaultPath, const FString& DefaultFile, const FString& FileTypes, uint32 Flags, TArray<FString>& OutFilenames, int32& OutFilterIndex )
 {
-	return FileDialogShared(false, ParentWindowHandle, DialogTitle, DefaultPath, DefaultFile, FileTypes, Flags, OutFilenames);
+	return FileDialogShared(false, ParentWindowHandle, DialogTitle, DefaultPath, DefaultFile, FileTypes, Flags, OutFilenames, OutFilterIndex );
+}
+
+bool FDesktopPlatformWindows::OpenFileDialog( const void* ParentWindowHandle, const FString& DialogTitle, const FString& DefaultPath, const FString& DefaultFile, const FString& FileTypes, uint32 Flags, TArray<FString>& OutFilenames)
+{
+	int32 DummyFilterIndex;
+	return FileDialogShared(false, ParentWindowHandle, DialogTitle, DefaultPath, DefaultFile, FileTypes, Flags, OutFilenames, DummyFilterIndex );
 }
 
 bool FDesktopPlatformWindows::SaveFileDialog(const void* ParentWindowHandle, const FString& DialogTitle, const FString& DefaultPath, const FString& DefaultFile, const FString& FileTypes, uint32 Flags, TArray<FString>& OutFilenames)
 {
-	return FileDialogShared(true, ParentWindowHandle, DialogTitle, DefaultPath, DefaultFile, FileTypes, Flags, OutFilenames);
+	int32 DummyFilterIndex = 0;
+	return FileDialogShared(true, ParentWindowHandle, DialogTitle, DefaultPath, DefaultFile, FileTypes, Flags, OutFilenames, DummyFilterIndex );
 }
 
 static ::INT CALLBACK BrowseCallbackProc(HWND hwnd, ::UINT uMsg, LPARAM lParam, LPARAM lpData)
@@ -184,7 +191,7 @@ bool FDesktopPlatformWindows::OpenLauncher(bool Install, FString CommandLinePara
 	return FPlatformProcess::CreateProc(*LaunchPath, *CommandLineParams, true, false, false, NULL, 0, *FPaths::GetPath(LaunchPath), NULL).IsValid();
 }
 
-bool FDesktopPlatformWindows::FileDialogShared(bool bSave, const void* ParentWindowHandle, const FString& DialogTitle, const FString& DefaultPath, const FString& DefaultFile, const FString& FileTypes, uint32 Flags, TArray<FString>& OutFilenames)
+bool FDesktopPlatformWindows::FileDialogShared(bool bSave, const void* ParentWindowHandle, const FString& DialogTitle, const FString& DefaultPath, const FString& DefaultFile, const FString& FileTypes, uint32 Flags, TArray<FString>& OutFilenames, int32& OutFilterIndex)
 {
 	FScopedSystemModalMode SystemModalScope;
 
@@ -320,10 +327,10 @@ bool FDesktopPlatformWindows::FileDialogShared(bool bSave, const void* ParentWin
 		}
 
 		// The index of the filter in OPENFILENAME starts at 1.
-		const uint32 FilterIndex = ofn.nFilterIndex - 1;
+		OutFilterIndex = ofn.nFilterIndex - 1;
 
 		// Get the extension to add to the filename (if one doesnt already exist)
-		const FString Extension = CleanExtensionList.IsValidIndex( FilterIndex ) ? CleanExtensionList[FilterIndex] : TEXT("");
+		FString Extension = CleanExtensionList.IsValidIndex( OutFilterIndex ) ? CleanExtensionList[OutFilterIndex] : TEXT("");
 
 		// Make sure all filenames gathered have their paths normalized and proper extensions added
 		for ( auto FilenameIt = OutFilenames.CreateIterator(); FilenameIt; ++FilenameIt )
diff --git a/Engine/Source/Developer/DesktopPlatform/Private/Windows/DesktopPlatformWindows.h b/Engine/Source/Developer/DesktopPlatform/Private/Windows/DesktopPlatformWindows.h
--- a/Engine/Source/Developer/DesktopPlatform/Private/Windows/DesktopPlatformWindows.h
+++ b/Engine/Source/Developer/DesktopPlatform/Private/Windows/DesktopPlatformWindows.h
@@ -10,6 +10,7 @@ class FDesktopPlatformWindows : public FDesktopPlatformBase
 public:
 	// IDesktopPlatform Implementation
 	virtual bool OpenFileDialog(const void* ParentWindowHandle, const FString& DialogTitle, const FString& DefaultPath, const FString& DefaultFile, const FString& FileTypes, uint32 Flags, TArray<FString>& OutFilenames) override;
+	virtual bool OpenFileDialog(const void* ParentWindowHandle, const FString& DialogTitle, const FString& DefaultPath, const FString& DefaultFile, const FString& FileTypes, uint32 Flags, TArray<FString>& OutFilenames, int32& OutFilterIndex) override;
 	virtual bool SaveFileDialog(const void* ParentWindowHandle, const FString& DialogTitle, const FString& DefaultPath, const FString& DefaultFile, const FString& FileTypes, uint32 Flags, TArray<FString>& OutFilenames) override;
 	virtual bool OpenDirectoryDialog(const void* ParentWindowHandle, const FString& DialogTitle, const FString& DefaultPath, FString& OutFolderName) override;
 	virtual bool OpenFontDialog(const void* ParentWindowHandle, FString& OutFontName, float& OutHeight, EFontImportFlags::Type& OutFlags) override;
@@ -28,7 +29,7 @@ public:
 	virtual FFeedbackContext* GetNativeFeedbackContext() override;
 
 private:
-	bool FileDialogShared(bool bSave, const void* ParentWindowHandle, const FString& DialogTitle, const FString& DefaultPath, const FString& DefaultFile, const FString& FileTypes, uint32 Flags, TArray<FString>& OutFilenames);
+	bool FileDialogShared(bool bSave, const void* ParentWindowHandle, const FString& DialogTitle, const FString& DefaultPath, const FString& DefaultFile, const FString& FileTypes, uint32 Flags, TArray<FString>& OutFilenames, int32& OutFilterIndex);
 	void GetRequiredRegistrySettings(TIndirectArray<FRegistryRootedKey> &RootedKeys);
 	int32 GetShellIntegrationVersion(const FString &FileName);
 };
diff --git a/Engine/Source/Developer/DesktopPlatform/Public/IDesktopPlatform.h b/Engine/Source/Developer/DesktopPlatform/Public/IDesktopPlatform.h
--- a/Engine/Source/Developer/DesktopPlatform/Public/IDesktopPlatform.h
+++ b/Engine/Source/Developer/DesktopPlatform/Public/IDesktopPlatform.h
@@ -73,6 +73,21 @@ public:
 	 */
 	virtual bool OpenFileDialog(const void* ParentWindowHandle, const FString& DialogTitle, const FString& DefaultPath, const FString& DefaultFile, const FString& FileTypes, uint32 Flags, TArray<FString>& OutFilenames) = 0;
 
+	/** 
+	 * Opens the "open file" dialog for the platform
+	 *
+	 * @param ParentWindowHandle		The native handle to the parent window for this dialog
+	 * @param DialogTitle				The text for the title of the dialog window
+	 * @param DefaultPath				The path where the file dialog will open initially
+	 * @param DefaultFile				The file that the dialog will select initially
+	 * @param Flags						Details about the dialog. See EFileDialogFlags.
+	 * @param FileTypes					The type filters to show in the dialog. This string should be a "|" delimited list of (Description|Extensionlist) pairs. Extensionlists are ";" delimited.
+	 * @param OutFilenames				The filenames that were selected in the dialog
+	 * @param OutFilterIndex			The type that was selected in the dialog
+	 * @return true if files were successfully selected
+	 */
+	virtual bool OpenFileDialog(const void* ParentWindowHandle, const FString& DialogTitle, const FString& DefaultPath, const FString& DefaultFile, const FString& FileTypes, uint32 Flags, TArray<FString>& OutFilenames, int32& outFilterIndex ) = 0;
+
 	/** 
 	 * Opens the "save file" dialog for the platform
 	 *
diff --git a/Engine/Source/Developer/ImageWrapper/Private/ImageWrapperModule.cpp b/Engine/Source/Developer/ImageWrapper/Private/ImageWrapperModule.cpp
--- a/Engine/Source/Developer/ImageWrapper/Private/ImageWrapperModule.cpp
+++ b/Engine/Source/Developer/ImageWrapper/Private/ImageWrapperModule.cpp
@@ -32,6 +32,9 @@ public:
 		case EImageFormat::JPEG:
 			ImageWrapper = new FJpegImageWrapper();
 			break;
+		case EImageFormat::GrayscaleJPEG:
+			ImageWrapper = new FJpegImageWrapper(1);
+			break;
 #endif	//WITH_UNREALJPEG
 
 		case EImageFormat::BMP:
diff --git a/Engine/Source/Developer/ImageWrapper/Private/JpegImageWrapper.cpp b/Engine/Source/Developer/ImageWrapper/Private/JpegImageWrapper.cpp
--- a/Engine/Source/Developer/ImageWrapper/Private/JpegImageWrapper.cpp
+++ b/Engine/Source/Developer/ImageWrapper/Private/JpegImageWrapper.cpp
@@ -29,8 +29,9 @@ FCriticalSection GJPEGSection;
 /* FJpegImageWrapper structors
  *****************************************************************************/
 
-FJpegImageWrapper::FJpegImageWrapper( )
-	: FImageWrapperBase()
+FJpegImageWrapper::FJpegImageWrapper( int32 InNumComponents )
+	: FImageWrapperBase(),
+	  NumComponents(InNumComponents)
 { }
 
 
@@ -71,7 +72,7 @@ bool FJpegImageWrapper::SetCompressed( const void* InCompressedData, int32 InCom
 
 bool FJpegImageWrapper::SetRaw( const void* InRawData, int32 InRawSize, const int32 InWidth, const int32 InHeight, const ERGBFormat::Type InFormat, const int32 InBitDepth )
 {
-	check((InFormat == ERGBFormat::RGBA || InFormat == ERGBFormat::BGRA) && InBitDepth == 8);
+	check((InFormat == ERGBFormat::RGBA || InFormat == ERGBFormat::BGRA || InFormat == ERGBFormat::Gray) && InBitDepth == 8);
 
 	bool bResult = FImageWrapperBase::SetRaw( InRawData, InRawSize, InWidth, InHeight, InFormat, InBitDepth );
 
@@ -114,7 +115,7 @@ void FJpegImageWrapper::Compress( int32 Quality )
 		jpge::params Parameters;
 		Parameters.m_quality = Quality;
 		bool bSuccess = jpge::compress_image_to_jpeg_file_in_memory(
-			CompressedData.GetTypedData(), OutBufferSize, Width, Height, 4, RawData.GetTypedData(), Parameters);
+			CompressedData.GetTypedData(), OutBufferSize, Width, Height, NumComponents, RawData.GetTypedData(), Parameters);
 		
 		check(bSuccess);
 
@@ -156,7 +157,11 @@ void FJpegImageWrapper::Uncompress( const ERGBFormat::Type InFormat, int32 InBit
 
 	RawData.Empty();
 	RawData.AddUninitialized( Width * Height * Channels );
-	FMemory::Memcpy( RawData.GetTypedData(), OutData, RawData.Num() );
+	if (OutData)
+	{
+		FMemory::Memcpy( RawData.GetTypedData(), OutData, RawData.Num() );
+		FMemory::Free(OutData);
+	}
 }
 
 
diff --git a/Engine/Source/Developer/ImageWrapper/Private/JpegImageWrapper.h b/Engine/Source/Developer/ImageWrapper/Private/JpegImageWrapper.h
--- a/Engine/Source/Developer/ImageWrapper/Private/JpegImageWrapper.h
+++ b/Engine/Source/Developer/ImageWrapper/Private/JpegImageWrapper.h
@@ -18,7 +18,7 @@ public:
 	/**
 	 * Default constructor.
 	 */
-	FJpegImageWrapper();
+	FJpegImageWrapper(int32 InNumComponents=4);
 
 public:
 
@@ -28,6 +28,9 @@ public:
 	virtual bool SetRaw( const void* InRawData, int32 InRawSize, const int32 InWidth, const int32 InHeight, const ERGBFormat::Type InFormat, const int32 InBitDepth ) override;
 	virtual void Uncompress( const ERGBFormat::Type InFormat, int32 InBitDepth ) override;
 	virtual void Compress(int32 Quality) override;
+
+private:
+	int32 NumComponents;
 };
 
 
diff --git a/Engine/Source/Developer/ImageWrapper/Public/Interfaces/IImageWrapper.h b/Engine/Source/Developer/ImageWrapper/Public/Interfaces/IImageWrapper.h
--- a/Engine/Source/Developer/ImageWrapper/Public/Interfaces/IImageWrapper.h
+++ b/Engine/Source/Developer/ImageWrapper/Public/Interfaces/IImageWrapper.h
@@ -16,6 +16,9 @@ namespace EImageFormat
 		/** Joint Photographic Experts Group */
 		JPEG,
 
+		/** Single channel jpeg */
+		GrayscaleJPEG,	
+
 		/** Windows Bitmap */
 		BMP,
 
diff --git a/Engine/Source/Editor/UnrealEd/Private/Editor.cpp b/Engine/Source/Editor/UnrealEd/Private/Editor.cpp
--- a/Engine/Source/Editor/UnrealEd/Private/Editor.cpp
+++ b/Engine/Source/Editor/UnrealEd/Private/Editor.cpp
@@ -2667,8 +2667,10 @@ void FReimportManager::GetNewReimportPath(UObject* Obj, TArray<FString>& InOutFi
 		return;
 	}
 
+	TMultiMap<uint32, UFactory*> DummyFilterIndexToFactory;
+
 	// Generate the file types and extensions represented by the selected factories
-	ObjectTools::GenerateFactoryFileExtensions( Factories, FileTypes, AllExtensions );
+	ObjectTools::GenerateFactoryFileExtensions( Factories, FileTypes, AllExtensions, DummyFilterIndexToFactory );
 
 	FileTypes = FString::Printf(TEXT("All Files (%s)|%s|%s"),*AllExtensions,*AllExtensions,*FileTypes);
 
diff --git a/Engine/Source/Editor/UnrealEd/Private/Factories/EditorFactories.cpp b/Engine/Source/Editor/UnrealEd/Private/Factories/EditorFactories.cpp
--- a/Engine/Source/Editor/UnrealEd/Private/Factories/EditorFactories.cpp
+++ b/Engine/Source/Editor/UnrealEd/Private/Factories/EditorFactories.cpp
@@ -2545,30 +2545,6 @@ public:
 	}
 };
 
-struct FTGAFileHeader
-{
-	uint8 IdFieldLength;
-	uint8 ColorMapType;
-	uint8 ImageTypeCode;		// 2 for uncompressed RGB format
-	uint16 ColorMapOrigin;
-	uint16 ColorMapLength;
-	uint8 ColorMapEntrySize;
-	uint16 XOrigin;
-	uint16 YOrigin;
-	uint16 Width;
-	uint16 Height;
-	uint8 BitsPerPixel;
-	uint8 ImageDescriptor;
-	friend FArchive& operator<<( FArchive& Ar, FTGAFileHeader& H )
-	{
-		Ar << H.IdFieldLength << H.ColorMapType << H.ImageTypeCode;
-		Ar << H.ColorMapOrigin << H.ColorMapLength << H.ColorMapEntrySize;
-		Ar << H.XOrigin << H.YOrigin << H.Width << H.Height << H.BitsPerPixel;
-		Ar << H.ImageDescriptor;
-		return Ar;
-	}
-};
-
 struct FTGAFileFooter
 {
 	uint32 ExtensionAreaOffset;
@@ -2865,6 +2841,376 @@ static void psd_GetPSDHeader( const uint8* Buffer, FPSDFileHeader& Info )
 		((int32)Buffer[25] <<  0);
 }
 
+
+void DecompressTGA_RLE_32bpp( const FTGAFileHeader* TGA, uint32* TextureData )
+{
+	uint8*	IdData		= (uint8*)TGA + sizeof(FTGAFileHeader); 
+	uint8*	ColorMap	= IdData + TGA->IdFieldLength;
+	uint8*	ImageData	= (uint8*) (ColorMap + (TGA->ColorMapEntrySize + 4) / 8 * TGA->ColorMapLength);				
+	uint32	Pixel		= 0;
+	int32     RLERun		= 0;
+	int32     RAWRun		= 0;
+
+	for(int32 Y = TGA->Height-1; Y >=0; Y--) // Y-flipped.
+	{					
+		for(int32 X = 0;X < TGA->Width;X++)
+		{						
+			if( RLERun > 0 )
+			{
+				RLERun--;  // reuse current Pixel data.
+			}
+			else if( RAWRun == 0 ) // new raw pixel or RLE-run.
+			{
+				uint8 RLEChunk = *(ImageData++);							
+				if( RLEChunk & 0x80 )
+				{
+					RLERun = ( RLEChunk & 0x7F ) + 1;
+					RAWRun = 1;
+				}
+				else
+				{
+					RAWRun = ( RLEChunk & 0x7F ) + 1;
+				}
+			}							
+			// Retrieve new pixel data - raw run or single pixel for RLE stretch.
+			if( RAWRun > 0 )
+			{
+				Pixel = *(uint32*)ImageData; // RGBA 32-bit dword.
+				ImageData += 4;
+				RAWRun--;
+				RLERun--;
+			}
+			// Store.
+			*( (TextureData + Y*TGA->Width)+X ) = Pixel;
+		}
+	}
+}
+
+void DecompressTGA_RLE_24bpp( const FTGAFileHeader* TGA, uint32* TextureData )
+{
+	uint8*	IdData = (uint8*)TGA + sizeof(FTGAFileHeader); 
+	uint8*	ColorMap = IdData + TGA->IdFieldLength;
+	uint8*	ImageData = (uint8*) (ColorMap + (TGA->ColorMapEntrySize + 4) / 8 * TGA->ColorMapLength);
+	uint8    Pixel[4];
+	int32     RLERun = 0;
+	int32     RAWRun = 0;
+
+	for(int32 Y = TGA->Height-1; Y >=0; Y--) // Y-flipped.
+	{					
+		for(int32 X = 0;X < TGA->Width;X++)
+		{						
+			if( RLERun > 0 )
+				RLERun--;  // reuse current Pixel data.
+			else if( RAWRun == 0 ) // new raw pixel or RLE-run.
+			{
+				uint8 RLEChunk = *(ImageData++);
+				if( RLEChunk & 0x80 )
+				{
+					RLERun = ( RLEChunk & 0x7F ) + 1;
+					RAWRun = 1;
+				}
+				else
+				{
+					RAWRun = ( RLEChunk & 0x7F ) + 1;
+				}
+			}							
+			// Retrieve new pixel data - raw run or single pixel for RLE stretch.
+			if( RAWRun > 0 )
+			{
+				Pixel[0] = *(ImageData++);
+				Pixel[1] = *(ImageData++);
+				Pixel[2] = *(ImageData++);
+				Pixel[3] = 255;
+				RAWRun--;
+				RLERun--;
+			}
+			// Store.
+			*( (TextureData + Y*TGA->Width)+X ) = *(uint32*)&Pixel;
+		}
+	}
+}
+
+void DecompressTGA_RLE_16bpp( const FTGAFileHeader* TGA, uint32* TextureData )
+{
+	uint8*	IdData = (uint8*)TGA + sizeof(FTGAFileHeader);
+	uint8*	ColorMap = IdData + TGA->IdFieldLength;				
+	uint16*	ImageData = (uint16*) (ColorMap + (TGA->ColorMapEntrySize + 4) / 8 * TGA->ColorMapLength);
+	uint16  FilePixel = 0;
+	uint32	TexturePixel = 0;
+	int32     RLERun = 0;
+	int32     RAWRun = 0;
+
+	for(int32 Y = TGA->Height-1; Y >=0; Y--) // Y-flipped.
+	{					
+		for( int32 X=0;X<TGA->Width;X++ )
+		{						
+			if( RLERun > 0 )
+				RLERun--;  // reuse current Pixel data.
+			else if( RAWRun == 0 ) // new raw pixel or RLE-run.
+			{
+				uint8 RLEChunk =  *((uint8*)ImageData);
+				ImageData = (uint16*)(((uint8*)ImageData)+1);
+				if( RLEChunk & 0x80 )
+				{
+					RLERun = ( RLEChunk & 0x7F ) + 1;
+					RAWRun = 1;
+				}
+				else
+				{
+					RAWRun = ( RLEChunk & 0x7F ) + 1;
+				}
+			}							
+			// Retrieve new pixel data - raw run or single pixel for RLE stretch.
+			if( RAWRun > 0 )
+			{ 
+				FilePixel = *(ImageData++);
+				RAWRun--;
+				RLERun--;
+			}
+			// Convert file format A1R5G5B5 into pixel format B8G8R8B8
+			TexturePixel = (FilePixel & 0x001F) << 3;
+			TexturePixel |= (FilePixel & 0x03E0) << 6;
+			TexturePixel |= (FilePixel & 0x7C00) << 9;
+			TexturePixel |= (FilePixel & 0x8000) << 16;
+			// Store.
+			*( (TextureData + Y*TGA->Width)+X ) = TexturePixel;
+		}
+	}
+}
+
+void DecompressTGA_32bpp( const FTGAFileHeader* TGA, uint32* TextureData )
+{
+
+	uint8*	IdData = (uint8*)TGA + sizeof(FTGAFileHeader);
+	uint8*	ColorMap = IdData + TGA->IdFieldLength;
+	uint32*	ImageData = (uint32*) (ColorMap + (TGA->ColorMapEntrySize + 4) / 8 * TGA->ColorMapLength);
+
+	for(int32 Y = 0;Y < TGA->Height;Y++)
+	{
+		FMemory::Memcpy(TextureData + Y * TGA->Width,ImageData + (TGA->Height - Y - 1) * TGA->Width,TGA->Width * 4);
+	}
+}
+
+void DecompressTGA_16bpp( const FTGAFileHeader* TGA, uint32* TextureData )
+{
+	uint8*	IdData = (uint8*)TGA + sizeof(FTGAFileHeader);
+	uint8*	ColorMap = IdData + TGA->IdFieldLength;
+	uint16*	ImageData = (uint16*) (ColorMap + (TGA->ColorMapEntrySize + 4) / 8 * TGA->ColorMapLength);
+	uint16    FilePixel = 0;
+	uint32	TexturePixel = 0;
+
+	for (int32 Y = TGA->Height - 1; Y >= 0; Y--)
+	{					
+		for (int32 X = 0; X<TGA->Width; X++)
+		{
+			FilePixel = *ImageData++;
+			// Convert file format A1R5G5B5 into pixel format B8G8R8A8
+			TexturePixel = (FilePixel & 0x001F) << 3;
+			TexturePixel |= (FilePixel & 0x03E0) << 6;
+			TexturePixel |= (FilePixel & 0x7C00) << 9;
+			TexturePixel |= (FilePixel & 0x8000) << 16;
+			// Store.
+			*((TextureData + Y*TGA->Width) + X) = TexturePixel;						
+		}
+	}
+}
+
+void DecompressTGA_24bpp( const FTGAFileHeader* TGA, uint32* TextureData )
+{
+	uint8*	IdData = (uint8*)TGA + sizeof(FTGAFileHeader);
+	uint8*	ColorMap = IdData + TGA->IdFieldLength;
+	uint8*	ImageData = (uint8*) (ColorMap + (TGA->ColorMapEntrySize + 4) / 8 * TGA->ColorMapLength);
+	uint8    Pixel[4];
+
+	for(int32 Y = 0;Y < TGA->Height;Y++)
+	{
+		for(int32 X = 0;X < TGA->Width;X++)
+		{
+			Pixel[0] = *(( ImageData+( TGA->Height-Y-1 )*TGA->Width*3 )+X*3+0);
+			Pixel[1] = *(( ImageData+( TGA->Height-Y-1 )*TGA->Width*3 )+X*3+1);
+			Pixel[2] = *(( ImageData+( TGA->Height-Y-1 )*TGA->Width*3 )+X*3+2);
+			Pixel[3] = 255;
+			*((TextureData+Y*TGA->Width)+X) = *(uint32*)&Pixel;
+		}
+	}
+}
+
+void DecompressTGA_8bpp( const FTGAFileHeader* TGA, uint8* TextureData )
+{
+	const uint8*  const IdData = (uint8*)TGA + sizeof(FTGAFileHeader);
+	const uint8*  const ColorMap = IdData + TGA->IdFieldLength;
+	const uint8*  const ImageData = (uint8*) (ColorMap + (TGA->ColorMapEntrySize + 4) / 8 * TGA->ColorMapLength);
+
+	int32 RevY = 0;
+	for (int32 Y = TGA->Height-1; Y >= 0; --Y)
+	{
+		const uint8* ImageCol = ImageData + (Y * TGA->Width); 
+		uint8* TextureCol = TextureData + (RevY++ * TGA->Width);
+		FMemory::Memcpy(TextureCol, ImageCol, TGA->Width);
+	}
+}
+
+bool DecompressTGA_helper(
+	const FTGAFileHeader* TGA,
+	uint32*& TextureData,
+	const int32 TextureDataSize,
+	FFeedbackContext* Warn )
+{
+	if(TGA->ImageTypeCode == 10) // 10 = RLE compressed 
+	{
+		// RLE compression: CHUNKS: 1 -byte header, high bit 0 = raw, 1 = compressed
+		// bits 0-6 are a 7-bit count; count+1 = number of raw pixels following, or rle pixels to be expanded. 
+		if(TGA->BitsPerPixel == 32)
+		{
+			DecompressTGA_RLE_32bpp(TGA, TextureData);
+		}
+		else if( TGA->BitsPerPixel == 24 )
+		{	
+			DecompressTGA_RLE_24bpp(TGA, TextureData);
+		}
+		else if( TGA->BitsPerPixel == 16 )
+		{
+			DecompressTGA_RLE_16bpp(TGA, TextureData);
+		}
+		else
+		{
+			Warn->Logf(ELogVerbosity::Error, TEXT("TGA uses an unsupported rle-compressed bit-depth: %u"),TGA->BitsPerPixel);
+			return false;
+		}
+	}
+	else if(TGA->ImageTypeCode == 2) // 2 = Uncompressed RGB
+	{
+		if(TGA->BitsPerPixel == 32)
+		{
+			DecompressTGA_32bpp(TGA, TextureData);
+		}
+		else if(TGA->BitsPerPixel == 16)
+		{
+			DecompressTGA_16bpp(TGA, TextureData);
+		}            
+		else if(TGA->BitsPerPixel == 24)
+		{
+			DecompressTGA_24bpp(TGA, TextureData);
+		}
+		else
+		{
+			Warn->Logf(ELogVerbosity::Error, TEXT("TGA uses an unsupported bit-depth: %u"),TGA->BitsPerPixel);
+			return false;
+		}
+	}
+	// Support for alpha stored as pseudo-color 8-bit TGA
+	else if(TGA->ColorMapType == 1 && TGA->ImageTypeCode == 1 && TGA->BitsPerPixel == 8)
+	{
+		DecompressTGA_8bpp(TGA, (uint8*)TextureData);
+	}
+	else
+	{
+		Warn->Logf(ELogVerbosity::Error, TEXT("TGA is an unsupported type: %u"),TGA->ImageTypeCode);
+		return false;
+	}
+
+	// Flip the image data if the flip bits are set in the TGA header.
+	bool FlipX = (TGA->ImageDescriptor & 0x10) ? 1 : 0;
+	bool FlipY = (TGA->ImageDescriptor & 0x20) ? 1 : 0;
+	if(FlipY || FlipX)
+	{
+		TArray<uint8> FlippedData;
+		FlippedData.AddUninitialized(TextureDataSize);
+
+		int32 NumBlocksX = TGA->Width;
+		int32 NumBlocksY = TGA->Height;
+		int32 BlockBytes = TGA->BitsPerPixel == 8 ? 1 : 4;
+
+		uint8* MipData = (uint8*)TextureData;
+
+		for(int32 Y = 0;Y < NumBlocksY;Y++)
+		{
+			for(int32 X  = 0;X < NumBlocksX;X++)
+			{
+				int32 DestX = FlipX ? (NumBlocksX - X - 1) : X;
+				int32 DestY = FlipY ? (NumBlocksY - Y - 1) : Y;
+				FMemory::Memcpy(
+					&FlippedData[(DestX + DestY * NumBlocksX) * BlockBytes],
+					&MipData[(X + Y * NumBlocksX) * BlockBytes],
+					BlockBytes
+					);
+			}
+		}
+		FMemory::Memcpy(MipData,FlippedData.GetData(),FlippedData.Num());
+	}
+
+	return true;
+}
+
+UTexture2D* DecompressTGA(
+	const FTGAFileHeader*	TGA,
+	UTextureFactory*		Factory,
+	UClass*					Class,
+	UObject*				InParent,
+	FName					Name,
+	EObjectFlags			Flags,
+	FFeedbackContext*	Warn)
+{
+	UTexture2D* Texture = Factory->CreateTexture2D( InParent, Name, Flags );
+	
+	if (TGA->ColorMapType == 1 && TGA->ImageTypeCode == 1 && TGA->BitsPerPixel == 8)
+	{
+		// Notes: The Scaleform GFx exporter (dll) strips all font glyphs into a single 8-bit texture.
+		// The targa format uses this for a palette index; GFx uses a palette of (i,i,i,i) so the index
+		// is also the alpha value.
+		//
+		// We store the image as PF_G8, where it will be used as alpha in the Glyph shader.
+
+		Texture->Source.Init(
+			TGA->Width,
+			TGA->Height,
+			/*NumSlices=*/ 1,
+			/*NumMips=*/ 1,
+			TSF_G8);
+
+		Texture->CompressionSettings = TC_Grayscale;
+	}
+	else
+	{
+		if(TGA->ImageTypeCode == 10) // 10 = RLE compressed 
+		{
+			if( TGA->BitsPerPixel != 32 &&
+				TGA->BitsPerPixel != 24 &&
+			    TGA->BitsPerPixel != 16 )
+			{
+				Warn->Logf(ELogVerbosity::Error, TEXT("TGA uses an unsupported rle-compressed bit-depth: %u"),TGA->BitsPerPixel);
+				return NULL;
+			}
+		}
+		else
+		{
+			if( TGA->BitsPerPixel != 32 &&
+				TGA->BitsPerPixel != 16 &&
+				TGA->BitsPerPixel != 24)
+			{
+				Warn->Logf(ELogVerbosity::Error, TEXT("TGA uses an unsupported bit-depth: %u"),TGA->BitsPerPixel);
+				return NULL;
+			}
+		}
+
+		Texture->Source.Init(
+			TGA->Width,
+			TGA->Height,
+			/*NumSlices=*/ 1,
+			/*NumMips=*/ 1,
+			TSF_BGRA8);
+	}
+
+	int32 TextureDataSize = Texture->Source.GetSizeX() * Texture->Source.GetSizeY() * Texture->Source.GetBytesPerPixel();
+	uint32* TextureData = (uint32*)Texture->Source.LockMip(0);
+
+	bool res = DecompressTGA_helper(TGA, TextureData, TextureDataSize, Warn);
+
+	Texture->Source.UnlockMip(0);
+
+	return Texture;
+}
+
 bool UTextureFactory::bSuppressImportOverwriteDialog = false;
 bool UTextureFactory::bSuppressImportResolutionWarnings = false;
 
@@ -3510,365 +3856,8 @@ UTexture* UTextureFactory::ImportTexture(UClass* Class, UObject* InParent, FName
 		{
 			return NULL;
 		}
-		if(TGA->ImageTypeCode == 10) // 10 = RLE compressed 
-		{
-			// RLE compression: CHUNKS: 1 -byte header, high bit 0 = raw, 1 = compressed
-			// bits 0-6 are a 7-bit count; count+1 = number of raw pixels following, or rle pixels to be expanded. 
-			if(TGA->BitsPerPixel == 32)
-			{
-				Texture = CreateTexture2D( InParent, Name, Flags );
-				if ( Texture )
-				{
-					Texture->Source.Init(
-						TGA->Width,
-						TGA->Height,
-						/*NumSlices=*/ 1,
-						/*NumMips=*/ 1,
-						TSF_BGRA8
-						);
-					uint32* TextureData = (uint32*)Texture->Source.LockMip(0);
-
-					uint8*	IdData		= (uint8*)TGA + sizeof(FTGAFileHeader); 
-					uint8*	ColorMap	= IdData + TGA->IdFieldLength;
-					uint8*	ImageData	= (uint8*) (ColorMap + (TGA->ColorMapEntrySize + 4) / 8 * TGA->ColorMapLength);				
-					uint32	Pixel		= 0;
-					int32     RLERun		= 0;
-					int32     RAWRun		= 0;
-
-					for(int32 Y = TGA->Height-1; Y >=0; Y--) // Y-flipped.
-					{					
-						for(int32 X = 0;X < TGA->Width;X++)
-						{						
-							if( RLERun > 0 )
-							{
-								RLERun--;  // reuse current Pixel data.
-							}
-							else if( RAWRun == 0 ) // new raw pixel or RLE-run.
-							{
-								uint8 RLEChunk = *(ImageData++);							
-								if( RLEChunk & 0x80 )
-								{
-									RLERun = ( RLEChunk & 0x7F ) + 1;
-									RAWRun = 1;
-								}
-								else
-								{
-									RAWRun = ( RLEChunk & 0x7F ) + 1;
-								}
-							}							
-							// Retrieve new pixel data - raw run or single pixel for RLE stretch.
-							if( RAWRun > 0 )
-							{
-								Pixel = *(uint32*)ImageData; // RGBA 32-bit dword.
-								ImageData += 4;
-								RAWRun--;
-								RLERun--;
-							}
-							// Store.
-							*( (TextureData + Y*TGA->Width)+X ) = Pixel;
-						}
-					}
-					Texture->Source.UnlockMip(0);
-				}
-			}
-			else if( TGA->BitsPerPixel == 24 )
-			{	
-				Texture = CreateTexture2D( InParent, Name, Flags );
-				if ( Texture )
-				{
-					Texture->Source.Init(
-						TGA->Width,
-						TGA->Height,
-						/*NumSlices=*/ 1,
-						/*NumMips=*/ 1,
-						TSF_BGRA8
-						);
-					uint32* TextureData = (uint32*)Texture->Source.LockMip(0);
-
-					uint8*	IdData = (uint8*)TGA + sizeof(FTGAFileHeader); 
-					uint8*	ColorMap = IdData + TGA->IdFieldLength;
-					uint8*	ImageData = (uint8*) (ColorMap + (TGA->ColorMapEntrySize + 4) / 8 * TGA->ColorMapLength);
-					uint8    Pixel[4];
-					int32     RLERun = 0;
-					int32     RAWRun = 0;
-				
-					for(int32 Y = TGA->Height-1; Y >=0; Y--) // Y-flipped.
-					{					
-						for(int32 X = 0;X < TGA->Width;X++)
-						{						
-							if( RLERun > 0 )
-								RLERun--;  // reuse current Pixel data.
-							else if( RAWRun == 0 ) // new raw pixel or RLE-run.
-							{
-								uint8 RLEChunk = *(ImageData++);
-								if( RLEChunk & 0x80 )
-								{
-									RLERun = ( RLEChunk & 0x7F ) + 1;
-									RAWRun = 1;
-								}
-								else
-								{
-									RAWRun = ( RLEChunk & 0x7F ) + 1;
-								}
-							}							
-							// Retrieve new pixel data - raw run or single pixel for RLE stretch.
-							if( RAWRun > 0 )
-							{
-								Pixel[0] = *(ImageData++);
-								Pixel[1] = *(ImageData++);
-								Pixel[2] = *(ImageData++);
-								Pixel[3] = 255;
-								RAWRun--;
-								RLERun--;
-							}
-							// Store.
-							*( (TextureData + Y*TGA->Width)+X ) = *(uint32*)&Pixel;
-						}
-					}
-					Texture->Source.UnlockMip(0);
-				}
-			}
-			else if( TGA->BitsPerPixel == 16 )
-			{
-				Texture = CreateTexture2D( InParent, Name, Flags );
-				if ( Texture )
-				{
-					Texture->Source.Init(
-						TGA->Width,
-						TGA->Height,
-						/*NumSlices=*/ 1,
-						/*NumMips=*/ 1,
-						TSF_BGRA8
-						);
-					uint32* TextureData = (uint32*)Texture->Source.LockMip(0);
-				
-					uint8*	IdData = (uint8*)TGA + sizeof(FTGAFileHeader);
-					uint8*	ColorMap = IdData + TGA->IdFieldLength;				
-					uint16*	ImageData = (uint16*) (ColorMap + (TGA->ColorMapEntrySize + 4) / 8 * TGA->ColorMapLength);
-					uint16  FilePixel = 0;
-					uint32	TexturePixel = 0;
-					int32     RLERun = 0;
-					int32     RAWRun = 0;
-
-					for(int32 Y = TGA->Height-1; Y >=0; Y--) // Y-flipped.
-					{					
-						for( int32 X=0;X<TGA->Width;X++ )
-						{						
-							if( RLERun > 0 )
-								RLERun--;  // reuse current Pixel data.
-							else if( RAWRun == 0 ) // new raw pixel or RLE-run.
-							{
-								uint8 RLEChunk =  *((uint8*)ImageData);
-								ImageData = (uint16*)(((uint8*)ImageData)+1);
-								if( RLEChunk & 0x80 )
-								{
-									RLERun = ( RLEChunk & 0x7F ) + 1;
-									RAWRun = 1;
-								}
-								else
-								{
-									RAWRun = ( RLEChunk & 0x7F ) + 1;
-								}
-							}							
-							// Retrieve new pixel data - raw run or single pixel for RLE stretch.
-							if( RAWRun > 0 )
-							{ 
-								FilePixel = *(ImageData++);
-								RAWRun--;
-								RLERun--;
-							}
-							// Convert file format A1R5G5B5 into pixel format B8G8R8B8
-							TexturePixel = (FilePixel & 0x001F) << 3;
-							TexturePixel |= (FilePixel & 0x03E0) << 6;
-							TexturePixel |= (FilePixel & 0x7C00) << 9;
-							TexturePixel |= (FilePixel & 0x8000) << 16;
-							// Store.
-							*( (TextureData + Y*TGA->Width)+X ) = TexturePixel;
-						}
-					}
-					Texture->Source.UnlockMip(0);
-				}
-			}
-			else
-			{
-				Warn->Logf(ELogVerbosity::Error, TEXT("TGA uses an unsupported rle-compressed bit-depth: %u"),TGA->BitsPerPixel);
-				return NULL;
-			}
-		}
-		else if(TGA->ImageTypeCode == 2) // 2 = Uncompressed RGB
-		{
-			if(TGA->BitsPerPixel == 32)
-			{
-				Texture = CreateTexture2D( InParent, Name, Flags );
-				if ( Texture )
-				{
-					Texture->Source.Init(
-						TGA->Width,
-						TGA->Height,
-						/*NumSlices=*/ 1,
-						/*NumMips=*/ 1,
-						TSF_BGRA8
-						);
-					uint32* TextureData = (uint32*)Texture->Source.LockMip(0);
-
-					uint8*	IdData = (uint8*)TGA + sizeof(FTGAFileHeader);
-					uint8*	ColorMap = IdData + TGA->IdFieldLength;
-					uint32*	ImageData = (uint32*) (ColorMap + (TGA->ColorMapEntrySize + 4) / 8 * TGA->ColorMapLength);
-
-					for(int32 Y = 0;Y < TGA->Height;Y++)
-					{
-						FMemory::Memcpy(TextureData + Y * TGA->Width,ImageData + (TGA->Height - Y - 1) * TGA->Width,TGA->Width * 4);
-					}
-					Texture->Source.UnlockMip(0);
-				}
-			}
-			else if(TGA->BitsPerPixel == 16)
-			{
-				Texture = CreateTexture2D( InParent, Name, Flags );
-				Texture->Source.Init(
-					TGA->Width,
-					TGA->Height,
-					/*NumSlices=*/ 1,
-					/*NumMips=*/ 1,
-					TSF_BGRA8
-					);
-				uint32* TextureData = (uint32*)Texture->Source.LockMip(0);
-
-				uint8*	IdData = (uint8*)TGA + sizeof(FTGAFileHeader);
-				uint8*	ColorMap = IdData + TGA->IdFieldLength;
-				uint16*	ImageData = (uint16*) (ColorMap + (TGA->ColorMapEntrySize + 4) / 8 * TGA->ColorMapLength);
-				uint16    FilePixel = 0;
-				uint32	TexturePixel = 0;
-
-				for (int32 Y = TGA->Height - 1; Y >= 0; Y--)
-				{					
-					for (int32 X = 0; X<TGA->Width; X++)
-					{
-						FilePixel = *ImageData++;
-						// Convert file format A1R5G5B5 into pixel format B8G8R8A8
-						TexturePixel = (FilePixel & 0x001F) << 3;
-						TexturePixel |= (FilePixel & 0x03E0) << 6;
-						TexturePixel |= (FilePixel & 0x7C00) << 9;
-						TexturePixel |= (FilePixel & 0x8000) << 16;
-						// Store.
-						*((TextureData + Y*TGA->Width) + X) = TexturePixel;						
-					}
-				}
-	   
-				Texture->Source.UnlockMip(0);
-			}            
-			else if(TGA->BitsPerPixel == 24)
-			{
-				Texture = CreateTexture2D( InParent, Name, Flags );
-				if ( Texture )
-				{
-					Texture->Source.Init(
-						TGA->Width,
-						TGA->Height,
-						/*NumSlices=*/ 1,
-						/*NumMips=*/ 1,
-						TSF_BGRA8
-						);
-					uint32* TextureData = (uint32*)Texture->Source.LockMip(0);
-
-					uint8*	IdData = (uint8*)TGA + sizeof(FTGAFileHeader);
-					uint8*	ColorMap = IdData + TGA->IdFieldLength;
-					uint8*	ImageData = (uint8*) (ColorMap + (TGA->ColorMapEntrySize + 4) / 8 * TGA->ColorMapLength);
-					uint8    Pixel[4];
-
-					for(int32 Y = 0;Y < TGA->Height;Y++)
-					{
-						for(int32 X = 0;X < TGA->Width;X++)
-						{
-							Pixel[0] = *(( ImageData+( TGA->Height-Y-1 )*TGA->Width*3 )+X*3+0);
-							Pixel[1] = *(( ImageData+( TGA->Height-Y-1 )*TGA->Width*3 )+X*3+1);
-							Pixel[2] = *(( ImageData+( TGA->Height-Y-1 )*TGA->Width*3 )+X*3+2);
-							Pixel[3] = 255;
-							*((TextureData+Y*TGA->Width)+X) = *(uint32*)&Pixel;
-						}
-					}
-					Texture->Source.UnlockMip(0);
-				}
-			}
-			else
-			{
-				Warn->Logf(ELogVerbosity::Error, TEXT("TGA uses an unsupported bit-depth: %u"),TGA->BitsPerPixel);
-				return NULL;
-			}
-		}
-		// Support for alpha stored as pseudo-color 8-bit TGA
-		else if(TGA->ColorMapType == 1 && TGA->ImageTypeCode == 1 && TGA->BitsPerPixel == 8)
-		{
-			// Notes: The Scaleform GFx exporter (dll) strips all font glyphs into a single 8-bit texture.
-			// The targa format uses this for a palette index; GFx uses a palette of (i,i,i,i) so the index
-			// is also the alpha value.
-			//
-			// We store the image as PF_G8, where it will be used as alpha in the Glyph shader.
-
-			Texture = CreateTexture2D( InParent, Name, Flags );
-			if ( Texture )
-			{
-				Texture->Source.Init(
-					TGA->Width,
-					TGA->Height,
-					/*NumSlices=*/ 1,
-					/*NumMips=*/ 1,
-					TSF_G8
-					);
-				uint8* TextureData = Texture->Source.LockMip(0);
-				Texture->CompressionSettings = TC_Grayscale;
-
-				const uint8*  const IdData = (uint8*)TGA + sizeof(FTGAFileHeader);
-				const uint8*  const ColorMap = IdData + TGA->IdFieldLength;
-				const uint8*  const ImageData = (uint8*) (ColorMap + (TGA->ColorMapEntrySize + 4) / 8 * TGA->ColorMapLength);
-
-				int32 RevY = 0;
-				for (int32 Y = TGA->Height-1; Y >= 0; --Y)
-				{
-					const uint8* ImageCol = ImageData + (Y * TGA->Width); 
-					uint8* TextureCol = TextureData + (RevY++ * TGA->Width);
-					FMemory::Memcpy(TextureCol, ImageCol, TGA->Width);
-				}
-				Texture->Source.UnlockMip(0);
-			}
-		}
-		else
-		{
-			Warn->Logf(ELogVerbosity::Error, TEXT("TGA is an unsupported type: %u"),TGA->ImageTypeCode);
-			return NULL;
-		}
-
-		// Flip the image data if the flip bits are set in the TGA header.
-		bool FlipX = (TGA->ImageDescriptor & 0x10) ? 1 : 0;
-		bool FlipY = (TGA->ImageDescriptor & 0x20) ? 1 : 0;
-		if(FlipY || FlipX)
-		{
-			TArray<uint8> FlippedData;
-			FlippedData.AddUninitialized(Texture->Source.CalcMipSize(0));
-
-			int32 NumBlocksX = Texture->Source.GetSizeX();
-			int32 NumBlocksY = Texture->Source.GetSizeY();
-			int32 BlockBytes = Texture->Source.GetBytesPerPixel();
-
-			uint8* MipData = Texture->Source.LockMip(0);
-			for(int32 Y = 0;Y < NumBlocksY;Y++)
-			{
-				for(int32 X  = 0;X < NumBlocksX;X++)
-				{
-					int32 DestX = FlipX ? (NumBlocksX - X - 1) : X;
-					int32 DestY = FlipY ? (NumBlocksY - Y - 1) : Y;
-					FMemory::Memcpy(
-						&FlippedData[(DestX + DestY * NumBlocksX) * BlockBytes],
-						&MipData[(X + Y * NumBlocksX) * BlockBytes],
-						BlockBytes
-						);
-				}
-			}
-			FMemory::Memcpy(MipData,FlippedData.GetData(),FlippedData.Num());
-			Texture->Source.UnlockMip(0);
-		}
 
-		return Texture;
+		return DecompressTGA(TGA, this, Class, InParent, Name, Flags, Warn);
 	}
 	//
 	// PSD File
diff --git a/Engine/Source/Editor/UnrealEd/Private/ObjectTools.cpp b/Engine/Source/Editor/UnrealEd/Private/ObjectTools.cpp
--- a/Engine/Source/Editor/UnrealEd/Private/ObjectTools.cpp
+++ b/Engine/Source/Editor/UnrealEd/Private/ObjectTools.cpp
@@ -2565,12 +2565,12 @@ namespace ObjectTools
 	 * @param	out_Filetypes	File types supported by the provided factory, concatenated into a string
 	 * @param	out_Extensions	Extensions supported by the provided factory, concatenated into a string
 	 */
-	void GenerateFactoryFileExtensions( UFactory* InFactory, FString& out_Filetypes, FString& out_Extensions )
+	void GenerateFactoryFileExtensions( UFactory* InFactory, FString& out_Filetypes, FString& out_Extensions, TMultiMap<uint32, UFactory*>& out_FilterIndexToFactory )
 	{
 		// Place the factory in an array and call the overloaded version of this function
 		TArray<UFactory*> FactoryArray;
 		FactoryArray.Add( InFactory );
-		GenerateFactoryFileExtensions( FactoryArray, out_Filetypes, out_Extensions );
+		GenerateFactoryFileExtensions( FactoryArray, out_Filetypes, out_Extensions, out_FilterIndexToFactory );
 	}
 
 	/**
@@ -2580,10 +2580,11 @@ namespace ObjectTools
 	 * @param	out_Filetypes	File types supported by the provided factory, concatenated into a string
 	 * @param	out_Extensions	Extensions supported by the provided factory, concatenated into a string
 	 */
-	void GenerateFactoryFileExtensions( const TArray<UFactory*>& InFactories, FString& out_Filetypes, FString& out_Extensions )
+	void GenerateFactoryFileExtensions( const TArray<UFactory*>& InFactories, FString& out_Filetypes, FString& out_Extensions, TMultiMap<uint32, UFactory*>& out_FilterIndexToFactory )
 	{
 		// Store all the descriptions and their corresponding extensions in a map
-		TMultiMap<FString, FString> DescToExtensionMap;		
+		TMultiMap<FString, FString> DescToExtensionMap;
+		TMultiMap<FString, UFactory*> DescToFactory;
 
 		// Iterate over each factory, retrieving their supported file descriptions and extensions, and storing them into the map
 		for ( TArray<UFactory*>::TConstIterator FactoryIter(InFactories); FactoryIter; ++FactoryIter )
@@ -2600,6 +2601,7 @@ namespace ObjectTools
 			for ( int32 FormatIndex = 0; FormatIndex < Descriptions.Num() && FormatIndex < Extensions.Num(); ++FormatIndex )
 			{
 				DescToExtensionMap.AddUnique( Descriptions[FormatIndex], Extensions[FormatIndex ] );
+				DescToFactory.AddUnique( Descriptions[FormatIndex], *FactoryIter );
 			}
 		}
 		
@@ -2615,6 +2617,8 @@ namespace ObjectTools
 		DescToExtensionMap.GetKeys( DescriptionKeyMap );
 		const TArray<FString>& DescriptionKeys = DescriptionKeyMap;
 
+		uint32 IdxFilter = 1; // the type list will start by an all supported files wildcard value
+
 		// Iterate over each unique map key, retrieving all of each key's associated values in order to populate the strings
 		for ( TArray<FString>::TConstIterator DescIter( DescriptionKeys ); DescIter; ++DescIter )
 		{
@@ -2649,7 +2653,18 @@ namespace ObjectTools
 					{
 						out_Filetypes += TEXT("|");
 					}
-					out_Filetypes += CurLine;	
+					out_Filetypes += CurLine;
+
+					// save the order in which descriptions are added to be able to identify 
+					// factories using filter index
+					TArray<UFactory*> Factories;
+					DescToFactory.MultiFind( CurDescription, Factories );
+					TArray<UFactory*>::TIterator FactIt(Factories);
+					for (;FactIt;++FactIt)
+					{
+						out_FilterIndexToFactory.Add( IdxFilter, *FactIt );
+					}
+					++IdxFilter;
 				}
 			}
 		}
diff --git a/Engine/Source/Editor/UnrealEd/Public/Factories.h b/Engine/Source/Editor/UnrealEd/Public/Factories.h
--- a/Engine/Source/Editor/UnrealEd/Public/Factories.h
+++ b/Engine/Source/Editor/UnrealEd/Public/Factories.h
@@ -94,4 +94,42 @@ public:
     virtual USkeletalMesh* CreateSkeletalMesh(const TCHAR* SrcFilename, EMorphImportError* Error ) = 0;
 };
 
+#pragma pack(push,1)
+
+struct FTGAFileHeader
+{
+	uint8 IdFieldLength;
+	uint8 ColorMapType;
+	uint8 ImageTypeCode;		// 2 for uncompressed RGB format
+	uint16 ColorMapOrigin;
+	uint16 ColorMapLength;
+	uint8 ColorMapEntrySize;
+	uint16 XOrigin;
+	uint16 YOrigin;
+	uint16 Width;
+	uint16 Height;
+	uint8 BitsPerPixel;
+	uint8 ImageDescriptor;
+	friend FArchive& operator<<( FArchive& Ar, FTGAFileHeader& H )
+	{
+		Ar << H.IdFieldLength << H.ColorMapType << H.ImageTypeCode;
+		Ar << H.ColorMapOrigin << H.ColorMapLength << H.ColorMapEntrySize;
+		Ar << H.XOrigin << H.YOrigin << H.Width << H.Height << H.BitsPerPixel;
+		Ar << H.ImageDescriptor;
+		return Ar;
+	}
+};
+
+#pragma pack(pop)
+
+/**
+ * This helper allows to decompress TGA data in a pre-allocated memory block.
+ * The pixel format is necessarily PF_A8R8G8B8.
+ */
+UNREALED_API bool DecompressTGA_helper(
+	const FTGAFileHeader* TGA,
+	uint32*& TextureData,
+	const int32 TextureDataSize,
+	FFeedbackContext* Warn );
+
 #endif // __FACTORIES_H__
diff --git a/Engine/Source/Editor/UnrealEd/Public/ObjectTools.h b/Engine/Source/Editor/UnrealEd/Public/ObjectTools.h
--- a/Engine/Source/Editor/UnrealEd/Public/ObjectTools.h
+++ b/Engine/Source/Editor/UnrealEd/Public/ObjectTools.h
@@ -407,7 +407,7 @@ namespace ObjectTools
 	 * @param	out_Filetypes	File types supported by the provided factory, concatenated into a string
 	 * @param	out_Extensions	Extensions supported by the provided factory, concatenated into a string
 	 */
-	UNREALED_API void GenerateFactoryFileExtensions( const UFactory* InFactory, FString& out_Filetypes, FString& out_Extensions );
+	UNREALED_API void GenerateFactoryFileExtensions( const UFactory* InFactory, FString& out_Filetypes, FString& out_Extensions, TMultiMap<uint32, UFactory*>& out_FilterIndexToFactory );
 	
 	/**
 	 * Populates two strings with all of the file types and extensions the provided factories support.
@@ -416,7 +416,7 @@ namespace ObjectTools
 	 * @param	out_Filetypes	File types supported by the provided factory, concatenated into a string
 	 * @param	out_Extensions	Extensions supported by the provided factory, concatenated into a string
 	 */
-	UNREALED_API void GenerateFactoryFileExtensions( const TArray<UFactory*>& InFactories, FString& out_Filetypes, FString& out_Extensions );
+	UNREALED_API void GenerateFactoryFileExtensions( const TArray<UFactory*>& InFactories, FString& out_Filetypes, FString& out_Extensions, TMultiMap<uint32, UFactory*>& out_FilterIndexToFactory );
 	
 	/**
 	 * Generates a list of file types for a given class.
diff --git a/Engine/Source/Programs/UnrealBuildTool/System/Plugins.cs b/Engine/Source/Programs/UnrealBuildTool/System/Plugins.cs
--- a/Engine/Source/Programs/UnrealBuildTool/System/Plugins.cs
+++ b/Engine/Source/Programs/UnrealBuildTool/System/Plugins.cs
@@ -41,6 +41,9 @@ namespace UnrealBuildTool
 			// List of platforms supported by this modules
 			public List<UnrealTargetPlatform> Platforms;
 		}
+		
+		// Whether or not the plugin should be built
+		public bool bShouldBuild;
 
 		// Plugin name
 		public string Name;
@@ -251,9 +254,22 @@ namespace UnrealBuildTool
 								}
 							}
 						}
-				
-						// add to list of modules
-						PluginInfo.Modules.Add( PluginModuleInfo );
+						
+						object ModuleShouldBuild;
+						if( ModuleDict.TryGetValue( "bShouldBuild", out ModuleShouldBuild ) )
+						{
+							PluginInfo.bShouldBuild = (Int64)ModuleShouldBuild == 1 ? true : false;
+						}
+						else
+						{
+							PluginInfo.bShouldBuild = true;
+						}
+
+						if (PluginInfo.bShouldBuild)
+						{
+							// add to list of modules
+							PluginInfo.Modules.Add(PluginModuleInfo);
+						}
 					}
 				}
 				else
diff --git a/Engine/Source/Runtime/Engine/Classes/Engine/Texture.h b/Engine/Source/Runtime/Engine/Classes/Engine/Texture.h
--- a/Engine/Source/Runtime/Engine/Classes/Engine/Texture.h
+++ b/Engine/Source/Runtime/Engine/Classes/Engine/Texture.h
@@ -283,10 +283,10 @@ struct FTexturePlatformData
 #endif
 
 	/** Default constructor. */
-	FTexturePlatformData();
+	ENGINE_API FTexturePlatformData();
 
 	/** Destructor. */
-	~FTexturePlatformData();
+	ENGINE_API ~FTexturePlatformData();
 
 	/**
 	 * Try to load mips from the derived data cache.
@@ -524,7 +524,7 @@ public:
 	/**
 	 * Creates a new resource for the texture, and updates any cached references to the resource.
 	 */
-	virtual void UpdateResource();
+	ENGINE_API virtual void UpdateResource();
 
 	/**
 	 * Implemented by subclasses to create a new resource for the texture.
@@ -575,7 +575,7 @@ public:
 	/**
 	 * Serializes cooked platform data.
 	 */
-	void SerializeCookedPlatformData(class FArchive& Ar);
+	ENGINE_API void SerializeCookedPlatformData(class FArchive& Ar);
 
 #if WITH_EDITORONLY_DATA
 	/**
@@ -607,7 +607,7 @@ public:
 	/**
 	 * Forces platform data to be rebuilt.
 	 */
-	void ForceRebuildPlatformData();
+	ENGINE_API void ForceRebuildPlatformData();
 
 	ENGINE_API void UpdateCachedLODBias( bool bIncTextureMips = true );
 
@@ -625,14 +625,14 @@ public:
 
 	// Begin UObject interface.
 #if WITH_EDITOR
-	virtual void PostEditChangeProperty(FPropertyChangedEvent& PropertyChangedEvent) override;
+	ENGINE_API virtual void PostEditChangeProperty(FPropertyChangedEvent& PropertyChangedEvent) override;
 #endif // WITH_EDITOR
-	virtual void Serialize(FArchive& Ar) override;
-	virtual void PostLoad() override;
-	virtual void PreSave() override;
-	virtual void BeginDestroy() override;
-	virtual bool IsReadyForFinishDestroy() override;
-	virtual void FinishDestroy() override;
+	ENGINE_API virtual void Serialize(FArchive& Ar) override;
+	ENGINE_API virtual void PostLoad() override;
+	ENGINE_API virtual void PreSave() override;
+	ENGINE_API virtual void BeginDestroy() override;
+	ENGINE_API virtual bool IsReadyForFinishDestroy() override;
+	ENGINE_API virtual void FinishDestroy() override;
 	// End UObject interface.
 
 	/**
@@ -643,7 +643,7 @@ public:
 	 *
 	 *	@return	float					The average brightness of the texture
 	 */
-	virtual float GetAverageBrightness(bool bIgnoreTrueBlack, bool bUseGrayscale);
+	ENGINE_API virtual float GetAverageBrightness(bool bIgnoreTrueBlack, bool bUseGrayscale);
 	
 	// @todo document
 	ENGINE_API static const TCHAR* GetTextureGroupString(TextureGroup InGroup);
diff --git a/Engine/Source/Runtime/Engine/Classes/Engine/Texture2DDynamic.h b/Engine/Source/Runtime/Engine/Classes/Engine/Texture2DDynamic.h
--- a/Engine/Source/Runtime/Engine/Classes/Engine/Texture2DDynamic.h
+++ b/Engine/Source/Runtime/Engine/Classes/Engine/Texture2DDynamic.h
@@ -4,7 +4,7 @@
 #pragma once
 #include "Texture2DDynamic.generated.h"
 
-UCLASS(hidecategories=Object)
+UCLASS(hidecategories=Object, MinimalAPI)
 class UTexture2DDynamic : public UTexture
 {
 	GENERATED_UCLASS_BODY()
@@ -28,10 +28,10 @@ class UTexture2DDynamic : public UTexture
 
 public:
 	// Begin UTexture Interface.
-	virtual FTextureResource* CreateResource() override;
+	ENGINE_API virtual FTextureResource* CreateResource() override;
 	virtual EMaterialValueType GetMaterialType() override { return MCT_Texture2D; }
-	virtual float GetSurfaceWidth() const override;
-	virtual float GetSurfaceHeight() const override;
+	ENGINE_API virtual float GetSurfaceWidth() const override;
+	ENGINE_API virtual float GetSurfaceHeight() const override;
 	// End UTexture Interface.
 	
 	/**
@@ -42,7 +42,7 @@ public:
 	 * @param InFormat			- Format of the texture, defaults to PF_B8G8R8A8
 	 * @param InIsResolveTarget	- Whether the texture can be used as a resolve target
 	 */
-	void Init(int32 InSizeX, int32 InSizeY, EPixelFormat InFormat = PF_B8G8R8A8, bool InIsResolveTarget = false);
+	ENGINE_API void Init(int32 InSizeX, int32 InSizeY, EPixelFormat InFormat = PF_B8G8R8A8, bool InIsResolveTarget = false);
 	
 	/** Creates and initializes a new Texture2DDynamic with the requested settings */
 	class UTexture2DDynamic* Create(int32 InSizeX, int32 InSizeY, EPixelFormat InFormat = PF_B8G8R8A8, bool InIsResolveTarget = false);
diff --git a/Engine/Source/Runtime/Engine/Private/Materials/MaterialExpressions.cpp b/Engine/Source/Runtime/Engine/Private/Materials/MaterialExpressions.cpp
--- a/Engine/Source/Runtime/Engine/Private/Materials/MaterialExpressions.cpp
+++ b/Engine/Source/Runtime/Engine/Private/Materials/MaterialExpressions.cpp
@@ -1688,6 +1688,10 @@ bool UMaterialExpressionTextureSampleParameter2D::TextureIsValid( UTexture* InTe
 		{
 			Result = true;
 		}
+		if ( InTexture->IsA(UTexture2DDynamic::StaticClass()) )
+		{
+			Result = true;
+		}
 	}
 	return Result;
 }
diff --git a/Engine/Source/Runtime/Engine/Public/MaterialCompiler.h b/Engine/Source/Runtime/Engine/Public/MaterialCompiler.h
--- a/Engine/Source/Runtime/Engine/Public/MaterialCompiler.h
+++ b/Engine/Source/Runtime/Engine/Public/MaterialCompiler.h
@@ -21,7 +21,7 @@ public:
 
 	virtual void SetMaterialProperty(EMaterialProperty InProperty, EShaderFrequency InShaderFrequency) = 0;
 	virtual int32 Error(const TCHAR* Text) = 0;
-	int32 Errorf(const TCHAR* Format,...);
+	ENGINE_API int32 Errorf(const TCHAR* Format,...);
 
 	virtual int32 CallExpression(FMaterialExpressionKey ExpressionKey,FMaterialCompiler* InCompiler) = 0;
 
diff --git a/Engine/Source/Runtime/Engine/Public/Texture.h b/Engine/Source/Runtime/Engine/Public/Texture.h
--- a/Engine/Source/Runtime/Engine/Public/Texture.h
+++ b/Engine/Source/Runtime/Engine/Public/Texture.h
@@ -166,7 +166,7 @@ struct FTexture2DMipMap
 	}
 
 	/** Serialization. */
-	void Serialize(FArchive& Ar, UObject* Owner, int32 MipIndex);
+	ENGINE_API void Serialize(FArchive& Ar, UObject* Owner, int32 MipIndex);
 
 #if WITH_EDITORONLY_DATA
 	/** Key if stored in the derived data cache. */
@@ -410,7 +410,7 @@ public:
 	virtual void ReleaseRHI();
 
 	/** Returns the Texture2DRHI, which can be used for locking/unlocking the mips. */
-	FTexture2DRHIRef GetTexture2DRHI();
+	ENGINE_API FTexture2DRHIRef GetTexture2DRHI();
 
 private:
 	/** The owner of this resource. */
