Support option to not "ConvertScene" when importing FBX

- i.e. blender doesn't provide correct scene causing the scene to be rotated wrong direction

diff --git a/Engine/Config/BaseEditorUserSettings.ini b/Engine/Config/BaseEditorUserSettings.ini
--- a/Engine/Config/BaseEditorUserSettings.ini
+++ b/Engine/Config/BaseEditorUserSettings.ini
@@ -539,6 +539,7 @@ bAutoCreateGroups=True
 bInvertNormalMaps=False
 bImportTextures=False
 bOverrideFullName=True
+bConvertScene=True
 bRemoveNamespaces=True
 bImportAnimations=False
 bResampleAnimations=False
diff --git a/Engine/Source/Editor/UnrealEd/Classes/Factories/FbxImportUI.h b/Engine/Source/Editor/UnrealEd/Classes/Factories/FbxImportUI.h
--- a/Engine/Source/Editor/UnrealEd/Classes/Factories/FbxImportUI.h
+++ b/Engine/Source/Editor/UnrealEd/Classes/Factories/FbxImportUI.h
@@ -35,6 +35,10 @@ class UFbxImportUI : public UObject
 	UPROPERTY(EditAnywhere, config, Category=General)
 	uint32 bOverrideFullName:1;
 
+	/** Whether to convert scene from FBX scene. */
+	UPROPERTY(EditAnywhere, config, Category=General)
+	uint32 bConvertScene:1;
+
 	/** For static meshes, enabling this option will combine all meshes in the FBX into a single monolithic mesh in Unreal */
 	UPROPERTY(EditAnywhere, config, Category=StaticMesh, meta=(ToolTip="If enabled, combines all meshes into a single mesh"))
 	uint32 bCombineMeshes:1;
diff --git a/Engine/Source/Editor/UnrealEd/Private/Fbx/FbxMainImport.cpp b/Engine/Source/Editor/UnrealEd/Private/Fbx/FbxMainImport.cpp
--- a/Engine/Source/Editor/UnrealEd/Private/Fbx/FbxMainImport.cpp
+++ b/Engine/Source/Editor/UnrealEd/Private/Fbx/FbxMainImport.cpp
@@ -160,6 +160,7 @@ void ApplyImportUIToImportOptions(UFbxImportUI* ImportUI, FBXImportOptions& InOu
 	InOutImportOptions.bInvertNormalMap = ImportUI->TextureImportData->bInvertNormalMaps;
 	InOutImportOptions.bImportTextures = ImportUI->bImportTextures;
 	InOutImportOptions.bUsedAsFullName = ImportUI->bOverrideFullName;
+	InOutImportOptions.bConvertScene = ImportUI->bConvertScene;
 	InOutImportOptions.bImportAnimations = ImportUI->bImportAnimations;
 	InOutImportOptions.SkeletonForAnimation = ImportUI->Skeleton;
 	if ( ImportUI->MeshTypeToImport == FBXIT_StaticMesh )
@@ -694,17 +695,20 @@ bool FFbxImporter::ImportFromFile(const FString& Filename, const FString& Type)
 			// The imported axis system is unknown for obj files
 			if( !Type.Equals( Obj, ESearchCase::IgnoreCase ) )
 			{
-				FbxAxisSystem::EFrontVector FrontVector = (FbxAxisSystem::EFrontVector) - FbxAxisSystem::eParityOdd;
-				const FbxAxisSystem UnrealZUp(FbxAxisSystem::eZAxis, FrontVector, FbxAxisSystem::eRightHanded);
-
-				if( Scene->GetGlobalSettings().GetAxisSystem() != UnrealZUp )
+				if (GetImportOptions()->bConvertScene)
 				{
-					// Converts the FBX data to Z-up, X-forward, Y-left.  Unreal is the same except with Y-right, 
-					// but the conversion to left-handed coordinates is not working properly
+					FbxAxisSystem::EFrontVector FrontVector = (FbxAxisSystem::EFrontVector) - FbxAxisSystem::eParityOdd;
+					const FbxAxisSystem UnrealZUp(FbxAxisSystem::eZAxis, FrontVector, FbxAxisSystem::eRightHanded);
+
+					if(Scene->GetGlobalSettings().GetAxisSystem() != UnrealZUp)
+					{
+						// Converts the FBX data to Z-up, X-forward, Y-left.  Unreal is the same except with Y-right, 
+						// but the conversion to left-handed coordinates is not working properly
 
-					// convert axis to Z-up
-					FbxRootNodeUtility::RemoveAllFbxRoots(Scene);
-					UnrealZUp.ConvertScene(Scene);
+						// convert axis to Z-up
+						FbxRootNodeUtility::RemoveAllFbxRoots(Scene);
+						UnrealZUp.ConvertScene(Scene);
+					}
 				}
 			}
 
diff --git a/Engine/Source/Editor/UnrealEd/Private/Fbx/FbxOptionWindow.cpp b/Engine/Source/Editor/UnrealEd/Private/Fbx/FbxOptionWindow.cpp
--- a/Engine/Source/Editor/UnrealEd/Private/Fbx/FbxOptionWindow.cpp
+++ b/Engine/Source/Editor/UnrealEd/Private/Fbx/FbxOptionWindow.cpp
@@ -545,6 +545,18 @@ TSharedRef<SWidget> SFbxOptionWindow::ConstructMiscOption()
 		]
 	];
 
+	NewBox->AddSlot().AutoHeight() .Padding(2)
+	[
+		SNew(SCheckBox)
+		.IsChecked(ImportUI->bConvertScene? true: false)
+		.OnCheckStateChanged(this, &SFbxOptionWindow::SetGeneral_ConvertScene)
+		.IsEnabled( !bIsObjFormat )
+		[
+			SNew(STextBlock)
+			.Text(LOCTEXT("FbxOptionWindow_ConvertScene", "Convert Scene"))
+		]
+	];
+
 	return NewBox;
 }
 
@@ -1126,6 +1138,11 @@ void SFbxOptionWindow::SetGeneral_OverrideFullName(ESlateCheckBoxState::Type New
 	ImportUI->bOverrideFullName = (NewType == ESlateCheckBoxState::Checked)? 1:0;
 }
 
+void SFbxOptionWindow::SetGeneral_ConvertScene(ESlateCheckBoxState::Type NewType)
+{
+	ImportUI->bConvertScene = (NewType == ESlateCheckBoxState::Checked)? 1:0;
+}
+
 void SFbxOptionWindow::SetSkeletalMesh_ImportMeshLODs(ESlateCheckBoxState::Type NewType)
 {
 	ImportUI->SkeletalMeshImportData->bImportMeshLODs = (NewType == ESlateCheckBoxState::Checked)? 1:0;
diff --git a/Engine/Source/Editor/UnrealEd/Private/FbxImporter.h b/Engine/Source/Editor/UnrealEd/Private/FbxImporter.h
--- a/Engine/Source/Editor/UnrealEd/Private/FbxImporter.h
+++ b/Engine/Source/Editor/UnrealEd/Private/FbxImporter.h
@@ -106,6 +106,7 @@ struct FBXImportOptions
 	bool bImportTextures;
 	bool bImportLOD;
 	bool bUsedAsFullName;
+	bool bConvertScene;
 	bool bRemoveNameSpace;
 	EFBXNormalImportMethod NormalImportMethod;
 	// Static Mesh options
diff --git a/Engine/Source/Editor/UnrealEd/Private/FbxOptionWindow.h b/Engine/Source/Editor/UnrealEd/Private/FbxOptionWindow.h
--- a/Engine/Source/Editor/UnrealEd/Private/FbxOptionWindow.h
+++ b/Engine/Source/Editor/UnrealEd/Private/FbxOptionWindow.h
@@ -98,6 +98,7 @@ private:
 	void SetImportType(EFBXImportType ImportType);
 
 	void SetGeneral_OverrideFullName(ESlateCheckBoxState::Type NewType);
+	void SetGeneral_ConvertScene(ESlateCheckBoxState::Type NewType);
 	void SetMesh_ImportTangents(ESlateCheckBoxState::Type NewType);
 	void SetSkeletalMesh_ImportMeshLODs(ESlateCheckBoxState::Type NewType);
 	void SetSkeletalMesh_ImportMorphTargets(ESlateCheckBoxState::Type NewType);
