UE4: Removed code attempting to branch files that were duplicated or renamed. This is unsafe to do at duplicate/rename time because the file was not yet created on disk, it is just an in-memory asset. The act of attempting to branch here causes a file to show up on disk while a package representing it is created in memory that is not linked to the file. This could cause many problems related to saving or loading this file. This also has asset registry implications. Fixes crash in TTP#338373

diff --git a/Engine/Source/Editor/UnrealEd/Private/ObjectTools.cpp b/Engine/Source/Editor/UnrealEd/Private/ObjectTools.cpp
--- a/Engine/Source/Editor/UnrealEd/Private/ObjectTools.cpp
+++ b/Engine/Source/Editor/UnrealEd/Private/ObjectTools.cpp
@@ -438,9 +438,6 @@ namespace ObjectTools
 		// Replace all references
 		FArchiveReplaceObjectRef<UObject> ReplaceAr( DupObject, ReplacementMap, false, true, true );
 
-		// if possible, we should branch in source control here.
-		SourceControlHelpers::BranchFile(SourceControlHelpers::PackageFilename(ReturnObject->GetOutermost()), SourceControlHelpers::PackageFilename(Object->GetOutermost()));
-
 		return ReturnObject;
 	}
 		
@@ -2421,9 +2418,6 @@ namespace ObjectTools
 				FAssetRegistryModule::AssetCreated(NewRedirector);
 			}
 
-			// if possible, we should branch in source control here.
-			SourceControlHelpers::BranchFile(SourceControlHelpers::PackageFilename(NewPackage), SourceControlHelpers::PackageFilename(OldPackage));
-
 			// Saw Successful Rename
 			InOutErrorMessage = FText::FromString( ErrorMessage );
 			return true;
