Made it possible for asset editors to maintain their own FEditorModeTools lists

diff --git a/Engine/Source/Editor/TextureAlignMode/Private/TextureAlignEdMode.cpp b/Engine/Source/Editor/TextureAlignMode/Private/TextureAlignEdMode.cpp
--- a/Engine/Source/Editor/TextureAlignMode/Private/TextureAlignEdMode.cpp
+++ b/Engine/Source/Editor/TextureAlignMode/Private/TextureAlignEdMode.cpp
@@ -10,25 +10,15 @@ DEFINE_LOG_CATEGORY_STATIC(LogTextureAlignMode, Log, All);
 
 void FTextureAlignModeModule::StartupModule()
 {
-	TSharedRef<FEdModeTexture> NewEditorMode = MakeShareable(new FEdModeTexture);
-	GEditorModeTools().RegisterMode(NewEditorMode);
-	EdModeTexture = NewEditorMode;
-
-	IBspModeModule& BspModule = FModuleManager::LoadModuleChecked<IBspModeModule>("BspMode");
-	GEditorModeTools().RegisterCompatibleModes( NewEditorMode->GetID(), BspModule.GetBspMode()->GetID() );
+	FEditorModeRegistry::Get().RegisterMode<FEdModeTexture>(
+		FBuiltinEditorModes::EM_Texture,
+		NSLOCTEXT("EditorModes", "TextureAlignmentMode", "Texture Alignment")
+		);
 }
 
 void FTextureAlignModeModule::ShutdownModule()
 {
-	GEditorModeTools().UnregisterMode(EdModeTexture.ToSharedRef());
-
-	if ( FModuleManager::Get().IsModuleLoaded( TEXT("BspMode") ) )
-	{
-		IBspModeModule& BspModule = FModuleManager::LoadModuleChecked<IBspModeModule>( TEXT("BspMode") );
-		GEditorModeTools().UnregisterCompatibleModes( EdModeTexture->GetID(), BspModule.GetBspMode()->GetID() );
-	}
-
-	EdModeTexture = NULL;
+	FEditorModeRegistry::Get().UnregisterMode(FBuiltinEditorModes::EM_Texture);
 }
 
 /*------------------------------------------------------------------------------
@@ -39,9 +29,6 @@ FEdModeTexture::FEdModeTexture()
 	:	ScopedTransaction( NULL )
 	, TrackingWorld( NULL )
 {
-	ID = FBuiltinEditorModes::EM_Texture;
-	Name = NSLOCTEXT("EditorModes", "TextureAlignmentMode", "Texture Alignment");
-
 	Tools.Add( new FModeTool_Texture() );
 	SetCurrentTool( MT_Texture );
 }
@@ -57,8 +44,8 @@ void FEdModeTexture::Enter()
 	FEdMode::Enter();
 
 	const bool bGetRawValue = true;
-	SaveCoordSystem = GEditorModeTools().GetCoordSystem(bGetRawValue);
-	GEditorModeTools().SetCoordSystem(COORD_Local);
+	SaveCoordSystem = GLevelEditorModeTools().GetCoordSystem(bGetRawValue);
+	GLevelEditorModeTools().SetCoordSystem(COORD_Local);
 }
 
 void FEdModeTexture::Exit()
@@ -71,7 +58,7 @@ void FEdModeTexture::Exit()
 
 	FEdMode::Exit();
 
-	GEditorModeTools().SetCoordSystem(SaveCoordSystem);
+	GLevelEditorModeTools().SetCoordSystem(SaveCoordSystem);
 	FEditorSupportDelegates::RedrawAllViewports.Broadcast();
 }
 
@@ -99,7 +86,7 @@ bool FEdModeTexture::ShouldDrawWidget() const
 bool FEdModeTexture::GetCustomDrawingCoordinateSystem( FMatrix& InMatrix, void* InData )
 {
 	// Texture mode is ALWAYS in local space
-	GEditorModeTools().SetCoordSystem(COORD_Local);
+	GLevelEditorModeTools().SetCoordSystem(COORD_Local);
 
 	FPoly* poly = NULL;
 
@@ -152,7 +139,7 @@ EAxisList::Type FEdModeTexture::GetWidgetAxisToDraw( FWidget::EWidgetMode InWidg
 	return EAxisList::XYZ;
 }
 
-bool FEdModeTexture::StartTracking(FLevelEditorViewportClient* InViewportClient, FViewport* InViewport)
+bool FEdModeTexture::StartTracking(FEditorViewportClient* InViewportClient, FViewport* InViewport)
 {
 	// call base version because it calls the StartModify() virtual method needed to track drag events
 	bool BaseRtn = FEdMode::StartTracking(InViewportClient, InViewport);
@@ -174,7 +161,7 @@ bool FEdModeTexture::StartTracking(FLevelEditorViewportClient* InViewportClient,
 	return BaseRtn;
 }
 
-bool FEdModeTexture::EndTracking(FLevelEditorViewportClient* InViewportClient, FViewport* InViewport)
+bool FEdModeTexture::EndTracking(FEditorViewportClient* InViewportClient, FViewport* InViewport)
 {
 	// Clean up the scoped transaction if one is still pending
 	if( ScopedTransaction != NULL )
@@ -194,6 +181,11 @@ bool FEdModeTexture::EndTracking(FLevelEditorViewportClient* InViewportClient, F
 	return FEdMode::EndTracking(InViewportClient, InViewport);
 }
 
+bool FEdModeTexture::IsCompatibleWith(FEditorModeID OtherModeID) const
+{
+	return OtherModeID == FBuiltinEditorModes::EM_Bsp;
+}
+
 
 
 
@@ -211,7 +203,7 @@ FModeTool_Texture::FModeTool_Texture()
 /**
  * @return		true if the delta was handled by this editor mode tool.
  */
-bool FModeTool_Texture::InputDelta(FLevelEditorViewportClient* InViewportClient,FViewport* InViewport,FVector& InDrag,FRotator& InRot,FVector& InScale)
+bool FModeTool_Texture::InputDelta(FEditorViewportClient* InViewportClient,FViewport* InViewport,FVector& InDrag,FRotator& InRot,FVector& InScale)
 {
 	if( InViewportClient->GetCurrentWidgetAxis() == EAxisList::None )
 	{
@@ -243,7 +235,7 @@ bool FModeTool_Texture::InputDelta(FLevelEditorViewportClient* InViewportClient,
 					Surf.pBase = Model->Points.Add(Base);
 				}
 			}
-			FMatrix Mat = GEditorModeTools().GetCustomDrawingCoordinateSystem();
+			FMatrix Mat = GLevelEditorModeTools().GetCustomDrawingCoordinateSystem();
 			FVector UVW = Mat.InverseTransformVector( deltaDrag );  // InverseTransformNormal because Mat is the transform from the surface/widget's coords to world coords
 			GEditor->polyTexPan( Model, UVW.X, UVW.Y, 0 );  // 0 is relative mode because UVW is made from deltaDrag - the user input since the last tick
 		}
diff --git a/Engine/Source/Editor/TextureAlignMode/Private/TextureAlignMode.cpp b/Engine/Source/Editor/TextureAlignMode/Private/TextureAlignMode.cpp
--- a/Engine/Source/Editor/TextureAlignMode/Private/TextureAlignMode.cpp
+++ b/Engine/Source/Editor/TextureAlignMode/Private/TextureAlignMode.cpp
@@ -26,7 +26,7 @@ FText FTextureAlignMode::GetBaseToolkitName() const
 
 class FEdMode* FTextureAlignMode::GetEditorMode() const
 {
-	return (FEdModeTexture*)GEditorModeTools().GetActiveMode(FBuiltinEditorModes::EM_Texture);
+	return (FEdModeTexture*)GLevelEditorModeTools().GetActiveMode(FBuiltinEditorModes::EM_Texture);
 }
 
 #undef LOCTEXT_NAMESPACE
\ No newline at end of file
diff --git a/Engine/Source/Editor/TextureAlignMode/Public/TextureAlignEdMode.h b/Engine/Source/Editor/TextureAlignMode/Public/TextureAlignEdMode.h
--- a/Engine/Source/Editor/TextureAlignMode/Public/TextureAlignEdMode.h
+++ b/Engine/Source/Editor/TextureAlignMode/Public/TextureAlignEdMode.h
@@ -40,9 +40,10 @@ public:
 	virtual bool GetCustomDrawingCoordinateSystem( FMatrix& InMatrix, void* InData );
 	virtual bool GetCustomInputCoordinateSystem( FMatrix& InMatrix, void* InData );
 	virtual EAxisList::Type GetWidgetAxisToDraw( FWidget::EWidgetMode InWidgetMode ) const;
-	virtual bool StartTracking(FLevelEditorViewportClient* InViewportClient, FViewport* InViewport);
-	virtual bool EndTracking(FLevelEditorViewportClient* InViewportClient, FViewport* InViewport);
+	virtual bool StartTracking(FEditorViewportClient* InViewportClient, FViewport* InViewport);
+	virtual bool EndTracking(FEditorViewportClient* InViewportClient, FViewport* InViewport);
 	virtual bool AllowWidgetMove() { return false; }
+	virtual bool IsCompatibleWith(FEditorModeID OtherModeID) const override;
 
 protected:
 	/** The current transaction. */
@@ -59,7 +60,7 @@ class FModeTool_Texture : public FModeTool
 public:
 	FModeTool_Texture();
 
-	virtual bool InputDelta(FLevelEditorViewportClient* InViewportClient,FViewport* InViewport,FVector& InDrag,FRotator& InRot,FVector& InScale);
+	virtual bool InputDelta(FEditorViewportClient* InViewportClient,FViewport* InViewport,FVector& InDrag,FRotator& InRot,FVector& InScale);
 
 	// override these to allow this tool to keep track of the users dragging during a single drag event
 	virtual bool StartModify()	{ PreviousInputDrag = FVector::ZeroVector; return true; }
diff --git a/Engine/Source/Editor/UnrealEd/Private/CameraController.cpp b/Engine/Source/Editor/UnrealEd/Private/CameraController.cpp
--- a/Engine/Source/Editor/UnrealEd/Private/CameraController.cpp
+++ b/Engine/Source/Editor/UnrealEd/Private/CameraController.cpp
@@ -356,7 +356,7 @@ void FEditorCameraController::UpdateRotation( const FCameraControllerUserImpulse
 			}
 			
 			//allow for unlocked pitch constraints while in matinee
-			if (Config.bLockedPitch || !GEditorModeTools().GetActiveMode(FBuiltinEditorModes::EM_InterpEdit)) 
+			if (Config.bLockedPitch || !GLevelEditorModeTools().GetActiveMode(FBuiltinEditorModes::EM_InterpEdit)) 
 			{
 				// Clamp the angle.
 				InOutCameraEuler[ CurRotationAxis ] =
diff --git a/Engine/Source/Editor/UnrealEd/Private/ComponentVisualizerManager.cpp b/Engine/Source/Editor/UnrealEd/Private/ComponentVisualizerManager.cpp
--- a/Engine/Source/Editor/UnrealEd/Private/ComponentVisualizerManager.cpp
+++ b/Engine/Source/Editor/UnrealEd/Private/ComponentVisualizerManager.cpp
@@ -2,6 +2,33 @@
 
 #include "UnrealEd.h"
 #include "ComponentVisualizerManager.h"
+#include "ILevelEditor.h"
+
+/** Handle a click on the specified level editor viewport client */
+bool FComponentVisualizerManager::HandleClick(FLevelEditorViewportClient* InViewportClient, HHitProxy *HitProxy, const FViewportClick &Click)
+{
+	bool bHandled = HandleProxyForComponentVis(HitProxy);
+	if (bHandled && Click.GetKey() == EKeys::RightMouseButton)
+	{
+		TSharedPtr<SWidget> MenuWidget = GenerateContextMenuForComponentVis();
+		if (MenuWidget.IsValid())
+		{
+			auto ParentLevelEditorPinned = InViewportClient->ParentLevelEditor.Pin();
+			if (ParentLevelEditorPinned.IsValid())
+			{
+				FSlateApplication::Get().PushMenu(
+					ParentLevelEditorPinned.ToSharedRef(),
+					MenuWidget.ToSharedRef(),
+					FSlateApplication::Get().GetCursorPos(),
+					FPopupTransitionEffect(FPopupTransitionEffect::ContextMenu));
+
+				return true;
+			}
+		}
+	}
+
+	return false;
+}
 
 bool FComponentVisualizerManager::HandleProxyForComponentVis(HHitProxy *HitProxy)
 {
@@ -46,7 +73,7 @@ void FComponentVisualizerManager::ClearActiveComponentVis()
 	}
 }
 
-bool FComponentVisualizerManager::HandleInputKey(FLevelEditorViewportClient* ViewportClient, FViewport* Viewport, FKey Key, EInputEvent Event) const
+bool FComponentVisualizerManager::HandleInputKey(FEditorViewportClient* ViewportClient, FViewport* Viewport, FKey Key, EInputEvent Event) const
 {
 	if(EditedVisualizer.IsValid())
 	{
@@ -59,7 +86,7 @@ bool FComponentVisualizerManager::HandleInputKey(FLevelEditorViewportClient* Vie
 	return false;
 }
 
-bool FComponentVisualizerManager::HandleInputDelta(FLevelEditorViewportClient* InViewportClient, FViewport* InViewport, FVector& InDrag, FRotator& InRot, FVector& InScale) const
+bool FComponentVisualizerManager::HandleInputDelta(FEditorViewportClient* InViewportClient, FViewport* InViewport, FVector& InDrag, FRotator& InRot, FVector& InScale) const
 {
 	if (EditedVisualizer.IsValid() && InViewportClient->GetCurrentWidgetAxis() != EAxisList::None)
 	{
diff --git a/Engine/Source/Editor/UnrealEd/Private/DragTool_BoxSelect.cpp b/Engine/Source/Editor/UnrealEd/Private/DragTool_BoxSelect.cpp
--- a/Engine/Source/Editor/UnrealEd/Private/DragTool_BoxSelect.cpp
+++ b/Engine/Source/Editor/UnrealEd/Private/DragTool_BoxSelect.cpp
@@ -122,7 +122,7 @@ void FDragTool_ActorBoxSelect::AddDelta( const FVector& InDelta )
 */
 void FDragTool_ActorBoxSelect::EndDrag()
 {
-	FEditorModeTools& EdModeTools = GEditorModeTools();
+	FEditorModeTools& EdModeTools = GLevelEditorModeTools();
 	const bool bGeometryMode = EdModeTools.IsModeActive( FBuiltinEditorModes::EM_Geometry );
 	
 	FScopedTransaction Transaction( NSLOCTEXT("ActorFrustumSelect", "MarqueeSelectTransation", "Marquee Select" ) );
@@ -139,11 +139,11 @@ void FDragTool_ActorBoxSelect::EndDrag()
 	else if( !bShiftDown )
 	{
 		// If the user is selecting, but isn't hold down SHIFT, remove all current selections.
-		GEditorModeTools().SelectNone();
+		GLevelEditorModeTools().SelectNone();
 	}
 
 	// Let the editor mode try to handle the box selection.
-	const bool bEditorModeHandledBoxSelection = GEditorModeTools().BoxSelect( SelBBox, bLeftMouseButtonDown );
+	const bool bEditorModeHandledBoxSelection = GLevelEditorModeTools().BoxSelect( SelBBox, bLeftMouseButtonDown );
 
 	// If the edit mode didn't handle the selection, try normal actor box selection.
 	if ( !bEditorModeHandledBoxSelection )
@@ -277,7 +277,7 @@ bool FDragTool_ActorBoxSelect::IntersectsBox( AActor& InActor, const FBox& InBox
 {
 	bool bActorHitByBox = false;
 
-	FEditorModeTools& EdModeTools = GEditorModeTools();
+	FEditorModeTools& EdModeTools = GLevelEditorModeTools();
 	const bool bGeometryMode = EdModeTools.IsModeActive( FBuiltinEditorModes::EM_Geometry );
 
 	// Check for special cases (like certain show flags that might hide an actor)
diff --git a/Engine/Source/Editor/UnrealEd/Private/DragTool_FrustumSelect.cpp b/Engine/Source/Editor/UnrealEd/Private/DragTool_FrustumSelect.cpp
--- a/Engine/Source/Editor/UnrealEd/Private/DragTool_FrustumSelect.cpp
+++ b/Engine/Source/Editor/UnrealEd/Private/DragTool_FrustumSelect.cpp
@@ -56,7 +56,7 @@ void FDragTool_ActorFrustumSelect::StartDrag(FEditorViewportClient* InViewportCl
 */
 void FDragTool_ActorFrustumSelect::EndDrag()
 {
-	FEditorModeTools& EdModeTools = GEditorModeTools();
+	FEditorModeTools& EdModeTools = GLevelEditorModeTools();
 	const bool bGeometryMode = EdModeTools.IsModeActive( FBuiltinEditorModes::EM_Geometry );
 
 	FSceneViewFamilyContext ViewFamily(FSceneViewFamily::ConstructionValues(ViewportClient->Viewport, ViewportClient->GetScene(), ViewportClient->EngineShowFlags ));
@@ -74,7 +74,7 @@ void FDragTool_ActorFrustumSelect::EndDrag()
 	if( !bShiftDown )
 	{
 		// If the user is selecting, but isn't hold down SHIFT, remove all current selections.
-		GEditorModeTools().SelectNone();
+		GLevelEditorModeTools().SelectNone();
 	}
 
 	// Does an actor have to be fully contained in the box to be selected
@@ -82,7 +82,7 @@ void FDragTool_ActorFrustumSelect::EndDrag()
 	bool bSelectionChanged = false;
 
 	// Let the editor mode try to handle the selection.
-	const bool bEditorModeHandledSelection = GEditorModeTools().FrustumSelect( Frustum, bLeftMouseButtonDown );
+	const bool bEditorModeHandledSelection = GLevelEditorModeTools().FrustumSelect( Frustum, bLeftMouseButtonDown );
 
 	if( !bEditorModeHandledSelection )
 	{
diff --git a/Engine/Source/Editor/UnrealEd/Private/Editor.cpp b/Engine/Source/Editor/UnrealEd/Private/Editor.cpp
--- a/Engine/Source/Editor/UnrealEd/Private/Editor.cpp
+++ b/Engine/Source/Editor/UnrealEd/Private/Editor.cpp
@@ -455,7 +455,7 @@ bool UEditorEngine::ShouldDrawBrushWireframe( AActor* InActor )
 {
 	bool bResult = true;
 
-	bResult = GEditorModeTools().ShouldDrawBrushWireframe( InActor );
+	bResult = GLevelEditorModeTools().ShouldDrawBrushWireframe( InActor );
 	
 	return bResult;
 }
@@ -2136,9 +2136,9 @@ void UEditorEngine::ApplyDeltaToActor(AActor* InActor,
 			if ( bDelta )
 			{
 				FVector NewActorLocation = InActor->GetActorLocation();
-				NewActorLocation -= GEditorModeTools().PivotLocation;
+				NewActorLocation -= GLevelEditorModeTools().PivotLocation;
 				NewActorLocation = FRotationMatrix( InDeltaRot ).TransformPosition( NewActorLocation );
-				NewActorLocation += GEditorModeTools().PivotLocation;
+				NewActorLocation += GLevelEditorModeTools().PivotLocation;
 				NewActorLocation -= InActor->GetActorLocation();
 				InActor->EditorApplyTranslation( NewActorLocation, bAltDown, bShiftDown, bControlDown );
 			}
@@ -2242,9 +2242,9 @@ void UEditorEngine::ApplyDeltaToActor(AActor* InActor,
 						for( int32 vertex = 0 ; vertex < Poly->Vertices.Num() ; vertex++ )
 						{
 							FVector Wk = BrushActorToWorld.TransformPosition( Poly->Vertices[vertex] );
-							Wk -= GEditorModeTools().PivotLocation;
+							Wk -= GLevelEditorModeTools().PivotLocation;
 							Wk += matrix.TransformPosition( Wk );
-							Wk += GEditorModeTools().PivotLocation;
+							Wk += GLevelEditorModeTools().PivotLocation;
 							Poly->Vertices[vertex] = BrushActorToWorld.InverseTransformPosition( Wk );
 						}
 
@@ -2255,9 +2255,9 @@ void UEditorEngine::ApplyDeltaToActor(AActor* InActor,
 						}
 
 						FVector Wk = BrushActorToWorld.TransformPosition( Poly->Base );
-						Wk -= GEditorModeTools().PivotLocation;
+						Wk -= GLevelEditorModeTools().PivotLocation;
 						Wk += matrix.TransformPosition( Wk );
-						Wk += GEditorModeTools().PivotLocation;
+						Wk += GLevelEditorModeTools().PivotLocation;
 						Poly->Base = BrushActorToWorld.InverseTransformPosition( Wk );
 
 						// Scale the texture vectors
@@ -2308,7 +2308,7 @@ void UEditorEngine::ApplyDeltaToActor(AActor* InActor,
 
 					InActor->EditorApplyScale( 
 						ModifiedScale,
-												&GEditorModeTools().PivotLocation,
+												&GLevelEditorModeTools().PivotLocation,
 												bAltDown,
 												bShiftDown,
 						bControlDown
@@ -3805,7 +3805,7 @@ bool UEditorEngine::ShouldOpenMatinee(AMatineeActor* MatineeActor) const
 	}
 
 	// Make sure we can't open the same action twice in Matinee.
-	if( GEditorModeTools().IsModeActive(FBuiltinEditorModes::EM_InterpEdit) )
+	if( GLevelEditorModeTools().IsModeActive(FBuiltinEditorModes::EM_InterpEdit) )
 	{
 		FMessageDialog::Open( EAppMsgType::Ok, NSLOCTEXT("UnrealEd", "MatineeActionAlreadyOpen", "An UnrealMatinee sequence is currently open in an editor.  Please close it before proceeding.") );
 		return false;
@@ -3839,14 +3839,14 @@ void UEditorEngine::OpenMatinee(AMatineeActor* MatineeActor, bool bWarnUser)
 	}
 
 	// If already in Matinee mode, exit out before going back in with new Interpolation.
-	if( GEditorModeTools().IsModeActive( FBuiltinEditorModes::EM_InterpEdit ) )
+	if( GLevelEditorModeTools().IsModeActive( FBuiltinEditorModes::EM_InterpEdit ) )
 	{
-		GEditorModeTools().DeactivateMode( FBuiltinEditorModes::EM_InterpEdit );
+		GLevelEditorModeTools().DeactivateMode( FBuiltinEditorModes::EM_InterpEdit );
 	}
 
-	GEditorModeTools().ActivateMode( FBuiltinEditorModes::EM_InterpEdit );
+	GLevelEditorModeTools().ActivateMode( FBuiltinEditorModes::EM_InterpEdit );
 
-	FEdModeInterpEdit* InterpEditMode = (FEdModeInterpEdit*)GEditorModeTools().GetActiveMode( FBuiltinEditorModes::EM_InterpEdit );
+	FEdModeInterpEdit* InterpEditMode = (FEdModeInterpEdit*)GLevelEditorModeTools().GetActiveMode( FBuiltinEditorModes::EM_InterpEdit );
 
 	InterpEditMode->InitInterpMode( MatineeActor );
 
diff --git a/Engine/Source/Editor/UnrealEd/Private/EditorActor.cpp b/Engine/Source/Editor/UnrealEd/Private/EditorActor.cpp
--- a/Engine/Source/Editor/UnrealEd/Private/EditorActor.cpp
+++ b/Engine/Source/Editor/UnrealEd/Private/EditorActor.cpp
@@ -455,7 +455,7 @@ void UUnrealEdEngine::edactDuplicateSelected( ULevel* InLevel, bool bOffsetLocat
 	}
 	
 	TArray<FEdMode*> ActiveModes;
-	GEditorModeTools().GetActiveModes( ActiveModes );
+	GLevelEditorModeTools().GetActiveModes( ActiveModes );
 
 	for( int32 ModeIndex = 0; ModeIndex < ActiveModes.Num(); ++ModeIndex )
 	{
@@ -722,7 +722,7 @@ bool UUnrealEdEngine::ShouldAbortActorDeletion() const
 
 	// Can't delete actors if Matinee is open.
 	const FText ErrorMsg = NSLOCTEXT("UnrealEd", "Error_WrongModeForActorDeletion", "Cannot delete actor while Matinee is open" );
-	if ( !GEditorModeTools().EnsureNotInMode( FBuiltinEditorModes::EM_InterpEdit, ErrorMsg, true ) )
+	if ( !GLevelEditorModeTools().EnsureNotInMode( FBuiltinEditorModes::EM_InterpEdit, ErrorMsg, true ) )
 	{
 		bResult = true;
 	}
@@ -2074,7 +2074,7 @@ void UUnrealEdEngine::edactAlignOrigin()
 			Brush->SetActorLocation(BrushLocation, false);
 
 			//Update EditorMode locations to match the new brush location
-			FEditorModeTools& Tools = GEditorModeTools();
+			FEditorModeTools& Tools = GLevelEditorModeTools();
 			Tools.SetPivotLocation( Brush->GetActorLocation(), true );
 
 			Brush->Brush->BuildBound();
@@ -2148,12 +2148,12 @@ void UUnrealEdEngine::edactAlignVertices()
 					}
 
 					// Determine if we are in geometry edit mode.
-					if ( GEditorModeTools().IsModeActive(FBuiltinEditorModes::EM_Geometry) )
+					if ( GLevelEditorModeTools().IsModeActive(FBuiltinEditorModes::EM_Geometry) )
 					{
 						// If we are in geometry mode, go through the list of geometry objects
 						// and find our current brush and update its source data as it might have changed 
 						// in RecomputePoly
-						FEdModeGeometry* GeomMode = (FEdModeGeometry*)GEditorModeTools().GetActiveMode(FBuiltinEditorModes::EM_Geometry);
+						FEdModeGeometry* GeomMode = (FEdModeGeometry*)GLevelEditorModeTools().GetActiveMode(FBuiltinEditorModes::EM_Geometry);
 						FEdModeGeometry::TGeomObjectIterator GeomModeIt = GeomMode->GeomObjectItor();
 						for( ; GeomModeIt; ++GeomModeIt )
 						{
diff --git a/Engine/Source/Editor/UnrealEd/Private/EditorComponents.cpp b/Engine/Source/Editor/UnrealEd/Private/EditorComponents.cpp
--- a/Engine/Source/Editor/UnrealEd/Private/EditorComponents.cpp
+++ b/Engine/Source/Editor/UnrealEd/Private/EditorComponents.cpp
@@ -603,7 +603,7 @@ void FEditorCommonDrawHelper::DrawPivot(const FSceneView* View,FPrimitiveDrawInt
 {
 	const FMatrix CameraToWorld = View->ViewMatrices.ViewMatrix.Inverse();
 
-	const FVector PivLoc = GEditorModeTools().SnappedLocation;
+	const FVector PivLoc = GLevelEditorModeTools().SnappedLocation;
 
 	const float ZoomFactor = FMath::Min<float>(View->ViewMatrices.ProjMatrix.M[0][0], View->ViewMatrices.ProjMatrix.M[1][1]);
 	const float WidgetRadius = View->ViewMatrices.ProjMatrix.TransformPosition(PivLoc).W * (PivotSize / ZoomFactor);
diff --git a/Engine/Source/Editor/UnrealEd/Private/EditorCsg.cpp b/Engine/Source/Editor/UnrealEd/Private/EditorCsg.cpp
--- a/Engine/Source/Editor/UnrealEd/Private/EditorCsg.cpp
+++ b/Engine/Source/Editor/UnrealEd/Private/EditorCsg.cpp
@@ -1276,7 +1276,7 @@ void UEditorEngine::MapBrushGet(UWorld* InWorld)
 void UEditorEngine::mapBrushPut()
 {
 	TArray<FEdMode*> ActiveModes; 
-	GEditorModeTools().GetActiveModes( ActiveModes );
+	GLevelEditorModeTools().GetActiveModes( ActiveModes );
 
 	for ( FSelectionIterator It( GEditor->GetSelectedActorIterator() ) ; It ; ++It )
 	{
diff --git a/Engine/Source/Editor/UnrealEd/Private/EditorLevelUtils.cpp b/Engine/Source/Editor/UnrealEd/Private/EditorLevelUtils.cpp
--- a/Engine/Source/Editor/UnrealEd/Private/EditorLevelUtils.cpp
+++ b/Engine/Source/Editor/UnrealEd/Private/EditorLevelUtils.cpp
@@ -321,13 +321,13 @@ namespace EditorLevelUtils
 	bool RemoveLevelFromWorld(ULevel* InLevel)
 	{
 		// If we're removing a level, lets make sure to close the level transform mode if its the same level currently selected for edit
-		FEdModeLevel* LevelMode = static_cast<FEdModeLevel*>(GEditorModeTools().GetActiveMode( FBuiltinEditorModes::EM_Level ));		
+		FEdModeLevel* LevelMode = static_cast<FEdModeLevel*>(GLevelEditorModeTools().GetActiveMode( FBuiltinEditorModes::EM_Level ));		
 		if( LevelMode )
 		{
 			ULevelStreaming* LevelStream = FLevelUtils::FindStreamingLevel( InLevel ); 
 			if( LevelMode->IsEditing( LevelStream ) )
 			{
-				GEditorModeTools().DeactivateMode( FBuiltinEditorModes::EM_Level );
+				GLevelEditorModeTools().DeactivateMode( FBuiltinEditorModes::EM_Level );
 			}
 		}
 
@@ -492,7 +492,7 @@ namespace EditorLevelUtils
 	ULevel* CreateNewLevel( UWorld* InWorld, bool bMoveSelectedActorsIntoNewLevel, UClass* LevelStreamingClass, const FString& DefaultFilename )
 	{
 		// Editor modes cannot be active when any level saving occurs.
-		GEditorModeTools().ActivateMode( FBuiltinEditorModes::EM_Default );
+		GLevelEditorModeTools().ActivateDefaultMode();
 
 		UWorld* NewWorld = nullptr;
 		if (FParse::Param(FCommandLine::Get(), TEXT("WorldAssets")))
@@ -641,10 +641,10 @@ namespace EditorLevelUtils
 			}
 
 			// If were hiding a level, lets make sure to close the level transform mode if its the same level currently selected for edit
-			FEdModeLevel* LevelMode = static_cast<FEdModeLevel*>(GEditorModeTools().GetActiveMode( FBuiltinEditorModes::EM_Level ));
+			FEdModeLevel* LevelMode = static_cast<FEdModeLevel*>(GLevelEditorModeTools().GetActiveMode( FBuiltinEditorModes::EM_Level ));
 			if( LevelMode && LevelMode->IsEditing( StreamingLevel ) )
 			{
-				GEditorModeTools().DeactivateMode( FBuiltinEditorModes::EM_Level );
+				GLevelEditorModeTools().DeactivateMode( FBuiltinEditorModes::EM_Level );
 			}
 
 			//create a transaction so we can undo the visibilty toggle
diff --git a/Engine/Source/Editor/UnrealEd/Private/EditorModeInterpolation.cpp b/Engine/Source/Editor/UnrealEd/Private/EditorModeInterpolation.cpp
--- a/Engine/Source/Editor/UnrealEd/Private/EditorModeInterpolation.cpp
+++ b/Engine/Source/Editor/UnrealEd/Private/EditorModeInterpolation.cpp
@@ -19,9 +19,6 @@ static const float	CurveHandleScale = 0.5f;
 
 FEdModeInterpEdit::FEdModeInterpEdit()
 {
-	ID = FBuiltinEditorModes::EM_InterpEdit;
-	Name = NSLOCTEXT("EditorModes", "MatineeMode", "Matinee Mode");
-
 	MatineeActor = NULL;
 	InterpEd = NULL;
 
@@ -35,7 +32,7 @@ FEdModeInterpEdit::~FEdModeInterpEdit()
 {
 }
 
-bool FEdModeInterpEdit::InputKey( FLevelEditorViewportClient* ViewportClient, FViewport* Viewport, FKey Key, EInputEvent Event )
+bool FEdModeInterpEdit::InputKey( FEditorViewportClient* ViewportClient, FViewport* Viewport, FKey Key, EInputEvent Event )
 {
 	// Enter key drops new key frames
 	if( Key == EKeys::Enter &&
@@ -98,7 +95,7 @@ void FEdModeInterpEdit::ActorMoveNotify()
 }
 
 
-void FEdModeInterpEdit::CamMoveNotify(FLevelEditorViewportClient* ViewportClient)
+void FEdModeInterpEdit::CamMoveNotify(FEditorViewportClient* ViewportClient)
 {
 	if(!InterpEd)
 		return;
@@ -141,7 +138,7 @@ void FEdModeInterpEdit::Render(const FSceneView* View, FViewport* Viewport, FPri
 {
 	FEdMode::Render(View, Viewport, PDI);
 
-	check( GEditorModeTools().IsModeActive(FBuiltinEditorModes::EM_InterpEdit) );
+	check( GLevelEditorModeTools().IsModeActive(FBuiltinEditorModes::EM_InterpEdit) );
 
 	if(InterpEd  && !InterpEd->Hide3DTrackView() && View->Family->EngineShowFlags.Splines)
 	{
@@ -149,7 +146,7 @@ void FEdModeInterpEdit::Render(const FSceneView* View, FViewport* Viewport, FPri
 	}
 }
 
-void FEdModeInterpEdit::DrawHUD(FLevelEditorViewportClient* ViewportClient,FViewport* Viewport,const FSceneView* View,FCanvas* Canvas)
+void FEdModeInterpEdit::DrawHUD(FEditorViewportClient* ViewportClient,FViewport* Viewport,const FSceneView* View,FCanvas* Canvas)
 {
 	FEdMode::DrawHUD(ViewportClient,Viewport,View,Canvas);
 
@@ -183,6 +180,14 @@ void FEdModeInterpEdit::ActorSelectionChangeNotify()
 	}
 }
 
+bool FEdModeInterpEdit::IsCompatibleWith(FEditorModeID OtherModeID) const
+{
+	return
+		OtherModeID == FBuiltinEditorModes::EM_MeshPaint	||
+		OtherModeID == FBuiltinEditorModes::EM_Geometry		||
+		OtherModeID == FBuiltinEditorModes::EM_Bsp;
+}
+
 //////////////////////////////////////////////////////////////////////////
 // FModeTool_InterpEdit
 //////////////////////////////////////////////////////////////////////////
@@ -202,16 +207,16 @@ FModeTool_InterpEdit::~FModeTool_InterpEdit()
 {
 }
 
-bool FModeTool_InterpEdit::MouseMove(FLevelEditorViewportClient* ViewportClient,FViewport* Viewport,int32 x, int32 y)
+bool FModeTool_InterpEdit::MouseMove(FEditorViewportClient* ViewportClient,FViewport* Viewport,int32 x, int32 y)
 {
-	check( GEditorModeTools().IsModeActive(FBuiltinEditorModes::EM_InterpEdit) );
+	check( GLevelEditorModeTools().IsModeActive(FBuiltinEditorModes::EM_InterpEdit) );
 
-	FEdModeInterpEdit* mode = (FEdModeInterpEdit*)GEditorModeTools().GetActiveMode(FBuiltinEditorModes::EM_InterpEdit);
+	FEdModeInterpEdit* mode = (FEdModeInterpEdit*)GLevelEditorModeTools().GetActiveMode(FBuiltinEditorModes::EM_InterpEdit);
 
 	return 0;
 }
 
-bool FModeTool_InterpEdit::InputAxis(FLevelEditorViewportClient* InViewportClient, FViewport* Viewport, int32 ControllerId, FKey Key, float Delta, float DeltaTime)
+bool FModeTool_InterpEdit::InputAxis(FEditorViewportClient* InViewportClient, FViewport* Viewport, int32 ControllerId, FKey Key, float Delta, float DeltaTime)
 {
 	if ( InViewportClient->GetCurrentWidgetAxis() == EAxisList::None )
 	{
@@ -234,11 +239,11 @@ bool FModeTool_InterpEdit::InputAxis(FLevelEditorViewportClient* InViewportClien
 	return false;
 }
 
-bool FModeTool_InterpEdit::InputKey(FLevelEditorViewportClient* ViewportClient, FViewport* Viewport, FKey Key, EInputEvent Event)
+bool FModeTool_InterpEdit::InputKey(FEditorViewportClient* ViewportClient, FViewport* Viewport, FKey Key, EInputEvent Event)
 {
-	check( GEditorModeTools().IsModeActive(FBuiltinEditorModes::EM_InterpEdit) );
+	check( GLevelEditorModeTools().IsModeActive(FBuiltinEditorModes::EM_InterpEdit) );
 
-	FEdModeInterpEdit* mode = (FEdModeInterpEdit*)GEditorModeTools().GetActiveMode(FBuiltinEditorModes::EM_InterpEdit);
+	FEdModeInterpEdit* mode = (FEdModeInterpEdit*)GLevelEditorModeTools().GetActiveMode(FBuiltinEditorModes::EM_InterpEdit);
 	if( !mode->InterpEd )
 	{
 		// Abort cleanly on InerpEd not yet being assigned.  This can occasionally be the case when receiving
@@ -339,11 +344,11 @@ bool FModeTool_InterpEdit::InputKey(FLevelEditorViewportClient* ViewportClient,
 	return FModeTool::InputKey(ViewportClient, Viewport, Key, Event);
 }
 
-bool FModeTool_InterpEdit::InputDelta(FLevelEditorViewportClient* InViewportClient, FViewport* InViewport, FVector& InDrag, FRotator& InRot, FVector& InScale)
+bool FModeTool_InterpEdit::InputDelta(FEditorViewportClient* InViewportClient, FViewport* InViewport, FVector& InDrag, FRotator& InRot, FVector& InScale)
 {
-	check( GEditorModeTools().IsModeActive(FBuiltinEditorModes::EM_InterpEdit) );
+	check( GLevelEditorModeTools().IsModeActive(FBuiltinEditorModes::EM_InterpEdit) );
 
-	FEdModeInterpEdit* mode = (FEdModeInterpEdit*)GEditorModeTools().GetActiveMode(FBuiltinEditorModes::EM_InterpEdit);
+	FEdModeInterpEdit* mode = (FEdModeInterpEdit*)GLevelEditorModeTools().GetActiveMode(FBuiltinEditorModes::EM_InterpEdit);
 	check(mode->InterpEd);
 
 	bool bShiftDown = InViewport->KeyState(EKeys::LeftShift) || InViewport->KeyState(EKeys::RightShift);
@@ -372,7 +377,7 @@ bool FModeTool_InterpEdit::InputDelta(FLevelEditorViewportClient* InViewportClie
 
 void FModeTool_InterpEdit::SelectNone()
 {
-	check( GEditorModeTools().IsModeActive(FBuiltinEditorModes::EM_InterpEdit) );
+	check( GLevelEditorModeTools().IsModeActive(FBuiltinEditorModes::EM_InterpEdit) );
 
-	FEdModeInterpEdit* mode = (FEdModeInterpEdit*)GEditorModeTools().GetActiveMode(FBuiltinEditorModes::EM_InterpEdit);
+	FEdModeInterpEdit* mode = (FEdModeInterpEdit*)GLevelEditorModeTools().GetActiveMode(FBuiltinEditorModes::EM_InterpEdit);
 }
diff --git a/Engine/Source/Editor/UnrealEd/Private/EditorModeTools.cpp b/Engine/Source/Editor/UnrealEd/Private/EditorModeTools.cpp
--- a/Engine/Source/Editor/UnrealEd/Private/EditorModeTools.cpp
+++ b/Engine/Source/Editor/UnrealEd/Private/EditorModeTools.cpp
@@ -15,7 +15,7 @@ FModeTool::~FModeTool()
 {
 }
 
-void FModeTool::DrawHUD(FLevelEditorViewportClient* ViewportClient,FViewport* Viewport,const FSceneView* View,FCanvas* Canvas)
+void FModeTool::DrawHUD(FEditorViewportClient* ViewportClient,FViewport* Viewport,const FSceneView* View,FCanvas* Canvas)
 {
 }
 
diff --git a/Engine/Source/Editor/UnrealEd/Private/EditorModes.cpp b/Engine/Source/Editor/UnrealEd/Private/EditorModes.cpp
--- a/Engine/Source/Editor/UnrealEd/Private/EditorModes.cpp
+++ b/Engine/Source/Editor/UnrealEd/Private/EditorModes.cpp
@@ -8,19 +8,11 @@
 #include "ScopedTransaction.h"
 #include "SurfaceIterators.h"
 #include "SoundDefinitions.h"
-#include "Editor/PlacementMode/Public/IPlacementModeModule.h"
-#include "Editor/LandscapeEditor/Public/LandscapeEditorModule.h"
-#include "Editor/BspMode/Public/IBspModeModule.h"
-#include "Editor/MeshPaint/Public/MeshPaintModule.h"
-#include "Editor/GeometryMode/Public/GeometryEdMode.h"
-#include "Editor/ActorPickerMode/Public/ActorPickerMode.h"
-#include "Editor/TextureAlignMode/Public/TextureAlignEdMode.h"
 #include "LevelEditor.h"
 #include "Toolkits/ToolkitManager.h"
 #include "EditorLevelUtils.h"
 #include "DynamicMeshBuilder.h"
 
-#include "Editor/FoliageEdit/Public/FoliageEditModule.h"
 #include "ActorEditorUtils.h"
 #include "EditorStyle.h"
 #include "ComponentVisualizer.h"
@@ -193,14 +185,11 @@ namespace
 // FEdMode
 
 FEdMode::FEdMode()
-	: Name( NSLOCTEXT("UnrealEd", "N/A", "N/A") )
-	, CurrentTool( NULL )
-	, ID( FBuiltinEditorModes::EM_None )
-	, bVisible(false)
-	, PriorityOrder(MAX_int32)
+	: CurrentTool( NULL )
 	, EditedPropertyName(TEXT(""))
 	, EditedPropertyIndex( INDEX_NONE )
 	, bEditedPropertyIsTransform( false )
+	, bPendingDeletion( false )
 {
 	bDrawKillZ = true;
 }
@@ -209,7 +198,16 @@ FEdMode::~FEdMode()
 {
 }
 
-bool FEdMode::MouseEnter( FLevelEditorViewportClient* ViewportClient,FViewport* Viewport,int32 x, int32 y )
+void FEdMode::OnModeUnregistered( FEditorModeID ModeID )
+{
+	if( ModeID == Info.ID )
+	{
+		// This should be synonymous with "delete this"
+		Owner->DestroyMode(ModeID);
+	}
+}
+
+bool FEdMode::MouseEnter( FEditorViewportClient* ViewportClient,FViewport* Viewport,int32 x, int32 y )
 {
 	if( GetCurrentTool() )
 	{
@@ -219,7 +217,7 @@ bool FEdMode::MouseEnter( FLevelEditorViewportClient* ViewportClient,FViewport*
 	return false;
 }
 
-bool FEdMode::MouseLeave( FLevelEditorViewportClient* ViewportClient,FViewport* Viewport )
+bool FEdMode::MouseLeave( FEditorViewportClient* ViewportClient,FViewport* Viewport )
 {
 	if( GetCurrentTool() )
 	{
@@ -229,7 +227,7 @@ bool FEdMode::MouseLeave( FLevelEditorViewportClient* ViewportClient,FViewport*
 	return false;
 }
 
-bool FEdMode::MouseMove(FLevelEditorViewportClient* ViewportClient,FViewport* Viewport,int32 x, int32 y)
+bool FEdMode::MouseMove(FEditorViewportClient* ViewportClient,FViewport* Viewport,int32 x, int32 y)
 {
 	if( GetCurrentTool() )
 	{
@@ -239,7 +237,7 @@ bool FEdMode::MouseMove(FLevelEditorViewportClient* ViewportClient,FViewport* Vi
 	return false;
 }
 
-bool FEdMode::ReceivedFocus(FLevelEditorViewportClient* ViewportClient,FViewport* Viewport)
+bool FEdMode::ReceivedFocus(FEditorViewportClient* ViewportClient,FViewport* Viewport)
 {
 	if( GetCurrentTool() )
 	{
@@ -249,7 +247,7 @@ bool FEdMode::ReceivedFocus(FLevelEditorViewportClient* ViewportClient,FViewport
 	return false;
 }
 
-bool FEdMode::LostFocus(FLevelEditorViewportClient* ViewportClient,FViewport* Viewport)
+bool FEdMode::LostFocus(FEditorViewportClient* ViewportClient,FViewport* Viewport)
 {
 	if( GetCurrentTool() )
 	{
@@ -269,7 +267,7 @@ bool FEdMode::LostFocus(FLevelEditorViewportClient* ViewportClient,FViewport* Vi
  *
  * @return	true if input was handled
  */
-bool FEdMode::CapturedMouseMove( FLevelEditorViewportClient* InViewportClient, FViewport* InViewport, int32 InMouseX, int32 InMouseY )
+bool FEdMode::CapturedMouseMove( FEditorViewportClient* InViewportClient, FViewport* InViewport, int32 InMouseX, int32 InMouseY )
 {
 	if( GetCurrentTool() )
 	{
@@ -280,13 +278,8 @@ bool FEdMode::CapturedMouseMove( FLevelEditorViewportClient* InViewportClient, F
 }
 
 
-bool FEdMode::InputKey(FLevelEditorViewportClient* ViewportClient, FViewport* Viewport, FKey Key, EInputEvent Event)
+bool FEdMode::InputKey(FEditorViewportClient* ViewportClient, FViewport* Viewport, FKey Key, EInputEvent Event)
 {
-	if (GUnrealEd->ComponentVisManager.HandleInputKey(ViewportClient, Viewport, Key, Event))
-	{
-		return true;
-	}
-
 	if( GetCurrentTool() && GetCurrentTool()->InputKey( ViewportClient, Viewport, Key, Event ) )
 	{
 		return true;
@@ -307,7 +300,7 @@ bool FEdMode::InputKey(FLevelEditorViewportClient* ViewportClient, FViewport* Vi
 	return 0;
 }
 
-bool FEdMode::InputAxis(FLevelEditorViewportClient* InViewportClient, FViewport* Viewport, int32 ControllerId, FKey Key, float Delta, float DeltaTime)
+bool FEdMode::InputAxis(FEditorViewportClient* InViewportClient, FViewport* Viewport, int32 ControllerId, FKey Key, float Delta, float DeltaTime)
 {
 	FModeTool* Tool = GetCurrentTool();
 	if (Tool)
@@ -318,13 +311,8 @@ bool FEdMode::InputAxis(FLevelEditorViewportClient* InViewportClient, FViewport*
 	return false;
 }
 
-bool FEdMode::InputDelta(FLevelEditorViewportClient* InViewportClient, FViewport* InViewport, FVector& InDrag, FRotator& InRot, FVector& InScale)
-{
-	if(GUnrealEd->ComponentVisManager.HandleInputDelta(InViewportClient, InViewport, InDrag, InRot, InScale))
-	{
-		return true;
-	}
-	
+bool FEdMode::InputDelta(FEditorViewportClient* InViewportClient, FViewport* InViewport, FVector& InDrag, FRotator& InRot, FVector& InScale)
+{	
 	if(UsesPropertyWidgets())
 	{
 		AActor* SelectedActor = GetFirstSelectedActorInstance();
@@ -428,13 +416,6 @@ bool FEdMode::UsesTransformWidget(FWidget::EWidgetMode CheckMode) const
 
 FVector FEdMode::GetWidgetLocation() const
 {
-	FVector ComponentVisWidgetLocation;
-	bool bUseComponentVisWidgetLocation = GUnrealEd->ComponentVisManager.GetWidgetLocation(ComponentVisWidgetLocation);
-	if(bUseComponentVisWidgetLocation)
-	{
-		return ComponentVisWidgetLocation;
-	}
-
 	if(UsesPropertyWidgets())
 	{
 		AActor* SelectedActor = GetFirstSelectedActorInstance();
@@ -462,7 +443,7 @@ FVector FEdMode::GetWidgetLocation() const
 	}
 
 	//UE_LOG(LogEditorModes, Log, TEXT("In FEdMode::GetWidgetLocation"));
-	return GEditorModeTools().PivotLocation;
+	return Owner->PivotLocation;
 }
 
 /**
@@ -527,7 +508,7 @@ void FEdMode::SelectNone()
 	}
 }
 
-void FEdMode::Tick(FLevelEditorViewportClient* ViewportClient,float DeltaTime)
+void FEdMode::Tick(FEditorViewportClient* ViewportClient,float DeltaTime)
 {
 	if( GetCurrentTool() )
 	{
@@ -540,12 +521,9 @@ void FEdMode::ActorSelectionChangeNotify()
 	EditedPropertyName = TEXT("");
 	EditedPropertyIndex = INDEX_NONE;
 	bEditedPropertyIsTransform = false;
-
-	// Clear active editing visualizer on selection change
-	GUnrealEd->ComponentVisManager.ClearActiveComponentVis();
 }
 
-bool FEdMode::HandleClick(FLevelEditorViewportClient* InViewportClient, HHitProxy *HitProxy, const FViewportClick &Click)
+bool FEdMode::HandleClick(FEditorViewportClient* InViewportClient, HHitProxy *HitProxy, const FViewportClick &Click)
 {
 	if(UsesPropertyWidgets() && HitProxy)
 	{
@@ -566,27 +544,6 @@ bool FEdMode::HandleClick(FLevelEditorViewportClient* InViewportClient, HHitProx
 		}
 	}
 
-	if(HitProxy)
-	{	
-		bool bHandled = GUnrealEd->ComponentVisManager.HandleProxyForComponentVis(HitProxy);
-		if (bHandled && Click.GetKey() == EKeys::RightMouseButton)
-		{
-			TSharedPtr<SWidget> MenuWidget = GUnrealEd->ComponentVisManager.GenerateContextMenuForComponentVis();
-			if (MenuWidget.IsValid())
-			{
-				auto ParentLevelEditorPinned = InViewportClient->ParentLevelEditor.Pin();
-				if (ParentLevelEditorPinned.IsValid())
-				{
-					FSlateApplication::Get().PushMenu(
-						ParentLevelEditorPinned.ToSharedRef(),
-						MenuWidget.ToSharedRef(),
-						FSlateApplication::Get().GetCursorPos(),
-						FPopupTransitionEffect(FPopupTransitionEffect::ContextMenu));
-				}
-			}
-		}
-	}
-
 	return false;
 }
 
@@ -603,13 +560,13 @@ void FEdMode::Enter()
 
 	FEditorDelegates::EditorModeEnter.Broadcast( this );
 	const bool bIsEnteringMode = true;
-	GEditorModeTools().BroadcastEditorModeChanged( this, bIsEnteringMode );
+	Owner->BroadcastEditorModeChanged( this, bIsEnteringMode );
 }
 
 void FEdMode::Exit()
 {
 	const bool bIsEnteringMode = false;
-	GEditorModeTools().BroadcastEditorModeChanged( this, bIsEnteringMode );
+	Owner->BroadcastEditorModeChanged( this, bIsEnteringMode );
 	FEditorDelegates::EditorModeExit.Broadcast( this );
 }
 
@@ -696,7 +653,7 @@ void FEdMode::Render(const FSceneView* View,FViewport* Viewport,FPrimitiveDrawIn
 	}
 
 	const bool bIsInGameView = !Viewport->GetClient() || Viewport->GetClient()->IsInGameView();
-	if (GEditorModeTools().ShouldDrawBrushVertices() && !bIsInGameView)
+	if (Owner->ShouldDrawBrushVertices() && !bIsInGameView)
 	{
 		UTexture2D* VertexTexture = GetVertexTexture();
 		const float TextureSizeX = VertexTexture->GetSizeX() * 0.170f;
@@ -803,7 +760,7 @@ void FEdMode::Render(const FSceneView* View,FViewport* Viewport,FPrimitiveDrawIn
 	}
 }
 
-void FEdMode::DrawHUD(FLevelEditorViewportClient* ViewportClient,FViewport* Viewport,const FSceneView* View,FCanvas* Canvas)
+void FEdMode::DrawHUD(FEditorViewportClient* ViewportClient,FViewport* Viewport,const FSceneView* View,FCanvas* Canvas)
 {
 	// Render the drag tool.
 	ViewportClient->RenderDragTool( View, Canvas );
@@ -969,7 +926,7 @@ void FEdMode::DrawHUD(FLevelEditorViewportClient* ViewportClient,FViewport* View
 }
 
 // Draw brackets around all selected objects
-void FEdMode::DrawBrackets( FLevelEditorViewportClient* ViewportClient, FViewport* Viewport, const FSceneView* View, FCanvas* Canvas )
+void FEdMode::DrawBrackets( FEditorViewportClient* ViewportClient, FViewport* Viewport, const FSceneView* View, FCanvas* Canvas )
 {
 	USelection& SelectedActors = *GEditor->GetSelectedActors();
 	for( int32 CurSelectedActorIndex = 0; CurSelectedActorIndex < SelectedActors.Num(); ++CurSelectedActorIndex )
@@ -992,7 +949,7 @@ bool FEdMode::UsesToolkits() const
 	return false;
 }
 
-bool FEdMode::StartTracking(FLevelEditorViewportClient* InViewportClient, FViewport* InViewport)
+bool FEdMode::StartTracking(FEditorViewportClient* InViewportClient, FViewport* InViewport)
 {
 	bool bResult = false;
 	if( GetCurrentTool() )
@@ -1002,7 +959,7 @@ bool FEdMode::StartTracking(FLevelEditorViewportClient* InViewportClient, FViewp
 	return bResult;
 }
 
-bool FEdMode::EndTracking(FLevelEditorViewportClient* InViewportClient, FViewport* InViewport)
+bool FEdMode::EndTracking(FEditorViewportClient* InViewportClient, FViewport* InViewport)
 {
 	bool bResult = false;
 	if( GetCurrentTool() )
@@ -1017,7 +974,7 @@ FVector FEdMode::GetWidgetNormalFromCurrentAxis( void* InData )
 	// Figure out the proper coordinate system.
 
 	FMatrix matrix = FMatrix::Identity;
-	if( GEditorModeTools().GetCoordSystem() == COORD_Local )
+	if( Owner->GetCoordSystem() == COORD_Local )
 	{
 		GetCustomDrawingCoordinateSystem( matrix, InData );
 	}
@@ -1186,8 +1143,6 @@ bool FEdMode::IsSnapRotationEnabled()
 
 FEdModeDefault::FEdModeDefault()
 {
-	ID = FBuiltinEditorModes::EM_Default;
-	Name = NSLOCTEXT("EditorModes", "DefaultMode", "Default");
 	bDrawKillZ = true;
 }
 
@@ -1208,110 +1163,23 @@ FEditorModeTools::FEditorModeTools()
 	,	bHideViewportUI(false)
 	,	CoordSystem(COORD_World)
 	,	bIsTracking(false)
-{}
-
-FEditorModeTools::~FEditorModeTools()
-{
-
-}
-
-void FEditorModeTools::RegisterMode(TSharedRef<FEdMode> ModeToRegister)
-{
-	check(ModeToRegister->ID != FBuiltinEditorModes::EM_None);
-	check(Modes.Find(ModeToRegister) == INDEX_NONE);
-
-	Modes.Add(ModeToRegister);
-
-	BroadcastRegisteredModesChanged();
-}
-
-void FEditorModeTools::UnregisterMode(TSharedRef<FEdMode> ModeToUnregister)
-{
-	// See if we need to deactivate the mode that is being unregistered.
-	DeactivateMode(ModeToUnregister->ID);
-
-	for (int32 ModeIndex = 0; ModeIndex < Modes.Num(); )
-	{
-		if (Modes[ModeIndex] == ModeToUnregister)
-		{
-			Modes.RemoveAtSwap(ModeIndex);
-		}
-		else
-		{
-			++ModeIndex;
-		}
-	}
-
-	BroadcastRegisteredModesChanged();
-}
-
-void FEditorModeTools::RegisterCompatibleModes( const FEditorModeID& Mode1, const FEditorModeID& Mode2 )
-{
-	ModeCompatabilityMap.Add( Mode1, Mode2 );
-	ModeCompatabilityMap.Add( Mode2, Mode1 );
-}
-
-void FEditorModeTools::UnregisterCompatibleModes( const FEditorModeID& Mode1, const FEditorModeID& Mode2 )
-{
-	ModeCompatabilityMap.Remove( Mode1, Mode2 );
-	ModeCompatabilityMap.Remove( Mode2, Mode1 );
-}
-
-// A simple macro to help set editor mode compatibility mappings
-#define SET_MODE_COMPATABILITY( InMode1, InMode2  ) \
-	ModeCompatabilityMap.Add( InMode1, InMode2 ); \
-	ModeCompatabilityMap.Add( InMode2, InMode1 ); \
-
-void FEditorModeTools::Init()
+	,	DefaultID(FBuiltinEditorModes::EM_Default)
 {
 	// Load the last used settings
 	LoadConfig();
 
-	// Editor modes
-	Modes.Add( MakeShareable(new FEdModeDefault()) );
-	Modes.Add( MakeShareable(new FEdModeInterpEdit()) );
-
-	// Load editor mode modules that will automatically register their editor modes, and clean themselves up on unload.
-	//@TODO: ROCKET: These are probably good plugin candidates, that shouldn't have to be force-loaded here but discovery loaded somehow
-	FModuleManager::LoadModuleChecked<IPlacementModeModule>(TEXT("PlacementMode"));
-	FModuleManager::LoadModuleChecked<IBspModeModule>(TEXT("BspMode"));
-	FModuleManager::LoadModuleChecked<FTextureAlignModeModule>(TEXT("TextureAlignMode"));
-	FModuleManager::LoadModuleChecked<FGeometryModeModule>(TEXT("GeometryMode"));
-	FModuleManager::LoadModuleChecked<FActorPickerModeModule>(TEXT("ActorPickerMode"));
-	FModuleManager::LoadModuleChecked<IMeshPaintModule>(TEXT("MeshPaint"));
-	FModuleManager::LoadModuleChecked<ILandscapeEditorModule>(TEXT("LandscapeEditor"));
-	FModuleManager::LoadModuleChecked<IFoliageEditModule>(TEXT("FoliageEdit"));
-
-	// Set which modes can run together
-	
-	// Interp edit and mesh paint can run together
-	SET_MODE_COMPATABILITY( FBuiltinEditorModes::EM_InterpEdit, FBuiltinEditorModes::EM_MeshPaint );
-	// Interp edit and geometry mode can run together
-	SET_MODE_COMPATABILITY( FBuiltinEditorModes::EM_InterpEdit, FBuiltinEditorModes::EM_Geometry );
-	// Interp and BSP mode can run together
-	SET_MODE_COMPATABILITY( FBuiltinEditorModes::EM_InterpEdit, TEXT("BSP") );
-
 	// Register our callback for actor selection changes
 	USelection::SelectNoneEvent.AddRaw(this, &FEditorModeTools::OnEditorSelectNone);
 	USelection::SelectionChangedEvent.AddRaw(this, &FEditorModeTools::OnEditorSelectionChanged);
 	USelection::SelectObjectEvent.AddRaw(this, &FEditorModeTools::OnEditorSelectionChanged);
 }
 
-/**
- * Called from UUnrealEdEngine::PreExit
- */
-void FEditorModeTools::Shutdown()
+FEditorModeTools::~FEditorModeTools()
 {
+	// Should we call Exit on any modes that are still active, or is it too late?
 	USelection::SelectionChangedEvent.RemoveAll(this);
 	USelection::SelectNoneEvent.RemoveAll(this);
 	USelection::SelectObjectEvent.RemoveAll(this);
-
-	//@TODO: Should we exit any active modes here, or is it too late?
-
-	// Clean up modes
-	ActiveModes.Empty();
-	ModeCompatabilityMap.Empty();
-	Modes.Empty();
 }
 
 /**
@@ -1378,17 +1246,9 @@ void FEditorModeTools::OnEditorSelectionChanged(UObject* NewSelection)
 		}
 	}
 
-	// Don't do auto-switch-to-cover-editing when in Matinee, as it will exit Matinee in a bad way.
-	//@TODO: Is this still a valid limitation, especially given that cover editing does not exist?
-	if( !IsModeActive(FBuiltinEditorModes::EM_InterpEdit) )
+	for (const auto& Pair : FEditorModeRegistry::Get().GetFactoryMap())
 	{
-		for (int32 ModeIndex = 0; ModeIndex < Modes.Num(); ++ModeIndex)
-		{
-			if (FEdMode* Mode = Modes[ModeIndex].Get())
-			{
-				Mode->PeekAtSelectionChangedEvent(NewSelection);
-			}
-		}
+		Pair.Value->OnSelectionChanged(*this, NewSelection);
 	}
 }
 
@@ -1398,31 +1258,6 @@ void FEditorModeTools::OnEditorSelectNone()
 	GEditor->ActorsThatWereSelected.Empty();
 }
 
-/** 
- * Returns a list of active modes that are incompatible with the passed in mode.
- * 
- * @param InID 				The mode to check for incompatibilites.
- * @param IncompatibleModes	The list of incompatible modes.
- */
-void FEditorModeTools::GetIncompatibleActiveModes( FEditorModeID InID, TArray<FEdMode*>& IncompatibleModes )
-{
-	// Find all modes which are compatible with the one passed in.
-	TArray<FEditorModeID> CompatibleModes;
-	ModeCompatabilityMap.MultiFind( InID, CompatibleModes );
-
-	// Search through all active modes for incompatible ones
-	for( int32 ModeIndex = 0; ModeIndex < ActiveModes.Num(); ++ModeIndex )
-	{
-		FEdMode* ActiveMode = ActiveModes[ModeIndex];
-
-		if( !CompatibleModes.Contains( ActiveMode->GetID() ) )
-		{
-			// This active mode is not compatible with the one that was passed in.
-			IncompatibleModes.Add ( ActiveMode );
-		}
-	}
-}
-
 /** 
  * Sets the pivot locations
  * 
@@ -1455,6 +1290,29 @@ void FEditorModeTools::SetCoordSystem(ECoordSystem NewCoordSystem)
 	CoordSystem = NewCoordSystem;
 }
 
+void FEditorModeTools::SetDefaultMode ( FEditorModeID InDefaultID )
+{
+	DefaultID = InDefaultID;
+}
+
+void FEditorModeTools::ActivateDefaultMode()
+{
+	ActivateMode( DefaultID );
+
+	check( IsModeActive( DefaultID ) );
+}
+
+void FEditorModeTools::DeactivateModeAtIndex(int32 InIndex)
+{
+	check( InIndex >= 0 && InIndex < Modes.Num() );
+
+	auto& Mode = Modes[InIndex];
+
+	Mode->Exit();
+	RecycledModes.Add( Mode->GetID(), Mode );
+	Modes.RemoveAt( InIndex );
+}
+
 /**
  * Deactivates an editor mode. 
  * 
@@ -1462,40 +1320,47 @@ void FEditorModeTools::SetCoordSystem(ECoordSystem NewCoordSystem)
  */
 void FEditorModeTools::DeactivateMode( FEditorModeID InID )
 {
-	// If the mode is already inactive do nothing
-	if( !IsModeActive(InID) )
+	// Find the mode from the ID and exit it.
+	for( int32 Index = Modes.Num() - 1; Index >= 0; --Index )
 	{
-		return;
-	}
-
-	// Find the active mode from the ID and exit it.
-	for( int32 ModeIndex = 0; ModeIndex < ActiveModes.Num(); ++ModeIndex )
-	{
-		FEdMode* Mode = ActiveModes[ModeIndex];
+		auto& Mode = Modes[Index];
 		if( Mode->GetID() == InID )
 		{
-			Mode->Exit();
-			ActiveModes.Remove( Mode );
+			DeactivateModeAtIndex(Index);
 			break;
 		}
 	}
 
-	if( ActiveModes.Num() == 0 )
+	if( Modes.Num() == 0 )
 	{
 		// Ensure the default mode is active if there are no active modes.
-		ActivateMode( FBuiltinEditorModes::EM_Default );
+		ActivateDefaultMode();
 	}
 }
 
 void FEditorModeTools::DeactivateAllModes()
 {
-	TArray< FEdMode* > CurrentActiveModes;
-	GetActiveModes(CurrentActiveModes);
+	for( int32 Index = Modes.Num() - 1; Index >= 0; --Index )
+	{
+		DeactivateModeAtIndex(Index);
+	}
+}
 
-	for( FEdMode* Mode : CurrentActiveModes )
+void FEditorModeTools::DestroyMode( FEditorModeID InID )
+{
+	// Find the mode from the ID and exit it.
+	for( int32 Index = Modes.Num() - 1; Index >= 0; --Index )
 	{
-		DeactivateMode(Mode->GetID());
+		auto& Mode = Modes[Index];
+		if ( Mode->GetID() == InID )
+		{
+			// Deactivate and destroy
+			DeactivateModeAtIndex(Index);
+			break;
+		}
 	}
+
+	RecycledModes.Remove(InID);
 }
 
 /**
@@ -1506,6 +1371,11 @@ void FEditorModeTools::DeactivateAllModes()
  */
 void FEditorModeTools::ActivateMode( FEditorModeID InID, bool bToggle )
 {
+	if (InID == FBuiltinEditorModes::EM_Default)
+	{
+		InID = DefaultID;
+	}
+
 	// Check to see if the mode is already active
 	if( IsModeActive(InID) )
 	{
@@ -1518,42 +1388,42 @@ void FEditorModeTools::ActivateMode( FEditorModeID InID, bool bToggle )
 		return;
 	}
 
-	// Get a list of modes which are incompatible with the one we are trying to activate
-	TArray<FEdMode*> IncompatibleModes;
-	GetIncompatibleActiveModes( InID, IncompatibleModes );
+	// Recycle a mode or factory a new one
+	TSharedPtr<FEdMode> Mode = RecycledModes.FindRef( InID );
 
-	// Exit all incompatible modes
-	for( int32 ModeIndex = 0; ModeIndex < IncompatibleModes.Num(); ++ModeIndex )
+	if ( Mode.IsValid() )
 	{
-		FEdMode* Mode = IncompatibleModes[ModeIndex];
-		if ( ActiveModes.Contains( Mode ) )
-		{
-			Mode->Exit();
-			ActiveModes.Remove( Mode );
-		}
+		RecycledModes.Remove( InID );
+	}
+	else
+	{
+		Mode = FEditorModeRegistry::Get().CreateMode( InID, *this );
 	}
 
-	// Find the new mode from its ID
-	FEdMode* NewMode = FindMode( InID );
-
-	if( !NewMode )
+	if( !Mode.IsValid() )
 	{
-		// Couldn't find a valid mode. Activate the default mode
-		UE_LOG(LogEditorModes, Log,  TEXT("FEditorModeTools::SetCurrentMode : Couldn't find mode '%s'.  Using default."), *InID.ToString() );
-		// We must call this function again to remove any other active modes that aren't compatible with the default
-		ActivateMode( FBuiltinEditorModes::EM_Default );
+		UE_LOG(LogEditorModes, Log, TEXT("FEditorModeTools::ActivateMode : Couldn't find mode '%s'."), *InID.ToString() );
+		// Just return and leave the mode list unmodified
+		return;
 	}
-	else
+
+	// Remove anything that isn't compatible with this mode
+	for( int32 ModeIndex = Modes.Num() - 1; ModeIndex >= 0; --ModeIndex )
 	{
-		// Add the new mode to the list of active ones
-		ActiveModes.Add( NewMode );
+		const bool bModesAreCompatible = Mode->IsCompatibleWith( Modes[ModeIndex]->GetID() ) || Modes[ModeIndex]->IsCompatibleWith( Mode->GetID() );
+		if ( !bModesAreCompatible )
+		{
+			DeactivateModeAtIndex(ModeIndex);
+		}
+	}
+
+	Modes.Add( Mode );
 
-		// Enter the new mode
-		NewMode->Enter();
+	// Enter the new mode
+	Mode->Enter();
 	
-		// Update the editor UI
-		FEditorSupportDelegates::UpdateUI.Broadcast();
-	}
+	// Update the editor UI
+	FEditorSupportDelegates::UpdateUI.Broadcast();
 }
 
 bool FEditorModeTools::EnsureNotInMode(FEditorModeID ModeID, const FText& ErrorMsg, bool bNotifyUser) const
@@ -1578,11 +1448,11 @@ bool FEditorModeTools::EnsureNotInMode(FEditorModeID ModeID, const FText& ErrorM
 
 FEdMode* FEditorModeTools::FindMode( FEditorModeID InID )
 {
-	for( int32 x = 0 ; x < Modes.Num() ; x++ )
+	for( auto& Mode : Modes )
 	{
-		if( Modes[x]->ID == InID )
+		if( Mode->GetID() == InID )
 		{
-			return Modes[x].Get();
+			return Mode.Get();
 		}
 	}
 
@@ -1605,9 +1475,9 @@ FMatrix FEditorModeTools::GetCustomDrawingCoordinateSystem()
 			// If it doesn't want to, create it by looking at the currently selected actors list.
 
 			bool CustomCoordinateSystemProvided = false;
-			for (int Index = 0; Index < ActiveModes.Num(); Index++)
+			for ( const auto& Mode : Modes)
 			{
-				if( ActiveModes[Index]->GetCustomDrawingCoordinateSystem( matrix, NULL ) )
+				if( Mode->GetCustomDrawingCoordinateSystem( matrix, NULL ) )
 				{
 					CustomCoordinateSystemProvided = true;
 					break;
@@ -1651,11 +1521,11 @@ FMatrix FEditorModeTools::GetCustomInputCoordinateSystem()
 EAxisList::Type FEditorModeTools::GetWidgetAxisToDraw( FWidget::EWidgetMode InWidgetMode ) const
 {
 	EAxisList::Type OutAxis = EAxisList::All;
-	for (int Index = ActiveModes.Num() - 1; Index >= 0 ; Index--)
+	for( int Index = Modes.Num() - 1; Index >= 0 ; Index-- )
 	{
-		if ( ActiveModes[Index]->ShouldDrawWidget() )
+		if ( Modes[Index]->ShouldDrawWidget() )
 		{
-			OutAxis = ActiveModes[Index]->GetWidgetAxisToDraw( InWidgetMode );
+			OutAxis = Modes[Index]->GetWidgetAxisToDraw( InWidgetMode );
 			break;
 		}
 	}
@@ -1664,30 +1534,30 @@ EAxisList::Type FEditorModeTools::GetWidgetAxisToDraw( FWidget::EWidgetMode InWi
 }
 
 /** Mouse tracking interface.  Passes tracking messages to all active modes */
-bool FEditorModeTools::StartTracking(FLevelEditorViewportClient* InViewportClient, FViewport* InViewport)
+bool FEditorModeTools::StartTracking(FEditorViewportClient* InViewportClient, FViewport* InViewport)
 {
 	bIsTracking = true;
 	bool bTransactionHandled = false;
 
 	CachedLocation = PivotLocation;	// Cache the pivot location
 
-	for( int32 ModeIndex = 0; ModeIndex < ActiveModes.Num(); ++ModeIndex )
+	for ( const auto& Mode : Modes)
 	{
-		bTransactionHandled |= ActiveModes[ModeIndex]->StartTracking(InViewportClient, InViewport);
+		bTransactionHandled |= Mode->StartTracking(InViewportClient, InViewport);
 	}
 
 	return bTransactionHandled;
 }
 
 /** Mouse tracking interface.  Passes tracking messages to all active modes */
-bool FEditorModeTools::EndTracking(FLevelEditorViewportClient* InViewportClient, FViewport* InViewport)
+bool FEditorModeTools::EndTracking(FEditorViewportClient* InViewportClient, FViewport* InViewport)
 {
 	bIsTracking = false;
 	bool bTransactionHandled = false;
 
-	for( int32 ModeIndex = 0; ModeIndex < ActiveModes.Num(); ++ModeIndex )
+	for ( const auto& Mode : Modes)
 	{
-		bTransactionHandled |= ActiveModes[ModeIndex]->EndTracking(InViewportClient, InViewportClient->Viewport);
+		bTransactionHandled |= Mode->EndTracking(InViewportClient, InViewportClient->Viewport);
 	}
 
 	CachedLocation = PivotLocation;	// Clear the pivot location
@@ -1698,9 +1568,9 @@ bool FEditorModeTools::EndTracking(FLevelEditorViewportClient* InViewportClient,
 /** Notifies all active modes that a map change has occured */
 void FEditorModeTools::MapChangeNotify()
 {
-	for( int32 ModeIndex = 0; ModeIndex < ActiveModes.Num(); ++ModeIndex )
+	for ( const auto& Mode : Modes)
 	{
-		ActiveModes[ModeIndex]->MapChangeNotify();
+		Mode->MapChangeNotify();
 	}
 }
 
@@ -1708,9 +1578,9 @@ void FEditorModeTools::MapChangeNotify()
 /** Notifies all active modes to empty their selections */
 void FEditorModeTools::SelectNone()
 {
-	for( int32 ModeIndex = 0; ModeIndex < ActiveModes.Num(); ++ModeIndex )
+	for ( const auto& Mode : Modes)
 	{
-		ActiveModes[ModeIndex]->SelectNone();
+		Mode->SelectNone();
 	}
 }
 
@@ -1718,9 +1588,9 @@ void FEditorModeTools::SelectNone()
 bool FEditorModeTools::BoxSelect( FBox& InBox, bool InSelect )
 {
 	bool bHandled = false;
-	for( int32 ModeIndex = 0; ModeIndex < ActiveModes.Num(); ++ModeIndex )
+	for( const auto& Mode : Modes)
 	{
-		bHandled |= ActiveModes[ModeIndex]->BoxSelect( InBox, InSelect );
+		bHandled |= Mode->BoxSelect( InBox, InSelect );
 	}
 	return bHandled;
 }
@@ -1729,9 +1599,9 @@ bool FEditorModeTools::BoxSelect( FBox& InBox, bool InSelect )
 bool FEditorModeTools::FrustumSelect( const FConvexVolume& InFrustum, bool InSelect )
 {
 	bool bHandled = false;
-	for( int32 ModeIndex = 0; ModeIndex < ActiveModes.Num(); ++ModeIndex )
+	for( const auto& Mode : Modes)
 	{
-		bHandled |= ActiveModes[ModeIndex]->FrustumSelect( InFrustum, InSelect );
+		bHandled |= Mode->FrustumSelect( InFrustum, InSelect );
 	}
 	return bHandled;
 }
@@ -1741,9 +1611,9 @@ bool FEditorModeTools::FrustumSelect( const FConvexVolume& InFrustum, bool InSel
 bool FEditorModeTools::UsesTransformWidget() const
 {
 	bool bUsesTransformWidget = false;
-	for( int32 ModeIndex = 0; ModeIndex < ActiveModes.Num(); ++ModeIndex )
+	for( const auto& Mode : Modes)
 	{
-		bUsesTransformWidget |= ActiveModes[ModeIndex]->UsesTransformWidget();
+		bUsesTransformWidget |= Mode->UsesTransformWidget();
 	}
 
 	return bUsesTransformWidget;
@@ -1753,9 +1623,9 @@ bool FEditorModeTools::UsesTransformWidget() const
 bool FEditorModeTools::UsesTransformWidget( FWidget::EWidgetMode CheckMode ) const
 {
 	bool bUsesTransformWidget = false;
-	for( int32 ModeIndex = 0; ModeIndex < ActiveModes.Num(); ++ModeIndex )
+	for( const auto& Mode : Modes)
 	{
-		bUsesTransformWidget |= ActiveModes[ModeIndex]->UsesTransformWidget(CheckMode);
+		bUsesTransformWidget |= Mode->UsesTransformWidget(CheckMode);
 	}
 
 	return bUsesTransformWidget;
@@ -1764,19 +1634,19 @@ bool FEditorModeTools::UsesTransformWidget( FWidget::EWidgetMode CheckMode ) con
 /** Sets the current widget axis */
 void FEditorModeTools::SetCurrentWidgetAxis( EAxisList::Type NewAxis )
 {
-	for( int32 ModeIndex = 0; ModeIndex < ActiveModes.Num(); ++ModeIndex )
+	for( const auto& Mode : Modes)
 	{
-		ActiveModes[ModeIndex]->SetCurrentWidgetAxis( NewAxis );
+		Mode->SetCurrentWidgetAxis( NewAxis );
 	}
 }
 
 /** Notifies all active modes of mouse click messages. */
-bool FEditorModeTools::HandleClick(FLevelEditorViewportClient* InViewportClient,  HHitProxy *HitProxy, const FViewportClick& Click )
+bool FEditorModeTools::HandleClick(FEditorViewportClient* InViewportClient,  HHitProxy *HitProxy, const FViewportClick& Click )
 {
 	bool bHandled = false;
-	for( int32 ModeIndex = 0; ModeIndex < ActiveModes.Num(); ++ModeIndex )
+	for( const auto& Mode : Modes)
 	{
-		bHandled |= ActiveModes[ModeIndex]->HandleClick(InViewportClient, HitProxy, Click);
+		bHandled |= Mode->HandleClick(InViewportClient, HitProxy, Click);
 	}
 
 	return bHandled;
@@ -1786,12 +1656,12 @@ bool FEditorModeTools::HandleClick(FLevelEditorViewportClient* InViewportClient,
 bool FEditorModeTools::ShouldDrawBrushWireframe( AActor* InActor ) const
 {
 	bool bShouldDraw = false;
-	for( int32 ModeIndex = 0; ModeIndex < ActiveModes.Num(); ++ModeIndex )
+	for( const auto& Mode : Modes)
 	{
-		bShouldDraw |= ActiveModes[ModeIndex]->ShouldDrawBrushWireframe( InActor );
+		bShouldDraw |= Mode->ShouldDrawBrushWireframe( InActor );
 	}
 
-	if( ActiveModes.Num() == 0 )
+	if( Modes.Num() == 0 )
 	{
 		// We can get into a state where there are no active modes at editor startup if the builder brush is created before the default mode is activated.
 		// Ensure we can see the builder brush when no modes are active.
@@ -1808,105 +1678,120 @@ bool FEditorModeTools::ShouldDrawBrushVertices() const
 }
 
 /** Ticks all active modes */
-void FEditorModeTools::Tick( FLevelEditorViewportClient* ViewportClient, float DeltaTime )
+void FEditorModeTools::Tick( FEditorViewportClient* ViewportClient, float DeltaTime )
 {
-	for( int32 ModeIndex = 0; ModeIndex < ActiveModes.Num(); ++ModeIndex )
+	// Remove anything pending destruction
+	for( int32 Index = Modes.Num() - 1; Index >= 0; --Index)
 	{
-		ActiveModes[ModeIndex]->Tick( ViewportClient, DeltaTime );
+		if (Modes[Index]->IsPendingDeletion())
+		{
+			DeactivateModeAtIndex(Index);
+		}
+	}
+	
+	if (Modes.Num() == 0)
+	{
+		// Ensure the default mode is active if there are no active modes.
+		ActivateDefaultMode();
+	}
+
+	for( const auto& Mode : Modes)
+	{
+		Mode->Tick( ViewportClient, DeltaTime );
 	}
 }
 
 /** Notifies all active modes of any change in mouse movement */
-bool FEditorModeTools::InputDelta( FLevelEditorViewportClient* InViewportClient,FViewport* InViewport,FVector& InDrag,FRotator& InRot,FVector& InScale )
+bool FEditorModeTools::InputDelta( FEditorViewportClient* InViewportClient,FViewport* InViewport,FVector& InDrag,FRotator& InRot,FVector& InScale )
 {
 	bool bHandled = false;
-	for( int32 ModeIndex = 0; ModeIndex < ActiveModes.Num(); ++ModeIndex )
+	for( const auto& Mode : Modes)
 	{
-		bHandled |= ActiveModes[ModeIndex]->InputDelta( InViewportClient, InViewport, InDrag, InRot, InScale );
+		bHandled |= Mode->InputDelta( InViewportClient, InViewport, InDrag, InRot, InScale );
 	}
 	return bHandled;
 }
 
 /** Notifies all active modes of captured mouse movement */	
-bool FEditorModeTools::CapturedMouseMove( FLevelEditorViewportClient* InViewportClient, FViewport* InViewport, int32 InMouseX, int32 InMouseY )
+bool FEditorModeTools::CapturedMouseMove( FEditorViewportClient* InViewportClient, FViewport* InViewport, int32 InMouseX, int32 InMouseY )
 {
 	bool bHandled = false;
-	for( int32 ModeIndex = 0; ModeIndex < ActiveModes.Num(); ++ModeIndex )
+	for( const auto& Mode : Modes)
 	{
-		bHandled |= ActiveModes[ModeIndex]->CapturedMouseMove( InViewportClient, InViewport, InMouseX, InMouseY );
+		bHandled |= Mode->CapturedMouseMove( InViewportClient, InViewport, InMouseX, InMouseY );
 	}
 	return bHandled;
 }
 
 /** Notifies all active modes of keyboard input */
-bool FEditorModeTools::InputKey(FLevelEditorViewportClient* InViewportClient, FViewport* Viewport, FKey Key, EInputEvent Event)
+bool FEditorModeTools::InputKey(FEditorViewportClient* InViewportClient, FViewport* Viewport, FKey Key, EInputEvent Event)
 {
 	bool bHandled = false;
-	for( int32 ModeIndex = 0; ModeIndex < ActiveModes.Num(); ++ModeIndex )
+	for (const auto& Mode : Modes)
 	{
-		bHandled |= ActiveModes[ModeIndex]->InputKey( InViewportClient, Viewport, Key, Event );
+		bHandled |= Mode->InputKey( InViewportClient, Viewport, Key, Event );
 	}
 	return bHandled;
 }
 
 /** Notifies all active modes of axis movement */
-bool FEditorModeTools::InputAxis(FLevelEditorViewportClient* InViewportClient, FViewport* Viewport, int32 ControllerId, FKey Key, float Delta, float DeltaTime)
+bool FEditorModeTools::InputAxis(FEditorViewportClient* InViewportClient, FViewport* Viewport, int32 ControllerId, FKey Key, float Delta, float DeltaTime)
 {
 	bool bHandled = false;
-	for( int32 ModeIndex = 0; ModeIndex < ActiveModes.Num(); ++ModeIndex )
+	for( const auto& Mode : Modes)
 	{
-		bHandled |= ActiveModes[ModeIndex]->InputAxis( InViewportClient, Viewport, ControllerId, Key, Delta, DeltaTime );
+		bHandled |= Mode->InputAxis( InViewportClient, Viewport, ControllerId, Key, Delta, DeltaTime );
 	}
 	return bHandled;
 }
 
-bool FEditorModeTools::MouseEnter( FLevelEditorViewportClient* InViewportClient, FViewport* Viewport, int32 X, int32 Y )
+bool FEditorModeTools::MouseEnter( FEditorViewportClient* InViewportClient, FViewport* Viewport, int32 X, int32 Y )
 {
 	bool bHandled = false;
-	for( int32 ModeIndex = 0; ModeIndex < ActiveModes.Num(); ++ModeIndex )
+	for( const auto& Mode : Modes)
 	{
-		bHandled |= ActiveModes[ModeIndex]->MouseEnter( InViewportClient, Viewport, X, Y );
+		bHandled |= Mode->MouseEnter( InViewportClient, Viewport, X, Y );
 	}
 	return bHandled;
 }
 
-bool FEditorModeTools::MouseLeave( FLevelEditorViewportClient* InViewportClient, FViewport* Viewport )
+bool FEditorModeTools::MouseLeave( FEditorViewportClient* InViewportClient, FViewport* Viewport )
 {
 	bool bHandled = false;
-	for( int32 ModeIndex = 0; ModeIndex < ActiveModes.Num(); ++ModeIndex )
+	for( const auto& Mode : Modes)
 	{
-		bHandled |= ActiveModes[ModeIndex]->MouseLeave( InViewportClient, Viewport );
+		bHandled |= Mode->MouseLeave( InViewportClient, Viewport );
 	}
 	return bHandled;
 }
 
 /** Notifies all active modes that the mouse has moved */
-bool FEditorModeTools::MouseMove( FLevelEditorViewportClient* InViewportClient, FViewport* Viewport, int32 X, int32 Y )
+bool FEditorModeTools::MouseMove( FEditorViewportClient* InViewportClient, FViewport* Viewport, int32 X, int32 Y )
 {
 	bool bHandled = false;
-	for( int32 ModeIndex = 0; ModeIndex < ActiveModes.Num(); ++ModeIndex )
+	for( const auto& Mode : Modes)
 	{
-		bHandled |= ActiveModes[ModeIndex]->MouseMove( InViewportClient, Viewport, X, Y );
+		bHandled |= Mode->MouseMove( InViewportClient, Viewport, X, Y );
 	}
 	return bHandled;
 }
 
-bool FEditorModeTools::ReceivedFocus( FLevelEditorViewportClient* InViewportClient, FViewport* Viewport )
+bool FEditorModeTools::ReceivedFocus( FEditorViewportClient* InViewportClient, FViewport* Viewport )
 {
 	bool bHandled = false;
-	for( int32 ModeIndex = 0; ModeIndex < ActiveModes.Num(); ++ModeIndex )
+	for( const auto& Mode : Modes)
 	{
-		bHandled |= ActiveModes[ModeIndex]->ReceivedFocus( InViewportClient, Viewport );
+		bHandled |= Mode->ReceivedFocus( InViewportClient, Viewport );
 	}
 	return bHandled;
 }
 
-bool FEditorModeTools::LostFocus( FLevelEditorViewportClient* InViewportClient, FViewport* Viewport )
+bool FEditorModeTools::LostFocus( FEditorViewportClient* InViewportClient, FViewport* Viewport )
 {
 	bool bHandled = false;
-	for( int32 ModeIndex = 0; ModeIndex < ActiveModes.Num(); ++ModeIndex )
+	for( const auto& Mode : Modes)
 	{
-		bHandled |= ActiveModes[ModeIndex]->LostFocus( InViewportClient, Viewport );
+		bHandled |= Mode->LostFocus( InViewportClient, Viewport );
 	}
 	return bHandled;
 }
@@ -1914,36 +1799,36 @@ bool FEditorModeTools::LostFocus( FLevelEditorViewportClient* InViewportClient,
 /** Draws all active mode components */	
 void FEditorModeTools::DrawActiveModes( const FSceneView* InView, FPrimitiveDrawInterface* PDI )
 {
-	for( int32 ModeIndex = 0; ModeIndex < ActiveModes.Num(); ++ModeIndex )
+	for( const auto& Mode : Modes)
 	{
-		ActiveModes[ModeIndex]->Draw( InView, PDI );
+		Mode->Draw( InView, PDI );
 	}
 }
 
 /** Renders all active modes */
 void FEditorModeTools::Render( const FSceneView* InView, FViewport* Viewport, FPrimitiveDrawInterface* PDI )
 {
-	for( int32 ModeIndex = 0; ModeIndex < ActiveModes.Num(); ++ModeIndex )
+	for( const auto& Mode : Modes)
 	{
-		ActiveModes[ModeIndex]->Render( InView, Viewport, PDI );
+		Mode->Render( InView, Viewport, PDI );
 	}
 }
 
 /** Draws the HUD for all active modes */
-void FEditorModeTools::DrawHUD( FLevelEditorViewportClient* InViewportClient,FViewport* Viewport, const FSceneView* View, FCanvas* Canvas )
+void FEditorModeTools::DrawHUD( FEditorViewportClient* InViewportClient,FViewport* Viewport, const FSceneView* View, FCanvas* Canvas )
 {
-	for( int32 ModeIndex = 0; ModeIndex < ActiveModes.Num(); ++ModeIndex )
+	for( const auto& Mode : Modes)
 	{
-		ActiveModes[ModeIndex]->DrawHUD( InViewportClient, Viewport, View, Canvas );
+		Mode->DrawHUD( InViewportClient, Viewport, View, Canvas );
 	}
 }
 
 /** Calls PostUndo on all active components */
 void FEditorModeTools::PostUndo()
 {
-	for( int32 ModeIndex = 0; ModeIndex < ActiveModes.Num(); ++ModeIndex )
+	for( const auto& Mode : Modes)
 	{
-		ActiveModes[ModeIndex]->PostUndo();
+		Mode->PostUndo();
 	}
 }
 
@@ -1951,9 +1836,9 @@ void FEditorModeTools::PostUndo()
 bool FEditorModeTools::AllowWidgetMove() const
 {
 	bool bAllow = false;
-	for( int32 ModeIndex = 0; ModeIndex < ActiveModes.Num(); ++ModeIndex )
+	for( const auto& Mode : Modes)
 	{
-		bAllow |= ActiveModes[ModeIndex]->AllowWidgetMove();
+		bAllow |= Mode->AllowWidgetMove();
 	}
 	return bAllow;
 }
@@ -1961,9 +1846,9 @@ bool FEditorModeTools::AllowWidgetMove() const
 bool FEditorModeTools::DisallowMouseDeltaTracking() const
 {
 	bool bDisallow = false;
-	for( int32 ModeIndex = 0; ModeIndex < ActiveModes.Num(); ++ModeIndex )
+	for( const auto& Mode : Modes)
 	{
-		bDisallow |= ActiveModes[ModeIndex]->DisallowMouseDeltaTracking();
+		bDisallow |= Mode->DisallowMouseDeltaTracking();
 	}
 	return bDisallow;
 }
@@ -1971,9 +1856,9 @@ bool FEditorModeTools::DisallowMouseDeltaTracking() const
 bool FEditorModeTools::GetCursor(EMouseCursor::Type& OutCursor) const
 {
 	bool bHandled = false;
-	for( int32 ModeIndex = 0; ModeIndex < ActiveModes.Num(); ++ModeIndex )
+	for( const auto& Mode : Modes)
 	{
-		bHandled |= ActiveModes[ModeIndex]->GetCursor(OutCursor);
+		bHandled |= Mode->GetCursor(OutCursor);
 	}
 	return bHandled;
 }
@@ -1986,7 +1871,7 @@ void FEditorModeTools::CycleWidgetMode (void)
 	//make sure we're not currently tracking mouse movement.  If we are, changing modes could cause a crash due to referencing an axis/plane that is incompatible with the widget
 	for(int32 ViewportIndex = 0;ViewportIndex < GEditor->LevelViewportClients.Num();ViewportIndex++)
 	{
-		FLevelEditorViewportClient* ViewportClient = GEditor->LevelViewportClients[ ViewportIndex ];
+		FEditorViewportClient* ViewportClient = GEditor->LevelViewportClients[ ViewportIndex ];
 		if (ViewportClient->IsTracking())
 		{
 			return;
@@ -2011,7 +1896,7 @@ void FEditorModeTools::CycleWidgetMode (void)
 				Wk -= FWidget::WM_Max;
 			}
 		}
-		while (!GEditorModeTools().UsesTransformWidget((FWidget::EWidgetMode)Wk) && Wk != CurrentWk);
+		while (!UsesTransformWidget((FWidget::EWidgetMode)Wk) && Wk != CurrentWk);
 		SetWidgetMode( (FWidget::EWidgetMode)Wk );
 		FEditorSupportDelegates::RedrawAllViewports.Broadcast();
 	}
@@ -2034,11 +1919,11 @@ void FEditorModeTools::LoadWidgetSettings(void)
 
 FVector FEditorModeTools::GetWidgetLocation() const
 {
-	for (int Index = ActiveModes.Num() - 1; Index >= 0 ; Index--)
+	for (int Index = Modes.Num() - 1; Index >= 0 ; Index--)
 	{
-		if ( ActiveModes[Index]->UsesTransformWidget() )
+		if ( Modes[Index]->UsesTransformWidget() )
 		{
-			 return ActiveModes[Index]->GetWidgetLocation();
+			 return Modes[Index]->GetWidgetLocation();
 		}
 	}
 	
@@ -2087,7 +1972,7 @@ bool FEditorModeTools::GetShowFriendlyVariableNames() const
  * Sets a bookmark in the levelinfo file, allocating it if necessary.
  */
 
-void FEditorModeTools::SetBookmark( uint32 InIndex, FLevelEditorViewportClient* InViewportClient )
+void FEditorModeTools::SetBookmark( uint32 InIndex, FEditorViewportClient* InViewportClient )
 {
 	UWorld* World = InViewportClient->GetWorld();
 	if ( World )
@@ -2138,7 +2023,7 @@ void FEditorModeTools::SetBookmark( uint32 InIndex, FLevelEditorViewportClient*
  * Checks to see if a bookmark exists at a given index
  */
 
-bool FEditorModeTools::CheckBookmark( uint32 InIndex, FLevelEditorViewportClient* InViewportClient )
+bool FEditorModeTools::CheckBookmark( uint32 InIndex, FEditorViewportClient* InViewportClient )
 {
 	UWorld* World = InViewportClient->GetWorld();
 	if ( World )
@@ -2157,7 +2042,7 @@ bool FEditorModeTools::CheckBookmark( uint32 InIndex, FLevelEditorViewportClient
  * Retrieves a bookmark from the list.
  */
 
-void FEditorModeTools::JumpToBookmark( uint32 InIndex, bool bShouldRestoreLevelVisibility, FLevelEditorViewportClient* InViewportClient )
+void FEditorModeTools::JumpToBookmark( uint32 InIndex, bool bShouldRestoreLevelVisibility, FEditorViewportClient* InViewportClient )
 {
 	UWorld* World = InViewportClient->GetWorld();
 	if ( World )
@@ -2188,7 +2073,7 @@ void FEditorModeTools::JumpToBookmark( uint32 InIndex, bool bShouldRestoreLevelV
 /**
  * Clears a bookmark
  */
-void FEditorModeTools::ClearBookmark( uint32 InIndex, FLevelEditorViewportClient* InViewportClient )
+void FEditorModeTools::ClearBookmark( uint32 InIndex, FEditorViewportClient* InViewportClient )
 {
 	UWorld* World = InViewportClient->GetWorld();
 	if( World )
@@ -2206,7 +2091,7 @@ void FEditorModeTools::ClearBookmark( uint32 InIndex, FLevelEditorViewportClient
 /**
 * Clears all book marks
 */
-void FEditorModeTools::ClearAllBookmarks( FLevelEditorViewportClient* InViewportClient )
+void FEditorModeTools::ClearAllBookmarks( FEditorViewportClient* InViewportClient )
 {
 	for( int i = 0; i <  AWorldSettings::MAX_BOOKMARK_NUMBER; ++i )
 	{
@@ -2232,16 +2117,14 @@ void FEditorModeTools::AddReferencedObjects( FReferenceCollector& Collector )
  */
 FEdMode* FEditorModeTools::GetActiveMode( FEditorModeID InID )
 {
-	FEdMode* ActiveMode = NULL;
-	for( int32 ModeIndex = 0; ModeIndex < ActiveModes.Num(); ++ModeIndex )
+	for( auto& Mode : Modes )
 	{
-		if( ActiveModes[ModeIndex]->GetID() == InID )
+		if( Mode->GetID() == InID )
 		{
-			ActiveMode = ActiveModes[ModeIndex];
-			break;
+			return Mode.Get();
 		}
 	}
-	return ActiveMode;
+	return nullptr;
 }
 
 /**
@@ -2250,16 +2133,15 @@ FEdMode* FEditorModeTools::GetActiveMode( FEditorModeID InID )
  */
 const FEdMode* FEditorModeTools::GetActiveMode( FEditorModeID InID ) const
 {
-	const FEdMode* ActiveMode = NULL;
-	for( int32 ModeIndex = 0; ModeIndex < ActiveModes.Num(); ++ModeIndex )
+	for (const auto& Mode : Modes)
 	{
-		if( ActiveModes[ModeIndex]->GetID() == InID )
+		if (Mode->GetID() == InID)
 		{
-			ActiveMode = ActiveModes[ModeIndex];
-			break;
+			return Mode.Get();
 		}
 	}
-	return ActiveMode;
+
+	return nullptr;
 }
 
 /**
@@ -2285,6 +2167,14 @@ bool FEditorModeTools::IsModeActive( FEditorModeID InID ) const
 	return GetActiveMode( InID ) != NULL;
 }
 
+/** 
+ * Returns true if the default editor mode is active 
+ */
+bool FEditorModeTools::IsDefaultModeActive() const
+{
+	return IsModeActive(DefaultID);
+}
+
 /** 
  * Returns an array of all active modes
  */
@@ -2292,15 +2182,8 @@ void FEditorModeTools::GetActiveModes( TArray<FEdMode*>& OutActiveModes )
 {
 	OutActiveModes.Empty();
 	// Copy into an array.  Do not let users modify the active list directly.
-	OutActiveModes = ActiveModes;
-}
-
-void FEditorModeTools::GetModes( TArray<FEdMode*>& OutModes )
-{
-	OutModes.Empty();
-	// Copy into an array.  Do not let users modify the modes list directly.
-	for(TSharedPtr<FEdMode> Mode : Modes)
+	for( auto& Mode : Modes)
 	{
-		OutModes.Add(Mode.Get());
+		OutActiveModes.Add(Mode.Get());
 	}
 }
diff --git a/Engine/Source/Editor/UnrealEd/Private/EditorSelectUtils.cpp b/Engine/Source/Editor/UnrealEd/Private/EditorSelectUtils.cpp
--- a/Engine/Source/Editor/UnrealEd/Private/EditorSelectUtils.cpp
+++ b/Engine/Source/Editor/UnrealEd/Private/EditorSelectUtils.cpp
@@ -37,7 +37,7 @@ void UUnrealEdEngine::NoteActorMovement()
 		GEditor->ClickFlags |= CF_MOVE_ACTOR;
 
 		const FScopedTransaction Transaction( NSLOCTEXT("UnrealEd", "ActorMovement", "Actor Movement") );
-		GEditorModeTools().Snapping=0;
+		GLevelEditorModeTools().Snapping=0;
 		
 		AActor* SelectedActor = NULL;
 		for ( FSelectionIterator It( GetSelectedActorIterator() ) ; It ; ++It )
@@ -62,7 +62,7 @@ void UUnrealEdEngine::NoteActorMovement()
 			AActor* Actor = static_cast<AActor*>( *It );
 			checkSlow( Actor->IsA(AActor::StaticClass()) );
 
-			GEditorModeTools().Snapping = 1;
+			GLevelEditorModeTools().Snapping = 1;
 			break;
 		}
 
@@ -144,13 +144,13 @@ void UUnrealEdEngine::Cleanse( bool ClearSelection, bool Redraw, const FText& Re
 
 FVector UUnrealEdEngine::GetPivotLocation()
 {
-	return GEditorModeTools().PivotLocation;
+	return GLevelEditorModeTools().PivotLocation;
 }
 
 
 void UUnrealEdEngine::SetPivot( FVector NewPivot, bool bSnapPivotToGrid, bool bIgnoreAxis, bool bAssignPivot/*=false*/ )
 {
-	FEditorModeTools& EditorModeTools = GEditorModeTools();
+	FEditorModeTools& EditorModeTools = GLevelEditorModeTools();
 
 	if( !bIgnoreAxis )
 	{
@@ -225,9 +225,9 @@ void UUnrealEdEngine::SetPivot( FVector NewPivot, bool bSnapPivotToGrid, bool bI
 
 void UUnrealEdEngine::ResetPivot()
 {
-	GEditorModeTools().PivotShown	= 0;
-	GEditorModeTools().Snapping		= 0;
-	GEditorModeTools().SnappedActor	= 0;
+	GLevelEditorModeTools().PivotShown	= 0;
+	GLevelEditorModeTools().Snapping		= 0;
+	GLevelEditorModeTools().SnappedActor	= 0;
 }
 
 /*-----------------------------------------------------------------------------
@@ -292,7 +292,7 @@ void UUnrealEdEngine::UpdatePivotLocationForSelection( bool bOnChange )
 	if( ActorCount > 0 ) 
 	{		
 		// For geometry mode use current pivot location as it's set to selected face, not actor
-		FEditorModeTools& Tools = GEditorModeTools();
+		FEditorModeTools& Tools = GLevelEditorModeTools();
 		if( Tools.IsModeActive(FBuiltinEditorModes::EM_Geometry) == false || bOnChange == true )
 		{
 			// Set pivot point to the actor's location
@@ -323,8 +323,11 @@ void UUnrealEdEngine::NoteSelectionChange()
 	// The selection changed, so make sure the pivot (widget) is located in the right place
 	UpdatePivotLocationForSelection( true );
 
+	// Clear active editing visualizer on selection change
+	GUnrealEd->ComponentVisManager.ClearActiveComponentVis();
+
 	TArray<FEdMode*> ActiveModes;
-	GEditorModeTools().GetActiveModes( ActiveModes );
+	GLevelEditorModeTools().GetActiveModes( ActiveModes );
 	for( int32 ModeIndex = 0; ModeIndex < ActiveModes.Num(); ++ModeIndex )
 	{
 		ActiveModes[ModeIndex]->ActorSelectionChangeNotify();
@@ -443,7 +446,7 @@ bool UUnrealEdEngine::CanSelectActor(AActor* Actor, bool bInSelected, bool bSele
 	bool bSelectionAllowed = false;
 
 	TArray<FEdMode*> ActiveModes;
-	GEditorModeTools().GetActiveModes( ActiveModes );
+	GLevelEditorModeTools().GetActiveModes( ActiveModes );
 	for( int32 ModeIndex = 0; ModeIndex < ActiveModes.Num(); ++ModeIndex )
 	{
 		bSelectionAllowed |= ActiveModes[ModeIndex]->IsSelectionAllowed( Actor, bInSelected );
@@ -463,7 +466,7 @@ void UUnrealEdEngine::SelectActor(AActor* Actor, bool bInSelected, bool bNotify,
 	bool bSelectionHandled = false;
 
 	TArray<FEdMode*> ActiveModes;
-	GEditorModeTools().GetActiveModes( ActiveModes );
+	GLevelEditorModeTools().GetActiveModes( ActiveModes );
 	for( int32 ModeIndex = 0; ModeIndex < ActiveModes.Num(); ++ModeIndex )
 	{
 		bSelectionHandled |= ActiveModes[ModeIndex]->Select( Actor, bInSelected );
diff --git a/Engine/Source/Editor/UnrealEd/Private/EditorServer.cpp b/Engine/Source/Editor/UnrealEd/Private/EditorServer.cpp
--- a/Engine/Source/Editor/UnrealEd/Private/EditorServer.cpp
+++ b/Engine/Source/Editor/UnrealEd/Private/EditorServer.cpp
@@ -807,7 +807,7 @@ bool UEditorEngine::Exec_Brush( UWorld* InWorld, const TCHAR* Str, FOutputDevice
 			}
 			WorldBrush->ReregisterAllComponents();
 
-			GEditorModeTools().MapChangeNotify();
+			GLevelEditorModeTools().MapChangeNotify();
 			RedrawLevelEditingViewports();
 			return true;
 		}
@@ -827,7 +827,7 @@ bool UEditorEngine::Exec_Brush( UWorld* InWorld, const TCHAR* Str, FOutputDevice
 			}
 			WorldBrush->ReregisterAllComponents();
 
-			GEditorModeTools().MapChangeNotify();
+			GLevelEditorModeTools().MapChangeNotify();
 			RedrawLevelEditingViewports();
 			return true;
 		}
@@ -1590,7 +1590,7 @@ void UEditorEngine::RebuildAlteredBSP()
 
 void UEditorEngine::BSPIntersectionHelper(UWorld* InWorld, ECsgOper Operation)
 {
-	FEdModeGeometry* Mode = GEditorModeTools().GetActiveModeTyped<FEdModeGeometry>(FBuiltinEditorModes::EM_Geometry);
+	FEdModeGeometry* Mode = GLevelEditorModeTools().GetActiveModeTyped<FEdModeGeometry>(FBuiltinEditorModes::EM_Geometry);
 	if (Mode)
 	{
 		Mode->GeometrySelectNone(true, true);
@@ -1749,27 +1749,27 @@ void UEditorEngine::CreateNewMapForEditing()
 	const FScopedBusyCursor BusyCursor;
 
 	// Change out of Matinee when opening new map, so we avoid editing data in the old one.
-	if( GEditorModeTools().IsModeActive( FBuiltinEditorModes::EM_InterpEdit ) )
+	if( GLevelEditorModeTools().IsModeActive( FBuiltinEditorModes::EM_InterpEdit ) )
 	{
-		GEditorModeTools().DeactivateMode( FBuiltinEditorModes::EM_InterpEdit );
+		GLevelEditorModeTools().DeactivateMode( FBuiltinEditorModes::EM_InterpEdit );
 	}
 
 	// Also change out of Landscape mode to ensure all references are cleared.
-	if( GEditorModeTools().IsModeActive( FBuiltinEditorModes::EM_Landscape ) )
+	if( GLevelEditorModeTools().IsModeActive( FBuiltinEditorModes::EM_Landscape ) )
 	{
-		GEditorModeTools().DeactivateMode( FBuiltinEditorModes::EM_Landscape );
+		GLevelEditorModeTools().DeactivateMode( FBuiltinEditorModes::EM_Landscape );
 	}
 
 	// Also change out of Foliage mode to ensure all references are cleared.
-	if( GEditorModeTools().IsModeActive( FBuiltinEditorModes::EM_Foliage ) )
+	if( GLevelEditorModeTools().IsModeActive( FBuiltinEditorModes::EM_Foliage ) )
 	{
-		GEditorModeTools().DeactivateMode( FBuiltinEditorModes::EM_Foliage );
+		GLevelEditorModeTools().DeactivateMode( FBuiltinEditorModes::EM_Foliage );
 	}
 
 	// Change out of mesh paint mode when opening a new map.
-	if( GEditorModeTools().IsModeActive( FBuiltinEditorModes::EM_MeshPaint ) )
+	if( GLevelEditorModeTools().IsModeActive( FBuiltinEditorModes::EM_MeshPaint ) )
 	{
-		GEditorModeTools().DeactivateMode( FBuiltinEditorModes::EM_MeshPaint );
+		GLevelEditorModeTools().DeactivateMode( FBuiltinEditorModes::EM_MeshPaint );
 	}
 
 	GUnrealEd->NewMap();
@@ -1993,10 +1993,7 @@ bool UEditorEngine::Map_Load(const TCHAR* Str, FOutputDevice& Ar)
 					FStatsViewerModule& StatsViewerModule = FModuleManager::Get().LoadModuleChecked<FStatsViewerModule>(TEXT("StatsViewer"));
 					StatsViewerModule.GetPage(EStatsPage::LightingBuildInfo)->Clear();
 
-					if( !GEditorModeTools().IsModeActive( FBuiltinEditorModes::EM_Default ) )
-					{
-						GEditorModeTools().ActivateMode( FBuiltinEditorModes::EM_Default );
-					}
+					GLevelEditorModeTools().ActivateDefaultMode();
 
 					OldOuter = Context.World()->GetOuter();
 
@@ -5509,7 +5506,7 @@ bool UEditorEngine::HandleSelectCommand( const TCHAR* Str, FOutputDevice& Ar, UW
 bool UEditorEngine::HandleDeleteCommand( const TCHAR* Str, FOutputDevice& Ar, UWorld* InWorld )
 {
 	// If geometry mode is active, give it a chance to handle this command.  If it does not, use the default handler
-	if( !GEditorModeTools().IsModeActive( FBuiltinEditorModes::EM_Geometry ) || !( (FEdModeGeometry*)GEditorModeTools().GetActiveMode(FBuiltinEditorModes::EM_Geometry) )->ExecDelete() )
+	if( !GLevelEditorModeTools().IsModeActive( FBuiltinEditorModes::EM_Geometry ) || !( (FEdModeGeometry*)GLevelEditorModeTools().GetActiveMode(FBuiltinEditorModes::EM_Geometry) )->ExecDelete() )
 	{
 		return Exec( InWorld, TEXT("ACTOR DELETE") );
 	}
diff --git a/Engine/Source/Editor/UnrealEd/Private/EditorViewportClient.cpp b/Engine/Source/Editor/UnrealEd/Private/EditorViewportClient.cpp
--- a/Engine/Source/Editor/UnrealEd/Private/EditorViewportClient.cpp
+++ b/Engine/Source/Editor/UnrealEd/Private/EditorViewportClient.cpp
@@ -180,7 +180,7 @@ int32 FEditorViewportClient::GetCameraSpeedSetting() const
 
 float const FEditorViewportClient::SafePadding = 0.075f;
 
-FEditorViewportClient::FEditorViewportClient(FPreviewScene* InPreviewScene)
+FEditorViewportClient::FEditorViewportClient(FEditorModeTools& InModeTools, FPreviewScene* InPreviewScene)
 	: CameraSpeedSetting(4)
 	, ImmersiveDelegate()
 	, VisibilityDelegate()
@@ -237,6 +237,7 @@ FEditorViewportClient::FEditorViewportClient(FPreviewScene* InPreviewScene)
 	, bIsSimulateInEditorViewport(false)
 	, bCameraLock(false)
 	, PreviewScene(InPreviewScene)
+	, ModeTools(&InModeTools)
 	, PerspViewModeIndex(VMI_Lit)
 	, OrthoViewModeIndex(VMI_BrushWireframe)
 	, NearPlane(-1.0f)
@@ -274,12 +275,12 @@ FEditorViewportClient::FEditorViewportClient(FPreviewScene* InPreviewScene)
 
 	SetViewMode(VMI_Lit);
 
-	GEditorModeTools().OnEditorModeChanged().AddRaw(this, &FEditorViewportClient::OnEditorModeChanged);
+	ModeTools->OnEditorModeChanged().AddRaw(this, &FEditorViewportClient::OnEditorModeChanged);
 }
 
 FEditorViewportClient::~FEditorViewportClient()
 {
-	GEditorModeTools().OnEditorModeChanged().RemoveAll(this);
+	ModeTools->OnEditorModeChanged().RemoveAll(this);
 
 	delete Widget;
 	delete MouseDeltaTracker;
@@ -969,7 +970,7 @@ EMouseCursor::Type FEditorViewportClient::GetCursor(FViewport* InViewport,int32
 	{
 		// allow editor modes to override cursor
 		EMouseCursor::Type EditorModeCursor = EMouseCursor::Default;
-		if(GEditorModeTools().GetCursor(EditorModeCursor))
+		if(ModeTools->GetCursor(EditorModeCursor))
 		{
 			MouseCursor = EditorModeCursor;
 		}
@@ -2608,6 +2609,11 @@ void FEditorViewportClient::Draw(const FSceneView* View,FPrimitiveDrawInterface*
 	FramesSinceLastDraw = 0;
 }
 
+void FEditorViewportClient::RenderDragTool(const FSceneView* View, FCanvas* Canvas)
+{
+	MouseDeltaTracker->RenderDragTool(View, Canvas);
+}
+
 FLinearColor FEditorViewportClient::GetBackgroundColor() const
 {
 	FLinearColor BackgroundColor = FColor(55, 55, 55);
@@ -3111,7 +3117,7 @@ bool FEditorViewportClient::IsFlightCameraInputModeActive() const
 		if( CameraController != NULL )
 		{
 			// Also check that we're not currently using a ModeWidget (for Vertex Paint etc)
-			const FEdMode* Mode = GEditorModeTools().GetActiveMode(FBuiltinEditorModes::EM_MeshPaint);
+			const FEdMode* Mode = ModeTools->GetActiveMode(FBuiltinEditorModes::EM_MeshPaint);
 			const bool bIsPaintingMesh = ( Mode ) ? ((FEdModeMeshPaint*)Mode)->IsPainting() : false;
 			const bool bLeftMouseButtonDown = Viewport->KeyState(EKeys::LeftMouseButton) && !bIsPaintingMesh;
 			const bool bMiddleMouseButtonDown = Viewport->KeyState( EKeys::MiddleMouseButton );
@@ -3940,6 +3946,35 @@ void FEditorViewportClient::DrawBoundingBox(FBox &Box, FCanvas* InCanvas, const
 	}
 }
 
+void FEditorViewportClient::DrawActorScreenSpaceBoundingBox( FCanvas* InCanvas, const FSceneView* InView, FViewport* InViewport, AActor* InActor, const FLinearColor& InColor, const bool bInDrawBracket, const FString& InLabelText )
+{
+	check( InActor != NULL );
+
+
+	// First check to see if we're dealing with a sprite, otherwise just use the normal bounding box
+	UBillboardComponent* Sprite = InActor->FindComponentByClass<UBillboardComponent>();
+
+	FBox ActorBox;
+	if( Sprite != NULL )
+	{
+		ActorBox = Sprite->Bounds.GetBox();
+	}
+	else
+	{
+		const bool bNonColliding = true;
+		ActorBox = InActor->GetComponentsBoundingBox( bNonColliding );
+	}
+
+
+	// If we didn't get a valid bounding box, just make a little one around the actor location
+	if( !ActorBox.IsValid || ActorBox.GetExtent().GetMin() < KINDA_SMALL_NUMBER )
+	{
+		ActorBox = FBox( InActor->GetActorLocation() - FVector( -20 ), InActor->GetActorLocation() + FVector( 20 ) );
+	}
+
+	DrawBoundingBox(ActorBox, InCanvas, InView, InViewport, InColor, bInDrawBracket, InLabelText);
+}
+
 void FEditorViewportClient::SetGameView(bool bGameViewEnable)
 {
 	// backup this state as we want to preserve it
diff --git a/Engine/Source/Editor/UnrealEd/Private/FileHelpers.cpp b/Engine/Source/Editor/UnrealEd/Private/FileHelpers.cpp
--- a/Engine/Source/Editor/UnrealEd/Private/FileHelpers.cpp
+++ b/Engine/Source/Editor/UnrealEd/Private/FileHelpers.cpp
@@ -189,7 +189,7 @@ namespace FileDialogHelpers
 static bool InInterpEditMode()
 {
 	// Must exit Interpolation Editing mode before you can save - so it can reset everything to its initial state.
-	if( GEditorModeTools().IsModeActive( FBuiltinEditorModes::EM_InterpEdit ) )
+	if( GLevelEditorModeTools().IsModeActive( FBuiltinEditorModes::EM_InterpEdit ) )
 	{
 		const bool ExitInterp = EAppReturnType::Yes == FMessageDialog::Open( EAppMsgType::YesNo, NSLOCTEXT("UnrealEd", "Prompt_21", "You must close Matinee before saving level.\nDo you wish to do this now and continue?") );
 		if(!ExitInterp)
@@ -197,7 +197,7 @@ static bool InInterpEditMode()
 			return true;
 		}
 
-		GEditorModeTools().DeactivateMode( FBuiltinEditorModes::EM_InterpEdit );
+		GLevelEditorModeTools().DeactivateMode( FBuiltinEditorModes::EM_InterpEdit );
 	}
 	return false;
 }
@@ -1680,21 +1680,21 @@ void FEditorFileUtils::LoadMap(const FString& InFilename, bool LoadAsTemplate, b
 	}
 
 	// Change out of Matinee when opening new map, so we avoid editing data in the old one.
-	if( GEditorModeTools().IsModeActive( FBuiltinEditorModes::EM_InterpEdit ) )
+	if( GLevelEditorModeTools().IsModeActive( FBuiltinEditorModes::EM_InterpEdit ) )
 	{
-		GEditorModeTools().ActivateMode( FBuiltinEditorModes::EM_Default );
+		GLevelEditorModeTools().ActivateDefaultMode();
 	}
 
 	// Also change out of Landscape mode to ensure all references are cleared.
-	if( GEditorModeTools().IsModeActive( FBuiltinEditorModes::EM_Landscape ) )
+	if( GLevelEditorModeTools().IsModeActive( FBuiltinEditorModes::EM_Landscape ) )
 	{
-		GEditorModeTools().DeactivateMode( FBuiltinEditorModes::EM_Landscape );
+		GLevelEditorModeTools().DeactivateMode( FBuiltinEditorModes::EM_Landscape );
 	}
 
 	// Change out of mesh paint mode when loading a map
-	if( GEditorModeTools().IsModeActive( FBuiltinEditorModes::EM_MeshPaint ) )
+	if( GLevelEditorModeTools().IsModeActive( FBuiltinEditorModes::EM_MeshPaint ) )
 	{
-		GEditorModeTools().DeactivateMode( FBuiltinEditorModes::EM_MeshPaint );
+		GLevelEditorModeTools().DeactivateMode( FBuiltinEditorModes::EM_MeshPaint );
 	}
 
 	FString LoadCommand = FString::Printf( TEXT("MAP LOAD FILE=\"%s\" TEMPLATE=%d SHOWPROGRESS=%d"), *Filename, LoadAsTemplate, bShowProgress );
diff --git a/Engine/Source/Editor/UnrealEd/Private/GroupActor.cpp b/Engine/Source/Editor/UnrealEd/Private/GroupActor.cpp
--- a/Engine/Source/Editor/UnrealEd/Private/GroupActor.cpp
+++ b/Engine/Source/Editor/UnrealEd/Private/GroupActor.cpp
@@ -640,7 +640,7 @@ void AGroupActor::UnlockSelectedGroups()
 void AGroupActor::ToggleGroupMode()
 {
 	// Group mode can only be toggled when not in InterpEdit mode
-	if( !GEditorModeTools().IsModeActive(FBuiltinEditorModes::EM_InterpEdit) )
+	if( !GLevelEditorModeTools().IsModeActive(FBuiltinEditorModes::EM_InterpEdit) )
 	{
 		GUnrealEd->bGroupingActive = !GUnrealEd->bGroupingActive;
 
