Made it possible for asset editors to maintain their own FEditorModeTools lists

diff --git a/Engine/Source/Editor/UnrealEd/Private/LevelEditorViewport.cpp b/Engine/Source/Editor/UnrealEd/Private/LevelEditorViewport.cpp
--- a/Engine/Source/Editor/UnrealEd/Private/LevelEditorViewport.cpp
+++ b/Engine/Source/Editor/UnrealEd/Private/LevelEditorViewport.cpp
@@ -1754,7 +1754,7 @@ void FTrackingTransaction::PromotePendingToActive()
 }
 
 FLevelEditorViewportClient::FLevelEditorViewportClient()
-	: FEditorViewportClient(NULL)
+	: FEditorViewportClient(GLevelEditorModeTools(), NULL)
 	, ViewHiddenLayers()
 	, VolumeActorVisibility()
 	, ActorLockedToCamera(NULL)
@@ -1767,7 +1767,6 @@ FLevelEditorViewportClient::FLevelEditorViewportClient()
 	, bEnableFading(false)
 	, bEnableColorScaling(false)
 	, bEditorCameraCut(false)
-	, bDrawVertices(false)
 	, bDrawBaseInfo(false)
 	, bDuplicateActorsOnNextDrag( false )
 	, bDuplicateActorsInProgress( false )
@@ -1779,7 +1778,6 @@ FLevelEditorViewportClient::FLevelEditorViewportClient()
 	, TrackingTransaction()
 	, DropPreviewMouseX(0)
 	, DropPreviewMouseY(0)
-	, bAllowMatineePreview(false)
 	, bWasControlledByOtherViewport(false)
 {
 	// By default a level editor viewport is pointed to the editor world
@@ -1787,7 +1785,7 @@ FLevelEditorViewportClient::FLevelEditorViewportClient()
 
 	GEditor->LevelViewportClients.Add(this);
 
-	Widget->SetUsesEditorModeTools( &GEditorModeTools() );
+	Widget->SetUsesEditorModeTools( ModeTools );
 
 	// Register for editor cleanse events so we can release references to hovered actors
 	FEditorSupportDelegates::CleanseEditor.AddRaw(this, &FLevelEditorViewportClient::OnEditorCleanse);
@@ -1963,7 +1961,7 @@ void FLevelEditorViewportClient::OverridePostProcessSettings( FSceneView& View )
 
 bool FLevelEditorViewportClient::ShouldLockPitch() const 
 {
-	return FEditorViewportClient::ShouldLockPitch() || !GEditorModeTools().GetActiveMode(FBuiltinEditorModes::EM_InterpEdit) ;
+	return FEditorViewportClient::ShouldLockPitch() || !ModeTools->GetActiveMode(FBuiltinEditorModes::EM_InterpEdit) ;
 }
 
 void FLevelEditorViewportClient::PerspectiveCameraMoved()
@@ -1978,7 +1976,7 @@ void FLevelEditorViewportClient::PerspectiveCameraMoved()
 	}
 
 	// Tell the editing mode that the camera moved, in case its interested.
-	FEdMode* Mode = GEditorModeTools().GetActiveMode(FBuiltinEditorModes::EM_InterpEdit);
+	FEdMode* Mode = ModeTools->GetActiveMode(FBuiltinEditorModes::EM_InterpEdit);
 	if( Mode )
 	{
 		((FEdModeInterpEdit*)Mode)->CamMoveNotify(this);
@@ -2011,7 +2009,7 @@ void FLevelEditorViewportClient::ResetCamera()
 	ViewTransform.SetOrthoZoom( DEFAULT_ORTHOZOOM );
 
 	// If interp mode is active, tell it about the camera movement.
-	FEdMode* Mode = GEditorModeTools().GetActiveMode(FBuiltinEditorModes::EM_InterpEdit);
+	FEdMode* Mode = ModeTools->GetActiveMode(FBuiltinEditorModes::EM_InterpEdit);
 	if( Mode )
 	{
 		((FEdModeInterpEdit*)Mode)->CamMoveNotify(this);
@@ -2064,12 +2062,16 @@ void FLevelEditorViewportClient::ProcessClick(FSceneView& View, HHitProxy* HitPr
 	static FName ProcessClickTrace = FName(TEXT("ProcessClickTrace"));
 
 	const FViewportClick Click(&View,this,Key,Event,HitX,HitY);
-	if (!GEditorModeTools().HandleClick(this, HitProxy,Click))
+	if (!ModeTools->HandleClick(this, HitProxy,Click))
 	{
 		if (HitProxy == NULL)
 		{
 			ClickHandlers::ClickBackdrop(this,Click);
 		}
+		else if (GUnrealEd->ComponentVisManager.HandleClick(this, HitProxy, Click))
+		{
+			// Component vis manager handled the click
+		}
 		else if (HitProxy->IsA(HActor::StaticGetType()))
 		{
 			ClickHandlers::ClickActor(this,((HActor*)HitProxy)->Actor,Click,true);
@@ -2193,7 +2195,7 @@ void FLevelEditorViewportClient::Tick(float DeltaTime)
 		GUnrealEd->UpdatePivotLocationForSelection();
 	}
 
-	GEditorModeTools().Tick(this,DeltaTime);
+	ModeTools->Tick(this, DeltaTime);
 
 	// Update the preview mesh for the preview mesh mode. 
 	GEditor->UpdatePreviewMesh();
@@ -2278,9 +2280,14 @@ namespace ViewportDeadZoneConstants
 
 bool FLevelEditorViewportClient::InputWidgetDelta( FViewport* Viewport, EAxisList::Type CurrentAxis, FVector& Drag, FRotator& Rot, FVector& Scale )
 {
+	if (GUnrealEd->ComponentVisManager.HandleInputDelta(this, Viewport, Drag, Rot, Scale))
+	{
+		return true;
+	}
+
 	bool bHandled = false;
 	// Give the current editor mode a chance to use the input first.  If it does, don't apply it to anything else.
-	FEditorModeTools& Tools = GEditorModeTools();
+	FEditorModeTools& Tools = GLevelEditorModeTools();
 	if( Tools.InputDelta( this, Viewport, Drag, Rot, Scale ) )
 	{
 		if( Tools.AllowWidgetMove() )
@@ -2393,14 +2400,14 @@ TSharedPtr<FDragTool> FLevelEditorViewportClient::MakeDragTool( EDragTool::Type
 void FLevelEditorViewportClient::UpdateMouseDelta()
 {
 	// Do nothing if a drag tool is being used.
-	if( MouseDeltaTracker->UsingDragTool() || GEditorModeTools().DisallowMouseDeltaTracking() )
+	if( MouseDeltaTracker->UsingDragTool() || GLevelEditorModeTools().DisallowMouseDeltaTracking() )
 	{
 		return;
 	}
 
 	// Stop tracking and do nothing else if we're tracking and the widget mode has changed mid-track.
 	// It can confuse the widget code that handles the mouse movements.
-	if (bIsTracking && MouseDeltaTracker->GetTrackingWidgetMode() != GEditorModeTools().GetWidgetMode())
+	if (bIsTracking && MouseDeltaTracker->GetTrackingWidgetMode() != GLevelEditorModeTools().GetWidgetMode())
 	{
 		StopTracking();
 		return;
@@ -2501,8 +2508,13 @@ bool FLevelEditorViewportClient::InputKey(FViewport* Viewport, int32 ControllerI
 		GEditor->ClickLocation = FVector((View->ViewMatrices.ViewMatrix * View->ViewMatrices.ProjMatrix).Inverse().TransformFVector4(FVector4((HitX - Viewport->GetSizeXY().X / 2.0f) / (Viewport->GetSizeXY().X / 2.0f),(HitY - Viewport->GetSizeXY().Y / 2.0f) / -(Viewport->GetSizeXY().Y / 2.0f),0.5f,1.0f)));
 	}
 
+	if (GUnrealEd->ComponentVisManager.HandleInputKey(this, Viewport, Key, Event))
+	{
+		return true;
+	}
+
 	// Let the current mode have a look at the input before reacting to it.
-	if( GEditorModeTools().InputKey(this, Viewport, Key, Event) )
+	if( GLevelEditorModeTools().InputKey(this, Viewport, Key, Event) )
 	{
 		return true;
 	}
@@ -2556,11 +2568,11 @@ bool FLevelEditorViewportClient::InputKey(FViewport* Viewport, int32 ControllerI
 		&& InputState.IsButtonPressed(EKeys::RightMouseButton) 
 		&& IsOrtho() )
 	{
-		GEditorModeTools().SetWidgetModeOverride( FWidget::WM_Rotate );
+		GLevelEditorModeTools().SetWidgetModeOverride( FWidget::WM_Rotate );
 	}
 	else
 	{
-		GEditorModeTools().SetWidgetModeOverride( FWidget::WM_None );
+		GLevelEditorModeTools().SetWidgetModeOverride( FWidget::WM_None );
 	}
 
 	bHandled |= FEditorViewportClient::InputKey(Viewport,ControllerId,Key,Event,AmountDepressed,bGamepad);
@@ -2578,7 +2590,7 @@ bool FLevelEditorViewportClient::InputKey(FViewport* Viewport, int32 ControllerI
 void FLevelEditorViewportClient::TrackingStarted( const FInputEventState& InInputState, bool bIsDraggingWidget, bool bNudge )
 {
 	// Begin transacting.  Give the current editor mode an opportunity to do the transacting.
-	const bool bTrackingHandledExternally = GEditorModeTools().StartTracking(this, Viewport);
+	const bool bTrackingHandledExternally = GLevelEditorModeTools().StartTracking(this, Viewport);
 
 	TrackingTransaction.End();
 
@@ -2718,7 +2730,7 @@ void FLevelEditorViewportClient::TrackingStopped()
 	bool bDidAnythingActuallyChange = false;
 
 	// Stop transacting.  Give the current editor mode an opportunity to do the transacting.
-	const bool bTransactingHandledByEditorMode = GEditorModeTools().EndTracking(this, Viewport);
+	const bool bTransactingHandledByEditorMode = GLevelEditorModeTools().EndTracking(this, Viewport);
 	if( !bTransactingHandledByEditorMode )
 	{
 		if( TrackingTransaction.TransCount > 0 )
@@ -2788,7 +2800,7 @@ void FLevelEditorViewportClient::TrackingStopped()
 	}
 
 	TArray<FEdMode*> ActiveModes; 
-	GEditorModeTools().GetActiveModes( ActiveModes );
+	GLevelEditorModeTools().GetActiveModes( ActiveModes );
 	for( int32 ModeIndex = 0; ModeIndex < ActiveModes.Num(); ++ModeIndex )
 	{
 		// Also notify the current editing modes if they are interested.
@@ -2943,7 +2955,7 @@ bool FLevelEditorViewportClient::InputAxis(FViewport* Viewport, int32 Controller
 
 	FScopedSetCurrentViewportClient( this );
 	// Let the current mode have a look at the input before reacting to it.
-	if ( GEditorModeTools().InputAxis(this, Viewport, ControllerId, Key, Delta, DeltaTime) )
+	if ( GLevelEditorModeTools().InputAxis(this, Viewport, ControllerId, Key, Delta, DeltaTime) )
 	{
 		return true;
 	}
@@ -3251,9 +3263,9 @@ bool FLevelEditorViewportClient::IsVolumeVisibleInViewport( const AActor& Volume
 
 void FLevelEditorViewportClient::SetWidgetMode( FWidget::EWidgetMode ActivatedMode )
 {
-	if( !GEditorModeTools().IsTracking() )
+	if( !GLevelEditorModeTools().IsTracking() )
 	{
-		GEditorModeTools().SetWidgetMode( ActivatedMode );
+		GLevelEditorModeTools().SetWidgetMode( ActivatedMode );
 
 		// force an invalidation (non-deferred) of the hit proxy here, otherwise we will
 		// end up checking against an incorrect hit proxy if the cursor is not moved
@@ -3261,7 +3273,7 @@ void FLevelEditorViewportClient::SetWidgetMode( FWidget::EWidgetMode ActivatedMo
 		bShouldCheckHitProxy = true;
 
 		// Fire event delegate
-		GEditorModeTools().BroadcastWidgetModeChanged(ActivatedMode);
+		GLevelEditorModeTools().BroadcastWidgetModeChanged(ActivatedMode);
 	}
 	
 	// Invalidate all viewports, so the new gizmo is rendered in each one
@@ -3276,33 +3288,39 @@ void FLevelEditorViewportClient::SetWidgetMode( FWidget::EWidgetMode ActivatedMo
 
 bool FLevelEditorViewportClient::CanSetWidgetMode( FWidget::EWidgetMode NewMode ) const
 {
-	return GEditorModeTools().GetShowWidget() == true;
+	return GLevelEditorModeTools().GetShowWidget() == true;
 }
 
 void FLevelEditorViewportClient::SetWidgetCoordSystemSpace( ECoordSystem NewCoordSystem )
 {
-	GEditorModeTools().SetCoordSystem(NewCoordSystem);
+	GLevelEditorModeTools().SetCoordSystem(NewCoordSystem);
 	Invalidate();
 }
 
 FWidget::EWidgetMode FLevelEditorViewportClient::GetWidgetMode() const
 {
-	return GEditorModeTools().GetWidgetMode();
+	return GLevelEditorModeTools().GetWidgetMode();
 }
 
 FVector FLevelEditorViewportClient::GetWidgetLocation() const
 {
-	return GEditorModeTools().GetWidgetLocation();
+	FVector ComponentVisWidgetLocation;
+	if (GUnrealEd->ComponentVisManager.GetWidgetLocation(ComponentVisWidgetLocation))
+	{
+		return ComponentVisWidgetLocation;
+	}
+
+	return GLevelEditorModeTools().GetWidgetLocation();
 }
 
 FMatrix FLevelEditorViewportClient::GetWidgetCoordSystem() const 
 {
-	return GEditorModeTools().GetCustomInputCoordinateSystem();
+	return GLevelEditorModeTools().GetCustomInputCoordinateSystem();
 }
 
 ECoordSystem FLevelEditorViewportClient::GetWidgetCoordSystemSpace() const
 {
-	return GEditorModeTools().GetCoordSystem();
+	return GLevelEditorModeTools().GetCoordSystem();
 }
 
 
@@ -3703,21 +3721,21 @@ void FLevelEditorViewportClient::ApplyDeltaToRotateWidget( const FRotator& InRot
 	//apply rotation to translate rotate widget
 	if (!InRot.IsZero())
 	{
-		FRotator TranslateRotateWidgetRotation(0, GEditorModeTools().TranslateRotateXAxisAngle, 0);
+		FRotator TranslateRotateWidgetRotation(0, GLevelEditorModeTools().TranslateRotateXAxisAngle, 0);
 		TranslateRotateWidgetRotation += InRot;
-		GEditorModeTools().TranslateRotateXAxisAngle = TranslateRotateWidgetRotation.Yaw;
+		GLevelEditorModeTools().TranslateRotateXAxisAngle = TranslateRotateWidgetRotation.Yaw;
 	}
 
 }
 
 void FLevelEditorViewportClient::MouseEnter( FViewport* Viewport,int32 x, int32 y )
 {
-	GEditorModeTools().MouseEnter(this, Viewport, x, y);
+	GLevelEditorModeTools().MouseEnter(this, Viewport, x, y);
 }
 
 void FLevelEditorViewportClient::MouseLeave( FViewport* Viewport ) 
 {
-	GEditorModeTools().MouseLeave(this, Viewport);
+	GLevelEditorModeTools().MouseLeave(this, Viewport);
 
 	FEditorViewportClient::MouseLeave(Viewport);
 }
@@ -3727,7 +3745,7 @@ void FLevelEditorViewportClient::MouseMove(FViewport* Viewport,int32 x, int32 y)
 	FEditorViewportClient::MouseMove(Viewport,x,y);
 
 	// Let the current editor mode know about the mouse movement.
-	if( IsLevelEditorClient() && GEditorModeTools().MouseMove(this, Viewport, x, y) )
+	if( IsLevelEditorClient() && GLevelEditorModeTools().MouseMove(this, Viewport, x, y) )
 	{
 		return;
 	}
@@ -3754,11 +3772,6 @@ EMouseCursor::Type FLevelEditorViewportClient::GetCursor(FViewport* Viewport,int
 
 }
 
-void FLevelEditorViewportClient::RenderDragTool( const FSceneView* View,FCanvas* Canvas )
-{
-	MouseDeltaTracker->RenderDragTool(View,Canvas);
-}
-
 FViewportCursorLocation FLevelEditorViewportClient::GetCursorWorldLocationFromMousePos()
 {
 	// Create the scene view context
@@ -3796,7 +3809,7 @@ void FLevelEditorViewportClient::CapturedMouseMove( FViewport* InViewport, int32
 	TrackingTransaction.PromotePendingToActive();
 
 	// Let the current editor mode know about the mouse movement.
-	if( IsLevelEditorClient() && GEditorModeTools().CapturedMouseMove( this, InViewport, InMouseX, InMouseY ) )
+	if( IsLevelEditorClient() && GLevelEditorModeTools().CapturedMouseMove( this, InViewport, InMouseX, InMouseY ) )
 	{
 		return;
 	}
@@ -3954,7 +3967,7 @@ bool FLevelEditorViewportClient::GetActiveSafeFrame(float& OutAspectRatio) const
 void FLevelEditorViewportClient::SetCurrentWidgetAxis( EAxisList::Type NewAxis )
 {
 	FEditorViewportClient::SetCurrentWidgetAxis( NewAxis );
-	GEditorModeTools().SetCurrentWidgetAxis( NewAxis );
+	GLevelEditorModeTools().SetCurrentWidgetAxis( NewAxis );
 }
 
 /** 
@@ -4036,10 +4049,10 @@ void FLevelEditorViewportClient::Draw(const FSceneView* View,FPrimitiveDrawInter
 	}
 
 
-	GEditorModeTools().DrawActiveModes(View, PDI);
+	GLevelEditorModeTools().DrawActiveModes(View, PDI);
 
 	// Draw the current editor mode.
-	GEditorModeTools().Render( View, Viewport, PDI );
+	GLevelEditorModeTools().Render( View, Viewport, PDI );
 
 	// Determine if a view frustum should be rendered in the viewport.
 	// The frustum should definitely be rendered if the viewport has a view parent.
@@ -4201,7 +4214,7 @@ void FLevelEditorViewportClient::SetupViewForRendering( FSceneViewFamily& ViewFa
 		}
 	}
 
-	if(GEditorModeTools().GetActiveMode(FBuiltinEditorModes::EM_InterpEdit) == 0 || !AllowMatineePreview())
+	if(GLevelEditorModeTools().GetActiveMode(FBuiltinEditorModes::EM_InterpEdit) == 0 || !AllowMatineePreview())
 	{
 		// in the editor, disable camera motion blur and other rendering features that rely on the former frame
 		// unless the view port is Matinee controlled
@@ -4228,9 +4241,9 @@ void FLevelEditorViewportClient::SetupViewForRendering( FSceneViewFamily& ViewFa
 void FLevelEditorViewportClient::DrawCanvas( FViewport& InViewport, FSceneView& View, FCanvas& Canvas )
 {
 	// Information string
-	Canvas.DrawShadowedString( 4,4, *GEditorModeTools().InfoString, GEngine->GetSmallFont(), FColor(255,255,255) );
+	Canvas.DrawShadowedString( 4,4, *GLevelEditorModeTools().InfoString, GEngine->GetSmallFont(), FColor(255,255,255) );
 
-	GEditorModeTools().DrawHUD(this,&InViewport,&View,&Canvas);
+	GLevelEditorModeTools().DrawHUD(this,&InViewport,&View,&Canvas);
 
 	// Testbed
 	FCanvasItemTestbed TestBed;
@@ -4239,48 +4252,6 @@ void FLevelEditorViewportClient::DrawCanvas( FViewport& InViewport, FSceneView&
 	DrawStaticLightingDebugInfo(&View, &Canvas);
 }
 
-
-
-/**
- * Draws a screen space bounding box around the specified actor
- *
- * @param	InCanvas		Canvas to draw on
- * @param	InView			View to render
- * @param	InViewport		Viewport we're rendering into
- * @param	InActor			Actor to draw a bounding box for
- * @param	InColor			Color of bounding box
- * @param	bInDrawBracket	True to draw a bracket, otherwise a box will be rendered
- * @param	bInLabelText	Optional label text to draw
- */
-void FLevelEditorViewportClient::DrawActorScreenSpaceBoundingBox( FCanvas* InCanvas, const FSceneView* InView, FViewport* InViewport, AActor* InActor, const FLinearColor& InColor, const bool bInDrawBracket, const FString& InLabelText )
-{
-	check( InActor != NULL );
-
-
-	// First check to see if we're dealing with a sprite, otherwise just use the normal bounding box
-	UBillboardComponent* Sprite = InActor->FindComponentByClass<UBillboardComponent>();
-
-	FBox ActorBox;
-	if( Sprite != NULL )
-	{
-		ActorBox = Sprite->Bounds.GetBox();
-	}
-	else
-	{
-		const bool bNonColliding = true;
-		ActorBox = InActor->GetComponentsBoundingBox( bNonColliding );
-	}
-
-
-	// If we didn't get a valid bounding box, just make a little one around the actor location
-	if( !ActorBox.IsValid || ActorBox.GetExtent().GetMin() < KINDA_SMALL_NUMBER )
-	{
-		ActorBox = FBox( InActor->GetActorLocation() - FVector( -20 ), InActor->GetActorLocation() + FVector( 20 ) );
-	}
-
-	DrawBoundingBox(ActorBox, InCanvas, InView, InViewport, InColor, bInDrawBracket, InLabelText);
-}
-
 /**
  *	Draw the texture streaming bounds.
  */
@@ -4472,13 +4443,13 @@ void FLevelEditorViewportClient::SetCameraSpeedSetting(int32 SpeedSetting)
 void FLevelEditorViewportClient::ReceivedFocus(FViewport* Viewport)
 {
 	FEditorViewportClient::ReceivedFocus(Viewport);
-	GEditorModeTools().ReceivedFocus(this, Viewport);
+	GLevelEditorModeTools().ReceivedFocus(this, Viewport);
 }
 
 void FLevelEditorViewportClient::LostFocus(FViewport* Viewport)
 {
 	FEditorViewportClient::LostFocus(Viewport);
-	GEditorModeTools().LostFocus(this, Viewport);
+	GLevelEditorModeTools().LostFocus(this, Viewport);
 }
 
 bool FLevelEditorViewportClient::OverrideHighResScreenshotCaptureRegion(FIntRect& OutCaptureRegion)
@@ -4643,14 +4614,14 @@ void FLevelEditorViewportClient::SetIsSimulateInEditorViewport( bool bInIsSimula
 { 
 	bIsSimulateInEditorViewport = bInIsSimulateInEditorViewport; 
 
-	static TSharedPtr<FPhysicsManipulationEdMode> Mode = MakeShareable(new FPhysicsManipulationEdMode());
 	if (bInIsSimulateInEditorViewport)
 	{
-		GEditorModeTools().RegisterMode(Mode.ToSharedRef());
+		TSharedRef<FPhysicsManipulationEdModeFactory> Factory = MakeShareable( new FPhysicsManipulationEdModeFactory );
+		FEditorModeRegistry::Get().RegisterMode( FBuiltinEditorModes::EM_Physics, Factory );
 	}
 	else
 	{
-		GEditorModeTools().UnregisterMode(Mode.ToSharedRef());
+		FEditorModeRegistry::Get().UnregisterMode( FBuiltinEditorModes::EM_Physics );
 	}
 }
 
diff --git a/Engine/Source/Editor/UnrealEd/Private/LevelViewportClickHandlers.cpp b/Engine/Source/Editor/UnrealEd/Private/LevelViewportClickHandlers.cpp
--- a/Engine/Source/Editor/UnrealEd/Private/LevelViewportClickHandlers.cpp
+++ b/Engine/Source/Editor/UnrealEd/Private/LevelViewportClickHandlers.cpp
@@ -328,7 +328,7 @@ namespace ClickHandlers
 			const FVector Delta = Snapped - World;
 			GEditor->SetPivot( Snapped, false, false );
 
-			if( GEditorModeTools().IsModeActive( FBuiltinEditorModes::EM_Default ))
+			if( GLevelEditorModeTools().IsDefaultModeActive() )
 			{		
 				// All selected actors need to move by the delta.
 				for ( FSelectionIterator It( GEditor->GetSelectedActorIterator() ) ; It ; ++It )
@@ -400,7 +400,7 @@ namespace ClickHandlers
 		else if( Click.GetKey() == EKeys::LeftMouseButton )
 		{
 			// This should only happen in geometry mode
-			FEdMode* Mode = GEditorModeTools().GetActiveMode( FBuiltinEditorModes::EM_Geometry );
+			FEdMode* Mode = GLevelEditorModeTools().GetActiveMode( FBuiltinEditorModes::EM_Geometry );
 			if( Mode )
 			{
 				if( ( InHitProxy != NULL) && ( InHitProxy->GeomObject != NULL ) && (InHitProxy->GeomObject->PolyPool.IsValidIndex( InHitProxy->PolyIndex ) == true ) )					
@@ -484,7 +484,7 @@ namespace ClickHandlers
 		}
 		else if( Click.GetKey() == EKeys::LeftMouseButton )
 		{
-			FEdMode* Mode = GEditorModeTools().GetActiveMode( FBuiltinEditorModes::EM_Geometry );
+			FEdMode* Mode = GLevelEditorModeTools().GetActiveMode( FBuiltinEditorModes::EM_Geometry );
 
 			if( Mode )
 			{ 
@@ -532,12 +532,12 @@ namespace ClickHandlers
 
 	bool ClickGeomVertex(FLevelEditorViewportClient* ViewportClient,HGeomVertexProxy* InHitProxy,const FViewportClick& Click)
 	{
-		if( !GEditorModeTools().IsModeActive( FBuiltinEditorModes::EM_Geometry ) )
+		if( !GLevelEditorModeTools().IsModeActive( FBuiltinEditorModes::EM_Geometry ) )
 		{
 			return false;
 		}
 
-		FEdModeGeometry* Mode = static_cast<FEdModeGeometry*>( GEditorModeTools().GetActiveMode(FBuiltinEditorModes::EM_Geometry) );
+		FEdModeGeometry* Mode = static_cast<FEdModeGeometry*>( GLevelEditorModeTools().GetActiveMode(FBuiltinEditorModes::EM_Geometry) );
 
 		// Note: The expected behavior is that right clicking on a vertex will snap the vertex that was
 		// right-clicked on to the nearest grid point, then move all SELECTED verts by the appropriate
diff --git a/Engine/Source/Editor/UnrealEd/Private/PackageAutoSaver.cpp b/Engine/Source/Editor/UnrealEd/Private/PackageAutoSaver.cpp
--- a/Engine/Source/Editor/UnrealEd/Private/PackageAutoSaver.cpp
+++ b/Engine/Source/Editor/UnrealEd/Private/PackageAutoSaver.cpp
@@ -304,7 +304,7 @@ bool FPackageAutoSaver::CanAutoSave() const
 
 	const bool bAutosaveEnabled	= GetDefault<UEditorLoadingSavingSettings>()->bAutoSaveEnable && bPackagesNeedAutoSave;
 	const bool bSlowTask = GIsSlowTask;
-	const bool bInterpEditMode = GEditorModeTools().IsModeActive(FBuiltinEditorModes::EM_InterpEdit);
+	const bool bInterpEditMode = GLevelEditorModeTools().IsModeActive(FBuiltinEditorModes::EM_InterpEdit);
 	const bool bPlayWorldValid = GUnrealEd->PlayWorld != nullptr;
 	const bool bAnyMenusVisible	= FSlateApplication::Get().AnyMenusVisible();
 	const bool bAutomationTesting = GIsAutomationTesting;
@@ -368,7 +368,7 @@ void FPackageAutoSaver::UpdateAutoSaveNotification()
 	if (UserAllowsAutosave && // The user has set to allow auto-save in preferences
 		TimeInSecondsUntilAutosave < LoadingSavingSettings->AutoSaveWarningInSeconds && 
 		!InGame && // we want to hide auto-save if we are simulating/playing
-		!GEditorModeTools().IsModeActive(FBuiltinEditorModes::EM_InterpEdit) // we want to hide auto-save if we are in matinee 
+		!GLevelEditorModeTools().IsModeActive(FBuiltinEditorModes::EM_InterpEdit) // we want to hide auto-save if we are in matinee 
 		)
 	{		
 		if (!bAutoSaveNotificationLaunched && !bDelayingDueToFailedSave)
diff --git a/Engine/Source/Editor/UnrealEd/Private/PhysicsManipulationMode.cpp b/Engine/Source/Editor/UnrealEd/Private/PhysicsManipulationMode.cpp
--- a/Engine/Source/Editor/UnrealEd/Private/PhysicsManipulationMode.cpp
+++ b/Engine/Source/Editor/UnrealEd/Private/PhysicsManipulationMode.cpp
@@ -5,30 +5,7 @@
 
 DEFINE_LOG_CATEGORY_STATIC(LogEditorPhysMode, Log, All);
 
-FPhysicsManipulationEdMode::FPhysicsManipulationEdMode()
-{
-	ID = FBuiltinEditorModes::EM_Physics;
-	Name = NSLOCTEXT("EditorModes", "PhysicsMode", "Physics Mode");
-
-	HandleComp = ConstructObject<UPhysicsHandleComponent>(UPhysicsHandleComponent::StaticClass(), GetTransientPackage(), NAME_None, RF_NoFlags);
-}
-
-FPhysicsManipulationEdMode::~FPhysicsManipulationEdMode()
-{
-	HandleComp = NULL;
-}
-
-void FPhysicsManipulationEdMode::Enter()
-{
-	HandleComp->RegisterComponentWithWorld(GetWorld());
-}
-
-void FPhysicsManipulationEdMode::Exit()
-{
-	HandleComp->UnregisterComponent();
-}
-
-void FPhysicsManipulationEdMode::PeekAtSelectionChangedEvent( UObject* ItemUndergoingChange )
+void FPhysicsManipulationEdModeFactory::OnSelectionChanged(FEditorModeTools& Tools, UObject* ItemUndergoingChange) const
 {
 	USelection* Selection = GEditor->GetSelectedActors();
 
@@ -40,18 +17,48 @@ void FPhysicsManipulationEdMode::PeekAtSelectionChangedEvent( UObject* ItemUnder
 			UPrimitiveComponent* PC = SelectedActor->GetRootPrimitiveComponent();
 			if (PC != NULL && PC->BodyInstance.bSimulatePhysics)
 			{
-				GEditorModeTools().ActivateMode(ID, false);
+				Tools.ActivateMode(FBuiltinEditorModes::EM_Physics);
 				return;
 			}
 		}
 	}
 	else if (ItemUndergoingChange != NULL && !ItemUndergoingChange->IsA(USelection::StaticClass()))
 	{
-		GEditorModeTools().DeactivateMode(ID);
+		Tools.DeactivateMode(FBuiltinEditorModes::EM_Physics);
 	}
 }
 
-bool FPhysicsManipulationEdMode::InputDelta( FLevelEditorViewportClient* InViewportClient,FViewport* InViewport,FVector& InDrag,FRotator& InRot,FVector& InScale )
+FEditorModeInfo FPhysicsManipulationEdModeFactory::GetModeInfo() const
+{
+	return FEditorModeInfo(FBuiltinEditorModes::EM_Physics, NSLOCTEXT("EditorModes", "PhysicsMode", "Physics Mode"));
+}
+
+TSharedRef<FEdMode> FPhysicsManipulationEdModeFactory::CreateMode() const
+{
+	return MakeShareable( new FPhysicsManipulationEdMode );
+}
+
+FPhysicsManipulationEdMode::FPhysicsManipulationEdMode()
+{
+	HandleComp = ConstructObject<UPhysicsHandleComponent>(UPhysicsHandleComponent::StaticClass(), GetTransientPackage(), NAME_None, RF_NoFlags);
+}
+
+FPhysicsManipulationEdMode::~FPhysicsManipulationEdMode()
+{
+	HandleComp = NULL;
+}
+
+void FPhysicsManipulationEdMode::Enter()
+{
+	HandleComp->RegisterComponentWithWorld(GetWorld());
+}
+
+void FPhysicsManipulationEdMode::Exit()
+{
+	HandleComp->UnregisterComponent();
+}
+
+bool FPhysicsManipulationEdMode::InputDelta( FEditorViewportClient* InViewportClient,FViewport* InViewport,FVector& InDrag,FRotator& InRot,FVector& InScale )
 {
 	UE_LOG(LogEditorPhysMode, Warning, TEXT("Mouse: %s InDrag: %s  InRot: %s"), *GEditor->MouseMovement.ToString(), *InDrag.ToString(), *InRot.ToString());
 
@@ -74,7 +81,7 @@ bool FPhysicsManipulationEdMode::InputDelta( FLevelEditorViewportClient* InViewp
 	}
 }
 
-bool FPhysicsManipulationEdMode::StartTracking( FLevelEditorViewportClient* InViewportClient, FViewport* InViewport )
+bool FPhysicsManipulationEdMode::StartTracking( FEditorViewportClient* InViewportClient, FViewport* InViewport )
 {
 	UE_LOG(LogEditorPhysMode, Warning, TEXT("Start Tracking"));
 
@@ -112,7 +119,7 @@ bool FPhysicsManipulationEdMode::StartTracking( FLevelEditorViewportClient* InVi
 	return FEdMode::StartTracking(InViewportClient, InViewport);
 }
 
-bool FPhysicsManipulationEdMode::EndTracking( FLevelEditorViewportClient* InViewportClient, FViewport* InViewport )
+bool FPhysicsManipulationEdMode::EndTracking( FEditorViewportClient* InViewportClient, FViewport* InViewport )
 {
 	UE_LOG(LogEditorPhysMode, Warning, TEXT("End Tracking"));
 
diff --git a/Engine/Source/Editor/UnrealEd/Private/PlayLevel.cpp b/Engine/Source/Editor/UnrealEd/Private/PlayLevel.cpp
--- a/Engine/Source/Editor/UnrealEd/Private/PlayLevel.cpp
+++ b/Engine/Source/Editor/UnrealEd/Private/PlayLevel.cpp
@@ -37,10 +37,10 @@ DEFINE_LOG_CATEGORY_STATIC(LogPlayLevel, Log, All);
 void UEditorEngine::EndPlayMap()
 {
 	// Matinee must be closed before PIE can stop - matinee during PIE will be editing a PIE-world actor
-	if( GEditorModeTools().IsModeActive(FBuiltinEditorModes::EM_InterpEdit) )
+	if( GLevelEditorModeTools().IsModeActive(FBuiltinEditorModes::EM_InterpEdit) )
 	{
 		FMessageDialog::Open( EAppMsgType::Ok, NSLOCTEXT("UnrealEd", "PIENeedsToCloseMatineeMessage", "Closing 'Play in Editor' must close UnrealMatinee.") );
-		GEditorModeTools().DeactivateMode( FBuiltinEditorModes::EM_InterpEdit );
+		GLevelEditorModeTools().DeactivateMode( FBuiltinEditorModes::EM_InterpEdit );
 	}
 
 	EndPlayOnLocalPc();
@@ -1556,14 +1556,14 @@ void UEditorEngine::PlayInEditor( UWorld* InWorld, bool bInSimulateInEditor )
 	double PIEStartTime = FPlatformTime::Seconds();
 
 	// Prompt the user that Matinee must be closed before PIE can occur.
-	if( GEditorModeTools().IsModeActive(FBuiltinEditorModes::EM_InterpEdit) )
+	if( GLevelEditorModeTools().IsModeActive(FBuiltinEditorModes::EM_InterpEdit) )
 	{
 		const bool bContinuePIE = EAppReturnType::Yes == FMessageDialog::Open( EAppMsgType::YesNo, NSLOCTEXT("UnrealEd", "PIENeedsToCloseMatineeQ", "'Play in Editor' must close UnrealMatinee.  Continue?") );
 		if ( !bContinuePIE )
 		{
 			return;
 		}
-		GEditorModeTools().DeactivateMode( FBuiltinEditorModes::EM_InterpEdit );
+		GLevelEditorModeTools().DeactivateMode( FBuiltinEditorModes::EM_InterpEdit );
 	}
 
 	FBlueprintEditorUtils::FindAndSetDebuggableBlueprintInstances();
diff --git a/Engine/Source/Editor/UnrealEd/Private/SEditorViewport.cpp b/Engine/Source/Editor/UnrealEd/Private/SEditorViewport.cpp
--- a/Engine/Source/Editor/UnrealEd/Private/SEditorViewport.cpp
+++ b/Engine/Source/Editor/UnrealEd/Private/SEditorViewport.cpp
@@ -305,7 +305,7 @@ void SEditorViewport::BindCommands()
 
 EVisibility SEditorViewport::OnGetViewportContentVisibility() const
 {
-	return GEditorModeTools().IsViewportUIHidden() ? EVisibility::Collapsed : EVisibility::Visible;
+	return GLevelEditorModeTools().IsViewportUIHidden() ? EVisibility::Collapsed : EVisibility::Visible;
 }
 
 void SEditorViewport::OnToggleRealtime()
diff --git a/Engine/Source/Editor/UnrealEd/Private/Settings/SettingsClasses.cpp b/Engine/Source/Editor/UnrealEd/Private/Settings/SettingsClasses.cpp
--- a/Engine/Source/Editor/UnrealEd/Private/Settings/SettingsClasses.cpp
+++ b/Engine/Source/Editor/UnrealEd/Private/Settings/SettingsClasses.cpp
@@ -197,11 +197,11 @@ void ULevelEditorViewportSettings::PostEditChangeProperty( struct FPropertyChang
 	{
 		if (bAllowTranslateRotateZWidget)
 		{
-			GEditorModeTools().SetWidgetMode(FWidget::WM_TranslateRotateZ);
+			GLevelEditorModeTools().SetWidgetMode(FWidget::WM_TranslateRotateZ);
 		}
-		else if (GEditorModeTools().GetWidgetMode() == FWidget::WM_TranslateRotateZ)
+		else if (GLevelEditorModeTools().GetWidgetMode() == FWidget::WM_TranslateRotateZ)
 		{
-			GEditorModeTools().SetWidgetMode(FWidget::WM_Translate);
+			GLevelEditorModeTools().SetWidgetMode(FWidget::WM_Translate);
 		}
 	}
 	else if (Name == FName(TEXT("bHighlightWithBrackets")))
diff --git a/Engine/Source/Editor/UnrealEd/Private/SnappingUtils.cpp b/Engine/Source/Editor/UnrealEd/Private/SnappingUtils.cpp
--- a/Engine/Source/Editor/UnrealEd/Private/SnappingUtils.cpp
+++ b/Engine/Source/Editor/UnrealEd/Private/SnappingUtils.cpp
@@ -105,7 +105,7 @@ bool FEditorViewportSnapping::IsSnapRotationEnabled()
 	// Ask Current Editor Mode if Rotation Snap is enabled
 	bool bSnapEnabled = false;
 	TArray<FEdMode*> ActiveModes; 
-	GEditorModeTools().GetActiveModes( ActiveModes );
+	GLevelEditorModeTools().GetActiveModes( ActiveModes );
 	for( int32 ModeIndex = 0; ModeIndex < ActiveModes.Num(); ++ModeIndex )
 	{
 		bSnapEnabled |= ActiveModes[ModeIndex]->IsSnapRotationEnabled();
@@ -186,7 +186,7 @@ void FEditorViewportSnapping::SetActorSnapDistance(float Distance)
 
 bool FEditorViewportSnapping::SnapActorsToNearestActor( FVector& Drag, FLevelEditorViewportClient* ViewportClient )
 {
-	FEditorModeTools& Tools = GEditorModeTools();
+	FEditorModeTools& Tools = GLevelEditorModeTools();
 
 	// Does the user have actor snapping enabled?
 	bool bSnapped = false;
@@ -288,7 +288,7 @@ void FEditorViewportSnapping::SnapRotatorToGrid(FRotator& Rotation)
 	if( IsSnapRotationEnabled() )
 	{		
 		TArray<FEdMode*> ActiveModes; 
-		GEditorModeTools().GetActiveModes( ActiveModes );
+		GLevelEditorModeTools().GetActiveModes( ActiveModes );
 		for( int32 ModeIndex = 0; ModeIndex < ActiveModes.Num(); ++ModeIndex )
 		{
 			if( ActiveModes[ModeIndex]->SnapRotatorToGridOverride( Rotation ) == true )
diff --git a/Engine/Source/Editor/UnrealEd/Private/StaticLightingSystem/StaticLightingSystem.cpp b/Engine/Source/Editor/UnrealEd/Private/StaticLightingSystem/StaticLightingSystem.cpp
--- a/Engine/Source/Editor/UnrealEd/Private/StaticLightingSystem/StaticLightingSystem.cpp
+++ b/Engine/Source/Editor/UnrealEd/Private/StaticLightingSystem/StaticLightingSystem.cpp
@@ -2048,7 +2048,7 @@ bool FStaticLightingSystem::CanAutoApplyLighting() const
 {
 	const bool bAutoApplyEnabled = GetDefault<ULevelEditorMiscSettings>()->bAutoApplyLightingEnable;
 	const bool bSlowTask = GIsSlowTask;
-	const bool bInterpEditMode = GEditorModeTools().IsModeActive( FBuiltinEditorModes::EM_InterpEdit );
+	const bool bInterpEditMode = GLevelEditorModeTools().IsModeActive( FBuiltinEditorModes::EM_InterpEdit );
 	const bool bPlayWorldValid = GUnrealEd->PlayWorld != nullptr;
 	const bool bAnyMenusVisible = FSlateApplication::Get().AnyMenusVisible();
 	const bool bAutomationTesting = GIsAutomationTesting;
diff --git a/Engine/Source/Editor/UnrealEd/Private/Tests/EditorAutomationTests.cpp b/Engine/Source/Editor/UnrealEd/Private/Tests/EditorAutomationTests.cpp
--- a/Engine/Source/Editor/UnrealEd/Private/Tests/EditorAutomationTests.cpp
+++ b/Engine/Source/Editor/UnrealEd/Private/Tests/EditorAutomationTests.cpp
@@ -1150,7 +1150,7 @@ bool FBSPValidation::RunTest(const FString& Parameters)
 	DirectionalLight->SetBrightness(3.142f);
 	DirectionalLight->SetLightColor(FColor::White);
 
-	GEditorModeTools().MapChangeNotify();
+	GLevelEditorModeTools().MapChangeNotify();
 	
 	return true;
 }
@@ -1224,7 +1224,7 @@ bool FStaticMeshValidation::RunTest(const FString& Parameters)
 	DirectionalLight->SetBrightness(3.142f);
 	DirectionalLight->SetLightColor(FColor::White);
 
-	GEditorModeTools().MapChangeNotify();
+	GLevelEditorModeTools().MapChangeNotify();
 
 	return true;
 }
@@ -1372,7 +1372,7 @@ bool FConvertToValidation::RunTest(const FString& Parameters)
 	//Wait to give the screenshot capture some time to complete.
 	ADD_LATENT_AUTOMATION_COMMAND(FWaitLatentCommand(0.1f));
 
-	GEditorModeTools().MapChangeNotify();
+	GLevelEditorModeTools().MapChangeNotify();
 
 	return true;
 }
diff --git a/Engine/Source/Editor/UnrealEd/Private/UnrealEd.cpp b/Engine/Source/Editor/UnrealEd/Private/UnrealEd.cpp
--- a/Engine/Source/Editor/UnrealEd/Private/UnrealEd.cpp
+++ b/Engine/Source/Editor/UnrealEd/Private/UnrealEd.cpp
@@ -27,9 +27,9 @@ UUnrealEdEngine* GUnrealEd;
 DEFINE_LOG_CATEGORY_STATIC(LogUnrealEd, Log, All);
 
 /**
- * Provides access to the FEditorModeTools singleton.
+ * Provides access to the FEditorModeTools for the level editor
  */
-class FEditorModeTools& GEditorModeTools()
+class FEditorModeTools& GLevelEditorModeTools()
 {
 	static FEditorModeTools* EditorModeToolsSingleton = new FEditorModeTools;
 	return *EditorModeToolsSingleton;
@@ -138,15 +138,15 @@ int32 EditorInit( IEngineLoop& EngineLoop )
 
 void EditorExit()
 {
-	if( GEditorModeTools().IsModeActive(FBuiltinEditorModes::EM_MeshPaint) ||
-		GEditorModeTools().IsModeActive(FBuiltinEditorModes::EM_InterpEdit) )
+	if( GLevelEditorModeTools().IsModeActive(FBuiltinEditorModes::EM_MeshPaint) ||
+		GLevelEditorModeTools().IsModeActive(FBuiltinEditorModes::EM_InterpEdit) )
 	{
-		GEditorModeTools().ActivateMode( FBuiltinEditorModes::EM_Default );
+		GLevelEditorModeTools().ActivateDefaultMode();
 	}
 
 	// Save out any config settings for the editor so they don't get lost
 	GEditor->SaveConfig();
-	GEditorModeTools().SaveConfig();
+	GLevelEditorModeTools().SaveConfig();
 
 	// Clean up the actor folders singleton
 	FActorFolders::Cleanup();
diff --git a/Engine/Source/Editor/UnrealEd/Private/UnrealEdEngine.cpp b/Engine/Source/Editor/UnrealEd/Private/UnrealEdEngine.cpp
--- a/Engine/Source/Editor/UnrealEd/Private/UnrealEdEngine.cpp
+++ b/Engine/Source/Editor/UnrealEd/Private/UnrealEdEngine.cpp
@@ -187,7 +187,7 @@ void UUnrealEdEngine::MakeSortedSpriteInfo(TArray<FSpriteCategoryInfo>& OutSorte
 void UUnrealEdEngine::PreExit()
 {
 	// Notify edit modes we're mode at exit
-	GEditorModeTools().Shutdown();
+	FEditorModeRegistry::Get().Shutdown();
 
 	FAVIWriter* AVIWriter = FAVIWriter::GetInstance();
 	if (AVIWriter)
@@ -438,7 +438,7 @@ void UUnrealEdEngine::OnPostWindowsMessage(FViewport* Viewport, uint32 Message)
 void UUnrealEdEngine::OnOpenMatinee()
 {
 	// Register a delegate to pickup when Matinee is closed.
-	GEditorModeTools().OnEditorModeChanged().AddUObject( this, &UUnrealEdEngine::OnMatineeEditorClosed );
+	GLevelEditorModeTools().OnEditorModeChanged().AddUObject( this, &UUnrealEdEngine::OnMatineeEditorClosed );
 }
 
 
@@ -1144,6 +1144,6 @@ void UUnrealEdEngine::OnMatineeEditorClosed( FEdMode* Mode, bool IsEntering )
 		}
 
 		// Remove this delegate. 
-		GEditorModeTools().OnEditorModeChanged().RemoveUObject( this, &UUnrealEdEngine::OnMatineeEditorClosed );
+		GLevelEditorModeTools().OnEditorModeChanged().RemoveUObject( this, &UUnrealEdEngine::OnMatineeEditorClosed );
 	}	
 }
diff --git a/Engine/Source/Editor/UnrealEd/Private/UnrealEdMisc.cpp b/Engine/Source/Editor/UnrealEd/Private/UnrealEdMisc.cpp
--- a/Engine/Source/Editor/UnrealEd/Private/UnrealEdMisc.cpp
+++ b/Engine/Source/Editor/UnrealEd/Private/UnrealEdMisc.cpp
@@ -69,8 +69,7 @@ void FUnrealEdMisc::OnInit()
 		return;
 	}
 	bInitialized = true;
-
-
+	
 	// Register all callback notifications
 	FEditorDelegates::SelectedProps.AddRaw(this, &FUnrealEdMisc::CB_SelectedProps);
 	FEditorDelegates::DisplayLoadErrors.AddRaw(this, &FUnrealEdMisc::CB_DisplayLoadErrors);
@@ -112,9 +111,8 @@ void FUnrealEdMisc::OnInit()
 	// Register navigation commands for all viewports
 	FViewportNavigationCommands::Register();
 
-	// Init the editor mode tools, and set default editor mode initially
-	GEditorModeTools().Init();
-	GEditorModeTools().ActivateMode( FBuiltinEditorModes::EM_Default );
+	FEditorModeRegistry::Initialize();
+	GLevelEditorModeTools().ActivateDefaultMode();
 
 	// Are we in immersive mode?
 	const TCHAR* ParsedCmdLine = FCommandLine::Get();
@@ -668,7 +666,7 @@ void FUnrealEdMisc::CB_MapChange( uint32 InFlags )
 
 	GEditor->EditorUpdateComponents();
 
-	GEditorModeTools().MapChangeNotify();
+	GLevelEditorModeTools().MapChangeNotify();
 
 	/*if ((InFlags&MapChangeEventFlags::MapRebuild) != 0)
 	{
@@ -709,7 +707,7 @@ void FUnrealEdMisc::CB_LevelActorsAdded(AActor* InActor)
 
 void FUnrealEdMisc::CB_Undo()
 {
-	GEditorModeTools().PostUndo();
+	GLevelEditorModeTools().PostUndo();
 }
 
 void FUnrealEdMisc::CB_PreAutomationTesting()
@@ -729,7 +727,7 @@ void FUnrealEdMisc::CB_PostAutomationTesting()
 
 void FUnrealEdMisc::OnEditorChangeMode(FEditorModeID NewEditorMode)
 {
-	GEditorModeTools().ActivateMode( NewEditorMode, true );
+	GLevelEditorModeTools().ActivateMode( NewEditorMode, true );
 }
 
 void FUnrealEdMisc::OnEditorPreModal()
diff --git a/Engine/Source/Editor/UnrealEd/Private/UnrealEdSrv.cpp b/Engine/Source/Editor/UnrealEd/Private/UnrealEdSrv.cpp
--- a/Engine/Source/Editor/UnrealEd/Private/UnrealEdSrv.cpp
+++ b/Engine/Source/Editor/UnrealEd/Private/UnrealEdSrv.cpp
@@ -1550,7 +1550,7 @@ bool UUnrealEdEngine::Exec_Edit( UWorld* InWorld, const TCHAR* Str, FOutputDevic
 	if( FParse::Command(&Str,TEXT("CUT")) )
 	{
 		TArray<FEdMode*> ActiveModes; 
-		GEditorModeTools().GetActiveModes( ActiveModes );
+		GLevelEditorModeTools().GetActiveModes( ActiveModes );
 		for( int32 ModeIndex = 0; ModeIndex < ActiveModes.Num(); ++ModeIndex )
 		{
 			if (ActiveModes[ModeIndex]->ProcessEditCut())
@@ -1563,7 +1563,7 @@ bool UUnrealEdEngine::Exec_Edit( UWorld* InWorld, const TCHAR* Str, FOutputDevic
 	else if( FParse::Command(&Str,TEXT("COPY")) )
 	{
 		TArray<FEdMode*> ActiveModes; 
-		GEditorModeTools().GetActiveModes( ActiveModes );
+		GLevelEditorModeTools().GetActiveModes( ActiveModes );
 		for( int32 ModeIndex = 0; ModeIndex < ActiveModes.Num(); ++ModeIndex )
 		{
 			if (ActiveModes[ModeIndex]->ProcessEditCopy())
@@ -1576,7 +1576,7 @@ bool UUnrealEdEngine::Exec_Edit( UWorld* InWorld, const TCHAR* Str, FOutputDevic
 	else if( FParse::Command(&Str,TEXT("PASTE")) )
 	{
 		TArray<FEdMode*> ActiveModes; 
-		GEditorModeTools().GetActiveModes( ActiveModes );
+		GLevelEditorModeTools().GetActiveModes( ActiveModes );
 		for( int32 ModeIndex = 0; ModeIndex < ActiveModes.Num(); ++ModeIndex )
 		{
 			if (ActiveModes[ModeIndex]->ProcessEditPaste())
@@ -1674,7 +1674,7 @@ static void MirrorActors(const FVector& MirrorScale)
 		AActor* Actor = static_cast<AActor*>( *It );
 		checkSlow( Actor->IsA(AActor::StaticClass()) );
 
-		const FVector PivotLocation = GEditorModeTools().PivotLocation;
+		const FVector PivotLocation = GLevelEditorModeTools().PivotLocation;
 		ABrush* Brush = Cast< ABrush >( Actor );
 		if( Brush )
 		{
@@ -1725,10 +1725,10 @@ static void MirrorActors(const FVector& MirrorScale)
 		LevelDirtyCallback.Request();
 	}
 
-	if ( GEditorModeTools().IsModeActive( FBuiltinEditorModes::EM_Geometry ) )
+	if ( GLevelEditorModeTools().IsModeActive( FBuiltinEditorModes::EM_Geometry ) )
 	{
 		// If we are in geometry mode, make sure to update the mode with new source data for selected brushes
-		FEdModeGeometry* Mode = (FEdModeGeometry*)GEditorModeTools().GetActiveMode( FBuiltinEditorModes::EM_Geometry );
+		FEdModeGeometry* Mode = (FEdModeGeometry*)GLevelEditorModeTools().GetActiveMode( FBuiltinEditorModes::EM_Geometry );
 		Mode->GetFromSource();
 	}
 
@@ -2316,7 +2316,7 @@ bool UUnrealEdEngine::Exec_Actor( UWorld* InWorld, const TCHAR* Str, FOutputDevi
 		FVector DeltaMove = FVector::ZeroVector;
 		GetFVECTOR( Str, DeltaMove );
 
-		FEditorModeTools& Tools = GEditorModeTools();
+		FEditorModeTools& Tools = GLevelEditorModeTools();
 		Tools.SetPivotLocation( Tools.PivotLocation + DeltaMove, false );
 
 		if (GCurrentLevelEditingViewportClient)
@@ -2560,7 +2560,7 @@ bool UUnrealEdEngine::Exec_Actor( UWorld* InWorld, const TCHAR* Str, FOutputDevi
 	{
 		bool bHandled = false;
 		TArray<FEdMode*> ActiveModes; 
-		GEditorModeTools().GetActiveModes( ActiveModes );
+		GLevelEditorModeTools().GetActiveModes( ActiveModes );
 		for( int32 ModeIndex = 0; ModeIndex < ActiveModes.Num(); ++ModeIndex )
 		{
 			bHandled |= ActiveModes[ModeIndex]->ProcessEditDelete();
@@ -2612,7 +2612,7 @@ bool UUnrealEdEngine::Exec_Actor( UWorld* InWorld, const TCHAR* Str, FOutputDevi
 
 		// Bakes the current pivot position into all selected brushes as their PrePivot
 
-		FEditorModeTools& EditorModeTools = GEditorModeTools();
+		FEditorModeTools& EditorModeTools = GLevelEditorModeTools();
 
 		for ( FSelectionIterator It( GetSelectedActorIterator() ) ; It ; ++It )
 		{
@@ -2642,7 +2642,7 @@ bool UUnrealEdEngine::Exec_Actor( UWorld* InWorld, const TCHAR* Str, FOutputDevi
 
 		// Resets the PrePivot of the selected brushes to 0,0,0 while leaving them in the same world location.
 
-		FEditorModeTools& EditorModeTools = GEditorModeTools();
+		FEditorModeTools& EditorModeTools = GLevelEditorModeTools();
 
 		for ( FSelectionIterator It( GetSelectedActorIterator() ) ; It ; ++It )
 		{
@@ -2759,7 +2759,7 @@ bool UUnrealEdEngine::Exec_Actor( UWorld* InWorld, const TCHAR* Str, FOutputDevi
 	{
 		bool bHandled = false;
 		TArray<FEdMode*> ActiveModes; 
-		GEditorModeTools().GetActiveModes( ActiveModes );
+		GLevelEditorModeTools().GetActiveModes( ActiveModes );
 		for( int32 ModeIndex = 0; ModeIndex < ActiveModes.Num(); ++ModeIndex )
 		{
 			bHandled |= ActiveModes[ModeIndex]->ProcessEditDuplicate();
@@ -2913,7 +2913,7 @@ bool UUnrealEdEngine::Exec_Mode( const TCHAR* Str, FOutputDevice& Ar )
 	if( FParse::Command(&Str, TEXT("WIDGETCOORDSYSTEMCYCLE")) )
 	{
 		const bool bGetRawValue = true;
-		int32 Wk = GEditorModeTools().GetCoordSystem(bGetRawValue);
+		int32 Wk = GLevelEditorModeTools().GetCoordSystem(bGetRawValue);
 		Wk++;
 
 		if( Wk == COORD_Max )
@@ -2921,14 +2921,14 @@ bool UUnrealEdEngine::Exec_Mode( const TCHAR* Str, FOutputDevice& Ar )
 			Wk -= COORD_Max;
 		}
 
-		GEditorModeTools().SetCoordSystem((ECoordSystem)Wk);
+		GLevelEditorModeTools().SetCoordSystem((ECoordSystem)Wk);
 		FEditorSupportDelegates::RedrawAllViewports.Broadcast();
 		FEditorSupportDelegates::UpdateUI.Broadcast();
 	}
 
 	if( FParse::Command(&Str, TEXT("WIDGETMODECYCLE")) )
 	{
-		GEditorModeTools().CycleWidgetMode();
+		GLevelEditorModeTools().CycleWidgetMode();
 	}
 
 	if( FParse::Value(Str, TEXT("GRID="), DWord1) )
@@ -3050,7 +3050,7 @@ bool UUnrealEdEngine::Exec_Mode( const TCHAR* Str, FOutputDevice& Ar )
 	if ( EditorMode == FBuiltinEditorModes::EM_None )
 	{
 		FString CommandToken = FParse::Token(Str, false);
-		FEdMode* FoundMode = GEditorModeTools().FindMode( FName( *CommandToken ) );
+		FEdMode* FoundMode = GLevelEditorModeTools().FindMode( FName( *CommandToken ) );
 
 		if ( FoundMode != NULL )
 		{
diff --git a/Engine/Source/Editor/UnrealEd/Private/UnrealWidget.cpp b/Engine/Source/Editor/UnrealEd/Private/UnrealWidget.cpp
--- a/Engine/Source/Editor/UnrealEd/Private/UnrealWidget.cpp
+++ b/Engine/Source/Editor/UnrealEd/Private/UnrealWidget.cpp
@@ -217,7 +217,7 @@ void FWidget::Render_Axis( const FSceneView* View, FPrimitiveDrawInterface* PDI,
 
 	if( bDrawWidget )
 	{
-		const bool bDisabled = EditorModeTools ? (EditorModeTools->IsModeActive(FBuiltinEditorModes::EM_Default) && GEditor->HasLockedActors() ) : false;
+		const bool bDisabled = EditorModeTools ? (EditorModeTools->IsDefaultModeActive() && GEditor->HasLockedActors() ) : false;
 		PDI->SetHitProxy( new HWidgetAxis( InAxis, bDisabled) );
 
 		const float AxisLength = AXIS_LENGTH + GetDefault<ULevelEditorViewportSettings>()->TransformWidgetSizeAdjustment;
@@ -1744,6 +1744,6 @@ EAxisList::Type FWidget::GetAxisToDraw( EWidgetMode WidgetMode ) const
 
 bool FWidget::IsWidgetDisabled() const
 {
-	return EditorModeTools ? (EditorModeTools->IsModeActive(FBuiltinEditorModes::EM_Default) && GEditor->HasLockedActors()) : false;
+	return EditorModeTools ? (EditorModeTools->IsDefaultModeActive() && GEditor->HasLockedActors()) : false;
 }
 
diff --git a/Engine/Source/Editor/UnrealEd/Private/VertexSnapping.cpp b/Engine/Source/Editor/UnrealEd/Private/VertexSnapping.cpp
--- a/Engine/Source/Editor/UnrealEd/Private/VertexSnapping.cpp
+++ b/Engine/Source/Editor/UnrealEd/Private/VertexSnapping.cpp
@@ -703,7 +703,7 @@ bool FVertexSnappingImpl::SnapDraggedActorsToNearestVertex( FVector& DragDelta,
 
 		FSceneView* View = ViewportClient->CalcSceneView( &ViewFamily );
 
-		FVector StartLocation = GEditorModeTools().PivotLocation;
+		FVector StartLocation = GLevelEditorModeTools().PivotLocation;
 
 		FVector DesiredUnsnappedLocation = StartLocation+DragDelta;
 					
diff --git a/Engine/Source/Editor/UnrealEd/Public/ComponentVisualizer.h b/Engine/Source/Editor/UnrealEd/Public/ComponentVisualizer.h
--- a/Engine/Source/Editor/UnrealEd/Public/ComponentVisualizer.h
+++ b/Engine/Source/Editor/UnrealEd/Public/ComponentVisualizer.h
@@ -39,9 +39,9 @@ public:
 	/** */
 	virtual bool GetWidgetLocation(FVector& OutLocation) const { return false; }
 	/** */
-	virtual bool HandleInputDelta(FLevelEditorViewportClient* ViewportClient, FViewport* Viewport, FVector& DeltaTranslate, FRotator& DeltalRotate, FVector& DeltaScale) { return false; }
+	virtual bool HandleInputDelta(FEditorViewportClient* ViewportClient, FViewport* Viewport, FVector& DeltaTranslate, FRotator& DeltalRotate, FVector& DeltaScale) { return false; }
 	/** */
-	virtual bool HandleInputKey(FLevelEditorViewportClient* ViewportClient,FViewport* Viewport,FKey Key,EInputEvent Event) { return false; }
+	virtual bool HandleInputKey(FEditorViewportClient* ViewportClient,FViewport* Viewport,FKey Key,EInputEvent Event) { return false; }
 	/** */
 	virtual TSharedPtr<SWidget> GenerateContextMenu() const { return TSharedPtr<SWidget>(); }
 
diff --git a/Engine/Source/Editor/UnrealEd/Public/ComponentVisualizerManager.h b/Engine/Source/Editor/UnrealEd/Public/ComponentVisualizerManager.h
--- a/Engine/Source/Editor/UnrealEd/Public/ComponentVisualizerManager.h
+++ b/Engine/Source/Editor/UnrealEd/Public/ComponentVisualizerManager.h
@@ -19,10 +19,12 @@ public:
 	/** Clear active component visualizer */
 	void ClearActiveComponentVis();
 
+	/** Handle a click on the specified level editor viewport client */
+	bool HandleClick(FLevelEditorViewportClient* InViewportClient, HHitProxy *HitProxy, const FViewportClick &Click);
 	/** Pass key input to active visualizer */
-	bool HandleInputKey(FLevelEditorViewportClient* ViewportClient, FViewport* Viewport, FKey Key, EInputEvent Event) const;
+	bool HandleInputKey(FEditorViewportClient* ViewportClient, FViewport* Viewport, FKey Key, EInputEvent Event) const;
 	/** Pass delta input to active visualizer */
-	bool HandleInputDelta(FLevelEditorViewportClient* InViewportClient, FViewport* InViewport, FVector& InDrag, FRotator& InRot, FVector& InScale) const;
+	bool HandleInputDelta(FEditorViewportClient* InViewportClient, FViewport* InViewport, FVector& InDrag, FRotator& InRot, FVector& InScale) const;
 	/** Get widget location from active visualizer */
 	bool GetWidgetLocation(FVector& OutLocation) const;
 
diff --git a/Engine/Source/Editor/UnrealEd/Public/Editor.h b/Engine/Source/Editor/UnrealEd/Public/Editor.h
--- a/Engine/Source/Editor/UnrealEd/Public/Editor.h
+++ b/Engine/Source/Editor/UnrealEd/Public/Editor.h
@@ -553,9 +553,9 @@ UNREALED_API bool IsUniqueObjectName( const FName& InName, UObject* Outer, FText
 
 
 /**
- * Provides access to the FEditorModeTools singleton.
+ * Provides access to the FEditorModeTools for the level editor
  */
-UNREALED_API class FEditorModeTools& GEditorModeTools();
+UNREALED_API class FEditorModeTools& GLevelEditorModeTools();
 
 
 namespace EditorUtilities
diff --git a/Engine/Source/Editor/UnrealEd/Public/EditorModeInterpolation.h b/Engine/Source/Editor/UnrealEd/Public/EditorModeInterpolation.h
--- a/Engine/Source/Editor/UnrealEd/Public/EditorModeInterpolation.h
+++ b/Engine/Source/Editor/UnrealEd/Public/EditorModeInterpolation.h
@@ -17,18 +17,19 @@ public:
 	FEdModeInterpEdit();
 	~FEdModeInterpEdit();
 
-	virtual bool InputKey( FLevelEditorViewportClient* ViewportClient, FViewport* Viewport, FKey Key, EInputEvent Event ) override;
+	virtual bool InputKey( FEditorViewportClient* ViewportClient, FViewport* Viewport, FKey Key, EInputEvent Event ) override;
 	virtual void Enter() override;
 	virtual void Exit() override;
 	virtual void ActorMoveNotify() override;
 	virtual void ActorPropChangeNotify() override;
 	virtual bool AllowWidgetMove() override;
 	virtual void ActorSelectionChangeNotify() override;
+	virtual bool IsCompatibleWith(FEditorModeID OtherModeID) const override;
 
 	virtual void Render(const FSceneView* View,FViewport* Viewport,FPrimitiveDrawInterface* PDI);
-	virtual void DrawHUD(FLevelEditorViewportClient* ViewportClient,FViewport* Viewport,const FSceneView* View,FCanvas* Canvas);
+	virtual void DrawHUD(FEditorViewportClient* ViewportClient,FViewport* Viewport,const FSceneView* View,FCanvas* Canvas);
 
-	void CamMoveNotify(FLevelEditorViewportClient* ViewportClient);
+	void CamMoveNotify(FEditorViewportClient* ViewportClient);
 	void InitInterpMode(class AMatineeActor* InMatineeActor);
 	UNREALED_API void UpdateSelectedActor();
 
@@ -57,14 +58,14 @@ public:
 	/**
 	 * @return		true if the key was handled by this editor mode tool.
 	 */
-	virtual bool InputKey(FLevelEditorViewportClient* ViewportClient, FViewport* Viewport, FKey Key, EInputEvent Event) override;
-	virtual bool InputAxis(FLevelEditorViewportClient* InViewportClient, FViewport* Viewport, int32 ControllerId, FKey Key, float Delta, float DeltaTime) override;
-	virtual bool MouseMove(FLevelEditorViewportClient* ViewportClient, FViewport* Viewport, int32 x, int32 y) override;
+	virtual bool InputKey(FEditorViewportClient* ViewportClient, FViewport* Viewport, FKey Key, EInputEvent Event) override;
+	virtual bool InputAxis(FEditorViewportClient* InViewportClient, FViewport* Viewport, int32 ControllerId, FKey Key, float Delta, float DeltaTime) override;
+	virtual bool MouseMove(FEditorViewportClient* ViewportClient, FViewport* Viewport, int32 x, int32 y) override;
 
 	/**
 	 * @return		true if the delta was handled by this editor mode tool.
 	 */
-	virtual bool InputDelta(FLevelEditorViewportClient* InViewportClient,FViewport* InViewport,FVector& InDrag,FRotator& InRot,FVector& InScale) override;
+	virtual bool InputDelta(FEditorViewportClient* InViewportClient,FViewport* InViewport,FVector& InDrag,FRotator& InRot,FVector& InScale) override;
 	virtual void SelectNone() override;
 
 
diff --git a/Engine/Source/Editor/UnrealEd/Public/EditorModeRegistry.cpp b/Engine/Source/Editor/UnrealEd/Public/EditorModeRegistry.cpp
new file mode 100644
--- /dev/null
+++ b/Engine/Source/Editor/UnrealEd/Public/EditorModeRegistry.cpp
@@ -0,0 +1,134 @@
+// Copyright 1998-2014 Epic Games, Inc. All Rights Reserved.
+
+#include "UnrealEd.h"
+#include "EditorModeRegistry.h"
+
+#include "Editor/PlacementMode/Public/IPlacementModeModule.h"
+#include "Editor/LandscapeEditor/Public/LandscapeEditorModule.h"
+#include "Editor/BspMode/Public/IBspModeModule.h"
+#include "Editor/MeshPaint/Public/MeshPaintModule.h"
+#include "Editor/GeometryMode/Public/GeometryEdMode.h"
+#include "Editor/ActorPickerMode/Public/ActorPickerMode.h"
+#include "Editor/TextureAlignMode/Public/TextureAlignEdMode.h"
+#include "Editor/FoliageEdit/Public/FoliageEditModule.h"
+
+FEditorModeInfo::FEditorModeInfo()
+	: ID(NAME_None)
+	, PriorityOrder(MAX_int32)
+	, bVisible(false)
+{
+}
+
+FEditorModeInfo::FEditorModeInfo(
+	FEditorModeID InID,
+	FText InName,
+	FSlateIcon InIconBrush,
+	bool InIsVisible,
+	int32 InPriorityOrder
+	)
+	: ID(InID)
+	, Name(InName)
+	, IconBrush(InIconBrush)
+	, bVisible(InIsVisible)
+	, PriorityOrder(InPriorityOrder)
+{
+	if (!InIconBrush.IsSet())
+	{
+		IconBrush = FSlateIcon(FEditorStyle::GetStyleSetName(), "LevelEditor.EditorModes");
+	}
+}
+
+FEditorModeRegistry* GModeRegistry = nullptr;
+
+void FEditorModeRegistry::Initialize()
+{
+	Get();
+
+	// Add default editor modes
+	GModeRegistry->RegisterMode<FEdModeDefault>(FBuiltinEditorModes::EM_Default);
+	GModeRegistry->RegisterMode<FEdModeInterpEdit>(FBuiltinEditorModes::EM_InterpEdit);
+
+	// Load editor mode modules that will automatically register their editor modes, and clean themselves up on unload.
+	//@TODO: ROCKET: These are probably good plugin candidates, that shouldn't have to be force-loaded here but discovery loaded somehow
+	FModuleManager::LoadModuleChecked<IPlacementModeModule>(TEXT("PlacementMode"));
+	FModuleManager::LoadModuleChecked<IBspModeModule>(TEXT("BspMode"));
+	FModuleManager::LoadModuleChecked<FTextureAlignModeModule>(TEXT("TextureAlignMode"));
+	FModuleManager::LoadModuleChecked<FGeometryModeModule>(TEXT("GeometryMode"));
+	FModuleManager::LoadModuleChecked<FActorPickerModeModule>(TEXT("ActorPickerMode"));
+	FModuleManager::LoadModuleChecked<IMeshPaintModule>(TEXT("MeshPaint"));
+	FModuleManager::LoadModuleChecked<ILandscapeEditorModule>(TEXT("LandscapeEditor"));
+	FModuleManager::LoadModuleChecked<IFoliageEditModule>(TEXT("FoliageEdit"));
+}
+
+void FEditorModeRegistry::Shutdown()
+{
+	delete GModeRegistry;
+	GModeRegistry = nullptr;
+}
+
+FEditorModeRegistry& FEditorModeRegistry::Get()
+{
+	if (!GModeRegistry)
+	{
+		GModeRegistry = new FEditorModeRegistry;
+	}
+	return *GModeRegistry;	
+}
+
+TArray<FEditorModeInfo> FEditorModeRegistry::GetSortedModeInfo() const
+{
+	TArray<FEditorModeInfo> ModeInfoArray;
+	
+	for (const auto& Pair : ModeFactories)
+	{
+		ModeInfoArray.Add(Pair.Value->GetModeInfo());
+	}
+
+	ModeInfoArray.Sort([](const FEditorModeInfo& A, const FEditorModeInfo& B){
+		return A.PriorityOrder < B.PriorityOrder;
+	});
+
+	return ModeInfoArray;
+}
+
+TSharedPtr<FEdMode> FEditorModeRegistry::CreateMode(FEditorModeID ModeID, FEditorModeTools& Owner)
+{
+	const TSharedRef<IEditorModeFactory>* ModeFactory = ModeFactories.Find(ModeID);
+	if (ModeFactory)
+	{
+		TSharedRef<FEdMode> Instance = (*ModeFactory)->CreateMode();
+
+		// Assign the mode info from the factory before we initialize
+		Instance->Info = (*ModeFactory)->GetModeInfo();
+		Instance->Owner = &Owner;
+
+		// This binding ensures the mode is destroyed if the type is unregistered
+		OnModeUnregistered().AddSP(Instance, &FEdMode::OnModeUnregistered);
+
+		Instance->Initialize();
+
+		return Instance;
+	}
+
+	return nullptr;
+}
+
+void FEditorModeRegistry::RegisterMode(FEditorModeID ModeID, TSharedRef<IEditorModeFactory> Factory)
+{
+	check(ModeID != FBuiltinEditorModes::EM_None);
+	check(!ModeFactories.Contains(ModeID));
+
+	ModeFactories.Add(ModeID, Factory);
+
+	OnModeRegisteredEvent.Broadcast(ModeID);
+	RegisteredModesChanged.Broadcast();
+}
+
+void FEditorModeRegistry::UnregisterMode(FEditorModeID ModeID)
+{
+	// First off delete the factory
+	ModeFactories.Remove(ModeID);
+
+	OnModeUnregisteredEvent.Broadcast(ModeID);
+	RegisteredModesChanged.Broadcast();
+}
\ No newline at end of file
diff --git a/Engine/Source/Editor/UnrealEd/Public/EditorModeRegistry.h b/Engine/Source/Editor/UnrealEd/Public/EditorModeRegistry.h
new file mode 100644
--- /dev/null
+++ b/Engine/Source/Editor/UnrealEd/Public/EditorModeRegistry.h
@@ -0,0 +1,187 @@
+// Copyright 1998-2014 Epic Games, Inc. All Rights Reserved.
+
+#pragma once
+
+// Required forward declarations
+class FEdMode;
+
+DECLARE_DELEGATE_RetVal(TSharedRef<FEdMode>, FEditorModeFactoryCallback);
+
+struct FEditorModeInfo
+{
+	/** Default constructor */
+	UNREALED_API FEditorModeInfo();
+
+	/** Helper constructor */
+	UNREALED_API FEditorModeInfo(
+		FEditorModeID InID,
+		FText InName = FText(),
+		FSlateIcon InIconBrush = FSlateIcon(),
+		bool bInIsVisible = false,
+		int32 InPriorityOrder = MAX_int32
+		);
+
+	/** The mode ID */
+	FEditorModeID ID;
+
+	/** Name for the editor to display */
+	FText Name;
+
+	/** The mode icon */
+	FSlateIcon IconBrush;
+
+	/** Whether or not the mode should be visible in the mode menu */
+	bool bVisible;
+
+	/** The priority of this mode which will determine its default order and shift+X command assignment */
+	int32 PriorityOrder;
+};
+
+struct IEditorModeFactory : public TSharedFromThis<IEditorModeFactory>
+{
+	/**
+	 * Allows mode factories to handle selection change events, and potentially activate/deactivate modes
+
+	 * @param Tools					The mode tools that are triggering the event
+	 * @param ItemUndergoingChange	Either an actor being selected or deselected, or a selection set being modified (typically emptied)
+	 */
+	virtual void OnSelectionChanged(FEditorModeTools& Tools, UObject* ItemUndergoingChange) const { }
+
+	/**
+	 * Gets the information pertaining to the mode type that this factory creates
+	 */
+	virtual FEditorModeInfo GetModeInfo() const = 0;
+
+	/**
+	 * Create a new instance of our mode
+	 */
+	virtual TSharedRef<FEdMode> CreateMode() const = 0;
+};
+
+struct FEditorModeFactory : IEditorModeFactory
+{
+	FEditorModeFactory(FEditorModeInfo InModeInfo) : ModeInfo(InModeInfo) {}
+
+	/** Information pertaining to this factory's mode */
+	FEditorModeInfo ModeInfo;
+
+	/** Callback used to create an instance of this mode type */
+	FEditorModeFactoryCallback FactoryCallback;
+
+	/**
+	 * Gets the information pertaining to the mode type that this factory creates
+	 */
+	virtual FEditorModeInfo GetModeInfo() const override { return ModeInfo; }
+
+	/**
+	 * Create a new instance of our mode
+	 */
+	virtual TSharedRef<FEdMode> CreateMode() const override { return FactoryCallback.Execute(); }
+};
+
+/**
+ * A registry of editor modes and factories
+ */
+class FEditorModeRegistry
+{
+	typedef TMap<FEditorModeID, TSharedRef<IEditorModeFactory>> FactoryMap;
+
+public:
+
+	/**
+	 * Initialize this registry
+	 */
+	static void Initialize();
+
+	/**
+	 * Shutdown this registry
+	 */
+	static void Shutdown();
+
+	/**
+	 * Singleton access
+	 */
+	UNREALED_API static FEditorModeRegistry& Get();
+
+	/**
+	 * Get a list of information for all currently registered modes, sorted by UI priority order
+	 */
+	UNREALED_API TArray<FEditorModeInfo> GetSortedModeInfo() const;
+
+	/**
+	 * Registers an editor mode. Typically called from a module's StartupModule() routine.
+	 *
+	 * @param ModeID	ID of the mode to register
+	 */
+	UNREALED_API void RegisterMode(FEditorModeID ModeID, TSharedRef<IEditorModeFactory> Factory);
+
+	/**
+	 * Registers an editor mode type. Typically called from a module's StartupModule() routine.
+	 *
+	 * @param ModeID	ID of the mode to register
+	 */
+	template<class T>
+	void RegisterMode(FEditorModeID ModeID, FText Name = FText(), FSlateIcon IconBrush = FSlateIcon(), bool bVisible = false, int32 PriorityOrder = MAX_int32)
+	{
+		TSharedRef<FEditorModeFactory> Factory = MakeShareable( new FEditorModeFactory(FEditorModeInfo(ModeID, Name, IconBrush, bVisible, PriorityOrder)) );
+
+		Factory->FactoryCallback = FEditorModeFactoryCallback::CreateStatic([]() -> TSharedRef<FEdMode>{
+			return MakeShareable(new T);
+		});
+		RegisterMode(ModeID, Factory);
+	}
+
+	/**
+	 * Unregisters an editor mode. Typically called from a module's ShutdownModule() routine.
+	 * Note: will exit the edit mode if it is currently active.
+	 *
+	 * @param ModeID	ID of the mode to unregister
+	 */
+	UNREALED_API void UnregisterMode(FEditorModeID ModeID);
+
+	/**
+	 * Event that is triggered whenever a mode is registered or unregistered
+	 */
+	DECLARE_EVENT(FEditorModeRegistry, FRegisteredModesChangedEvent);
+	FRegisteredModesChangedEvent& OnRegisteredModesChanged() { return RegisteredModesChanged; }
+	
+	/**
+	 * Event that is triggered whenever a mode is registered
+	 */
+	DECLARE_EVENT_OneParam(FEditorModeRegistry, FOnModeRegistered, FEditorModeID);
+	FOnModeRegistered& OnModeRegistered() { return OnModeRegisteredEvent; }
+	
+	/**
+	 * Event that is triggered whenever a mode is unregistered
+	 */
+	DECLARE_EVENT_OneParam(FEditorModeRegistry, FOnModeUnregistered, FEditorModeID);
+	FOnModeUnregistered& OnModeUnregistered() { return OnModeUnregisteredEvent; }
+
+	/**
+	 * Create a new instance of the mode registered under the specified ID
+	 */
+	TSharedPtr<FEdMode> CreateMode(FEditorModeID ModeID, FEditorModeTools& Owner);
+
+	/**
+	 * Const access to the internal factory map
+	 */
+	const FactoryMap& GetFactoryMap() const { return ModeFactories; }
+
+
+private:
+
+	/** A map of editor mode IDs to factory callbacks */
+	FactoryMap ModeFactories;
+	
+	/** A list of all modes created */
+	TArray<TWeakPtr<FEdMode>> CreatedModes;
+
+	/** Event that is triggered whenever a mode is registered or unregistered */
+	FRegisteredModesChangedEvent RegisteredModesChanged;
+
+	/** Event that is triggered whenever a mode is registered */
+	FOnModeRegistered OnModeRegisteredEvent;
+
+	/** Event that is triggered whenever a mode is unregistered */
+	FOnModeUnregistered OnModeUnregisteredEvent;
+};
\ No newline at end of file
diff --git a/Engine/Source/Editor/UnrealEd/Public/EditorModeTools.h b/Engine/Source/Editor/UnrealEd/Public/EditorModeTools.h
--- a/Engine/Source/Editor/UnrealEd/Public/EditorModeTools.h
+++ b/Engine/Source/Editor/UnrealEd/Public/EditorModeTools.h
@@ -46,15 +46,15 @@ public:
 
 	// User input
 
-	virtual bool MouseEnter( FLevelEditorViewportClient* ViewportClient,FViewport* Viewport,int32 x, int32 y ) { return false; }
+	virtual bool MouseEnter( FEditorViewportClient* ViewportClient,FViewport* Viewport,int32 x, int32 y ) { return false; }
 
-	virtual bool MouseLeave( FLevelEditorViewportClient* ViewportClient,FViewport* Viewport ) { return false; }
+	virtual bool MouseLeave( FEditorViewportClient* ViewportClient,FViewport* Viewport ) { return false; }
 
-	virtual bool MouseMove(FLevelEditorViewportClient* ViewportClient,FViewport* Viewport,int32 x, int32 y) { return false; }
+	virtual bool MouseMove(FEditorViewportClient* ViewportClient,FViewport* Viewport,int32 x, int32 y) { return false; }
 
-	virtual bool ReceivedFocus(FLevelEditorViewportClient* ViewportClient,FViewport* Viewport) { return false; }
+	virtual bool ReceivedFocus(FEditorViewportClient* ViewportClient,FViewport* Viewport) { return false; }
 
-	virtual bool LostFocus(FLevelEditorViewportClient* ViewportClient,FViewport* Viewport) { return false; }
+	virtual bool LostFocus(FEditorViewportClient* ViewportClient,FViewport* Viewport) { return false; }
 
 	/**
 	 * Called when the mouse is moved while a window input capture is in effect
@@ -66,13 +66,13 @@ public:
 	 *
 	 * @return	true if input was handled
 	 */
-	virtual bool CapturedMouseMove( FLevelEditorViewportClient* InViewportClient, FViewport* InViewport, int32 InMouseX, int32 InMouseY ) { return false; }
+	virtual bool CapturedMouseMove( FEditorViewportClient* InViewportClient, FViewport* InViewport, int32 InMouseX, int32 InMouseY ) { return false; }
 
 
 	/**
 	 * @return		true if the delta was handled by this editor mode tool.
 	 */
-	virtual bool InputAxis(FLevelEditorViewportClient* InViewportClient, FViewport* Viewport, int32 ControllerId, FKey Key, float Delta, float DeltaTime)
+	virtual bool InputAxis(FEditorViewportClient* InViewportClient, FViewport* Viewport, int32 ControllerId, FKey Key, float Delta, float DeltaTime)
 	{
 		return false;
 	}
@@ -80,7 +80,7 @@ public:
 	/**
 	 * @return		true if the delta was handled by this editor mode tool.
 	 */
-	virtual bool InputDelta(FLevelEditorViewportClient* InViewportClient, FViewport* InViewport, FVector& InDrag, FRotator& InRot, FVector& InScale)
+	virtual bool InputDelta(FEditorViewportClient* InViewportClient, FViewport* InViewport, FVector& InDrag, FRotator& InRot, FVector& InScale)
 	{
 		return false;
 	}
@@ -88,13 +88,13 @@ public:
 	/**
 	 * @return		true if the key was handled by this editor mode tool.
 	 */
-	virtual bool InputKey(FLevelEditorViewportClient* ViewportClient, FViewport* Viewport, FKey Key, EInputEvent Event)
+	virtual bool InputKey(FEditorViewportClient* ViewportClient, FViewport* Viewport, FKey Key, EInputEvent Event)
 	{
 		return false;
 	}
 
 	virtual void Render(const FSceneView* View,FViewport* Viewport,FPrimitiveDrawInterface* PDI);
-	virtual void DrawHUD(FLevelEditorViewportClient* ViewportClient,FViewport* Viewport,const FSceneView* View,FCanvas* Canvas);
+	virtual void DrawHUD(FEditorViewportClient* ViewportClient,FViewport* Viewport,const FSceneView* View,FCanvas* Canvas);
 
 	//@{
 	virtual bool StartModify()	{ return 0; }
@@ -107,7 +107,7 @@ public:
 	//@}
 
 	// Tick
-	virtual void Tick(FLevelEditorViewportClient* ViewportClient,float DeltaTime) {}
+	virtual void Tick(FEditorViewportClient* ViewportClient,float DeltaTime) {}
 
 	/** @name Selections */
 	//@{
diff --git a/Engine/Source/Editor/UnrealEd/Public/EditorModes.h b/Engine/Source/Editor/UnrealEd/Public/EditorModes.h
--- a/Engine/Source/Editor/UnrealEd/Public/EditorModes.h
+++ b/Engine/Source/Editor/UnrealEd/Public/EditorModes.h
@@ -2,6 +2,7 @@
 
 #pragma once
 
+#include "EditorModeRegistry.h"
 
 /** Outcomes when determining whether it's possible to perform an action on the edit modes*/
 namespace EEditAction
@@ -97,18 +98,23 @@ public:
 class UNREALED_API FEdMode : public TSharedFromThis<FEdMode>, public FGCObject, public FEditorCommonDrawHelper
 {
 public:
+	/** Friends so it can access mode's internals on construction */
+	friend class FEditorModeRegistry;
+
 	FEdMode();
 	virtual ~FEdMode();
 
-	virtual bool MouseEnter( FLevelEditorViewportClient* ViewportClient,FViewport* Viewport,int32 x, int32 y );
+	virtual void Initialize() {}
+
+	virtual bool MouseEnter( FEditorViewportClient* ViewportClient,FViewport* Viewport,int32 x, int32 y );
 
-	virtual bool MouseLeave( FLevelEditorViewportClient* ViewportClient,FViewport* Viewport );
+	virtual bool MouseLeave( FEditorViewportClient* ViewportClient,FViewport* Viewport );
 
-	virtual bool MouseMove(FLevelEditorViewportClient* ViewportClient,FViewport* Viewport,int32 x, int32 y);
+	virtual bool MouseMove(FEditorViewportClient* ViewportClient,FViewport* Viewport,int32 x, int32 y);
 
-	virtual bool ReceivedFocus(FLevelEditorViewportClient* ViewportClient,FViewport* Viewport);
+	virtual bool ReceivedFocus(FEditorViewportClient* ViewportClient,FViewport* Viewport);
 
-	virtual bool LostFocus(FLevelEditorViewportClient* ViewportClient,FViewport* Viewport);
+	virtual bool LostFocus(FEditorViewportClient* ViewportClient,FViewport* Viewport);
 
 	/**
 	 * Called when the mouse is moved while a window input capture is in effect
@@ -120,13 +126,13 @@ public:
 	 *
 	 * @return	true if input was handled
 	 */
-	virtual bool CapturedMouseMove( FLevelEditorViewportClient* InViewportClient, FViewport* InViewport, int32 InMouseX, int32 InMouseY );
+	virtual bool CapturedMouseMove( FEditorViewportClient* InViewportClient, FViewport* InViewport, int32 InMouseX, int32 InMouseY );
 
-	virtual bool InputKey(FLevelEditorViewportClient* ViewportClient,FViewport* Viewport,FKey Key,EInputEvent Event);
-	virtual bool InputAxis(FLevelEditorViewportClient* InViewportClient,FViewport* Viewport,int32 ControllerId,FKey Key,float Delta,float DeltaTime);
-	virtual bool InputDelta(FLevelEditorViewportClient* InViewportClient,FViewport* InViewport,FVector& InDrag,FRotator& InRot,FVector& InScale);
-	virtual bool StartTracking(FLevelEditorViewportClient* InViewportClient, FViewport* InViewport);
-	virtual bool EndTracking(FLevelEditorViewportClient* InViewportClient, FViewport* InViewport);
+	virtual bool InputKey(FEditorViewportClient* ViewportClient,FViewport* Viewport,FKey Key,EInputEvent Event);
+	virtual bool InputAxis(FEditorViewportClient* InViewportClient,FViewport* Viewport,int32 ControllerId,FKey Key,float Delta,float DeltaTime);
+	virtual bool InputDelta(FEditorViewportClient* InViewportClient,FViewport* InViewport,FVector& InDrag,FRotator& InRot,FVector& InScale);
+	virtual bool StartTracking(FEditorViewportClient* InViewportClient, FViewport* InViewport);
+	virtual bool EndTracking(FEditorViewportClient* InViewportClient, FViewport* InViewport);
 	// Added for handling EDIT Command...
 	virtual EEditAction::Type GetActionEditDuplicate() { return EEditAction::Skip; }
 	virtual EEditAction::Type GetActionEditDelete() { return EEditAction::Skip; }
@@ -139,8 +145,9 @@ public:
 	virtual bool ProcessEditCopy() { return false; }
 	virtual bool ProcessEditPaste() { return false; }
 
-	virtual void Tick(FLevelEditorViewportClient* ViewportClient,float DeltaTime);
+	virtual void Tick(FEditorViewportClient* ViewportClient,float DeltaTime);
 
+	virtual bool IsCompatibleWith(FEditorModeID OtherModeID) const { return false; }
 	virtual void ActorMoveNotify() {}
 	virtual void ActorsDuplicatedNotify(TArray<AActor*>& PreDuplicateSelection, TArray<AActor*>& PostDuplicateSelection, bool bOffsetLocations) {}
 	virtual void ActorSelectionChangeNotify();
@@ -234,12 +241,7 @@ public:
 	void SelectNone();
 	virtual void SelectionChanged() {}
 
-	// Allows modes to preview selection change events, and potentially make themselves active or inactive
-	// Note: This method is called on *all* modes, not just ones that are currently active, and it is not currently called if matinee is active.
-	// @param ItemUndergoingChange	Either an actor being selected or deselected, or a selection set being modified (typically emptied)
-	virtual void PeekAtSelectionChangedEvent(UObject* ItemUndergoingChange) {}
-
-	virtual bool HandleClick(FLevelEditorViewportClient* InViewportClient, HHitProxy *HitProxy, const FViewportClick &Click);
+	virtual bool HandleClick(FEditorViewportClient* InViewportClient, HHitProxy *HitProxy, const FViewportClick &Click);
 
 	/** Handling SelectActor */
 	virtual bool Select( AActor* InActor, bool bInSelected ) { return 0; }
@@ -247,28 +249,11 @@ public:
 	/** Check to see if an actor can be selected in this mode - no side effects */
 	virtual bool IsSelectionAllowed( AActor* InActor, bool bInSelection ) const { return true; }
 
-	/** Returns the editor mode. */
-	FEditorModeID GetID() const { return ID; }
+	/** Returns the editor mode identifier. */
+	FEditorModeID GetID() const { return Info.ID; }
 
-	/** Returns the name of this mode (already localized) */
-	const FText& GetName() const
-	{
-		return Name;
-	}
-
-	/** Returns the icon for the mode */
-	FSlateIcon GetIcon() const
-	{
-		return IconBrush.IsSet() ? IconBrush : FSlateIcon(FEditorStyle::GetStyleSetName(), "LevelEditor.EditorModes");
-	}
-
-	/** Returns true if the editor mode should be explicit and visible in any menu/toolbar for modes */
-	bool IsVisible() const { return bVisible; }
-
-	/** Returns the priority of this editor mode */
-	int32 GetPriorityOrder() const { return PriorityOrder; }
-
-	friend class FEditorModeTools;
+	/** Returns the editor mode information. */
+	const FEditorModeInfo& GetModeInfo() const { return Info; }
 
 	// Tools
 
@@ -299,7 +284,7 @@ public:
 	//void DrawGridSection(int32 ViewportLocX,int32 ViewportGridY,FVector* A,FVector* B,float* AX,float* BX,int32 Axis,int32 AlphaCase,FSceneView* View,FPrimitiveDrawInterface* PDI);
 
 	/** Overlays the editor hud (brushes, drag tools, static mesh vertices, etc*. */
-	virtual void DrawHUD(FLevelEditorViewportClient* ViewportClient,FViewport* Viewport,const FSceneView* View,FCanvas* Canvas);
+	virtual void DrawHUD(FEditorViewportClient* ViewportClient,FViewport* Viewport,const FSceneView* View,FCanvas* Canvas);
 	//@}
 
 	/**
@@ -311,7 +296,7 @@ public:
 	/**
 	 * Called when the mode wants to draw brackets around selected objects
 	 */
-	virtual void DrawBrackets( FLevelEditorViewportClient* ViewportClient, FViewport* Viewport, const FSceneView* View, FCanvas* Canvas );
+	virtual void DrawBrackets( FEditorViewportClient* ViewportClient, FViewport* Viewport, const FSceneView* View, FCanvas* Canvas );
 
 	/** True if this mode uses a toolkit mode (eventually they all should) */
 	virtual bool UsesToolkits() const;
@@ -340,10 +325,22 @@ public:
 	/** See if we should create a widget for the supplied property when selecting an actor instance */
 	static bool ShouldCreateWidgetForProperty(UProperty* InProp);
 
-protected:
-	/** Name for the editor to display. */
-	FText Name;
+public:
 
+	/** Request that this mode be deleted at the next convenient opportunity (FEditorModeTools::Tick) */
+	void RequestDeletion() { bPendingDeletion = true; }
+	
+	/** returns true if this mode is to be deleted at the next convenient opportunity (FEditorModeTools::Tick) */
+	bool IsPendingDeletion() const { return bPendingDeletion; }
+
+private:
+	/** true if this mode is pending removal from its owner */
+	bool bPendingDeletion;
+
+	/** Called whenever a mode type is unregistered */
+	void OnModeUnregistered(FEditorModeID ModeID);
+
+protected:
 	/** The current axis that is being dragged on the widget. */
 	EAxisList::Type CurrentWidgetAxis;
 
@@ -353,21 +350,15 @@ protected:
 	/** The tool that is currently active within this mode. */
 	FModeTool* CurrentTool;
 
-	/** The mode ID. */
-	FEditorModeID ID;
-
-	/** The mode icon. */
-	FSlateIcon IconBrush;
-
-	/** Whether or not the mode should be visible in the mode menu. */
-	bool bVisible;
-
-	/** The priority of this mode which will determine its default order and shift+X command assignment */
-	int32 PriorityOrder;
+	/** Information pertaining to this mode. Assigned by FEditorModeRegistry. */
+	FEditorModeInfo Info;
 
 	/** Editor Mode Toolkit that is associated with this toolkit mode */
 	TSharedPtr<class FModeToolkit> Toolkit;
 
+	/** Pointer back to the mode tools that we are registered with */
+	FEditorModeTools* Owner;
+
 	// Property Widgets
 
 	/** Structure that holds info about our optional property widget */
@@ -429,38 +420,29 @@ public:
 };
 
 
-
 /**
  * A helper class to store the state of the various editor modes.
  */
-class FEditorModeTools : public FGCObject
+class UNREALED_API FEditorModeTools : public FGCObject
 {
 public:
 	FEditorModeTools();
 	virtual ~FEditorModeTools();
 
-	void Init();
-
-	void Shutdown();
-
 	/**
-	 * Registers an editor mode.  Typically called from a module's StartupModule() routine.
-	 *
-	 * @param ModeToRegister	The editor mode instance to register
+	 * Set the default editor mode for these tools
 	 */
-	UNREALED_API void RegisterMode(TSharedRef<FEdMode> ModeToRegister);
+	void SetDefaultMode ( FEditorModeID DefaultID );
 
 	/**
-	 * Unregisters an editor mode.  Typically called from a module's ShutdownModule() routine.
-	 * Note: will exit the edit mode if it is currently active.
-	 *
-	 * @param ModeToUnregister	The editor mode instance to unregister
+	 * Activates the default mode defined by this class
 	 */
-	UNREALED_API void UnregisterMode(TSharedRef<FEdMode> ModeToUnregister);
+	void ActivateDefaultMode();
 
-	UNREALED_API void RegisterCompatibleModes( const FEditorModeID& Mode, const FEditorModeID& CompatibleMode );
-
-	UNREALED_API void UnregisterCompatibleModes( const FEditorModeID& Mode, const FEditorModeID& CompatibleMode );
+	/** 
+	 * Returns true if the default editor mode is active 
+	 */
+	bool IsDefaultModeActive() const;
 
 	/**
 	 * Activates an editor mode. Shuts down all other active modes which cannot run with the passed in mode.
@@ -468,24 +450,36 @@ public:
 	 * @param InID		The ID of the editor mode to activate.
 	 * @param bToggle	true if the passed in editor mode should be toggled off if it is already active.
 	 */
-	UNREALED_API void ActivateMode( FEditorModeID InID, bool bToggle = false );
+	void ActivateMode( FEditorModeID InID, bool bToggle = false );
 
 	/**
 	 * Deactivates an editor mode. 
 	 * 
 	 * @param InID		The ID of the editor mode to deactivate.
 	 */
-	UNREALED_API void DeactivateMode( FEditorModeID InID );
+	void DeactivateMode(FEditorModeID InID);
+
+	/**
+	 * Deactivate the mode and entirely purge it from memory. Used when a mode type is unregistered
+	 */
+	void DestroyMode(FEditorModeID InID);
+
+protected:
+	
+	/** Deactivates the editor mode at the specified index */
+	void DeactivateModeAtIndex( int32 InIndex );
+		
+public:
 
 	/**
 	 * Deactivates all modes, note some modes can never be deactivated.
 	 */
-	UNREALED_API void DeactivateAllModes();
+	void DeactivateAllModes();
 
 	/** 
 	 * Returns the editor mode specified by the passed in ID
 	 */
-	UNREALED_API FEdMode* FindMode( FEditorModeID InID );
+	FEdMode* FindMode( FEditorModeID InID );
 
 	/**
 	 * Returns true if the current mode is not the specified ModeID.  Also optionally warns the user.
@@ -495,22 +489,22 @@ public:
 	 * @param	bNotifyUser		If true, display the error as a notification, instead of a dialog
 	 * @return					true if the current mode is not the specified mode.
 	 */
-	UNREALED_API bool EnsureNotInMode(FEditorModeID ModeID, const FText& ErrorMsg = FText::GetEmpty(), bool bNotifyUser = false) const;
+	bool EnsureNotInMode(FEditorModeID ModeID, const FText& ErrorMsg = FText::GetEmpty(), bool bNotifyUser = false) const;
 
-	UNREALED_API FMatrix GetCustomDrawingCoordinateSystem();
+	FMatrix GetCustomDrawingCoordinateSystem();
 	FMatrix GetCustomInputCoordinateSystem();
 	
 	/** 
 	 * Returns true if the passed in editor mode is active 
 	 */
-	UNREALED_API bool IsModeActive( FEditorModeID InID ) const;
+	bool IsModeActive( FEditorModeID InID ) const;
 
 	/**
 	 * Returns a pointer to an active mode specified by the passed in ID
 	 * If the editor mode is not active, NULL is returned
 	 */
-	UNREALED_API FEdMode* GetActiveMode( FEditorModeID InID );
-	UNREALED_API const FEdMode* GetActiveMode( FEditorModeID InID ) const;
+	FEdMode* GetActiveMode( FEditorModeID InID );
+	const FEdMode* GetActiveMode( FEditorModeID InID ) const;
 
 	template <typename SpecificModeType>
 	SpecificModeType* GetActiveModeTyped( FEditorModeID InID )
@@ -533,12 +527,7 @@ public:
 	/** 
 	 * Returns an array of all active modes
 	 */
-	UNREALED_API void GetActiveModes( TArray<FEdMode*>& OutActiveModes );
-
-	/** 
-	 * Returns an array of all modes.
-	 */
-	UNREALED_API void GetModes( TArray<FEdMode*>& OutModes );
+	void GetActiveModes( TArray<FEdMode*>& OutActiveModes );
 
 	void SetShowWidget( bool InShowWidget )	{ bShowWidget = InShowWidget; }
 	bool GetShowWidget() const				{ return bShowWidget; }
@@ -554,12 +543,12 @@ public:
 	EAxisList::Type GetWidgetAxisToDraw( FWidget::EWidgetMode InWidgetMode ) const;
 
 	/** Mouse tracking interface.  Passes tracking messages to all active modes */
-	bool StartTracking(FLevelEditorViewportClient* InViewportClient, FViewport* InViewport);
-	bool EndTracking(FLevelEditorViewportClient* InViewportClient, FViewport* InViewport);
+	bool StartTracking(FEditorViewportClient* InViewportClient, FViewport* InViewport);
+	bool EndTracking(FEditorViewportClient* InViewportClient, FViewport* InViewport);
 	bool IsTracking() const { return bIsTracking; }
 
 	/** Notifies all active modes that a map change has occured */
-	UNREALED_API void MapChangeNotify();
+	void MapChangeNotify();
 
 	/** Notifies all active modes to empty their selections */
 	void SelectNone();
@@ -580,7 +569,7 @@ public:
 	void SetCurrentWidgetAxis( EAxisList::Type NewAxis );
 
 	/** Notifies all active modes of mouse click messages. */
-	bool HandleClick(FLevelEditorViewportClient* InViewportClient,  HHitProxy *HitProxy, const FViewportClick &Click );
+	bool HandleClick(FEditorViewportClient* InViewportClient,  HHitProxy *HitProxy, const FViewportClick &Click );
 
 	/** true if the passed in brush actor should be drawn in wireframe */	
 	bool ShouldDrawBrushWireframe( AActor* InActor ) const;
@@ -589,41 +578,41 @@ public:
 	bool ShouldDrawBrushVertices() const;
 
 	/** Ticks all active modes */
-	void Tick( FLevelEditorViewportClient* ViewportClient, float DeltaTime );
+	void Tick( FEditorViewportClient* ViewportClient, float DeltaTime );
 
 	/** Notifies all active modes of any change in mouse movement */
-	bool InputDelta( FLevelEditorViewportClient* InViewportClient,FViewport* InViewport,FVector& InDrag,FRotator& InRot,FVector& InScale );
+	bool InputDelta( FEditorViewportClient* InViewportClient,FViewport* InViewport,FVector& InDrag,FRotator& InRot,FVector& InScale );
 
 	/** Notifies all active modes of captured mouse movement */	
-	bool CapturedMouseMove( FLevelEditorViewportClient* InViewportClient, FViewport* InViewport, int32 InMouseX, int32 InMouseY );
+	bool CapturedMouseMove( FEditorViewportClient* InViewportClient, FViewport* InViewport, int32 InMouseX, int32 InMouseY );
 
 	/** Notifies all active modes of keyboard input */
-	bool InputKey( FLevelEditorViewportClient* InViewportClient, FViewport* Viewport, FKey Key, EInputEvent Event);
+	bool InputKey( FEditorViewportClient* InViewportClient, FViewport* Viewport, FKey Key, EInputEvent Event);
 
 	/** Notifies all active modes of axis movement */
-	bool InputAxis( FLevelEditorViewportClient* InViewportClient, FViewport* Viewport, int32 ControllerId, FKey Key, float Delta, float DeltaTime);
+	bool InputAxis( FEditorViewportClient* InViewportClient, FViewport* Viewport, int32 ControllerId, FKey Key, float Delta, float DeltaTime);
 
-	bool MouseEnter( FLevelEditorViewportClient* InViewportClient, FViewport* Viewport, int32 X, int32 Y );
+	bool MouseEnter( FEditorViewportClient* InViewportClient, FViewport* Viewport, int32 X, int32 Y );
 	
-	bool MouseLeave( FLevelEditorViewportClient* InViewportClient, FViewport* Viewport );
+	bool MouseLeave( FEditorViewportClient* InViewportClient, FViewport* Viewport );
 
 	/** Notifies all active modes that the mouse has moved */
-	bool MouseMove( FLevelEditorViewportClient* InViewportClient, FViewport* Viewport, int32 X, int32 Y );
+	bool MouseMove( FEditorViewportClient* InViewportClient, FViewport* Viewport, int32 X, int32 Y );
 
 	/** Notifies all active modes that a viewport has received focus */
-	bool ReceivedFocus( FLevelEditorViewportClient* InViewportClient, FViewport* Viewport );
+	bool ReceivedFocus( FEditorViewportClient* InViewportClient, FViewport* Viewport );
 
 	/** Notifies all active modes that a viewport has lost focus */
-	bool LostFocus( FLevelEditorViewportClient* InViewportClient, FViewport* Viewport );
+	bool LostFocus( FEditorViewportClient* InViewportClient, FViewport* Viewport );
 
 	/** Draws all active modes */	
 	void DrawActiveModes( const FSceneView* InView, FPrimitiveDrawInterface* PDI );
 
 	/** Renders all active modes */
-	UNREALED_API void Render( const FSceneView* InView, FViewport* Viewport, FPrimitiveDrawInterface* PDI );
+	void Render( const FSceneView* InView, FViewport* Viewport, FPrimitiveDrawInterface* PDI );
 
 	/** Draws the HUD for all active modes */
-	void DrawHUD( FLevelEditorViewportClient* InViewportClient,FViewport* Viewport,const FSceneView* View,FCanvas* Canvas );
+	void DrawHUD( FEditorViewportClient* InViewportClient,FViewport* Viewport,const FSceneView* View,FCanvas* Canvas );
 
 	/** Calls PostUndo on all active components */
 	void PostUndo();
@@ -640,12 +629,12 @@ public:
 	/**
 	 * Returns a good location to draw the widget at.
 	 */
-	UNREALED_API FVector GetWidgetLocation() const;
+	FVector GetWidgetLocation() const;
 
 	/**
 	 * Changes the current widget mode.
 	 */
-	UNREALED_API void SetWidgetMode( FWidget::EWidgetMode InWidgetMode );
+	void SetWidgetMode( FWidget::EWidgetMode InWidgetMode );
 
 	/**
 	 * Allows you to temporarily override the widget mode.  Call this function again
@@ -656,7 +645,7 @@ public:
 	/**
 	 * Retrieves the current widget mode, taking overrides into account.
 	 */
-	UNREALED_API FWidget::EWidgetMode GetWidgetMode() const;
+	FWidget::EWidgetMode GetWidgetMode() const;
 
 	/**
 	 * Gets the current state of script editor usage of show friendly names
@@ -681,7 +670,7 @@ public:
 	 * @param InIndex			Index of the bookmark to set.
 	 * @param InViewportClient	Level editor viewport client used to get pointer the world that the bookmark is for.
 	 */
-	UNREALED_API void SetBookmark( uint32 InIndex, FLevelEditorViewportClient* InViewportClient );
+	void SetBookmark( uint32 InIndex, FEditorViewportClient* InViewportClient );
 
 	/**
 	 * Checks to see if a bookmark exists at a given index
@@ -689,7 +678,7 @@ public:
 	 * @param InIndex			Index of the bookmark to set.
 	 * @param InViewportClient	Level editor viewport client used to get pointer the world to check for the bookmark in.
 	 */
-	UNREALED_API bool CheckBookmark( uint32 InIndex, FLevelEditorViewportClient* InViewportClient );
+	bool CheckBookmark( uint32 InIndex, FEditorViewportClient* InViewportClient );
 
 	/**
 	 * Retrieves a bookmark from the list.
@@ -698,7 +687,7 @@ public:
 	 * @param					Set to true to restore the visibility of any streaming levels.
 	 * @param InViewportClient	Level editor viewport client used to get pointer the world that the bookmark is relevant to.
 	 */
-	UNREALED_API void JumpToBookmark( uint32 InIndex, bool bShouldRestoreLevelVisibility, FLevelEditorViewportClient* InViewportClient );
+	void JumpToBookmark( uint32 InIndex, bool bShouldRestoreLevelVisibility, FEditorViewportClient* InViewportClient );
 
 	/**
 	 * Clears a bookmark from the list.
@@ -706,14 +695,14 @@ public:
 	 * @param InIndex			Index of the bookmark to clear.
 	 * @param InViewportClient	Level editor viewport client used to get pointer the world that the bookmark is relevant to.
 	 */
-	UNREALED_API void ClearBookmark( uint32 InIndex, FLevelEditorViewportClient* InViewportClient );
+	void ClearBookmark( uint32 InIndex, FEditorViewportClient* InViewportClient );
 
 	/**
 	 * Clears all book marks
 	 * 
 	 * @param InViewportClient	Level editor viewport client used to get pointer the world to clear the bookmarks from.
 	 */
-	UNREALED_API void ClearAllBookmarks( FLevelEditorViewportClient* InViewportClient );
+	void ClearAllBookmarks( FEditorViewportClient* InViewportClient );
 
 	// FGCObject interface
 	virtual void AddReferencedObjects( FReferenceCollector& Collector ) override;
@@ -727,15 +716,7 @@ public:
 	/**
 	 * Saves the current state to the INI file
 	 */
-	UNREALED_API void SaveConfig(void);
-
-	/** 
-	 * Returns a list of active modes that are incompatible with the passed in mode.
-	 * 
-	 * @param InID 				The mode to check for incompatibilites.
-	 * @param IncompatibleModes	The list of incompatible modes.
-	 */
-	void GetIncompatibleActiveModes( FEditorModeID InID, TArray<FEdMode*>& IncompatibleModes );
+	void SaveConfig(void);
 
 	/** 
 	 * Sets the pivot locations
@@ -743,7 +724,7 @@ public:
 	 * @param Location 		The location to set
 	 * @param bIncGridBase	Whether or not to also set the GridBase
 	 */
-	UNREALED_API void SetPivotLocation( const FVector& Location, const bool bIncGridBase );
+	void SetPivotLocation( const FVector& Location, const bool bIncGridBase );
 
 	/**
 	 * Multicast delegate for OnModeEntered and OnModeExited callbacks.
@@ -759,23 +740,20 @@ public:
 	FWidgetModeChangedEvent& OnWidgetModeChanged() { return WidgetModeChangedEvent; }
 
 	/**	Broadcasts the WidgetModeChanged event */
-	void BroadcastWidgetModeChanged(FWidget::EWidgetMode InWidgetMode) { WidgetModeChangedEvent.Broadcast( InWidgetMode ); }
+	void BroadcastWidgetModeChanged(FWidget::EWidgetMode InWidgetMode) { WidgetModeChangedEvent.Broadcast(InWidgetMode); }
 
-	DECLARE_EVENT( FEditorModeTools, FRegisteredModesChangedEvent );
-	FRegisteredModesChangedEvent& OnRegisteredModesChanged() { return RegisteredModesChanged; }
-
-	/**	Broadcasts the registered modes changed event */
-	void BroadcastRegisteredModesChanged() { RegisteredModesChanged.Broadcast( ); }
+	/**	Broadcasts the EditorModeChanged event */
+	void BroadcastEditorModeChanged(FEdMode* Mode, bool IsEnteringMode) { EditorModeChangedEvent.Broadcast(Mode, IsEnteringMode); }
 
 	/**
 	 * Returns the current CoordSystem
 	 * 
 	 * @param bGetRawValue true when you want the actual value of CoordSystem, not the value modified by the state.
 	 */
-	UNREALED_API ECoordSystem GetCoordSystem(bool bGetRawValue = false);
+	ECoordSystem GetCoordSystem(bool bGetRawValue = false);
 
 	/** Sets the current CoordSystem */
-	UNREALED_API void SetCoordSystem(ECoordSystem NewCoordSystem);
+	void SetCoordSystem(ECoordSystem NewCoordSystem);
 
 	/** Sets the hide viewport UI state */
 	void SetHideViewportUI( bool bInHideViewportUI ) { bHideViewportUI = bInHideViewportUI; }
@@ -789,7 +767,7 @@ public:
 	/** The angle for the translate rotate widget */
 	float TranslateRotateXAxisAngle;
 
-	/** Draws in the top level corner of all FLevelEditorViewportClient windows (can be used to relay info to the user). */
+	/** Draws in the top level corner of all FEditorViewportClient windows (can be used to relay info to the user). */
 	FString InfoString;
 
 protected:
@@ -799,14 +777,14 @@ protected:
 	void OnEditorSelectionChanged(UObject* NewSelection);
 	void OnEditorSelectNone();
 
-	/** A mapping of editor modes that can be active at the same time */
-	TMultiMap<FEditorModeID, FEditorModeID> ModeCompatabilityMap;
+	/** The mode ID to use as a default */
+	FEditorModeID DefaultID;
 
-	/** A list of all available editor modes. */
+	/** A list of active editor modes. */
 	TArray< TSharedPtr<FEdMode> > Modes;
 
-	/** The editor modes currently in use. */
-	TArray<FEdMode*> ActiveModes;
+	/** A list of previously active editor modes that we will potentially recycle */
+	TMap< FEditorModeID, TSharedPtr<FEdMode> > RecycledModes;
 
 	/** The mode that the editor viewport widget is in. */
 	FWidget::EWidgetMode WidgetMode;
@@ -823,11 +801,6 @@ protected:
 	/** if true, the viewports will hide all UI overlays */
 	bool bHideViewportUI;
 
-	/**	Broadcasts the EditorModeChanged event */
-	void BroadcastEditorModeChanged(FEdMode* Mode, bool IsEnteringMode) { EditorModeChangedEvent.Broadcast( Mode, IsEnteringMode ); }
-
-	friend class FEdMode;
-
 private:
 
 	/** The coordinate system the widget is operating within. */
@@ -839,8 +812,6 @@ private:
 	/** Multicast delegate that is broadcast when a widget mode is changed */
 	FWidgetModeChangedEvent WidgetModeChangedEvent;
 
-	FRegisteredModesChangedEvent RegisteredModesChanged;
-
 	/** Flag set between calls to StartTracking() and EndTracking() */
 	bool bIsTracking;
 };
diff --git a/Engine/Source/Editor/UnrealEd/Public/EditorViewportClient.h b/Engine/Source/Editor/UnrealEd/Public/EditorViewportClient.h
--- a/Engine/Source/Editor/UnrealEd/Public/EditorViewportClient.h
+++ b/Engine/Source/Editor/UnrealEd/Public/EditorViewportClient.h
@@ -15,6 +15,7 @@ class FWidget;
 class FCameraControllerConfig;
 class FCachedJoystickState;
 class FPreviewScene;
+class FEditorModeTools;
 
 /** Delegate called by FEditorViewportClient to check its visibility */
 DECLARE_DELEGATE_RetVal( bool, FViewportStateGetter );
@@ -144,6 +145,19 @@ private:
 };
 
 
+struct FDropQuery
+{
+	FDropQuery()
+		: bCanDrop(false)
+	{}
+
+	/** True if it's valid to drop the object at the location queried */
+	bool bCanDrop;
+
+	/** Optional hint text that may be returned to the user. */
+	FText HintText;
+};
+
 /** 
  * Stores the transformation data for the viewport camera
  */
@@ -227,7 +241,7 @@ class UNREALED_API FEditorViewportClient : public FCommonViewportClient, public
 public:
 	friend class FMouseDeltaTracker;
 
-	FEditorViewportClient(class FPreviewScene* InPreviewScene = NULL);
+	FEditorViewportClient(FEditorModeTools& InModeTools, FPreviewScene* InPreviewScene = nullptr);
 	virtual ~FEditorViewportClient();
 
 	/**
@@ -530,6 +544,11 @@ public:
 	 */
 	virtual void DrawCanvas( FViewport& InViewport, FSceneView& View, FCanvas& Canvas ) {};
 
+	/**
+	 * Render the drag tool in the viewport
+	 */
+	void RenderDragTool(const FSceneView* View, FCanvas* Canvas);
+
 	/**
 	 * Configures the specified FSceneView object with the view and projection matrices for this viewport.
 	 * @param	View		The view to be configured.  Must be valid.
@@ -608,6 +627,64 @@ public:
 	 */
 	virtual void CheckHoveredHitProxy( HHitProxy* HoveredHitProxy );
 
+	/** Returns true if a placement dragging actor exists */
+	virtual bool HasDropPreviewActors() const { return false; }
+
+	/**
+	 * If dragging an actor for placement, this function updates its position.
+	 *
+	 * @param MouseX						The position of the mouse's X coordinate
+	 * @param MouseY						The position of the mouse's Y coordinate
+	 * @param DroppedObjects				The Objects that were used to create preview objects
+	 * @param out_bDroppedObjectsVisible	Output, returns if preview objects are visible or not
+	 *
+	 * Returns true if preview actors were updated
+	 */
+	virtual bool UpdateDropPreviewActors(int32 MouseX, int32 MouseY, const TArray<UObject*>& DroppedObjects, bool& out_bDroppedObjectsVisible, class UActorFactory* FactoryToUse = NULL) { return false; }
+
+	/**
+	 * If dragging an actor for placement, this function destroys the actor.
+	 */
+	virtual void DestroyDropPreviewActors() {}
+
+	/**
+	 * Checks the viewport to see if the given object can be dropped using the given mouse coordinates local to this viewport
+	 *
+	 * @param MouseX			The position of the mouse's X coordinate
+	 * @param MouseY			The position of the mouse's Y coordinate
+	 * @param AssetInfo			Asset in question to be dropped
+	 */
+	virtual FDropQuery CanDropObjectsAtCoordinates(int32 MouseX, int32 MouseY, const FAssetData& AssetInfo) { return FDropQuery(); }
+
+	/**
+	 * Attempts to intelligently drop the given objects in the viewport, using the given mouse coordinates local to this viewport
+	 *
+	 * @param MouseX			 The position of the mouse's X coordinate
+	 * @param MouseY			 The position of the mouse's Y coordinate
+	 * @param DroppedObjects	 The Objects to be placed into the editor via this viewport
+	 * @param OutNewActors		 The new actor objects that were created
+	 * @param bOnlyDropOnTarget  Flag that when True, will only attempt a drop on the actor targeted by the Mouse position. Defaults to false.
+	 * @param bCreateDropPreview If true, a drop preview actor will be spawned instead of a normal actor.
+	 * @param bSelectActors		 If true, select the newly dropped actors (defaults: true)
+	 * @param FactoryToUse		 The preferred actor factory to use (optional)
+	 */
+	virtual bool DropObjectsAtCoordinates(int32 MouseX, int32 MouseY, const TArray<UObject*>& DroppedObjects, TArray<AActor*>& OutNewActors, bool bOnlyDropOnTarget = false, bool bCreateDropPreview = false, bool bSelectActors = true, UActorFactory* FactoryToUse = NULL ) { return false; }
+
+	/** Returns true if the viewport is allowed to be possessed by Matinee for previewing sequences */
+	bool AllowMatineePreview() const { return bAllowMatineePreview; }
+
+	/** Sets whether or not this viewport is allowed to be possessed by Matinee */
+	void SetAllowMatineePreview(const bool bInAllowMatineePreview)
+	{
+		bAllowMatineePreview = bInAllowMatineePreview;
+	}
+
+protected:
+
+	/** true if this window is allowed to be possessed by Matinee for previewing sequences in real-time */
+	bool bAllowMatineePreview;
+
+public:
 	/** True if the window is maximized or floating */
 	bool IsVisible() const;
 
@@ -793,6 +870,12 @@ public:
 
 	bool IsForcedRealtimeAudio() const { return bForceAudioRealtime; }
 
+	/** true to force realtime audio to be true, false to stop forcing it */
+	void SetForcedAudioRealtime(bool bShouldForceAudioRealtime)
+	{
+		bForceAudioRealtime = bShouldForceAudioRealtime;
+	}
+
 	/** @return true if a mouse button is down and it's movement being tracked for operations inside the viewport */
 	bool IsTracking() const { return bIsTracking; }
 
@@ -833,6 +916,19 @@ protected:
 public:
 
 	void DrawBoundingBox(FBox &Box, FCanvas* InCanvas, const FSceneView* InView, const FViewport* InViewport, const FLinearColor& InColor, const bool bInDrawBracket, const FString &InLabelText) ;
+	
+	/**
+	 * Draws a screen space bounding box around the specified actor
+	 *
+	 * @param	InCanvas		Canvas to draw on
+	 * @param	InView			View to render
+	 * @param	InViewport		Viewport we're rendering into
+	 * @param	InActor			Actor to draw a bounding box for
+	 * @param	InColor			Color of bounding box
+	 * @param	bInDrawBracket	True to draw a bracket, otherwise a box will be rendered
+	 * @param	bInLabelText	Optional label text to draw
+	 */
+	void DrawActorScreenSpaceBoundingBox( FCanvas* InCanvas, const FSceneView* InView, FViewport* InViewport, AActor* InActor, const FLinearColor& InColor, const bool bInDrawBracket, const FString& InLabelText = TEXT( "" ) );
 
 	void SetGameView(bool bGameViewEnable);
 
@@ -1134,8 +1230,14 @@ public:
 	// Override the LOD of landscape in this viewport
 	int8 LandscapeLODOverride;
 
+	/** If true, draw vertices for selected BSP brushes and static meshes if the large vertices show flag is set. */
+	bool bDrawVertices;
+
 protected:
 
+	/** Editor mode tools provided to this instance. Assumed to be managed externally */
+	FEditorModeTools*		ModeTools;
+
 	FWidget*				Widget;
 
 	FMouseDeltaTracker*		MouseDeltaTracker;
