Made it possible for asset editors to maintain their own FEditorModeTools lists

diff --git a/Engine/Source/Editor/UnrealEd/Public/LevelEditorViewport.h b/Engine/Source/Editor/UnrealEd/Public/LevelEditorViewport.h
--- a/Engine/Source/Editor/UnrealEd/Public/LevelEditorViewport.h
+++ b/Engine/Source/Editor/UnrealEd/Public/LevelEditorViewport.h
@@ -102,20 +102,6 @@ private:
 };
 
 
-struct FDropQuery
-{
-	FDropQuery()
-		: bCanDrop(false)
-	{}
-
-	/** True if it's valid to drop the object at the location queried */
-	bool bCanDrop;
-
-	/** Optional hint text that may be returned to the user. */
-	FText HintText;
-};
-
-
 /** */
 class UNREALED_API FLevelEditorViewportClient : public FEditorViewportClient
 {
@@ -202,13 +188,6 @@ public:
 	 */
 	void RestoreCameraFromPIE();
 
-
-	/** true to force realtime audio to be true, false to stop forcing it true */
-	void SetForcedAudioRealtime( bool bShouldForceAudioRealtime )
-	{
-		bForceAudioRealtime = bShouldForceAudioRealtime;
-	}
-
 	/**
 	 * Updates the audio listener for this viewport 
 	 *
@@ -269,15 +248,6 @@ public:
 	 */
 	EAxisList::Type GetVertAxis() const;
 
-	/** Returns true if the viewport is allowed to be possessed by Matinee for previewing sequences */
-	bool AllowMatineePreview() const { return bAllowMatineePreview; }
-
-	/** Sets whether or not this viewport is allowed to be possessed by Matinee */
-	void SetAllowMatineePreview( const bool bInAllowMatineePreview )
-	{
-		bAllowMatineePreview = bInAllowMatineePreview;
-	}
-
 	void NudgeSelectedObjects( const struct FInputEventState& InputState );
 
 	/**
@@ -303,19 +273,6 @@ public:
 
 	virtual void SetIsSimulateInEditorViewport( bool bInIsSimulateInEditorViewport ) override;
 
-	/**
-	 * Draws a screen space bounding box around the specified actor
-	 *
-	 * @param	InCanvas		Canvas to draw on
-	 * @param	InView			View to render
-	 * @param	InViewport		Viewport we're rendering into
-	 * @param	InActor			Actor to draw a bounding box for
-	 * @param	InColor			Color of bounding box
-	 * @param	bInDrawBracket	True to draw a bracket, otherwise a box will be rendered
-	 * @param	bInLabelText	Optional label text to draw
-	 */
-	void DrawActorScreenSpaceBoundingBox( FCanvas* InCanvas, const FSceneView* InView, FViewport* InViewport, AActor* InActor, const FLinearColor& InColor, const bool bInDrawBracket, const FString& InLabelText = TEXT( "" ) );
-
 	/**
 	 *	Draw the texture streaming bounds.
 	 */
@@ -383,7 +340,7 @@ public:
 	void RemoveReferenceToWorldContext(FWorldContext& WorldContext);
 
 	/** Returns true if a placement dragging actor exists */
-	bool HasDropPreviewActors() const;
+	virtual bool HasDropPreviewActors() const override;
 
 	/**
 	 * If dragging an actor for placement, this function updates its position.
@@ -395,12 +352,12 @@ public:
 	 *
 	 * Returns true if preview actors were updated
 	 */
-	bool UpdateDropPreviewActors(int32 MouseX, int32 MouseY, const TArray<UObject*>& DroppedObjects, bool& out_bDroppedObjectsVisible, class UActorFactory* FactoryToUse = NULL);
+	virtual bool UpdateDropPreviewActors(int32 MouseX, int32 MouseY, const TArray<UObject*>& DroppedObjects, bool& out_bDroppedObjectsVisible, class UActorFactory* FactoryToUse = NULL) override;
 
 	/**
 	 * If dragging an actor for placement, this function destroys the actor.
 	 */
-	void DestroyDropPreviewActors();
+	virtual void DestroyDropPreviewActors() override;
 
 	/**
 	 * Checks the viewport to see if the given object can be dropped using the given mouse coordinates local to this viewport
@@ -409,7 +366,7 @@ public:
 	 * @param MouseY			The position of the mouse's Y coordinate
 	 * @param AssetInfo			Asset in question to be dropped
 	 */
-	FDropQuery CanDropObjectsAtCoordinates(int32 MouseX, int32 MouseY, const FAssetData& AssetInfo);
+	virtual FDropQuery CanDropObjectsAtCoordinates(int32 MouseX, int32 MouseY, const FAssetData& AssetInfo) override;
 
 	/**
 	 * Attempts to intelligently drop the given objects in the viewport, using the given mouse coordinates local to this viewport
@@ -423,7 +380,7 @@ public:
 	 * @param bSelectActors		 If true, select the newly dropped actors (defaults: true)
 	 * @param FactoryToUse		 The preferred actor factory to use (optional)
 	 */
-	bool DropObjectsAtCoordinates(int32 MouseX, int32 MouseY, const TArray<UObject*>& DroppedObjects, TArray<AActor*>& OutNewActors, bool bOnlyDropOnTarget = false, bool bCreateDropPreview = false, bool bSelectActors = true, class UActorFactory* FactoryToUse = NULL );
+	virtual bool DropObjectsAtCoordinates(int32 MouseX, int32 MouseY, const TArray<UObject*>& DroppedObjects, TArray<AActor*>& OutNewActors, bool bOnlyDropOnTarget = false, bool bCreateDropPreview = false, bool bSelectActors = true, UActorFactory* FactoryToUse = NULL ) override;
 
 	/**
 	 * Sets GWorld to the appropriate world for this client
@@ -470,8 +427,6 @@ public:
 	 */
 	void ModifyScale( USceneComponent* InComponent, FVector& ScaleDelta ) const;
 
-	void RenderDragTool( const FSceneView* View,FCanvas* Canvas );
-
 	/** Set the global ptr to the current viewport */
 	void SetCurrentViewport();
 
@@ -798,9 +753,6 @@ public:
 	/** If true, we switched between two different cameras. Set by matinee, used by the motion blur to invalidate this frames motion vectors */
 	bool					bEditorCameraCut;
 
-	/** If true, draw vertices for selected BSP brushes and static meshes if the large vertices show flag is set. */
-	bool					bDrawVertices;
-
 	/** Indicates whether, of not, the base attachment volume should be drawn for this viewport. */
 	bool bDrawBaseInfo;
 
@@ -844,9 +796,6 @@ private:
 	int32 DropPreviewMouseX;
 	int32 DropPreviewMouseY;
 
-	/** true if this window is allowed to be possessed by Matinee for previewing sequences in real-time */
-	bool bAllowMatineePreview;
-
 	/** If this view was controlled by another view this/last frame, don't update itself */
 	bool bWasControlledByOtherViewport;
 
diff --git a/Engine/Source/Editor/UnrealEd/Public/PhysicsManipulationMode.h b/Engine/Source/Editor/UnrealEd/Public/PhysicsManipulationMode.h
--- a/Engine/Source/Editor/UnrealEd/Public/PhysicsManipulationMode.h
+++ b/Engine/Source/Editor/UnrealEd/Public/PhysicsManipulationMode.h
@@ -1,6 +1,19 @@
 // Copyright 1998-2014 Epic Games, Inc. All Rights Reserved.
 #pragma once
 
+
+class FPhysicsManipulationEdModeFactory : public IEditorModeFactory
+{
+	/** Called when the selection changes in the editor */
+	virtual void OnSelectionChanged(FEditorModeTools& Tools, UObject* ItemUndergoingChange) const override;
+
+	/** Gets the information pertaining to the mode type that this factory creates */
+	virtual FEditorModeInfo GetModeInfo() const override;
+
+	/** Create a new instance of our mode */
+	virtual TSharedRef<FEdMode> CreateMode() const override;
+};
+
 /** 
  * Editor mode to manipulate physics objects in level editor viewport simulation
  **/
@@ -12,17 +25,15 @@ public:
 
 	virtual void AddReferencedObjects( FReferenceCollector& Collector ) override;
 
-	virtual void Enter();
-	virtual void Exit();
-	
-	virtual void PeekAtSelectionChangedEvent(UObject* ItemUndergoingChange);
+	virtual void Enter() override;
+	virtual void Exit() override;
 
-	virtual bool InputDelta(FLevelEditorViewportClient* InViewportClient,FViewport* InViewport,FVector& InDrag,FRotator& InRot,FVector& InScale);
-	virtual bool StartTracking(FLevelEditorViewportClient* InViewportClient, FViewport* InViewport);
-	virtual bool EndTracking(FLevelEditorViewportClient* InViewportClient, FViewport* InViewport);
+	virtual bool InputDelta(FEditorViewportClient* InViewportClient,FViewport* InViewport,FVector& InDrag,FRotator& InRot,FVector& InScale) override;
+	virtual bool StartTracking(FEditorViewportClient* InViewportClient, FViewport* InViewport) override;
+	virtual bool EndTracking(FEditorViewportClient* InViewportClient, FViewport* InViewport) override;
 
-	virtual bool ShouldDrawWidget() const		{ return true; };
-	virtual bool UsesTransformWidget() const	{ return true; };
+	virtual bool ShouldDrawWidget() const override		{ return true; };
+	virtual bool UsesTransformWidget() const override	{ return true; };
 
 private:
 	class UPhysicsHandleComponent* HandleComp;
diff --git a/Engine/Source/Editor/UnrealEd/Public/UnrealEd.h b/Engine/Source/Editor/UnrealEd/Public/UnrealEd.h
--- a/Engine/Source/Editor/UnrealEd/Public/UnrealEd.h
+++ b/Engine/Source/Editor/UnrealEd/Public/UnrealEd.h
@@ -25,6 +25,7 @@
 #include "EditorViewportClient.h"
 #include "LevelEditorViewport.h"
 
+#include "EditorModeRegistry.h"
 #include "EditorModes.h"
 
 #include "MRUList.h"
diff --git a/Engine/Source/Editor/WorldBrowser/Private/LevelCollectionModel.cpp b/Engine/Source/Editor/WorldBrowser/Private/LevelCollectionModel.cpp
--- a/Engine/Source/Editor/WorldBrowser/Private/LevelCollectionModel.cpp
+++ b/Engine/Source/Editor/WorldBrowser/Private/LevelCollectionModel.cpp
@@ -524,20 +524,20 @@ void FLevelCollectionModel::UnloadLevels(const FLevelModelList& InLevelList)
 	}
 
 	// If matinee is opened, and if it belongs to the level being removed, close it
-	if (GEditorModeTools().IsModeActive(FBuiltinEditorModes::EM_InterpEdit))
+	if (GLevelEditorModeTools().IsModeActive(FBuiltinEditorModes::EM_InterpEdit))
 	{
 		TArray<ULevel*> LevelsToRemove = GetLevelObjectList(InLevelList);
 		
-		const FEdModeInterpEdit* InterpEditMode = (const FEdModeInterpEdit*)GEditorModeTools().GetActiveMode(FBuiltinEditorModes::EM_InterpEdit);
+		const FEdModeInterpEdit* InterpEditMode = (const FEdModeInterpEdit*)GLevelEditorModeTools().GetActiveMode(FBuiltinEditorModes::EM_InterpEdit);
 
 		if (InterpEditMode && InterpEditMode->MatineeActor && LevelsToRemove.Contains(InterpEditMode->MatineeActor->GetLevel()))
 		{
-			GEditorModeTools().ActivateMode(FBuiltinEditorModes::EM_Default);
+			GLevelEditorModeTools().ActivateDefaultMode();
 		}
 	}
-	else if(GEditorModeTools().IsModeActive(FBuiltinEditorModes::EM_Landscape))
+	else if(GLevelEditorModeTools().IsModeActive(FBuiltinEditorModes::EM_Landscape))
 	{
-		GEditorModeTools().ActivateMode(FBuiltinEditorModes::EM_Default);
+		GLevelEditorModeTools().ActivateDefaultMode();
 	}
 
 	
diff --git a/Engine/Source/Editor/WorldBrowser/Private/SWorldDetails.cpp b/Engine/Source/Editor/WorldBrowser/Private/SWorldDetails.cpp
--- a/Engine/Source/Editor/WorldBrowser/Private/SWorldDetails.cpp
+++ b/Engine/Source/Editor/WorldBrowser/Private/SWorldDetails.cpp
@@ -13,7 +13,7 @@ SWorldDetails::SWorldDetails()
 
 SWorldDetails::~SWorldDetails()
 {
-	GEditorModeTools().DeactivateMode(FBuiltinEditorModes::EM_StreamingLevel);
+	GLevelEditorModeTools().DeactivateMode(FBuiltinEditorModes::EM_StreamingLevel);
 	
 	FWorldBrowserModule& WorldBrowserModule = FModuleManager::GetModuleChecked<FWorldBrowserModule>("WorldBrowser");
 	WorldBrowserModule.OnBrowseWorld.RemoveAll(this);
diff --git a/Engine/Source/Editor/WorldBrowser/Private/StreamingLevels/StreamingLevelCollectionModel.cpp b/Engine/Source/Editor/WorldBrowser/Private/StreamingLevels/StreamingLevelCollectionModel.cpp
--- a/Engine/Source/Editor/WorldBrowser/Private/StreamingLevels/StreamingLevelCollectionModel.cpp
+++ b/Engine/Source/Editor/WorldBrowser/Private/StreamingLevels/StreamingLevelCollectionModel.cpp
@@ -509,9 +509,9 @@ bool FStreamingLevelCollectionModel::AddExistingLevel()
 		}
 
 		// For safety
-		if( GEditorModeTools().IsModeActive( FBuiltinEditorModes::EM_Landscape ) )
+		if( GLevelEditorModeTools().IsModeActive( FBuiltinEditorModes::EM_Landscape ) )
 		{
-			GEditorModeTools().ActivateMode( FBuiltinEditorModes::EM_Default );
+			GLevelEditorModeTools().ActivateDefaultMode();
 		}
 
 		// refresh editor windows
diff --git a/Engine/Source/Editor/WorldBrowser/Private/StreamingLevels/StreamingLevelCustomization.cpp b/Engine/Source/Editor/WorldBrowser/Private/StreamingLevels/StreamingLevelCustomization.cpp
--- a/Engine/Source/Editor/WorldBrowser/Private/StreamingLevels/StreamingLevelCustomization.cpp
+++ b/Engine/Source/Editor/WorldBrowser/Private/StreamingLevels/StreamingLevelCustomization.cpp
@@ -233,7 +233,7 @@ bool FStreamingLevelCustomization::LevelEditTextTransformAllowed() const
 	ULevelStreaming* LevelStreaming = LevelModel->GetLevelStreaming().Get();
 	
 	auto* ActiveMode = static_cast<FStreamingLevelEdMode*>(
-		GEditorModeTools().GetActiveMode(FBuiltinEditorModes::EM_StreamingLevel)
+		GLevelEditorModeTools().GetActiveMode(FBuiltinEditorModes::EM_StreamingLevel)
 		);		
 	
 	if (ActiveMode && ActiveMode->IsEditing(LevelStreaming) == true)
@@ -260,21 +260,20 @@ FReply FStreamingLevelCustomization::OnEditLevelClicked()
 		return FReply::Handled();
 	}		
 
-	if (!GEditorModeTools().IsModeActive(FBuiltinEditorModes::EM_StreamingLevel))
+	if (!GLevelEditorModeTools().IsModeActive(FBuiltinEditorModes::EM_StreamingLevel))
 	{
 		// Activate Level Mode if it was not active
-		GEditorModeTools().ActivateMode(FBuiltinEditorModes::EM_StreamingLevel);
+		GLevelEditorModeTools().ActivateMode(FBuiltinEditorModes::EM_StreamingLevel);
 	}
 			
-	auto* ActiveMode = static_cast<FStreamingLevelEdMode*>(
-		GEditorModeTools().GetActiveMode(FBuiltinEditorModes::EM_StreamingLevel)
-		);		
+	auto* ActiveMode = GLevelEditorModeTools().GetActiveModeTyped<FStreamingLevelEdMode>(FBuiltinEditorModes::EM_StreamingLevel);
 	check(ActiveMode != NULL);
 
 	if (ActiveMode->IsEditing(LevelStreaming) == true)
 	{
 		// Toggle this mode off if already editing this level
-		GEditorModeTools().DeactivateMode(FBuiltinEditorModes::EM_StreamingLevel);
+		GLevelEditorModeTools().DeactivateMode(FBuiltinEditorModes::EM_StreamingLevel);
+		// ActiveMode is now invalid
 	}
 	else
 	{
diff --git a/Engine/Source/Editor/WorldBrowser/Private/StreamingLevels/StreamingLevelEdMode.cpp b/Engine/Source/Editor/WorldBrowser/Private/StreamingLevels/StreamingLevelEdMode.cpp
--- a/Engine/Source/Editor/WorldBrowser/Private/StreamingLevels/StreamingLevelEdMode.cpp
+++ b/Engine/Source/Editor/WorldBrowser/Private/StreamingLevels/StreamingLevelEdMode.cpp
@@ -9,9 +9,6 @@
 FStreamingLevelEdMode::FStreamingLevelEdMode() 
 	: FEdMode()
 {
-	ID = FBuiltinEditorModes::EM_StreamingLevel;
-	Name = LOCTEXT("StreamingLevelMode", "Level Transform Editing");
-
 	UMaterial* GizmoMaterial = (UMaterial*)StaticLoadObject( UMaterial::StaticClass(),NULL,TEXT("/Engine/EditorMaterials/LevelTransformMaterial.LevelTransformMaterial") );
 	BoxMaterial = UMaterialInstanceDynamic::Create( GizmoMaterial, NULL );
 
@@ -106,7 +103,7 @@ void FStreamingLevelEdMode::SetLevel( ULevelStreaming* LevelStream )
 	GEditor->RedrawAllViewports();
 }
 
-bool FStreamingLevelEdMode::InputDelta( FLevelEditorViewportClient* InViewportClient, FViewport* InViewport, FVector& InDrag, FRotator& InRot, FVector& InScale )
+bool FStreamingLevelEdMode::InputDelta( FEditorViewportClient* InViewportClient, FViewport* InViewport, FVector& InDrag, FRotator& InRot, FVector& InScale )
 {
 	// Only Update the LevelTransform if the user has clicked on the Widget
 	if (InViewportClient->GetCurrentWidgetAxis() != EAxisList::None)
@@ -145,13 +142,13 @@ void FStreamingLevelEdMode::Render( const FSceneView* View, FViewport* Viewport,
 	}
 }
 
-bool FStreamingLevelEdMode::StartTracking( FLevelEditorViewportClient* InViewportClient, FViewport* InViewport )
+bool FStreamingLevelEdMode::StartTracking( FEditorViewportClient* InViewportClient, FViewport* InViewport )
 {	
 	bIsTracking = true;
 	return true;
 }
 
-bool FStreamingLevelEdMode::EndTracking( FLevelEditorViewportClient* InViewportClient, FViewport* InViewport )
+bool FStreamingLevelEdMode::EndTracking( FEditorViewportClient* InViewportClient, FViewport* InViewport )
 {	
 	bIsTracking = false;
 	if( SelectedLevel.IsValid() && !LevelTransform.Equals( SelectedLevel->LevelTransform ) )
diff --git a/Engine/Source/Editor/WorldBrowser/Private/StreamingLevels/StreamingLevelEdMode.h b/Engine/Source/Editor/WorldBrowser/Private/StreamingLevels/StreamingLevelEdMode.h
--- a/Engine/Source/Editor/WorldBrowser/Private/StreamingLevels/StreamingLevelEdMode.h
+++ b/Engine/Source/Editor/WorldBrowser/Private/StreamingLevels/StreamingLevelEdMode.h
@@ -29,13 +29,13 @@ public:
 
 	virtual bool AllowWidgetMove() override { return true; }
 
-	virtual bool InputDelta( FLevelEditorViewportClient* InViewportClient, FViewport* InViewport, FVector& InDrag, FRotator& InRot, FVector& InScale ) override;
+	virtual bool InputDelta( FEditorViewportClient* InViewportClient, FViewport* InViewport, FVector& InDrag, FRotator& InRot, FVector& InScale ) override;
 
 	virtual void Render( const FSceneView* View, FViewport* Viewport, FPrimitiveDrawInterface* PDI ) override;
 
-	virtual bool StartTracking(FLevelEditorViewportClient* InViewportClient, FViewport* InViewport) override;
+	virtual bool StartTracking(FEditorViewportClient* InViewportClient, FViewport* InViewport) override;
 
-	virtual bool EndTracking(FLevelEditorViewportClient* InViewportClient, FViewport* InViewport) override;
+	virtual bool EndTracking(FEditorViewportClient* InViewportClient, FViewport* InViewport) override;
 
 	virtual bool IsSnapRotationEnabled() override;
 
diff --git a/Engine/Source/Editor/WorldBrowser/Private/Tiles/WorldTileCollectionModel.cpp b/Engine/Source/Editor/WorldBrowser/Private/Tiles/WorldTileCollectionModel.cpp
--- a/Engine/Source/Editor/WorldBrowser/Private/Tiles/WorldTileCollectionModel.cpp
+++ b/Engine/Source/Editor/WorldBrowser/Private/Tiles/WorldTileCollectionModel.cpp
@@ -889,8 +889,7 @@ TSharedPtr<FLevelModel> FWorldTileCollectionModel::CreateNewEmptyLevel()
 		return NewLevelModel;
 	}
 	
-	// Editor modes cannot be active when any level saving occurs.
-	GEditorModeTools().ActivateMode(FBuiltinEditorModes::EM_Default);
+	GLevelEditorModeTools().ActivateDefaultMode();
 
 	// Save new level to the same directory where selected level/folder is
 	FString Directory = FPaths::GetPath(GetWorldRootModel()->GetPackageFileName());
diff --git a/Engine/Source/Editor/WorldBrowser/Private/WorldBrowserModule.cpp b/Engine/Source/Editor/WorldBrowser/Private/WorldBrowserModule.cpp
--- a/Engine/Source/Editor/WorldBrowser/Private/WorldBrowserModule.cpp
+++ b/Engine/Source/Editor/WorldBrowser/Private/WorldBrowserModule.cpp
@@ -20,9 +20,9 @@ void FWorldBrowserModule::StartupModule()
 	FLevelCollectionCommands::Register();
 
 	// register the editor mode
-	TSharedRef<FStreamingLevelEdMode> NewEditorMode = MakeShareable(new FStreamingLevelEdMode);
-	GEditorModeTools().RegisterMode(NewEditorMode);
-	EdModeStreamingLevel = NewEditorMode;
+	FEditorModeRegistry::Get().RegisterMode<FStreamingLevelEdMode>(
+		FBuiltinEditorModes::EM_StreamingLevel,
+		NSLOCTEXT("WorldBrowser", "StreamingLevelMode", "Level Transform Editing"));
 
 	FWorldDelegates::OnPostWorldInitialization.AddRaw(this, &FWorldBrowserModule::OnWorldCreated);
 	FWorldDelegates::OnWorldCleanup.AddRaw(this, &FWorldBrowserModule::OnWorldCleanup);
@@ -40,8 +40,7 @@ void FWorldBrowserModule::ShutdownModule()
 	FLevelCollectionCommands::Unregister();
 
 	// unregister the editor mode
-	GEditorModeTools().UnregisterMode(EdModeStreamingLevel.ToSharedRef());
-	EdModeStreamingLevel = NULL;
+	FEditorModeRegistry::Get().UnregisterMode(FBuiltinEditorModes::EM_StreamingLevel);
 }
 
 TSharedRef<SWidget> FWorldBrowserModule::CreateWorldBrowserHierarchy()
@@ -129,4 +128,4 @@ TSharedPtr<FLevelCollectionModel> FWorldBrowserModule::SharedWorldModel(UWorld*
 
 
 
-#undef LOCTEXT_NAMESPACE
\ No newline at end of file
+#undef LOCTEXT_NAMESPACE
diff --git a/Engine/Source/Editor/WorldBrowser/Public/WorldBrowserModule.h b/Engine/Source/Editor/WorldBrowser/Public/WorldBrowserModule.h
--- a/Engine/Source/Editor/WorldBrowser/Public/WorldBrowserModule.h
+++ b/Engine/Source/Editor/WorldBrowser/Public/WorldBrowserModule.h
@@ -54,5 +54,4 @@ private:
 		
 private:
 	TWeakPtr<class FLevelCollectionModel>	WorldModel;
-	TSharedPtr<class FStreamingLevelEdMode> EdModeStreamingLevel;
 };
diff --git a/Engine/Source/Runtime/Engine/Private/AI/Navigation/NavigationSystem.cpp b/Engine/Source/Runtime/Engine/Private/AI/Navigation/NavigationSystem.cpp
--- a/Engine/Source/Runtime/Engine/Private/AI/Navigation/NavigationSystem.cpp
+++ b/Engine/Source/Runtime/Engine/Private/AI/Navigation/NavigationSystem.cpp
@@ -221,7 +221,7 @@ UNavigationSystem::UNavigationSystem(const class FPostConstructInitializePropert
 #if WITH_EDITOR
 	if (GIsEditor && HasAnyFlags(RF_ClassDefaultObject) == false)
 	{
-		GEditorModeTools().OnEditorModeChanged().AddUObject(this, &UNavigationSystem::OnEditorModeChanged);
+		GLevelEditorModeTools().OnEditorModeChanged().AddUObject(this, &UNavigationSystem::OnEditorModeChanged);
 	}
 #endif // WITH_EDITOR
 }
@@ -232,7 +232,7 @@ UNavigationSystem::~UNavigationSystem()
 #if WITH_EDITOR
 	if (GIsEditor)
 	{
-		GEditorModeTools().OnEditorModeChanged().RemoveAll(this);
+		GLevelEditorModeTools().OnEditorModeChanged().RemoveAll(this);
 	}
 #endif // WITH_EDITOR
 }
