Auto-save improvements and performance tweaks

https://answers.unrealengine.com/questions/53990/cant-find-file-for-asset.html

This change restores the behaviour of the package auto-saver that was removed in CL# 1950261 due to changes in the way UPackage::PackageDirtyStateUpdatedEvent was called (it was changed to be called only when the package became dirty state was actually changed).

There is now another callback, UPackage::FOnPackageMarkedDirty, which is always called when a package is marked as dirty, even if the package is already dirty. This allows the package auto-saver to watch for changes in-between auto-save periods, rather than be forced to assume that everything that was dirty had become dirty during the last auto-save period. FOnPackageDirtyStateUpdated and UPackage::PackageDirtyStateUpdateEvent have been renamed to FOnPackageDirtyStateChanged and UPackage::PackageDirtyStateChangedEvent respectively.

Additionally, you're now able to restore asset files that haven't yet been saved to disk as a .uasset file. This can happen for newly added or imported assets, as UE4 doesn't generate the associated .uasset file until you do something which causes dirty packages to be saved.

diff --git a/Engine/Source/Editor/ContentBrowser/Private/FrontendFilters.cpp b/Engine/Source/Editor/ContentBrowser/Private/FrontendFilters.cpp
--- a/Engine/Source/Editor/ContentBrowser/Private/FrontendFilters.cpp
+++ b/Engine/Source/Editor/ContentBrowser/Private/FrontendFilters.cpp
@@ -97,12 +97,12 @@ FFrontendFilter_Modified::FFrontendFilter_Modified(TSharedPtr<FFrontendFilterCat
 	: FFrontendFilter(InCategory)
 	, bIsCurrentlyActive(false)
 {
-	UPackage::PackageDirtyStateUpdatedEvent.AddRaw(this, &FFrontendFilter_Modified::OnPackageDirtyStateUpdated);
+	UPackage::PackageDirtyStateChangedEvent.AddRaw(this, &FFrontendFilter_Modified::OnPackageDirtyStateUpdated);
 }
 
 FFrontendFilter_Modified::~FFrontendFilter_Modified()
 {
-	UPackage::PackageDirtyStateUpdatedEvent.RemoveAll(this);
+	UPackage::PackageDirtyStateChangedEvent.RemoveAll(this);
 }
 
 void FFrontendFilter_Modified::ActiveStateChanged(bool bActive)
diff --git a/Engine/Source/Editor/Kismet/Private/FormatTextDetails.cpp b/Engine/Source/Editor/Kismet/Private/FormatTextDetails.cpp
--- a/Engine/Source/Editor/Kismet/Private/FormatTextDetails.cpp
+++ b/Engine/Source/Editor/Kismet/Private/FormatTextDetails.cpp
@@ -36,12 +36,12 @@ void FFormatTextDetails::CustomizeDetails( IDetailLayoutBuilder& DetailLayout )
 		InputsCategory.AddCustomBuilder( Layout.ToSharedRef() );
 	}
 
-	UPackage::PackageDirtyStateUpdatedEvent.AddSP(this, &FFormatTextDetails::OnEditorPackageModified);
+	UPackage::PackageDirtyStateChangedEvent.AddSP(this, &FFormatTextDetails::OnEditorPackageModified);
 }
 
 FFormatTextDetails::~FFormatTextDetails()
 {
-	UPackage::PackageDirtyStateUpdatedEvent.RemoveAll(this);
+	UPackage::PackageDirtyStateChangedEvent.RemoveAll(this);
 }
 
 void FFormatTextDetails::OnForceRefresh()
diff --git a/Engine/Source/Editor/LevelEditor/Private/SLevelEditorBuildAndSubmit.cpp b/Engine/Source/Editor/LevelEditor/Private/SLevelEditorBuildAndSubmit.cpp
--- a/Engine/Source/Editor/LevelEditor/Private/SLevelEditorBuildAndSubmit.cpp
+++ b/Engine/Source/Editor/LevelEditor/Private/SLevelEditorBuildAndSubmit.cpp
@@ -122,7 +122,7 @@ private:
 
 SLevelEditorBuildAndSubmit::~SLevelEditorBuildAndSubmit()
 {
-	UPackage::PackageDirtyStateUpdatedEvent.RemoveAll(this);
+	UPackage::PackageDirtyStateChangedEvent.RemoveAll(this);
 
 	ISourceControlModule::Get().GetProvider().UnregisterSourceControlStateChanged(FSourceControlStateChanged::FDelegate::CreateRaw(this, &SLevelEditorBuildAndSubmit::OnSourceControlStateChanged));
 }
@@ -276,7 +276,7 @@ void SLevelEditorBuildAndSubmit::Construct( const FArguments& InArgs, const TSha
 
 	ISourceControlModule::Get().GetProvider().RegisterSourceControlStateChanged(FSourceControlStateChanged::FDelegate::CreateRaw(this, &SLevelEditorBuildAndSubmit::OnSourceControlStateChanged));
 
-	UPackage::PackageDirtyStateUpdatedEvent.AddRaw(this, &SLevelEditorBuildAndSubmit::OnEditorPackageModified);
+	UPackage::PackageDirtyStateChangedEvent.AddRaw(this, &SLevelEditorBuildAndSubmit::OnEditorPackageModified);
 }
 
 void SLevelEditorBuildAndSubmit::OnEditorPackageModified(UPackage* Package)
diff --git a/Engine/Source/Editor/UnrealEd/Private/FileHelpers.cpp b/Engine/Source/Editor/UnrealEd/Private/FileHelpers.cpp
--- a/Engine/Source/Editor/UnrealEd/Private/FileHelpers.cpp
+++ b/Engine/Source/Editor/UnrealEd/Private/FileHelpers.cpp
@@ -1814,7 +1814,7 @@ void FEditorFileUtils::ResetLevelFilenames()
 	MainFrameModule.SetLevelNameForWindowTitle(EmptyFilename);
 }
 
-bool FEditorFileUtils::AutosaveMap(const FString& AbsoluteAutosaveDir, const int32 AutosaveIndex)
+bool FEditorFileUtils::AutosaveMap(const FString& AbsoluteAutosaveDir, const int32 AutosaveIndex, const bool bForceIfNotInList, const TSet< TWeakObjectPtr<UPackage> >& DirtyPackagesForAutoSave)
 {
 	const FScopedBusyCursor BusyCursor;
 	bool bResult  = false;
@@ -1841,7 +1841,7 @@ bool FEditorFileUtils::AutosaveMap(const FString& AbsoluteAutosaveDir, const int
 			check( Package );
 
 			// If this world needs saving . . .
-			if ( Package->IsDirty() )
+			if ( Package->IsDirty() && (bForceIfNotInList || DirtyPackagesForAutoSave.Contains(Package)) )
 			{
 				const FString AutosaveFilename = GetAutoSaveFilename(Package, AbsoluteAutosaveDir, AutosaveIndex, FPackageName::GetMapPackageExtension());
 				//UE_LOG(LogFileHelpers, Log,  TEXT("Autosaving '%s'"), *AutosaveFilename );
@@ -1881,7 +1881,7 @@ bool FEditorFileUtils::AutosaveMap(const FString& AbsoluteAutosaveDir, const int
 	return bResult;
 }
 
-bool FEditorFileUtils::AutosaveContentPackages(const FString& AbsoluteAutosaveDir, const int32 AutosaveIndex)
+bool FEditorFileUtils::AutosaveContentPackages(const FString& AbsoluteAutosaveDir, const int32 AutosaveIndex, const bool bForceIfNotInList, const TSet< TWeakObjectPtr<UPackage> >& DirtyPackagesForAutoSave)
 {
 	const FScopedBusyCursor BusyCursor;
 	double SaveStartTime = FPlatformTime::Seconds();
@@ -1895,7 +1895,7 @@ bool FEditorFileUtils::AutosaveContentPackages(const FString& AbsoluteAutosaveDi
 		UPackage* CurPackage = *PackageIter;
 
 		// If the package is dirty and is not the transient package, we'd like to autosave it
-		if ( CurPackage && ( CurPackage != TransientPackage ) && CurPackage->IsDirty() )
+		if ( CurPackage && ( CurPackage != TransientPackage ) && CurPackage->IsDirty() && (bForceIfNotInList || DirtyPackagesForAutoSave.Contains(CurPackage)) )
 		{
 			UWorld* MapWorld = UWorld::FindWorldInPackage(CurPackage);
 			// Also, make sure this is not a map package
diff --git a/Engine/Source/Editor/UnrealEd/Private/ObjectTools.cpp b/Engine/Source/Editor/UnrealEd/Private/ObjectTools.cpp
--- a/Engine/Source/Editor/UnrealEd/Private/ObjectTools.cpp
+++ b/Engine/Source/Editor/UnrealEd/Private/ObjectTools.cpp
@@ -1370,6 +1370,9 @@ namespace ObjectTools
 
 		GWarn->EndSlowTask();
 
+		// Let the package auto-saver know that it needs to ignore the deleted packages
+		GUnrealEd->GetPackageAutoSaver().OnPackagesDeleted(PackagesToDelete);
+
 		// Unload the packages and collect garbage.
 		if ( PackagesToDelete.Num() > 0 )
 		{
diff --git a/Engine/Source/Editor/UnrealEd/Private/PackageAutoSaver.cpp b/Engine/Source/Editor/UnrealEd/Private/PackageAutoSaver.cpp
--- a/Engine/Source/Editor/UnrealEd/Private/PackageAutoSaver.cpp
+++ b/Engine/Source/Editor/UnrealEd/Private/PackageAutoSaver.cpp
@@ -42,6 +42,14 @@ namespace PackageAutoSaverJson
 	 * @param DirtyPackages		Packages that may have auto-saves that they could be restored from
 	 */
 	void SaveRestoreFile(const bool bRestoreEnabled, const TMap< TWeakObjectPtr<UPackage>, FString >& DirtyPackages);
+
+	/** @return whether the auto-save restore should be enabled (you can force this to true when testing with a debugger attached) */
+	bool IsRestoreEnabled()
+	{
+		// Note: Restore is disabled when running under the debugger, as programmers
+		// like to just kill applications and we don't want this to count as a crash
+		return !FPlatformMisc::IsDebuggerPresent();
+	}
 }
 
 
@@ -57,8 +65,17 @@ FPackageAutoSaver::FPackageAutoSaver()
 	, bIsAutoSaving(false)
 	, bDelayingDueToFailedSave(false)
 {
-	// Register for the package dirty state updated callback to catch packages that have been modified and need to be saved
-	UPackage::PackageDirtyStateUpdatedEvent.AddRaw(this, &FPackageAutoSaver::OnPackageDirtyStateUpdated);
+	const UEditorLoadingSavingSettings* LoadingSavingSettings = GetDefault<UEditorLoadingSavingSettings>();
+
+	// Set these to the current setting value
+	bWasAutoSavingMaps = LoadingSavingSettings->bAutoSaveMaps;
+	bWasAutoSavingContent = LoadingSavingSettings->bAutoSaveContent;
+
+	// Register for the package dirty state updated callback to catch packages that have been cleaned without being saved
+	UPackage::PackageDirtyStateChangedEvent.AddRaw(this, &FPackageAutoSaver::OnPackageDirtyStateUpdated);
+
+	// Register for the "MarkPackageDirty" callback to catch packages that have been modified and need to be saved
+	UPackage::PackageMarkedDirtyEvent.AddRaw(this, &FPackageAutoSaver::OnMarkPackageDirty);
 
 	// Register for the package modified callback to catch packages that have been saved
 	UPackage::PackageSavedEvent.AddRaw(this, &FPackageAutoSaver::OnPackageSaved);
@@ -66,7 +83,8 @@ FPackageAutoSaver::FPackageAutoSaver()
 
 FPackageAutoSaver::~FPackageAutoSaver()
 {
-	UPackage::PackageDirtyStateUpdatedEvent.RemoveAll(this);
+	UPackage::PackageDirtyStateChangedEvent.RemoveAll(this);
+	UPackage::PackageMarkedDirtyEvent.RemoveAll(this);
 	UPackage::PackageSavedEvent.RemoveAll(this);
 }
 
@@ -151,12 +169,18 @@ void FPackageAutoSaver::AttemptAutoSave()
 
 			if (LoadingSavingSettings->bAutoSaveMaps)
 			{
-				bLevelSaved = FEditorFileUtils::AutosaveMap(AutoSaveDir, NewAutoSaveIndex);
+				// If we weren't saving worlds when we last ran an auto-save, we need to forcibly save any that are dirty
+				// as DirtyPackagesForAutoSave gets cleared each time an auto-save is run, so may not be up-to-date for the new settings
+				const bool bForceSaveWorlds = !bWasAutoSavingMaps && LoadingSavingSettings->bAutoSaveMaps;
+				bLevelSaved = FEditorFileUtils::AutosaveMap(AutoSaveDir, NewAutoSaveIndex, bForceSaveWorlds, DirtyPackagesForAutoSave);
 			}
 
 			if (LoadingSavingSettings->bAutoSaveContent && UnrealEdMisc.GetAutosaveState() != FUnrealEdMisc::EAutosaveState::Cancelled)
 			{
-				bAssetsSaved = FEditorFileUtils::AutosaveContentPackages(AutoSaveDir, NewAutoSaveIndex);
+				// If we weren't saving content packages when we last ran an auto-save, we need to forcibly save any that are dirty
+				// as DirtyPackagesForAutoSave gets cleared each time an auto-save is run, so may not be up-to-date for the new settings
+				const bool bForceSaveContentPackages = !bWasAutoSavingContent && LoadingSavingSettings->bAutoSaveContent;
+				bAssetsSaved = FEditorFileUtils::AutosaveContentPackages(AutoSaveDir, NewAutoSaveIndex, bForceSaveContentPackages, DirtyPackagesForAutoSave);
 			}
 
 			if (bLevelSaved || bAssetsSaved)
@@ -165,10 +189,7 @@ void FPackageAutoSaver::AttemptAutoSave()
 				AutoSaveIndex = NewAutoSaveIndex;
 
 				// Update the restore information
-				// Note: Restore is disabled when running under the debugger, as programmers
-				// like to just kill applications and we don't want this to count as a crash
-				const bool bIsRunningUnderDebugger = FPlatformMisc::IsDebuggerPresent();
-				UpdateRestoreFile(!bIsRunningUnderDebugger);
+				UpdateRestoreFile(PackageAutoSaverJson::IsRestoreEnabled());
 			}
 
 			ResetAutoSaveTimer();
@@ -178,6 +199,13 @@ void FPackageAutoSaver::AttemptAutoSave()
 			{
 				UE_LOG(PackageAutoSaver, Warning, TEXT("Autosave was cancelled."));
 			}
+			else
+			{
+				DirtyPackagesForAutoSave.Empty();
+
+				bWasAutoSavingMaps = LoadingSavingSettings->bAutoSaveMaps;
+				bWasAutoSavingContent = LoadingSavingSettings->bAutoSaveContent;
+			}
 
 			bIsAutoSaving = false;
 			UnrealEdMisc.SetAutosaveState(FUnrealEdMisc::EAutosaveState::Inactive);
@@ -238,11 +266,29 @@ void FPackageAutoSaver::OfferToRestorePackages()
 	}
 }
 
+void FPackageAutoSaver::OnPackagesDeleted(const TArray<UPackage*>& DeletedPackages)
+{
+	ClearStalePointers();
+
+	for(UPackage* DeletedPackage : DeletedPackages)
+	{
+		DirtyPackagesForAutoSave.Remove(DeletedPackage);
+		DirtyPackagesForUserSave.Remove(DeletedPackage);
+	}
+
+	UpdateRestoreFile(PackageAutoSaverJson::IsRestoreEnabled());
+}
+
 void FPackageAutoSaver::OnPackageDirtyStateUpdated(UPackage* Pkg)
 {
 	UpdateDirtyListsForPackage(Pkg);
 }
 
+void FPackageAutoSaver::OnMarkPackageDirty(UPackage* Pkg, bool bWasDirty)
+{
+	UpdateDirtyListsForPackage(Pkg);
+}
+
 void FPackageAutoSaver::OnPackageSaved(const FString& Filename, UObject* Obj)
 {
 	UPackage* const Pkg = Cast<UPackage>(Obj);
@@ -275,11 +321,20 @@ void FPackageAutoSaver::UpdateDirtyListsForPackage(UPackage* Pkg)
 		if(Pkg->IsDirty())
 		{
 			// Always add the package to the user list
-			// Note: Packages get dirtied again after they're auto-saved
 			DirtyPackagesForUserSave.FindOrAdd(Pkg);
+
+			// Only add the package to the auto-save list if we're not auto-saving
+			// Note: Packages get dirtied again after they're auto-saved, so this would add them back again, which we don't want
+			if(!IsAutoSaving())
+			{
+				DirtyPackagesForAutoSave.Add(Pkg);
+			}
 		}
 		else
 		{
+			// Always remove the package from the auto-save list
+			DirtyPackagesForAutoSave.Remove(Pkg);
+
 			// Only remove the package from the user list if we're not auto-saving
 			// Note: Packages call this even when auto-saving, so this would remove them from the user list, which we don't want as they're still dirty
 			if(!IsAutoSaving())
@@ -287,10 +342,7 @@ void FPackageAutoSaver::UpdateDirtyListsForPackage(UPackage* Pkg)
 				if(DirtyPackagesForUserSave.Remove(Pkg))
 				{
 					// Update the restore information too
-					// Note: Restore is disabled when running under the debugger, as programmers
-					// like to just kill applications and we don't want this to count as a crash
-					const bool bIsRunningUnderDebugger = FPlatformMisc::IsDebuggerPresent();
-					UpdateRestoreFile(!bIsRunningUnderDebugger);
+					UpdateRestoreFile(PackageAutoSaverJson::IsRestoreEnabled());
 				}
 			}
 		}
@@ -318,9 +370,14 @@ bool FPackageAutoSaver::DoPackagesNeedAutoSave() const
 {
 	const UEditorLoadingSavingSettings* LoadingSavingSettings = GetDefault<UEditorLoadingSavingSettings>();
 
-	const bool bHasDirtyPackagesForAutoSave = DirtyPackagesForUserSave.Num() != 0;
-	const bool bWorldsMightBeDirty = LoadingSavingSettings->bAutoSaveMaps && bHasDirtyPackagesForAutoSave;
-	const bool bContentPackagesMightBeDirty = LoadingSavingSettings->bAutoSaveContent && bHasDirtyPackagesForAutoSave;
+	// If we weren't saving worlds or content packages when we last ran an auto-save, we need to forcibly save any that are dirty
+	// as DirtyPackagesForAutoSave gets cleared each time an auto-save is run, so may not be up-to-date for the new settings
+	const bool bForceSaveWorlds = !bWasAutoSavingMaps && LoadingSavingSettings->bAutoSaveMaps;
+	const bool bForceSaveContentPackages = !bWasAutoSavingContent && LoadingSavingSettings->bAutoSaveContent;
+
+	const bool bHasDirtyPackagesForAutoSave = DirtyPackagesForAutoSave.Num() != 0;
+	const bool bWorldsMightBeDirty = LoadingSavingSettings->bAutoSaveMaps && (bForceSaveWorlds || bHasDirtyPackagesForAutoSave);
+	const bool bContentPackagesMightBeDirty = LoadingSavingSettings->bAutoSaveContent && (bForceSaveContentPackages || bHasDirtyPackagesForAutoSave);
 	const bool bPackagesNeedAutoSave = bWorldsMightBeDirty || bContentPackagesMightBeDirty;
 
 	return bPackagesNeedAutoSave;
@@ -334,7 +391,7 @@ FText FPackageAutoSaver::GetAutoSaveNotificationText(const int32 TimeInSecondsUn
 		// Count down the time
 		FFormatNamedArguments Args;
 		Args.Add(TEXT("TimeInSecondsUntilAutosave"), TimeInSecondsUntilAutosave);
-		Args.Add(TEXT("DirtyPackagesCount"), DirtyPackagesForUserSave.Num());
+		Args.Add(TEXT("DirtyPackagesCount"), DirtyPackagesForAutoSave.Num());
 		return (DirtyPackagesForUserSave.Num() == 1)
 			? FText::Format(NSLOCTEXT("AutoSaveNotify", "AutoSaveIn", "Autosave in {TimeInSecondsUntilAutosave} seconds"), Args)
 			: FText::Format(NSLOCTEXT("AutoSaveNotify", "AutoSaveXPackagesIn", "Autosave in {TimeInSecondsUntilAutosave} seconds for {DirtyPackagesCount} items"), Args);
@@ -502,6 +559,16 @@ void FPackageAutoSaver::ClearStalePointers()
 			DirtyPackagesForUserSave.Remove(Package);
 		}
 	}
+
+	auto DirtyPackagesForAutoSaveTmp = DirtyPackagesForAutoSave;
+	for(auto It = DirtyPackagesForAutoSaveTmp.CreateConstIterator(); It; ++It)
+	{
+		const TWeakObjectPtr<UPackage>& Package = *It;
+		if(!Package.IsValid())
+		{
+			DirtyPackagesForAutoSave.Remove(Package);
+		}
+	}
 }
 
 
diff --git a/Engine/Source/Editor/UnrealEd/Private/PackageAutoSaver.h b/Engine/Source/Editor/UnrealEd/Private/PackageAutoSaver.h
--- a/Engine/Source/Editor/UnrealEd/Private/PackageAutoSaver.h
+++ b/Engine/Source/Editor/UnrealEd/Private/PackageAutoSaver.h
@@ -1,7 +1,6 @@
 // Copyright 1998-2014 Epic Games, Inc. All Rights Reserved.
 
-#ifndef __PACKAGEAUTOSAVER_H__
-#define __PACKAGEAUTOSAVER_H__
+#pragma once
 
 #include "IPackageAutoSaver.h"
 
@@ -24,6 +23,7 @@ public:
 	virtual void UpdateRestoreFile(const bool bRestoreEnabled) const override;
 	virtual bool HasPackagesToRestore() const override;
 	virtual void OfferToRestorePackages() override;
+	virtual void OnPackagesDeleted(const TArray<UPackage*>& DeletedPackages) override;
 	virtual bool IsAutoSaving() const override
 	{
 		return bIsAutoSaving;
@@ -37,6 +37,14 @@ private:
 	 */
 	void OnPackageDirtyStateUpdated(UPackage* Pkg);
 
+	/**
+	 * Called when a package is marked as dirty
+	 *
+	 * @param Pkg The package that was marked as dirty
+	 * @param bWasDirty Whether the package was previously dirty before the call to MarkPackageDirty
+	 */
+	void OnMarkPackageDirty(UPackage* Pkg, bool bWasDirty);
+
 	/**
 	 * Called when a package has been saved
 	 *
@@ -98,6 +106,12 @@ private:
 	/** Flag for whether auto-save warning notification has been launched */
 	bool bAutoSaveNotificationLaunched;
 
+	/** Were we saving worlds last time an auto-save ran? */
+	bool bWasAutoSavingMaps;
+
+	/** Were we saving worlds last time an auto-save ran? */
+	bool bWasAutoSavingContent;
+
 	/** If we are delaying the time a little bit because we failed to save */
 	bool bDelayingDueToFailedSave;
 
@@ -107,8 +121,9 @@ private:
 	/** Packages that have been dirtied and not saved by the user, mapped to their latest auto-save file */
 	TMap<TWeakObjectPtr<UPackage>, FString> DirtyPackagesForUserSave;
 
+	/** Packages that have been dirtied and not saved by the auto-saver */
+	TSet<TWeakObjectPtr<UPackage>> DirtyPackagesForAutoSave;
+
 	/** Restore information that was loaded following a crash */
 	TMap<FString, FString> PackagesThatCanBeRestored;
 };
-
-#endif // __PACKAGEAUTOSAVER_H__
diff --git a/Engine/Source/Editor/UnrealEd/Private/PackageRestore.cpp b/Engine/Source/Editor/UnrealEd/Private/PackageRestore.cpp
--- a/Engine/Source/Editor/UnrealEd/Private/PackageRestore.cpp
+++ b/Engine/Source/Editor/UnrealEd/Private/PackageRestore.cpp
@@ -12,10 +12,11 @@ namespace PackageRestore
 	class FPackageRestoreItem : public TSharedFromThis<FPackageRestoreItem>
 	{
 	public:
-		FPackageRestoreItem(const FString& InPackageName, const FString& InPackageFilename, const FString& InAutoSaveFilename)
+		FPackageRestoreItem(const FString& InPackageName, const FString& InPackageFilename, const FString& InAutoSaveFilename, const bool bInIsExistingPackage)
 			: PackageName(InPackageName)
 			, PackageFilename(InPackageFilename)
 			, AutoSaveFilename(InAutoSaveFilename)
+			, bIsExistingPackage(bInIsExistingPackage)
 			, State(ESlateCheckBoxState::Unchecked)
 		{
 		}
@@ -38,6 +39,12 @@ namespace PackageRestore
 			return AutoSaveFilename;
 		}
 
+		/** @return True if this item is to replace an existing package, or false if it is to add a new package */
+		bool IsExistingPackage() const
+		{
+			return bIsExistingPackage;
+		}
+
 		/** @return The state of this item (checked, unchecked) */
 		ESlateCheckBoxState::Type GetState() const
 		{
@@ -89,6 +96,7 @@ namespace PackageRestore
 		FString PackageName;
 		FString PackageFilename;
 		FString AutoSaveFilename;
+		bool bIsExistingPackage;
 		ESlateCheckBoxState::Type State;
 	};
 
@@ -400,18 +408,14 @@ namespace PackageRestore
 		 * Get the return code for this dlg, as well as some useful information about what was selected
 		 *
 		 * @param SelectedPackageItems Array to fill with the list items the user wants to restore
-		 * @param SelectedPackageNames Array to fill with the package names the user wants to restore
 		 *
 		 * @return true if we should perform an import, false if the user cancelled
 		 */
-		bool GetReturnType(FPackageRestoreItems& SelectedPackageItems, TArray<FString>& SelectedPackageNames) const
+		bool GetReturnType(FPackageRestoreItems& SelectedPackageItems) const
 		{
 			SelectedPackageItems.Empty();
 			SelectedPackageItems.Reserve(PackageRestoreItems->Num());
 
-			SelectedPackageNames.Empty();
-			SelectedPackageNames.Reserve(PackageRestoreItems->Num());
-
 			// Get the list of packages selected to be restored
 			for(auto It = PackageRestoreItems->CreateConstIterator(); It; ++It)
 			{
@@ -420,7 +424,6 @@ namespace PackageRestore
 				if(ListItem->GetState() == ESlateCheckBoxState::Checked)
 				{
 					SelectedPackageItems.Add(ListItem);
-					SelectedPackageNames.Add(ListItem->GetPackageName());
 				}
 			}
 
@@ -446,6 +449,11 @@ namespace PackageRestore
 		{
 			const FPackageRestoreItemPtr& RestoreItem = *It;
 
+			if(!RestoreItem->IsExistingPackage())
+			{
+				continue;
+			}
+
 			UPackage* const Package = FindPackage(nullptr, *RestoreItem->GetPackageName());
 			if(Package)
 			{
@@ -519,7 +527,16 @@ FEditorFileUtils::EPromptReturnCode PackageRestore::PromptToRestorePackages(cons
 		FString PackageFilename;
 		if(FPackageName::DoesPackageExist(PackageFullPath, nullptr, &PackageFilename))
 		{
-			FPackageRestoreItemPtr PackageItemPtr = MakeShareable(new FPackageRestoreItem(PackageFullPath, PackageFilename, AutoSaveDir / AutoSavePath));
+			FPackageRestoreItemPtr PackageItemPtr = MakeShareable(new FPackageRestoreItem(PackageFullPath, PackageFilename, AutoSaveDir / AutoSavePath, true/*bIsExistingPackage*/));
+			PackageRestoreItems.Add(PackageItemPtr);
+		}
+		else
+		{
+			// A package may not exist on disk if it was for a newly added or imported asset, which hasn't yet had SaveDirtyPackages called for it
+			PackageFilename = FPackageName::LongPackageNameToFilename(PackageFullPath); // no extension yet
+			PackageFilename += FPaths::GetExtension(AutoSavePath, true/*bIncludeDot*/);
+
+			FPackageRestoreItemPtr PackageItemPtr = MakeShareable(new FPackageRestoreItem(PackageFullPath, PackageFilename, AutoSaveDir / AutoSavePath, false/*bIsExistingPackage*/));
 			PackageRestoreItems.Add(PackageItemPtr);
 		}
 	}
@@ -546,17 +563,32 @@ FEditorFileUtils::EPromptReturnCode PackageRestore::PromptToRestorePackages(cons
 
 	// Get the return code, and work out what we need to restore
 	FPackageRestoreItems SelectedPackageItems;
-	TArray<FString> SelectedPackageNames;
-	if(!PackageRestoreDlgRef->GetReturnType(SelectedPackageItems, SelectedPackageNames))
+	if(!PackageRestoreDlgRef->GetReturnType(SelectedPackageItems))
 	{
 		return FEditorFileUtils::PR_Declined;
 	}
 
 	// Try and ensure that these packages are checked-out by the source control system
-	// Note: This may fail and present the user with an error message, however we still 
-	// want to continue as they may have checked out some packages that could now be restored
-	const bool bErrorIfAlreadyCheckedOut = false; // some of the packages might already be checked out; that isn't an error
-	FEditorFileUtils::CheckoutPackages(SelectedPackageNames, nullptr, bErrorIfAlreadyCheckedOut);
+	{
+		TArray<FString> SelectedPackageNames;
+		SelectedPackageNames.Reserve(SelectedPackageItems.Num());
+
+		// Get an array of selected package names to check out
+		for(auto It = SelectedPackageItems.CreateConstIterator(); It; ++It)
+		{
+			const FPackageRestoreItemPtr& SelectedPackageItem = *It;
+
+			if(SelectedPackageItem->IsExistingPackage())
+			{
+				SelectedPackageNames.Add(SelectedPackageItem->GetPackageName());
+			}
+		}
+
+		// Note: This may fail and present the user with an error message, however we still 
+		// want to continue as they may have checked out some packages that could now be restored
+		const bool bErrorIfAlreadyCheckedOut = false; // some of the packages might already be checked out; that isn't an error
+		FEditorFileUtils::CheckoutPackages(SelectedPackageNames, nullptr, bErrorIfAlreadyCheckedOut);
+	}
 
 	// It's possible that some packages may have already been loaded by the editor
 	// If they have, we need to forcibly unload them so that we can overwrite their files
diff --git a/Engine/Source/Editor/UnrealEd/Private/UnrealEdEngine.cpp b/Engine/Source/Editor/UnrealEd/Private/UnrealEdEngine.cpp
--- a/Engine/Source/Editor/UnrealEd/Private/UnrealEdEngine.cpp
+++ b/Engine/Source/Editor/UnrealEd/Private/UnrealEdEngine.cpp
@@ -39,7 +39,7 @@ void UUnrealEdEngine::Init(IEngineLoop* InEngineLoop)
 	PackageAutoSaver->LoadRestoreFile();
 
 	// Register for the package dirty state updated callback to catch packages that have been modified and need to be checked out.
-	UPackage::PackageDirtyStateUpdatedEvent.AddUObject(this, &UUnrealEdEngine::OnPackageDirtyStateUpdated);
+	UPackage::PackageDirtyStateChangedEvent.AddUObject(this, &UUnrealEdEngine::OnPackageDirtyStateUpdated);
 	
 	// Register to the PostGarbageCollect delegate, as we want to use this to trigger the RefreshAllBroweser delegate from 
 	// here rather then from Core
@@ -223,7 +223,7 @@ void UUnrealEdEngine::FinishDestroy()
 		PackageAutoSaver.Reset();
 	}
 
-	UPackage::PackageDirtyStateUpdatedEvent.RemoveAll(this);
+	UPackage::PackageDirtyStateChangedEvent.RemoveAll(this);
 	FCoreDelegates::PostGarbageCollect.RemoveAll(this);
 	FCoreDelegates::ColorPickerChanged.RemoveAll(this);
 	Super::FinishDestroy();
diff --git a/Engine/Source/Editor/UnrealEd/Public/FileHelpers.h b/Engine/Source/Editor/UnrealEd/Public/FileHelpers.h
--- a/Engine/Source/Editor/UnrealEd/Public/FileHelpers.h
+++ b/Engine/Source/Editor/UnrealEd/Public/FileHelpers.h
@@ -105,18 +105,22 @@ public:
 	 *
 	 * @param	AbsoluteAutosaveDir			Autosave directory.
 	 * @param	AutosaveIndex				Integer prepended to autosave filenames..
+	 * @param	bForceIfNotInList			Should the save be forced if the package is dirty, but not in DirtyPackagesForAutoSave?
+	 * @param	DirtyPackagesForAutoSave	A set of packages that are considered by the auto-save system to be dirty, you should check this to see if a package needs saving
 	 */
-	static bool AutosaveMap(const FString& AbsoluteAutosaveDir, const int32 AutosaveIndex);
+	static bool AutosaveMap(const FString& AbsoluteAutosaveDir, const int32 AutosaveIndex, const bool bForceIfNotInList, const TSet< TWeakObjectPtr<UPackage> >& DirtyPackagesForAutoSave);
 
 	/**
 	 * Saves all asset packages to the specified directory.
 	 *
 	 * @param	AbsoluteAutosaveDir			Autosave directory.
 	 * @param	AutosaveIndex				Integer prepended to autosave filenames.
+	 * @param	bForceIfNotInList			Should the save be forced if the package is dirty, but not in DirtyPackagesForAutoSave?
+	 * @param	DirtyPackagesForAutoSave	A set of packages that are considered by the auto-save system to be dirty, you should check this to see if a package needs saving
 	 *
 	 * @return	true if one or more packages were autosaved; false otherwise
 	 */
-	static bool AutosaveContentPackages(const FString& AbsoluteAutosaveDir, const int32 AutosaveIndex);
+	static bool AutosaveContentPackages(const FString& AbsoluteAutosaveDir, const int32 AutosaveIndex, const bool bForceIfNotInList, const TSet< TWeakObjectPtr<UPackage> >& DirtyPackagesForAutoSave);
 
 	/**
 	 * Looks at all currently loaded packages and saves them if their "bDirty" flag is set, optionally prompting the user to select which packages to save)
diff --git a/Engine/Source/Editor/UnrealEd/Public/IPackageAutoSaver.h b/Engine/Source/Editor/UnrealEd/Public/IPackageAutoSaver.h
--- a/Engine/Source/Editor/UnrealEd/Public/IPackageAutoSaver.h
+++ b/Engine/Source/Editor/UnrealEd/Public/IPackageAutoSaver.h
@@ -1,7 +1,6 @@
 // Copyright 1998-2014 Epic Games, Inc. All Rights Reserved.
 
-#ifndef __IPACKAGEAUTOSAVER_H__
-#define __IPACKAGEAUTOSAVER_H__
+#pragma once
 
 /** An interface to handle the creation, destruction, and restoration of auto-saved packages */
 class UNREALED_API IPackageAutoSaver
@@ -56,6 +55,7 @@ public:
 
 	/** Offer the user the chance to restore any packages that were dirty and have auto-saves */
 	virtual void OfferToRestorePackages() = 0;
-};
 
-#endif // __IPACKAGEAUTOSAVER_H__
+	/** Called when packages are deleted in the editor */
+	virtual void OnPackagesDeleted(const TArray<UPackage*>& DeletedPackages) = 0;
+};
diff --git a/Engine/Source/Runtime/CoreUObject/Private/UObject/Package.cpp b/Engine/Source/Runtime/CoreUObject/Private/UObject/Package.cpp
--- a/Engine/Source/Runtime/CoreUObject/Private/UObject/Package.cpp
+++ b/Engine/Source/Runtime/CoreUObject/Private/UObject/Package.cpp
@@ -9,10 +9,17 @@
 /** Delegate to notify subscribers when a package has been saved. This is triggered when the package saving
  *  has completed and was successful. */
 UPackage::FOnPackageSaved UPackage::PackageSavedEvent;
-/** Delegate to notify subscribers about the dirty state of a package being set.
+/** Delegate to notify subscribers when the dirty state of a package is changed.
  *  Allows the editor to register the modified package as one that should be prompted for source control checkout. 
  *  Use Package->IsDirty() to get the updated dirty state of the package */
-UPackage::FOnPackageDirtyStateUpdated UPackage::PackageDirtyStateUpdatedEvent;
+UPackage::FOnPackageDirtyStateChanged UPackage::PackageDirtyStateChangedEvent;
+/** 
+ * Delegate to notify subscribers when a package is marked as dirty via UObjectBaseUtilty::MarkPackageDirty 
+ * Note: Unlike FOnPackageDirtyStateChanged, this is always called, even when the package is already dirty
+ * Use bWasDirty to check the previous dirty state of the package
+ * Use Package->IsDirty() to get the updated dirty state of the package
+ */
+UPackage::FOnPackageMarkedDirty UPackage::PackageMarkedDirtyEvent;
 
 void UPackage::PostInitProperties()
 {
@@ -57,7 +64,7 @@ void UPackage::SetDirtyFlag( bool bIsDirty )
 			&& GetTransientPackage() != this )			// Skip the transient package
 		{
 			// Package is changing dirty state, let the editor know so we may prompt for source control checkout
-			PackageDirtyStateUpdatedEvent.Broadcast(this);
+			PackageDirtyStateChangedEvent.Broadcast(this);
 		}
 	}
 }
diff --git a/Engine/Source/Runtime/CoreUObject/Private/UObject/UObjectBaseUtility.cpp b/Engine/Source/Runtime/CoreUObject/Private/UObject/UObjectBaseUtility.cpp
--- a/Engine/Source/Runtime/CoreUObject/Private/UObject/UObjectBaseUtility.cpp
+++ b/Engine/Source/Runtime/CoreUObject/Private/UObject/UObjectBaseUtility.cpp
@@ -131,12 +131,19 @@ void UObjectBaseUtility::MarkPackageDirty() const
 			// It is against policy to dirty a map or package during load in the Editor, to enforce this policy
 			// we explicitly disable the ability to dirty a package or map during load.  Commandlets can still
 			// set the dirty state on load.
-			// We also prevent needless re-dirtying as this can be an expensive operation.
-			if( !Package->IsDirty() &&
-				(!GIsEditor || IsRunningCommandlet() || 
+			if( (!GIsEditor || IsRunningCommandlet() || 
 				(GIsEditor && !GIsEditorLoadingPackage)))
 			{
-				Package->SetDirtyFlag(true);
+				const bool bIsDirty = Package->IsDirty();
+
+				// We prevent needless re-dirtying as this can be an expensive operation.
+				if( !bIsDirty )
+				{
+					Package->SetDirtyFlag(true);
+				}
+
+				// Always call PackageMarkedDirtyEvent, even when the package is already dirty
+				Package->PackageMarkedDirtyEvent.Broadcast(Package, bIsDirty);
 			}
 		}
 	}
diff --git a/Engine/Source/Runtime/CoreUObject/Public/UObject/CoreObject.h b/Engine/Source/Runtime/CoreUObject/Public/UObject/CoreObject.h
--- a/Engine/Source/Runtime/CoreUObject/Public/UObject/CoreObject.h
+++ b/Engine/Source/Runtime/CoreUObject/Public/UObject/CoreObject.h
@@ -25,17 +25,26 @@ class COREUOBJECT_API UPackage : public UObject
 
 public:
 	/** delegate type for package dirty state events.  ( Params: UPackage* ModifiedPackage ) */
-	DECLARE_MULTICAST_DELEGATE_OneParam(FOnPackageDirtyStateUpdated, class UPackage*);
+	DECLARE_MULTICAST_DELEGATE_OneParam(FOnPackageDirtyStateChanged, class UPackage*);
 	/** delegate type for package saved events ( Params: const FString& PackageFileName, UObject* Outer ) */
 	DECLARE_MULTICAST_DELEGATE_TwoParams(FOnPackageSaved, const FString&, UObject*);					
+	/** delegate type for when a package is marked as dirty via UObjectBaseUtilty::MarkPackageDirty ( Params: UPackage* ModifiedPackage, bool bWasDirty ) */
+	DECLARE_MULTICAST_DELEGATE_TwoParams(FOnPackageMarkedDirty, class UPackage*, bool);
 
 	/** Delegate to notify subscribers when a package has been saved. This is triggered when the package saving
 	 *  has completed and was successfull. */
 	static FOnPackageSaved PackageSavedEvent;
-	/** Delegate to notify subscribers about the dirty state of a package being set.
+	/** Delegate to notify subscribers when the dirty state of a package is changed.
 	 *  Allows the editor to register the modified package as one that should be prompted for source control checkout. 
 	 *  Use Package->IsDirty() to get the updated dirty state of the package */
-	static FOnPackageDirtyStateUpdated PackageDirtyStateUpdatedEvent;
+	static FOnPackageDirtyStateChanged PackageDirtyStateChangedEvent;
+	/** 
+	 * Delegate to notify subscribers when a package is marked as dirty via UObjectBaseUtilty::MarkPackageDirty 
+	 * Note: Unlike FOnPackageDirtyStateChanged, this is always called, even when the package is already dirty
+	 * Use bWasDirty to check the previous dirty state of the package
+	 * Use Package->IsDirty() to get the updated dirty state of the package
+	 */
+	static FOnPackageMarkedDirty PackageMarkedDirtyEvent;
 
 private:
 #if !IS_MONOLITHIC
