Fix FBX material counting problem causing some materials to not be imported

diff --git a/Engine/Source/Editor/UnrealEd/Private/Fbx/FbxStaticMeshImport.cpp b/Engine/Source/Editor/UnrealEd/Private/Fbx/FbxStaticMeshImport.cpp
--- a/Engine/Source/Editor/UnrealEd/Private/Fbx/FbxStaticMeshImport.cpp
+++ b/Engine/Source/Editor/UnrealEd/Private/Fbx/FbxStaticMeshImport.cpp
@@ -657,7 +657,7 @@ bool UnFbx::FFbxImporter::BuildStaticMeshFromGeometry(FbxMesh* Mesh, UStaticMesh
 			MaterialIndex = 0;
 		}
 	
-		RawMesh.FaceMaterialIndices[DestTriangleIndex] = FMath::Min( MaterialIndex, MAX_MESH_MATERIAL_INDEX );
+		RawMesh.FaceMaterialIndices[DestTriangleIndex] = MaterialIndex;
 	}
 
 
