UE4: Added a warning message letting the user know that their UWorld will be lost if they created it, never saved it, and declined the last minute save. Unlike normal assets, worlds get GC'd when they are no longer being edited. TTP#338376

diff --git a/Engine/Source/Editor/UnrealEd/Classes/Editor/EditorEngine.h b/Engine/Source/Editor/UnrealEd/Classes/Editor/EditorEngine.h
--- a/Engine/Source/Editor/UnrealEd/Classes/Editor/EditorEngine.h
+++ b/Engine/Source/Editor/UnrealEd/Classes/Editor/EditorEngine.h
@@ -2171,6 +2171,13 @@ public:
 	 */
 	bool ShouldAbortBecauseOfPIEWorld() const;
 
+	/**
+	 * If an unsaved world exists that would be lost in a map transition, give the user the option to cancel a map load.
+	 *
+	 * @return				true if an unsaved world exists and the user refused to continue, false otherwise.
+	 */
+	bool ShouldAbortBecauseOfUnsavedWorld() const;
+
 	/**
 	 * Assigns a new label to an actor. If the name exists it will be appended with a number to make it unique. Actor labels are only available in development builds.
 	 *
diff --git a/Engine/Source/Editor/UnrealEd/Private/EditorServer.cpp b/Engine/Source/Editor/UnrealEd/Private/EditorServer.cpp
--- a/Engine/Source/Editor/UnrealEd/Private/EditorServer.cpp
+++ b/Engine/Source/Editor/UnrealEd/Private/EditorServer.cpp
@@ -1724,6 +1724,34 @@ bool UEditorEngine::ShouldAbortBecauseOfPIEWorld() const
 	return false;
 }
 
+bool UEditorEngine::ShouldAbortBecauseOfUnsavedWorld() const
+{
+	// If an unsaved world exists that would be lost in a map transition, give the user the option to cancel a map load.
+
+	// First check if we have a world and it is dirty
+	UWorld* EditorWorld = GEditor->GetEditorWorldContext().World();
+	if (EditorWorld && EditorWorld->GetOutermost()->IsDirty())
+	{
+		// Now check if the world is in a path that can be saved (otherwise it is in something like the transient package or temp)
+		const FString PackageName = EditorWorld->GetOutermost()->GetName();
+		const bool bIncludeReadOnlyRoots = false;
+		if ( FPackageName::IsValidLongPackageName(PackageName, bIncludeReadOnlyRoots) )
+		{
+			// Now check if the file exists on disk. If it does, it won't be "lost" when GC'd.
+			if ( !FPackageName::DoesPackageExist(PackageName) )
+			{
+				// This world will be completely lost if a map transition happens. Warn the user that this is happening and ask him/her how to proceed.
+				if (EAppReturnType::Yes != FMessageDialog::Open(EAppMsgType::YesNo, FText::Format(NSLOCTEXT("UnrealEd", "Prompt_ThisActionWillDiscardWorldContinue", "The unsaved level {0} will be lost.  Continue?"), FText::FromString(EditorWorld->GetName()))))
+				{
+					// User doesn't want to lose the world -- abort the load.
+					return true;
+				}
+			}
+		}
+	}
+	return false;
+}
+
 /**
  * Prompts the user to save the current map if necessary, then creates a new (blank) map.
  */
@@ -1745,6 +1773,11 @@ void UEditorEngine::CreateNewMapForEditing()
 		// something went wrong or the user pressed cancel.  Return to the editor so the user doesn't lose their changes		
 		return;
 	}
+
+	if ( ShouldAbortBecauseOfUnsavedWorld() )
+	{
+		return;
+	}
 	
 	const FScopedBusyCursor BusyCursor;
 
diff --git a/Engine/Source/Editor/UnrealEd/Private/FileHelpers.cpp b/Engine/Source/Editor/UnrealEd/Private/FileHelpers.cpp
--- a/Engine/Source/Editor/UnrealEd/Private/FileHelpers.cpp
+++ b/Engine/Source/Editor/UnrealEd/Private/FileHelpers.cpp
@@ -1679,6 +1679,12 @@ void FEditorFileUtils::LoadMap(const FString& InFilename, bool LoadAsTemplate, b
 		return;
 	}
 
+	// If a level is in memory but never saved to disk, warn the user that the level will be lost.
+	if (GEditor->ShouldAbortBecauseOfUnsavedWorld())
+	{
+		return;
+	}
+
 	// Change out of Matinee when opening new map, so we avoid editing data in the old one.
 	if( GLevelEditorModeTools().IsModeActive( FBuiltinEditorModes::EM_InterpEdit ) )
 	{
