Editor: Add option to TriangulatePoly to avoid discarding colinear points (useful when subdividing something that will later be deformed and need to retain specified detail, etc...)

diff --git a/Engine/Source/Editor/UnrealEd/Private/GeomTools.cpp b/Engine/Source/Editor/UnrealEd/Private/GeomTools.cpp
--- a/Engine/Source/Editor/UnrealEd/Private/GeomTools.cpp
+++ b/Engine/Source/Editor/UnrealEd/Private/GeomTools.cpp
@@ -477,7 +477,7 @@ static bool AreEdgesMergeable(
 }
 
 /** Given a polygon, decompose into triangles and append to OutTris. */
-bool TriangulatePoly(TArray<FClipSMTriangle>& OutTris, const FClipSMPolygon& InPoly)
+bool TriangulatePoly(TArray<FClipSMTriangle>& OutTris, const FClipSMPolygon& InPoly, bool bKeepColinearVertices)
 {
 	// Can't work if not enough verts for 1 triangle
 	if(InPoly.Vertices.Num() < 3)
@@ -492,16 +492,19 @@ bool TriangulatePoly(TArray<FClipSMTriangle>& OutTris, const FClipSMPolygon& InP
 	// Keep iterating while there are still vertices
 	while(true)
 	{
-		// Cull redundant vertex edges from the polygon.
-		for(int32 VertexIndex = 0;VertexIndex < PolyVerts.Num();VertexIndex++)
+		if (!bKeepColinearVertices)
 		{
-			const int32 I0 = (VertexIndex + 0) % PolyVerts.Num();
-			const int32 I1 = (VertexIndex + 1) % PolyVerts.Num();
-			const int32 I2 = (VertexIndex + 2) % PolyVerts.Num();
-			if(AreEdgesMergeable(PolyVerts[I0],PolyVerts[I1],PolyVerts[I2]))
+			// Cull redundant vertex edges from the polygon.
+			for (int32 VertexIndex = 0; VertexIndex < PolyVerts.Num(); VertexIndex++)
 			{
-				PolyVerts.RemoveAt(I1);
-				VertexIndex--;
+				const int32 I0 = (VertexIndex + 0) % PolyVerts.Num();
+				const int32 I1 = (VertexIndex + 1) % PolyVerts.Num();
+				const int32 I2 = (VertexIndex + 2) % PolyVerts.Num();
+				if (AreEdgesMergeable(PolyVerts[I0], PolyVerts[I1], PolyVerts[I2]))
+				{
+					PolyVerts.RemoveAt(I1);
+					VertexIndex--;
+				}
 			}
 		}
 
diff --git a/Engine/Source/Editor/UnrealEd/Public/GeomTools.h b/Engine/Source/Editor/UnrealEd/Public/GeomTools.h
--- a/Engine/Source/Editor/UnrealEd/Public/GeomTools.h
+++ b/Engine/Source/Editor/UnrealEd/Public/GeomTools.h
@@ -1,10 +1,6 @@
 // Copyright 1998-2014 Epic Games, Inc. All Rights Reserved.
 
-/*=============================================================================
- =============================================================================*/
-
-#ifndef __GEOMTOOLS_H__
-#define __GEOMTOOLS_H__
+#pragma once
 
 /** An edge in 3D space, used by these utility functions. */
 struct FUtilEdge3D
@@ -137,7 +133,7 @@ void Buid2DPolysFromEdges(TArray<FUtilPoly2D>& OutPolys, const TArray<FUtilEdge2
 /** Given a polygon, decompose into triangles and append to OutTris.
   * @return	true if the triangulation was successful.
   */
-UNREALED_API bool TriangulatePoly(TArray<FClipSMTriangle>& OutTris, const FClipSMPolygon& InPoly);
+UNREALED_API bool TriangulatePoly(TArray<FClipSMTriangle>& OutTris, const FClipSMPolygon& InPoly, bool bKeepColinearVertices = false);
 
 /** Transform triangle from 2D to 3D static-mesh triangle. */
 FClipSMPolygon Transform2DPolygonToSMPolygon(const FUtilPoly2D& InTri, const FMatrix& InMatrix);
@@ -150,5 +146,3 @@ void RemoveRedundantTriangles(TArray<FClipSMTriangle>& Tris);
 
 /** Split 2D polygons with a 3D plane. */
 void Split2DPolysWithPlane(FUtilPoly2DSet& PolySet, const FPlane& Plane, const FColor& ExteriorVertColor, const FColor& InteriorVertColor);
-
-#endif
