Sorted map selection list in Project Settings -> Maps & Modes dialog.
#ttp 338498 - EDITOR: SETTINGS: Default map selection list is not sorted
#branch UE4
#proj Editor.UnrealEd
#add Added calls to Sort in FGameMapsSettingsCustomization::MakeMapMenu.

diff --git a/Engine/Source/Editor/UnrealEd/Private/Settings/GameMapsSettingsCustomization.h b/Engine/Source/Editor/UnrealEd/Private/Settings/GameMapsSettingsCustomization.h
--- a/Engine/Source/Editor/UnrealEd/Private/Settings/GameMapsSettingsCustomization.h
+++ b/Engine/Source/Editor/UnrealEd/Private/Settings/GameMapsSettingsCustomization.h
@@ -155,6 +155,7 @@ protected:
 		MenuBuilder.BeginSection(NAME_None, LOCTEXT("ProjectMapsSectionHeader", "Project"));
 		{
 			IFileManager::Get().FindFilesRecursive(MapNames, *FPaths::GameContentDir(), *MapFileWildCard, true, false);
+			MapNames.Sort([](const FString& A, const FString& B) { return FPaths::GetBaseFilename(A) < FPaths::GetBaseFilename(B); });
 
 			for (int32 i = 0; i < MapNames.Num(); i++)
 			{
@@ -173,6 +174,7 @@ protected:
 		MenuBuilder.BeginSection(NAME_None, LOCTEXT("EngineMapsSectionHeader", "Engine"));
 		{
 			IFileManager::Get().FindFilesRecursive(MapNames, *FPaths::EngineContentDir(), *MapFileWildCard, true, false);
+			MapNames.Sort([](const FString& A, const FString& B) { return FPaths::GetBaseFilename(A) < FPaths::GetBaseFilename(B); });
 
 			for (int32 i = 0; i < MapNames.Num(); i++)
 			{
