Pass command line renderer choice through for movie playback command line so it uses same renderer as Editor.

diff --git a/Engine/Source/Editor/UnrealEd/Private/PlayLevel.cpp b/Engine/Source/Editor/UnrealEd/Private/PlayLevel.cpp
--- a/Engine/Source/Editor/UnrealEd/Private/PlayLevel.cpp
+++ b/Engine/Source/Editor/UnrealEd/Private/PlayLevel.cpp
@@ -1335,6 +1335,18 @@ void UEditorEngine::PlayForMovieCapture()
 	FString EditorCommandLine = SavedMapNames[0];
 	EditorCommandLine += TEXT(" -game");
 
+	// renderer overrides - hack
+	EditorCommandLine += FParse::Param(FCommandLine::Get(), TEXT("d3d11"))		?	TEXT(" -d3d11")		: TEXT("");
+	EditorCommandLine += FParse::Param(FCommandLine::Get(), TEXT("sm5"))		?	TEXT(" -sm5")		: TEXT("");
+	EditorCommandLine += FParse::Param(FCommandLine::Get(), TEXT("dx11"))		?	TEXT(" -dx11")		: TEXT("");
+	EditorCommandLine += FParse::Param(FCommandLine::Get(), TEXT("d3d10"))		?	TEXT(" -d3d10")		: TEXT("");
+	EditorCommandLine += FParse::Param(FCommandLine::Get(), TEXT("sm4"))		?	TEXT(" -sm4")		: TEXT("");
+	EditorCommandLine += FParse::Param(FCommandLine::Get(), TEXT("dx10"))		?	TEXT(" -dx10")		: TEXT("");
+	EditorCommandLine += FParse::Param(FCommandLine::Get(), TEXT("opengl"))		?	TEXT(" -opengl")	: TEXT("");
+	EditorCommandLine += FParse::Param(FCommandLine::Get(), TEXT("opengl3"))	?	TEXT(" -opengl3")	: TEXT("");
+	EditorCommandLine += FParse::Param(FCommandLine::Get(), TEXT("opengl4"))	?	TEXT(" -opengl4")	: TEXT("");
+
+
 	// this parameter tells UGameEngine to add the auto-save dir to the paths array and repopulate the package file cache
 	// this is needed in order to support streaming levels as the streaming level packages will be loaded only when needed (thus
 	// their package names need to be findable by the package file caching system)
