Cleaning trans buffer when level removed from world

diff --git a/Engine/Source/Editor/Levels/Private/LevelCollectionViewModel.cpp b/Engine/Source/Editor/Levels/Private/LevelCollectionViewModel.cpp
--- a/Engine/Source/Editor/Levels/Private/LevelCollectionViewModel.cpp
+++ b/Engine/Source/Editor/Levels/Private/LevelCollectionViewModel.cpp
@@ -1553,9 +1553,7 @@ void FLevelCollectionViewModel::RemoveSelectedLevels_Executed()
 		
 		EditorLevelUtils::RemoveLevelFromWorld( CurLevel );
 	}
-
-	GEditor->ResetTransaction( LOCTEXT("RemoveLevelTransReset", "Removing Levels from World") );
-
+	
 	// Collect garbage to clear out the destroyed level
 	CollectGarbage( GARBAGE_COLLECTION_KEEPFLAGS );
 }
diff --git a/Engine/Source/Editor/UnrealEd/Private/EditorLevelUtils.cpp b/Engine/Source/Editor/UnrealEd/Private/EditorLevelUtils.cpp
--- a/Engine/Source/Editor/UnrealEd/Private/EditorLevelUtils.cpp
+++ b/Engine/Source/Editor/UnrealEd/Private/EditorLevelUtils.cpp
@@ -335,6 +335,9 @@ namespace EditorLevelUtils
 
 		GEditor->CloseEditedWorldAssets(CastChecked<UWorld>(InLevel->GetOuter()));
 
+		// Need to clean refs to a removed level 
+		GEditor->ResetTransaction( LOCTEXT("RemoveLevelTransReset", "Removing Levels from World") );
+				
 		UWorld* OwningWorld = InLevel->OwningWorld;
 		const FName LevelPackageName		= InLevel->GetOutermost()->GetFName();
 		const bool bRemovingCurrentLevel	= InLevel && InLevel->IsCurrentLevel();
