A message box is now displayed when the user is exited due to infinite loop
#TTP 331819 - LIVE: FEATURE REQUEST: BLUEPRINTS: Exiting PIE due to a Runaway Loop might benefit from a pop-up message

#branch UE4

#proj Editor.UnrealEd

#add Added message box which displays the text Infinite loop found in blueprint: <blueprint name>. Reasons for not using a notification: In practice I found that because the infinite loop essentially locked the editor for some time I found that I would often miss the display of the notification. A message box demands attention and highlights the serious issue.

diff --git a/Engine/Source/Editor/UnrealEd/Private/Kismet2/KismetDebugUtilities.cpp b/Engine/Source/Editor/UnrealEd/Private/Kismet2/KismetDebugUtilities.cpp
--- a/Engine/Source/Editor/UnrealEd/Private/Kismet2/KismetDebugUtilities.cpp
+++ b/Engine/Source/Editor/UnrealEd/Private/Kismet2/KismetDebugUtilities.cpp
@@ -272,6 +272,14 @@ void FKismetDebugUtilities::OnScriptException(const UObject* ActiveObject, const
 				{
 					GEditor->RequestEndPlayMap();
 					FSlateApplication::Get().LeaveDebuggingMode();
+					
+					// Launch a message box notifying the user why they have been booted
+					{
+						FString BlueprintName;
+						BlueprintObj->GetName(BlueprintName);
+
+						FMessageDialog::Open( EAppMsgType::Ok, FText::Format(LOCTEXT("InfiniteLoopWarning", "Infinite Loop detected in blueprint: {0}"), FText::FromString(BlueprintName)));
+					}
 				}
 			}
 			break;
