Removed level actor added binding for editor actor folders.

This callback was not necessary because actors are only ever spawned with no folder path as the path is imported at a later stage.

diff --git a/Engine/Source/Editor/UnrealEd/Private/EditorActorFolders.cpp b/Engine/Source/Editor/UnrealEd/Private/EditorActorFolders.cpp
--- a/Engine/Source/Editor/UnrealEd/Private/EditorActorFolders.cpp
+++ b/Engine/Source/Editor/UnrealEd/Private/EditorActorFolders.cpp
@@ -38,7 +38,6 @@ FActorFolders::FActorFolders()
 	check(GEngine);
 	GEngine->OnLevelActorFolderChanged().AddRaw(this, &FActorFolders::OnActorFolderChanged);
 	GEngine->OnLevelActorListChanged().AddRaw(this, &FActorFolders::OnLevelActorListChanged);
-	GEngine->OnLevelActorAdded().AddRaw(this, &FActorFolders::OnLevelActorAdded);
 
 	FEditorDelegates::MapChange.AddRaw(this, &FActorFolders::OnMapChange);
 	FEditorDelegates::PostSaveWorld.AddRaw(this, &FActorFolders::OnWorldSaved);
@@ -140,24 +139,6 @@ void FActorFolders::OnWorldSaved(uint32 SaveFlags, UWorld* World, bool bSuccess)
 	}
 }
 
-void FActorFolders::OnLevelActorAdded(AActor* InActor)
-{
-	check(InActor);
-
-	UWorld* World = InActor->GetWorld();
-	check(World);
-
-	const FName FolderPath = InActor->GetFolderPath();
-
-	if (WillAddFolderToWorld(*World, FolderPath))
-	{
-		FScopedTransaction Transaction(LOCTEXT("UndoAction_ActorAdded", "Actor Folder Added"));
-
-		const bool bAddedFolder = AddFolderToWorld(*World, FolderPath);
-		check(bAddedFolder);
-	}
-}
-
 void FActorFolders::OnActorFolderChanged(const AActor* InActor, FName OldPath)
 {
 	check(InActor && InActor->GetWorld());
@@ -437,19 +418,6 @@ bool FActorFolders::RenameFolderInWorld(UWorld& World, FName OldPath, FName NewP
 	return RenamedFolders.Num() != 0;
 }
 
-bool FActorFolders::WillAddFolderToWorld(UWorld& InWorld, FName Path) const
-{
-	bool bWillAddFolders = false;
-
-	if (!Path.IsNone() && FoldersExistForWorld(InWorld))
-	{
-		UEditorActorFolders& Folders = *TemporaryWorldFolders.FindChecked(&InWorld);
-		bWillAddFolders = !Folders.Folders.Contains(Path);
-	}
-
-	return bWillAddFolders;
-}
-
 bool FActorFolders::AddFolderToWorld(UWorld& InWorld, FName Path)
 {
 	if (!Path.IsNone())
diff --git a/Engine/Source/Editor/UnrealEd/Public/EditorActorFolders.h b/Engine/Source/Editor/UnrealEd/Public/EditorActorFolders.h
--- a/Engine/Source/Editor/UnrealEd/Public/EditorActorFolders.h
+++ b/Engine/Source/Editor/UnrealEd/Public/EditorActorFolders.h
@@ -115,15 +115,9 @@ private:
 	/** Called after a world has been saved */
 	void OnWorldSaved(uint32 SaveFlags, UWorld* World, bool bSuccess);
 
-	/** Called when an actor is added to the world */
-	void OnLevelActorAdded(AActor* Actor);
-
 	/** Remove any references to folder arrays for dead worlds */
 	void Housekeeping();
 
-	/** Returns true if calling AddFolderToWorld with the provided details would add a new folder */
-	bool WillAddFolderToWorld(UWorld& InWorld, FName Path) const;
-
 	/** Add a folder to the folder map for the specified world. Does not trigger any events. */
 	bool AddFolderToWorld(UWorld& InWorld, FName Path);
 
