Extending the DumpBlueprintInfo commandlet so that is shows more contextual menu actions.

diff --git a/Engine/Source/Editor/UnrealEd/Private/Commandlets/DumpBlueprintsInfoCommandlet.cpp b/Engine/Source/Editor/UnrealEd/Private/Commandlets/DumpBlueprintsInfoCommandlet.cpp
--- a/Engine/Source/Editor/UnrealEd/Private/Commandlets/DumpBlueprintsInfoCommandlet.cpp
+++ b/Engine/Source/Editor/UnrealEd/Private/Commandlets/DumpBlueprintsInfoCommandlet.cpp
@@ -6,141 +6,1747 @@
 #include "K2ActionMenuBuilder.h"
 #include "EdGraphSchema_K2_Actions.h"
 #include "K2Node.h"
+#include "AnimationGraph.h"
+#include "AnimGraphNode_StateMachine.h"
+#include "BlueprintEditorUtils.h"
+#include "AnimStateConduitNode.h"
+#include "AnimStateNode.h"
+#include "Engine/LevelScriptBlueprint.h"
+#include "BehaviorTree/BehaviorTree.h"
+#include "AssetSelection.h"
+
+DEFINE_LOG_CATEGORY_STATIC(LogBlueprintInfoDump, Log, All);
+
+/*******************************************************************************
+ * Static Helpers
+ ******************************************************************************/
+
+namespace DumpBlueprintInfoUtils
+{
+	static const FString HelpString = TEXT("\n\
+\n\
+DumpBlueprintsInfo commandlet params: \n\
+\n\
+    -class=<Class>      Used to specify the blueprint's parent class, if left  \n\
+                        unset then it will go through EVERY blueprint parent   \n\
+                        class available.                                       \n\
+\n\
+    -multifile          Used to keep file size down, will split each blueprint \n\
+                        into its own file (meaning only one file will be       \n\
+                        created when used with -class).                        \n\
+\n\
+    -palette            Defaults to on, unless other flags were specified.     \n\
+                        Dumps all actions from the blueprint's palette menu    \n\
+                        (constant across all graphs).                          \n\
+\n\
+    -palfilter=<Class>  Simulates picking a class from the blueprint palette's \n\
+                        drop down. Setting -palfilter=all will dump the palette\n\
+                        for every possible class.                              \n\
+\n\
+    -context            Will dump all actions from the context menu, for each  \n\
+                        graph and every pin type within (dumps a lot of info). \n\
+\n\
+    -noPinContext       Use after the -context switch to only dump context info\n\
+                        for each graph (not for each pin type within each      \n\
+                        graph).                                                \n\
+\n\
+    -noActionInfo       Whenever an action is dumped it will omit its sub      \n\
+                        fields (can help cut down on the lines that have to be \n\
+                        compared in a diff).                                   \n\
+\n\
+    -graph=<Filter>     Used to only dump contextual info for specific graph   \n\
+                        types. The <Filter> param can only be certain values:  \n\
+                        function, ubergraph, macro, animation, or statemachine.\n\
+\n\
+    -pin=<PinType>      Used to specify a single pin type to dump contextual   \n\
+                        data for. The <PinType> param can be a POD type like   \n\
+                        'int' or 'bool', or an object class like 'Actor', etc. \n\
+\n\
+    -classpin=<Class>   There is no way with the -pin switch to specify a class\n\
+                        pin, this will override any previous -pin switches and \n\
+                        replace them with the class specified (like 'Actor').  \n\
+\n\
+    -select=<Class>     For graph context dumps, this will simulate an object  \n\
+                        being selected during action menu generation. <Class>  \n\
+                        must be an Actor or ActorComponent sub-class. You can  \n\
+                        also specify -select=all to go through every viable    \n\
+                        Actor and ActorComponent class.                        \n\
+\n\
+    -time               When enabled, will record timings during menu building \n\
+                        (as it has been time sync in the past). This is not    \n\
+                        ideal for diffs though (since times can easily vary).  \n\
+\n\
+    -diff=<DiffPath>    Can be a folder or file path, if a folder then it will \n\
+                        look for a matching filename within that folder to     \n\
+                        compare against. Can help automate looking for changes.\n\
+\n\
+    -diffcmd=<Command>  Diffing will attempt to use the P4 diff-tool, but if   \n\
+                        you wish to use your own tool then specify it here. Use\n\
+                        '{1}' and '{2}' as placeholders for filenames, like so:\n\
+                        -diffcmd=\"AraxisP4Diff.exe {2} {1}\".                 \n\
+\n\
+    -help, -h, -?       Display this message and then exit.                    \n\
+\n");
+
+	/** Flags that govern the verbosity of the dump. */
+	enum EDumpFlags
+	{
+		BPDUMP_UnfilteredPalette	= (1<<0), 
+		BPDUMP_FilteredPalette		= (1<<1), 
+		BPDUMP_GraphContextActions	= (1<<2), 
+		BPDUMP_PinContextActions	= (1<<3), 
+		BPDUMP_LogHelp				= (1<<4),
+
+		BPDUMP_FilePerBlueprint		= (1<<5), 
+		BPDUMP_PinTypeIsClass		= (1<<6),
+		BPDUMP_DoNotDumpActionInfo  = (1<<7),
+		BPDUMP_RecordTiming			= (1<<8),
+		BPDUMP_SelectAllObjTypes    = (1<<9), 
+
+		BPDUMP_PaletteMask = (BPDUMP_UnfilteredPalette | BPDUMP_FilteredPalette),
+		BPDUMP_ContextMask = (BPDUMP_GraphContextActions | BPDUMP_PinContextActions),
+	};
+
+	/**
+	 * A collection of variables that represent the various command switches 
+	 * that users can specify when running the commandlet. See the HelpString 
+	 * variable for a listing of supported switches.
+	 */
+	struct CommandletOptions
+	{
+		CommandletOptions()
+			: BlueprintClass(nullptr)
+			, DumpFlags(DumpBlueprintInfoUtils::BPDUMP_UnfilteredPalette)
+			, PaletteFilter(nullptr)
+			, GraphFilter(GT_MAX)
+			, SelectedObjectType(nullptr)
+		{
+		}
+
+		/**
+		* Parses the string command switches into flags, class pointers, and 
+		* booleans that will govern what should be dumped. Logs errors if any 
+		* switch was misused.
+		*/
+		CommandletOptions(TArray<FString> const& Switches);
+
+		UClass*    BlueprintClass;
+		uint32     DumpFlags;
+		UClass*    PaletteFilter;
+		EGraphType GraphFilter;
+		FString    PinType;
+		UClass*    SelectedObjectType;
+		FString    DiffPath;
+		FString    DiffCommand;
+	};
+
+	/** Static instance of the command switches (so we don't have to pass one along the call stack) */
+	static CommandletOptions CommandOptions;
+
+	/** Tracks spawned level actors (so we don't have to create more the we have to). */
+	static TMap<UClass*, AActor*> LevelActors;
+
+	/**
+	 * Utility class for tracking the duration of a scoped action (tracks time 
+	 * in seconds and adds it to the specified variable on destruction).
+	 */
+	class FScopedDurationTimer
+	{
+	public:
+		FScopedDurationTimer(double& AccumulatorIn)
+			: StartTime(FPlatformTime::Seconds())
+			, Accumulator(AccumulatorIn)
+		{}
+
+		/** Dtor, updating seconds with time delta. */
+		~FScopedDurationTimer()
+		{
+			Accumulator += FPlatformTime::Seconds() - StartTime;
+		}
+	private:
+		/** Start time, captured in ctor. */
+		double StartTime;
+		/** Time variable to update. */
+		double& Accumulator;
+	};
+
+	/**
+	 * Certain blueprints (like level blueprints) require a level outer, and 
+	 * for certain actions we need a level actor selected. This utility function
+	 * provides an easy way to grab the world (which has a level that we can use 
+	 * for these purposes).
+	 */
+	static UWorld* GetWorld();
+
+	/** 
+	 * Utility function for spawning and selecting an actor of the specified 
+	 * type (utilized when the user uses the -select switch with an actor class).
+	 */
+	static AActor* SpawnLevelActor(UClass* ActorClass, bool bSelect = true);
+
+	/**
+	 * Spawns a transient blueprint of the specified type. Adds all possible 
+	 * graph types (function, macro, etc.), and does some additional setup for
+	 * unique blueprint types (like level and anim blueprints).
+	 */
+	static UBlueprint* MakeTempBlueprint(UClass* BlueprintClass);
+
+	/**
+	 * Adds an instance of the specified component type to the supplied 
+	 * blueprint (invoked when the user has specified the -select command switch
+	 * with a component class).
+	 */
+	static bool AddComponentToBlueprint(UBlueprint* Blueprint, UClass* ComponentClass);
+
+	/**
+	 * Certain nodes add specific graph types that we want to dump info on (like 
+	 * state machine graphs for anim blueprints). The best way to add those 
+	 * graphs is through the natural process of adding those nodes (which this
+	 * method is intended for).
+	 */
+	template<class NodeType>
+	static NodeType* AddNodeToGraph(UEdGraph* Graph);
+
+	/**
+	 * Builds a fully qualified file path for a new dump file. If using the 
+	 * -multifile command switch, then this will create a sub-directory and name
+	 * the file after the class. Generally, dump files are placed in the 
+	 * project's ".../Saved/Commandlets/" directory.
+	 */
+	static FString BuildDumpFilePath(UClass* BlueprintClass);
+
+	/**
+	 * Utility function to convert a tab integer into a string of whitespace.
+	 * Defaults to tab characters, but if bUseSpaces is enabled, then single 
+	 * spaces are used.
+	 */
+	static FString BuildIndentString(uint32 IndentCount, bool bUseSpaces = false);
+
+	/**
+	 * Utility function to convert a graph's EGraphType into a string. Used as 
+	 * an aid when writing graph information to json.
+	 */
+	static FString GetGraphTypeString(UEdGraph const* Graph);
+
+	/**
+	 * Concatenates the action's category with its menu name (to help
+	 * distinguish similarly named actions). Can then be used to sort and 
+	 * uniquely identify actions.
+	 */
+	static FString GetActionKey(FGraphActionListBuilderBase::ActionGroup const& Action);
+
+	/**
+	* Goes through all of the blueprint skeleton's object properties and pulls
+	* out the ones that are associated with an UActorComponent (and are visbile
+	* to the blueprint).
+	*/
+	static void GetComponentProperties(UBlueprint* Blueprint, TArray<UObjectProperty*>& PropertiesOut);
+	
+	/**
+	 * Constructs a temporary blueprint (of the class type specified) and kicks 
+	 * off a dump of all its nested information (palette, graph, contextual 
+	 * actions, etc.). 
+	 */
+	static void DumpInfoForClass(uint32 Indent, UClass* BlueprintClass, FArchive* FileOutWriter);
+	
+	/**
+	 * Assumes that the specified ActionListBuilder is configured with the 
+	 * proper blueprint. Starts by clearing any actions it contained and then
+	 * runs through and re-adds all actions.
+	 *
+	 * @return The amount of time (in seconds) that the menu building took.
+	 */
+	static double GetPaletteMenuActions(FBlueprintPaletteListBuilder& ActionListBuilder, UClass* PaletteFilter);
+
+	/**
+	 * Dumps all palette actions listed for the specified blueprint. Determines
+	 * if the user specified any filter class for the palette and adjusts 
+	 * accordingly (can dump multiple palettes if -palfilter=all was specified).
+	 */
+	static void DumpPalette(uint32 Indent, UBlueprint* Blueprint, FArchive* FileOutWriter);
+
+	/**
+	 * Dumps a single instance of the blueprint's palette (using the ClassFilter).
+	 * ClassFilter can be null and the full unfiltered palette will be dumped.
+	 */
+	static void DumpPalette(uint32 Indent, UBlueprint* Blueprint, UClass* ClassFilter, FArchive* FileOutWriter);
+
+	/**
+	 * Generic function utilized by both palette and context-menu dumps. Take a
+	 * GraphActionListBuilder and writes out every action that it has captured.
+	 */
+	static void DumpActionList(uint32 Indent, FGraphActionListBuilderBase& ActionList, FArchive* FileOutWriter);
+
+	/**
+	 * Generic function that dumps information on a single action (like it's 
+	 * name, category, an associated node if it has one, etc.).
+	 */
+	static void DumpActionMenuItem(uint32 Indent, FGraphActionListBuilderBase::ActionGroup const& Action, FArchive* FileOutWriter);
+
+	/**
+	 * Emulates the blueprint's context menu and goes through each of its graphs,
+	 * dumping the context menu(s) for each. Entry point for dumping all nested
+	 * context actions.
+	 */
+	static void DumpContextualActions(uint32 Indent, UBlueprint* Blueprint, FArchive* FileOutWriter);
+
+	/**
+	 * Dumps the context menu actions that can be found when right clicking in a
+	 * bare graph (of the specified type). If pin context dumping is enabled, 
+	 * then this will continue down into that.
+	 */
+	static void DumpGraphContextActions(uint32 Indent, UEdGraph* Graph, FArchive* FileOutWriter);
+
+	/**
+	 * Will go through every pin type that a user can create a pin for and dumps
+	 * the full context menu as if it were dragged from each type.
+	 *
+	 * WARNING: This will dump A LOT of info, as it accounts for every type as
+	 *          an input/output and as an array input/output (use -pin= to reign it in).
+	 */
+	static bool DumpPinContextActions(uint32 Indent, UEdGraph* Graph, FArchive* FileOutWriter);
+
+	/**
+	 * Pin type info comes in a tree format (intended for menus), this recursive 
+	 * function traverses the tree and dumps contextual pin actions for each 
+	 * leaf it encounters (utilized by DumpPinContextActions).
+	 */
+	static bool DumpTypeTreeActions(uint32 Indent, UEdGraph* Graph, TSharedPtr<UEdGraphSchema_K2::FPinTypeTreeInfo> PinTypeInfo, FArchive* FileOutWriter);
+
+	/**
+	 * Generic function that takes a specific FEdGraphPinType and dumps all the 
+	 * actions available to that pin (as if you dragged and spawned a context 
+	 * menu from it).
+	 */
+	static void DumpContextualPinTypeActions(uint32 Indent, UEdGraph* Graph, FEdGraphPinType const& PinType, FArchive* FileOutWriter);
+	
+	/**
+	 * Generic function that takes a contextual GraphActionList and calls down 
+	 * into the DumpActionList(), while dumping context information to go along 
+	 * with it (to give the reader context).
+	 */
+	static void DumpContextActionList(uint32 Indent, FBlueprintGraphActionListBuilder ActionBuilder, FArchive* FileOutWriter);
+
+	/**
+	 * Assumes that the specified ActionListBuilder is configured with all the 
+	 * proper filter/contextual data. Starts by clearing any actions it 
+	 * contained and then runs through building adding new ones that match the
+	 * context.
+	 *
+	 * @return The amount of time (in seconds) that the menu building took.
+	 */
+	static double GetContextMenuActions(FBlueprintGraphActionListBuilder& ActionListBuilder);
+
+	/**
+	 * Looks at the filter/context items set on the specified GraphActionListBuilder
+	 * and writes them out, to provide context with any dumped actions.
+	 */
+	static void DumpContextInfo(uint32 Indent, FBlueprintGraphActionListBuilder const& ActionBuilder, FArchive* FileOutWriter);
+
+	/**
+	 * A utility function aimed at determining whether we should open a diff for
+	 * the dump file (with -multifile we only want to open files that differ, 
+	 * since we could be going through every class in the project).
+	 *
+	 * @TODO: Right now this is very platform specific (it uses dos commands). 
+	 *        Maybe it should be moved into a platform generic file like IFileManager?
+	 */
+	static bool DoFilesDiffer(FString const& NewFilePath, FString const& OldFilePath);
+
+	/**
+	 * Looks at the two dump files, and opens the specified diff tool if there 
+	 * are differences detected (defaults to p4merge.exe if a diffcmd wasn't 
+	 * specified by the user).
+	 */
+	static void DiffDumpFiles(FString const& NewFilePath, FString const& OldFilePath, FString const& UserDiffCmd);
+}
+
+//------------------------------------------------------------------------------
+DumpBlueprintInfoUtils::CommandletOptions::CommandletOptions(TArray<FString> const& Switches)
+	: BlueprintClass(nullptr)
+	, DumpFlags(DumpBlueprintInfoUtils::BPDUMP_UnfilteredPalette)
+	, PaletteFilter(nullptr)
+	, GraphFilter(GT_MAX)
+	, SelectedObjectType(nullptr)
+{
+	uint32 NewDumpFlags = 0u;
+	for (FString const& Switch : Switches)
+	{
+		if (Switch.StartsWith("class="))
+		{
+			FString ClassSwitch, ClassName;
+			Switch.Split(TEXT("="), &ClassSwitch, &ClassName);
+			BlueprintClass = FindObject<UClass>(ANY_PACKAGE, *ClassName);
+
+			if (BlueprintClass == nullptr)
+			{
+				UE_LOG(LogBlueprintInfoDump, Error, TEXT("Unrecognized blueprint class '%s', defaulting to 'Actor'"), *ClassName);
+				BlueprintClass = AActor::StaticClass();
+			}
+		}
+		else if (Switch.StartsWith("palfilter="))
+		{
+			FString ClassSwitch, ClassName;
+			Switch.Split(TEXT("="), &ClassSwitch, &ClassName);
+			PaletteFilter = FindObject<UClass>(ANY_PACKAGE, *ClassName);
+
+			NewDumpFlags |= BPDUMP_FilteredPalette;
+			if (PaletteFilter == nullptr)
+			{
+				if (ClassName.Compare("all", ESearchCase::IgnoreCase))
+				{
+					UE_LOG(LogBlueprintInfoDump, Error, TEXT("Unrecognized palette filter '%s', defaulting to unfiltered"), *ClassName);
+					NewDumpFlags &= ~(DumpBlueprintInfoUtils::BPDUMP_FilteredPalette);
+					NewDumpFlags |= DumpBlueprintInfoUtils::BPDUMP_UnfilteredPalette;
+				}
+			}
+		}
+		else if (Switch.StartsWith("graph="))
+		{
+			FString GraphSwitch, FilterName;
+			Switch.Split(TEXT("="), &GraphSwitch, &FilterName);
+
+#define CHECK_GRAPH_SWITCH(GraphType) \
+			if ((GraphFilter == GT_MAX) && !FilterName.Compare(#GraphType, ESearchCase::IgnoreCase)) \
+			{ \
+				GraphFilter = GT_##GraphType; \
+			}
+
+			CHECK_GRAPH_SWITCH(Function)
+			CHECK_GRAPH_SWITCH(Ubergraph)
+			CHECK_GRAPH_SWITCH(Macro)
+			CHECK_GRAPH_SWITCH(Animation)
+			CHECK_GRAPH_SWITCH(StateMachine)
+
+#undef CHECK_GRAPH_SWITCH
+
+			if (GraphFilter == GT_MAX)
+			{
+				UE_LOG(LogBlueprintInfoDump, Error, TEXT("Invalid graph type '%s', can only be one of the following: function, ubergraph, macro, animation, or statemachine"), *FilterName);
+			}
+
+			NewDumpFlags |= BPDUMP_GraphContextActions;
+		}
+		else if (Switch.StartsWith("pin="))
+		{
+			NewDumpFlags &= ~BPDUMP_PinTypeIsClass;
+
+			FString PinSwitch;
+			Switch.Split(TEXT("="), &PinSwitch, &PinType);
+
+			NewDumpFlags |= DumpBlueprintInfoUtils::BPDUMP_PinContextActions;
+			// implies that we want contextual actions
+			NewDumpFlags |= DumpBlueprintInfoUtils::BPDUMP_GraphContextActions;
+		}
+		else if (Switch.StartsWith("classpin="))
+		{
+			NewDumpFlags |= BPDUMP_PinTypeIsClass;
+
+			FString PinSwitch;
+			Switch.Split(TEXT("="), &PinSwitch, &PinType);
+			
+			NewDumpFlags |= BPDUMP_PinContextActions;
+			// implies that we want contextual actions
+			NewDumpFlags |= BPDUMP_GraphContextActions;
+		}
+		else if (Switch.StartsWith("select="))
+		{
+			FString ClassSwitch, ClassName;
+			Switch.Split(TEXT("="), &ClassSwitch, &ClassName);
+			SelectedObjectType = FindObject<UClass>(ANY_PACKAGE, *ClassName);
+
+			if (!ClassName.Compare("all", ESearchCase::IgnoreCase))
+			{
+				NewDumpFlags |= BPDUMP_SelectAllObjTypes;
+			}
+			else if (SelectedObjectType == nullptr)
+			{
+				SelectedObjectType = UStaticMeshComponent::StaticClass();
+				UE_LOG(LogBlueprintInfoDump, Error, TEXT("Unrecognized selection class '%s', defaulting to '%s'"), *ClassName, *SelectedObjectType->GetName());
+			}
+			else if (SelectedObjectType->IsChildOf(UActorComponent::StaticClass()))
+			{
+				if (SelectedObjectType->HasAnyClassFlags(CLASS_Abstract))
+				{
+					SelectedObjectType = UStaticMeshComponent::StaticClass();
+					UE_LOG(LogBlueprintInfoDump, Error, TEXT("Selection component cannot be a abstract ('%s' is), defaulting to '%s'"), *ClassName, *SelectedObjectType->GetName());
+				}
+			}
+			else if (!SelectedObjectType->IsChildOf(AActor::StaticClass()))
+			{
+				SelectedObjectType = AActor::StaticClass();
+				UE_LOG(LogBlueprintInfoDump, Error, TEXT("Selection class must be an actor or component ('%s' is not), defaulting to '%s'"), *ClassName, *SelectedObjectType->GetName());
+			}
+			else if (GetWorld() == nullptr)
+			{
+				UE_LOG(LogBlueprintInfoDump, Error, TEXT("Cannot select a level actor without a valid editor world (clearing the selection)"), *ClassName, *SelectedObjectType->GetName());
+				SelectedObjectType = nullptr;
+			}
+
+			// implies that we want contextual actions
+			NewDumpFlags |= BPDUMP_GraphContextActions;
+		}
+		else if (Switch.StartsWith("diff="))
+		{
+			FString DiffSwitch;
+			Switch.Split(TEXT("="), &DiffSwitch, &DiffPath);
+		}
+		else if (Switch.StartsWith("diffcmd="))
+		{
+			FString DiffSwitch;
+			Switch.Split(TEXT("="), &DiffSwitch, &DiffCommand);
+		}
+		else if (!Switch.Compare("palette", ESearchCase::IgnoreCase))
+		{
+			NewDumpFlags |= BPDUMP_UnfilteredPalette;
+		}
+		else if (!Switch.Compare("context", ESearchCase::IgnoreCase))
+		{
+			NewDumpFlags |= BPDUMP_GraphContextActions;
+			NewDumpFlags |= BPDUMP_PinContextActions;
+		}
+		else if (!Switch.Compare("noPinContext", ESearchCase::IgnoreCase))
+		{
+			NewDumpFlags &= (~BPDUMP_PinContextActions);
+		}
+		else if (!Switch.Compare("h", ESearchCase::IgnoreCase) || 
+		         !Switch.Compare("?", ESearchCase::IgnoreCase) || 
+				 !Switch.Compare("help", ESearchCase::IgnoreCase))
+		{
+			NewDumpFlags |= BPDUMP_LogHelp;
+		}
+		else if (!Switch.Compare("multifile", ESearchCase::IgnoreCase))
+		{
+			NewDumpFlags |= BPDUMP_FilePerBlueprint;
+		}
+		else if (!Switch.Compare("noActionInfo", ESearchCase::IgnoreCase))
+		{
+			NewDumpFlags |= BPDUMP_DoNotDumpActionInfo;
+		}
+		else if (!Switch.Compare("time", ESearchCase::IgnoreCase))
+		{
+			NewDumpFlags |= BPDUMP_RecordTiming;
+		}
+		else
+		{
+			UE_LOG(LogBlueprintInfoDump, Warning, TEXT("Unrecognized command switch '%s', use -help for a listing of all accepted params"), *Switch);
+		}
+	}
+
+	if (NewDumpFlags != 0)
+	{
+		DumpFlags = NewDumpFlags;
+	}
+	if (!(DumpFlags & (BPDUMP_PaletteMask | BPDUMP_ContextMask)))
+	{
+		DumpFlags |= BPDUMP_UnfilteredPalette;
+	}
+}
+
+//------------------------------------------------------------------------------
+static UWorld* DumpBlueprintInfoUtils::GetWorld()
+{
+	UWorld* World = nullptr;
+	for (FWorldContext const& WorldContext : GEngine->GetWorldContexts())
+	{
+		World = WorldContext.World();
+		if (World != nullptr)
+		{
+			break;
+		}
+	}
+
+	if (World == nullptr)
+	{
+		static UWorld* CommandletWorld = nullptr;
+		if (CommandletWorld == nullptr)
+		{
+			if (GUnrealEd == nullptr)
+			{
+				UE_LOG(LogBlueprintInfoDump, Error, TEXT("Cannot create a temp map to test within, without a valid editor world"));
+			}
+			else
+			{
+				CommandletWorld = GEditor->NewMap();
+			}
+		}
+
+		World = CommandletWorld;
+	}
+
+	return World;
+}
+
+//------------------------------------------------------------------------------
+static AActor* DumpBlueprintInfoUtils::SpawnLevelActor(UClass* ActorClass, bool bSelect)
+{
+	check(ActorClass->IsChildOf(AActor::StaticClass()));
+
+	AActor* SpawnedActor = nullptr;
+	if (AActor** ExistingActorPtr = LevelActors.Find(ActorClass))
+	{
+		SpawnedActor = *ExistingActorPtr;
+	}
+	else if (FKismetEditorUtilities::CanCreateBlueprintOfClass(ActorClass))
+	{
+		UActorFactory* NewFactory = ConstructObject<UActorFactory>(UActorFactoryBlueprint::StaticClass());
+
+		UBlueprint* ActorTemplate = MakeTempBlueprint(ActorClass);
+		SpawnedActor = FActorFactoryAssetProxy::AddActorForAsset(ActorTemplate,
+			/*ActorLocation =*/nullptr,
+			/*bUseSurfaceOrientation =*/false,
+			/*SelectActor =*/bSelect,
+			RF_Transient, NewFactory, NAME_None);
+	}
+	// @TODO: What about non-blueprintable actors (brushes, etc.)... the code below crashes
+// 	else
+// 	{
+// 		UWorld* World = GetWorld();
+// 		SpawnedActor = GEditor->AddActor(World->GetCurrentLevel(), ActorClass, FVector::ZeroVector);
+// 	}
+
+	if (SpawnedActor != nullptr)
+	{
+		LevelActors.Add(ActorClass, SpawnedActor);
+	}
+
+	if (bSelect)
+	{
+		USelection* SelectedActors = GEditor->GetSelectedActors();
+		SelectedActors->DeselectAll();
+
+		if (SpawnedActor != nullptr)
+		{
+			SelectedActors->Select(SpawnedActor);
+		}
+	}
+
+	return SpawnedActor;
+}
+
+//------------------------------------------------------------------------------
+static UBlueprint* DumpBlueprintInfoUtils::MakeTempBlueprint(UClass* ParentClass)
+{
+	UObject* BlueprintOuter = GetTransientPackage();
+
+	bool bIsAnimBlueprint = ParentClass->IsChildOf(UAnimInstance::StaticClass());
+	bool bIsLevelBlueprint = ParentClass->IsChildOf(ALevelScriptActor::StaticClass());
+
+	UClass* BlueprintClass = UBlueprint::StaticClass();
+	UClass* GeneratedClass = UBlueprintGeneratedClass::StaticClass();
+	EBlueprintType BlueprintType = BPTYPE_Normal;
+
+	if (bIsAnimBlueprint)
+	{
+		BlueprintClass = UAnimBlueprint::StaticClass();
+		GeneratedClass = UAnimBlueprintGeneratedClass::StaticClass();
+	}
+	else if (bIsLevelBlueprint)
+	{
+		UWorld* World = GetWorld();
+		if (World == nullptr)
+		{
+			UE_LOG(LogBlueprintInfoDump, Error, TEXT("Cannot make a proper level blueprint without a valid editor level for its outer."));
+		}
+		else
+		{
+			BlueprintClass = ULevelScriptBlueprint::StaticClass();
+			BlueprintType = BPTYPE_LevelScript;
+			BlueprintOuter = World->GetCurrentLevel();
+		}
+	}
+	// @TODO: UEditorUtilityBlueprint
+
+	FString const ClassName = ParentClass->GetName();
+	FString const DesiredName = FString::Printf(TEXT("COMMANDLET_TEMP_Blueprint_%s"), *ClassName);
+	FName   const TempBpName = MakeUniqueObjectName(BlueprintOuter, BlueprintClass, FName(*DesiredName));
+
+	check(FKismetEditorUtilities::CanCreateBlueprintOfClass(ParentClass));
+	UBlueprint* TempBlueprint = FKismetEditorUtilities::CreateBlueprint(ParentClass, BlueprintOuter, TempBpName, BlueprintType, BlueprintClass, GeneratedClass);
+
+	// if this is an animation blueprint, then we want anim specific graphs to test as well...
+	if (bIsAnimBlueprint)
+	{
+		check(TempBlueprint->FunctionGraphs.Num() > 0);
+		UAnimationGraph* AnimGraph = CastChecked<UAnimationGraph>(TempBlueprint->FunctionGraphs[0]);
+
+		// should add a state-machine graph
+		UAnimGraphNode_StateMachine* StateMachineNode = AddNodeToGraph<UAnimGraphNode_StateMachine>(AnimGraph);
+
+		UAnimationStateMachineGraph* StateMachineGraph = StateMachineNode->EditorStateMachineGraph;
+		// should add an conduit graph 
+		UAnimStateConduitNode* ConduitNode = AddNodeToGraph<UAnimStateConduitNode>((UEdGraph*)StateMachineGraph);
+
+		UAnimStateNode* StateNode = AddNodeToGraph<UAnimStateNode>((UEdGraph*)StateMachineGraph);
+		// should create a transition graph
+		StateNode->AutowireNewNode(ConduitNode->GetOutputPin());
+	}
+
+	// may have been altered in CreateBlueprint()
+	BlueprintType = TempBlueprint->BlueprintType;
+
+	bool bCanAddFunctions = (BlueprintType != BPTYPE_MacroLibrary); // taken from FBlueprintEditor::NewDocument_IsVisibleForType()
+	if (bCanAddFunctions)
+	{
+		// add a functions graph that isn't the construction script (or an animation graph)
+		FName FuncGraphName = MakeUniqueObjectName(TempBlueprint, UEdGraph::StaticClass(), FName(TEXT("NewFunction")));
+		UEdGraph* FuncGraph = FBlueprintEditorUtils::CreateNewGraph(TempBlueprint, FuncGraphName, UEdGraph::StaticClass(), UEdGraphSchema_K2::StaticClass());
+		FBlueprintEditorUtils::AddFunctionGraph<UClass>(TempBlueprint, FuncGraph, /*bIsUserCreated =*/true, nullptr);
+	}
+
+	bool bCanAddMacros = ((BlueprintType == BPTYPE_MacroLibrary) || (BlueprintType == BPTYPE_Normal) || (BlueprintType == BPTYPE_LevelScript));
+	if (bCanAddMacros)
+	{
+		FName MacroGraphName = MakeUniqueObjectName(TempBlueprint, UEdGraph::StaticClass(), FName(TEXT("NewMacro")));
+		UEdGraph* MacroGraph = FBlueprintEditorUtils::CreateNewGraph(TempBlueprint, MacroGraphName, UEdGraph::StaticClass(), UEdGraphSchema_K2::StaticClass());
+		FBlueprintEditorUtils::AddMacroGraph(TempBlueprint, MacroGraph, /*bIsUserCreated =*/true, nullptr);
+	}
+
+	UClass* ObjTypeToAdd = CommandOptions.SelectedObjectType;
+	if (ObjTypeToAdd && ObjTypeToAdd->IsChildOf(UActorComponent::StaticClass()))
+	{
+		if (!AddComponentToBlueprint(TempBlueprint, ObjTypeToAdd))
+		{
+			UE_LOG(LogBlueprintInfoDump, Error, TEXT("Cannot add a '%s' to a '%s' blueprint."), *ObjTypeToAdd->GetName(), *ClassName);
+		}
+	}
+	else if (CommandOptions.DumpFlags & BPDUMP_SelectAllObjTypes)
+	{
+		for (TObjectIterator<UClass> ClassIt; ClassIt; ++ClassIt)
+		{
+			UClass* Class = *ClassIt;
+			if (Class->IsChildOf(UActorComponent::StaticClass()))
+			{
+				AddComponentToBlueprint(TempBlueprint, Class);
+			}
+		}
+	}
+
+	return TempBlueprint;
+}
+
+//------------------------------------------------------------------------------
+static bool DumpBlueprintInfoUtils::AddComponentToBlueprint(UBlueprint* Blueprint, UClass* ComponentClass)
+{
+	// copied from FBlueprintEditor::CanAccessComponentsMode()
+	bool const bCanUserAddComponents = (Blueprint->SimpleConstructionScript != nullptr) // An SCS must be present (otherwise there is nothing valid to edit)
+		&& FBlueprintEditorUtils::IsActorBased(Blueprint)        // Must be parented to an AActor-derived class (some older BPs may have an SCS but may not be Actor-based)
+		&& (Blueprint->BlueprintType != BPTYPE_MacroLibrary)     // Must not be a macro-type Blueprint
+		&& (Blueprint->BlueprintType != BPTYPE_FunctionLibrary); // Must not be a function library
+
+	bool const bClassIsActorComponent = ComponentClass->IsChildOf(UActorComponent::StaticClass());
+	bool const bCanBeAddedToBlueprint = !ComponentClass->HasAnyClassFlags(CLASS_Abstract) && ComponentClass->HasMetaData(FBlueprintMetadata::MD_BlueprintSpawnableComponent);
+	bool const bCanMakeComponent = (bCanUserAddComponents && bClassIsActorComponent && bCanBeAddedToBlueprint);
+
+	if (bCanMakeComponent)
+	{
+		USimpleConstructionScript* SCS = Blueprint->SimpleConstructionScript;
+
+		USCS_Node* NewNode = SCS->CreateNode(ComponentClass);
+
+		if (ComponentClass->IsChildOf(USceneComponent::StaticClass()) || (SCS->GetRootNodes().Num() == 0))
+		{
+			SCS->AddNode(NewNode);
+		}
+		else
+		{
+			USCS_Node* RootNode = SCS->GetDefaultSceneRootNode();
+			check(RootNode != nullptr);
+			RootNode->AddChildNode(NewNode);
+		}
+
+		// regenerate the skeleton class
+		FBlueprintEditorUtils::MarkBlueprintAsStructurallyModified(Blueprint);
+	}
+
+	return bCanMakeComponent;
+}
+
+//------------------------------------------------------------------------------
+template<class NodeType>
+static NodeType* DumpBlueprintInfoUtils::AddNodeToGraph(UEdGraph* Graph)
+{
+	check(Graph != nullptr);
+	NodeType* NewNode = NewObject<NodeType>(Graph);
+	Graph->AddNode(NewNode, /*bFromUI =*/true, /*bSelectNewNode =*/false);
+
+	NewNode->PostPlacedNewNode();
+	NewNode->AllocateDefaultPins();
+	return NewNode;
+}
+
+//------------------------------------------------------------------------------
+static FString DumpBlueprintInfoUtils::BuildDumpFilePath(UClass* BlueprintClass)
+{
+	FString CommandletSaveDir = FPaths::GameSavedDir() + TEXT("Commandlets/");
+	CommandletSaveDir = FPaths::ConvertRelativePathToFull(CommandletSaveDir);
+	IFileManager::Get().MakeDirectory(*CommandletSaveDir);
+
+	FString Pathname = FString::Printf(TEXT("BlueprintsInfoDump_%s"), FPlatformTime::StrTimestamp());
+	Pathname = Pathname.Replace(TEXT(" "), TEXT("_"));
+	Pathname = Pathname.Replace(TEXT("/"), TEXT("-"));
+	Pathname = Pathname.Replace(TEXT(":"), TEXT("."));
+
+	bool const bSplitBlueprintsByFile = (CommandOptions.DumpFlags & BPDUMP_FilePerBlueprint) != 0;
+	if (bSplitBlueprintsByFile && (BlueprintClass != nullptr))
+	{
+		CommandletSaveDir += Pathname + "/";
+		IFileManager::Get().MakeDirectory(*CommandletSaveDir);
+
+		Pathname = ("BlueprintInfo_" + BlueprintClass->GetName() + ".json");
+	}
+	else
+	{
+		Pathname += ".json";
+	}
+
+	return CommandletSaveDir / *Pathname;
+}
+
+//------------------------------------------------------------------------------
+static FString DumpBlueprintInfoUtils::BuildIndentString(uint32 IndentCount, bool bUseSpaces)
+{
+	char RepatingChar = '\t';
+	if (bUseSpaces)
+	{
+		RepatingChar = ' ';
+	}
+
+	FString IndentString;
+	while (IndentCount > 0)
+	{
+		IndentString += RepatingChar;
+		--IndentCount;
+	}
+	return IndentString;
+}
+
+//------------------------------------------------------------------------------
+static FString DumpBlueprintInfoUtils::GetGraphTypeString(UEdGraph const* Graph)
+{
+#define RETURN_GRAPH_TYPE(GraphType) \
+	case GraphType: \
+	{ \
+		return FString(#GraphType); \
+	} \
+	break
+
+	UEdGraphSchema const* GraphSchema = GetDefault<UEdGraphSchema>(Graph->Schema);
+	switch (GraphSchema->GetGraphType(Graph))
+	{
+		RETURN_GRAPH_TYPE(GT_Ubergraph);
+		RETURN_GRAPH_TYPE(GT_Function);
+		RETURN_GRAPH_TYPE(GT_Macro);
+		RETURN_GRAPH_TYPE(GT_Animation);
+		RETURN_GRAPH_TYPE(GT_StateMachine);
+
+		default:
+			break;
+	}
+#undef RETURN_GRAPH_TYPE
+
+	return FString("<UNRECOGNIZED>");
+}
+
+//------------------------------------------------------------------------------
+static FString DumpBlueprintInfoUtils::GetActionKey(FGraphActionListBuilderBase::ActionGroup const& Action)
+{
+	TArray<FString> MenuHierarchy;
+	Action.GetCategoryChain(MenuHierarchy);
+
+	FString ActionKey;
+	for (FString const& SubCategory : MenuHierarchy)
+	{
+		ActionKey += SubCategory + TEXT("|");
+	}
+	if (MenuHierarchy.Num() > 0)
+	{
+		ActionKey.RemoveAt(ActionKey.Len() - 1); // remove the trailing '|'
+	}
+
+	TSharedPtr<FEdGraphSchemaAction> MainAction = Action.Actions[0];
+	ActionKey += MainAction->MenuDescription.ToString();
+
+	return ActionKey;
+}
+
+//------------------------------------------------------------------------------
+static void DumpBlueprintInfoUtils::GetComponentProperties(UBlueprint* Blueprint, TArray<UObjectProperty*>& PropertiesOut)
+{
+	if (AActor* BlueprintCDO = Cast<AActor>(Blueprint->SkeletonGeneratedClass->GetDefaultObject()))
+	{
+		for (TFieldIterator<UObjectProperty> PropertyIt(Blueprint->SkeletonGeneratedClass, EFieldIteratorFlags::IncludeSuper); PropertyIt; ++PropertyIt)
+		{
+			// SMyBlueprint filters out component variables in SMyBlueprint::CollectAllActions() using CPF_BlueprintVisible/CPF_Parm flags
+			if (PropertyIt->PropertyClass->IsChildOf(UActorComponent::StaticClass()) &&
+				PropertyIt->HasAnyPropertyFlags(CPF_BlueprintVisible) && !PropertyIt->HasAnyPropertyFlags(CPF_Parm))
+			{
+				PropertiesOut.Add(*PropertyIt);
+			}
+		}
+	}
+}
+
+//------------------------------------------------------------------------------
+static void DumpBlueprintInfoUtils::DumpInfoForClass(uint32 Indent, UClass* BlueprintClass, FArchive* FileOutWriter)
+{
+	FString const ClassName = BlueprintClass->GetName();
+	UE_LOG(LogBlueprintInfoDump, Display, TEXT("%sDumping BP class: '%s'..."), *BuildIndentString(Indent, true), *ClassName);
+
+	FString const ClassEntryIndent = BuildIndentString(Indent);
+	FString BeginClassEntry = FString::Printf(TEXT("%s\"%s\" : {"), *ClassEntryIndent, *ClassName);
+
+	FString IndentedNewline = "\n" + BuildIndentString(Indent + 1);
+
+	BeginClassEntry += IndentedNewline + "\"ClassContext\" : \"" + ClassName + "\",\n";
+	FileOutWriter->Serialize(TCHAR_TO_ANSI(*BeginClassEntry), BeginClassEntry.Len());
+
+	UBlueprint* TempBlueprint = MakeTempBlueprint(BlueprintClass);
+
+	bool bNeedsClosingComma = false;
+	if (CommandOptions.DumpFlags & BPDUMP_PaletteMask)
+	{
+		DumpPalette(Indent + 1, TempBlueprint, FileOutWriter);
+		bNeedsClosingComma = true;
+	}
+
+	if (CommandOptions.DumpFlags & BPDUMP_ContextMask)
+	{
+		if (bNeedsClosingComma)
+		{
+			FileOutWriter->Serialize(TCHAR_TO_ANSI(TEXT(",\n")), 2);
+		}
+		DumpContextualActions(Indent + 1, TempBlueprint, FileOutWriter);
+		bNeedsClosingComma = true;
+	}
+
+	FString EndClassEntry = "\n" + ClassEntryIndent + "}";
+	FileOutWriter->Serialize(TCHAR_TO_ANSI(*EndClassEntry), EndClassEntry.Len());
+}
+
+//------------------------------------------------------------------------------
+static double DumpBlueprintInfoUtils::GetPaletteMenuActions(FBlueprintPaletteListBuilder& PaletteBuilder, UClass* PaletteFilter)
+{
+	PaletteBuilder.Empty();
+	UEdGraphSchema_K2 const* K2Schema = GetDefault<UEdGraphSchema_K2>();
+
+	double MenuBuildDuration = 0.0;
+	{
+		FScopedDurationTimer DurationTimer(MenuBuildDuration);
+		FK2ActionMenuBuilder(K2Schema).GetPaletteActions(PaletteBuilder, PaletteFilter);
+	}
+
+	return MenuBuildDuration;
+}
+
+//------------------------------------------------------------------------------
+static void DumpBlueprintInfoUtils::DumpPalette(uint32 Indent, UBlueprint* Blueprint, FArchive* FileOutWriter)
+{
+	UClass* PaletteFilter = CommandOptions.PaletteFilter;
+	DumpPalette(Indent, Blueprint, PaletteFilter, FileOutWriter);
+	bool bNeedsEndline = true;
+
+	if ((CommandOptions.DumpFlags & BPDUMP_FilteredPalette) && (PaletteFilter == nullptr))
+	{
+		// anim blueprints don't have a palette, so it is ok to assume this
+		UEdGraphSchema_K2 const* K2Schema = GetDefault<UEdGraphSchema_K2>();
+
+		for (TObjectIterator<UClass> ClassIt; ClassIt; ++ClassIt)
+		{
+			UClass* Class = *ClassIt;
+			if (!K2Schema->ClassHasBlueprintAccessibleMembers(Class))
+			{
+				continue;
+			}
+
+			FileOutWriter->Serialize(TCHAR_TO_ANSI(TEXT(",\n")), 2);
+			DumpPalette(Indent, Blueprint, Class, FileOutWriter);
+		}
+	}
+}
+
+//------------------------------------------------------------------------------
+static void DumpBlueprintInfoUtils::DumpPalette(uint32 Indent, UBlueprint* Blueprint, UClass* ClassFilter, FArchive* FileOutWriter)
+{
+	FString const PaletteEntryIndent = BuildIndentString(Indent);
+	FString BeginPaletteEntry = FString::Printf(TEXT("%s\"Palette"), *PaletteEntryIndent);
+
+	FString FilterClassName("<UNFILTERED>");
+	if (ClassFilter != nullptr)
+	{
+		FilterClassName = ClassFilter->GetName();
+		BeginPaletteEntry += "-" + FilterClassName;
+	}
+	BeginPaletteEntry += "\" : {\n";
+
+	FString const NestedIndent = BuildIndentString(Indent + 1);
+	UE_LOG(LogBlueprintInfoDump, Display, TEXT("%sDumping palette: %s"), *BuildIndentString(Indent, true), *FilterClassName);
+
+	bool bIsAnimBlueprint = (Cast<UAnimBlueprint>(Blueprint) != nullptr);
+	// animation blueprints don't have a palette
+	if (bIsAnimBlueprint)
+	{
+		BeginPaletteEntry += NestedIndent + "\"IsAnimBlueprint\" : true";
+		FileOutWriter->Serialize(TCHAR_TO_ANSI(*BeginPaletteEntry), BeginPaletteEntry.Len());
+	}
+	else
+	{
+		FBlueprintPaletteListBuilder PaletteBuilder(Blueprint);
+		double MenuBuildDuration = GetPaletteMenuActions(PaletteBuilder, ClassFilter);
+
+		BeginPaletteEntry += NestedIndent + "\"FilterClass\" : \"" + FilterClassName + "\",\n";
+		if (CommandOptions.DumpFlags & BPDUMP_RecordTiming)
+		{
+			BeginPaletteEntry += FString::Printf(TEXT("%s\"MenuBuildTime_Seconds\" : %f,\n"), *NestedIndent, MenuBuildDuration);
+		}
+
+		FileOutWriter->Serialize(TCHAR_TO_ANSI(*BeginPaletteEntry), BeginPaletteEntry.Len());
+		DumpActionList(Indent + 1, PaletteBuilder, FileOutWriter);
+	}
+
+	FString EndPaletteEntry = "\n" + PaletteEntryIndent + "}";
+	FileOutWriter->Serialize(TCHAR_TO_ANSI(*EndPaletteEntry), EndPaletteEntry.Len());
+}
+
+//------------------------------------------------------------------------------
+static void DumpBlueprintInfoUtils::DumpActionList(uint32 Indent, FGraphActionListBuilderBase& ActionList, FArchive* FileOutWriter)
+{
+	TArray< FGraphActionListBuilderBase::ActionGroup const* > SortedActions;
+	for (int32 ActionIndex = 0; ActionIndex < ActionList.GetNumActions(); ++ActionIndex)
+	{
+		FGraphActionListBuilderBase::ActionGroup& Action = ActionList.GetAction(ActionIndex);
+		if (Action.Actions.Num() <= 0)
+		{
+			continue;
+		}
+
+		SortedActions.Add(&Action);
+	}
+
+	FString const ActionListIndent = BuildIndentString(Indent);
+	FString const NestedIndent     = BuildIndentString(Indent + 1);
+
+	FString BeginActionListEntry = FString::Printf(TEXT("%s\"ActionSet\" : {\n%s\"ActionCount\" : %d"), *ActionListIndent, *NestedIndent, SortedActions.Num());
+	BeginActionListEntry += ",\n" + NestedIndent + "\"Actions\" : ";
+
+	bool const bNoActionInfo = (CommandOptions.DumpFlags & BPDUMP_DoNotDumpActionInfo) != 0;
+	if (bNoActionInfo)
+	{
+		BeginActionListEntry += "[";
+	}
+	else
+	{
+		BeginActionListEntry += "{";
+	}
+	FileOutWriter->Serialize(TCHAR_TO_ANSI(*BeginActionListEntry), BeginActionListEntry.Len());
+
+	struct ActionSortFunctor
+	{
+		bool operator()(FGraphActionListBuilderBase::ActionGroup const& LHS, FGraphActionListBuilderBase::ActionGroup const& RHS) const
+		{
+			FString LhKey = GetActionKey(LHS);
+			FString RhKey = GetActionKey(RHS);
+			return (LhKey.Compare(RhKey) < 0);
+		}
+	};
+	// need to sort so we can easily compare from one generation to the next
+	SortedActions.Sort(ActionSortFunctor());
+
+	FString LineEnding("\n");
+	for (FGraphActionListBuilderBase::ActionGroup const* Action : SortedActions)
+	{
+		FileOutWriter->Serialize(TCHAR_TO_ANSI(*LineEnding), LineEnding.Len());
+		DumpActionMenuItem(Indent + 2, *Action, FileOutWriter);
+		LineEnding = TEXT(",\n");
+	}
+
+	FString EndActionListEntry = "\n" + NestedIndent;
+	if (bNoActionInfo)
+	{
+		EndActionListEntry += "]";
+	}
+	else
+	{
+		EndActionListEntry += "}";
+	}
+	EndActionListEntry += "\n" + ActionListIndent + "}";
+	FileOutWriter->Serialize(TCHAR_TO_ANSI(*EndActionListEntry), EndActionListEntry.Len());
+}
+
+//------------------------------------------------------------------------------
+static void DumpBlueprintInfoUtils::DumpActionMenuItem(uint32 Indent, FGraphActionListBuilderBase::ActionGroup const& Action, FArchive* FileOutWriter)
+{
+	check(Action.Actions.Num() > 0);
+
+	// Get action category info
+	TArray<FString> MenuHierarchy;
+	Action.GetCategoryChain(MenuHierarchy);
+
+	FString ActionCategory = TEXT("");
+
+	bool bHasCategory = (MenuHierarchy.Num() > 0);
+	if (bHasCategory)
+	{
+		for (FString const& SubCategory : MenuHierarchy)
+		{
+			ActionCategory += SubCategory + TEXT("|");
+		}
+	}
+
+	TSharedPtr<FEdGraphSchemaAction> PrimeAction = Action.Actions[0];
+	FString const ActionName = PrimeAction->MenuDescription.ToString();
+
+	FString const ActionEntryIndent = BuildIndentString(Indent);
+	FString ActionEntry = ActionEntryIndent + "\"" + ActionCategory + ActionName + "\"";;
+
+	bool const bNoActionInfo = (CommandOptions.DumpFlags & BPDUMP_DoNotDumpActionInfo) != 0;
+	if (!bNoActionInfo)
+	{
+		++Indent;
+		FString const IndentedNewline = "\n" + BuildIndentString(Indent);
+
+		ActionEntry += " : {";
+		ActionEntry += IndentedNewline + "\"Name\"        : \"" + ActionName + "\",";
+		ActionEntry += IndentedNewline + "\"Category\"    : \"";
+		if (bHasCategory)
+		{
+			ActionEntry += ActionCategory;
+			ActionEntry.RemoveAt(ActionEntry.Len() - 1); // remove the trailing '|'
+		}
+		ActionEntry += "\","; // end action category data
+
+		FString TooltipStr = PrimeAction->TooltipDescription.Replace(TEXT("\\\""), TEXT("'")).Replace(TEXT("\""), TEXT("'"));
+		FString const TooltipFieldLabel("\"Tooltip\"     : \"");
+		TooltipStr = TooltipStr.Replace(TEXT("\n"), *(IndentedNewline + BuildIndentString(TooltipFieldLabel.Len(), /*bUseSpaces =*/true)));
+
+		ActionEntry += IndentedNewline + TooltipFieldLabel + TooltipStr + "\",";
+		ActionEntry += IndentedNewline + "\"Keywords\"    : \"" + PrimeAction->Keywords + "\",";
+		ActionEntry += IndentedNewline + "\"SearchTitle\" : \"" + PrimeAction->SearchTitle.ToString() + "\",";
+		ActionEntry += IndentedNewline + FString::Printf(TEXT("\"Grouping\"    : %d"), PrimeAction->Grouping);
+		
+		// Get action node type info
+		UK2Node const* NodeTemplate = FK2SchemaActionUtils::ExtractNodeTemplateFromAction(PrimeAction);
+		if (NodeTemplate != nullptr)
+		{
+			// Build action node type data
+			ActionEntry += "," + IndentedNewline + "\"Node\"        : \"";
+			ActionEntry += NodeTemplate->GetClass()->GetPathName();
+			ActionEntry += "\"";
+		}
+		// Finish action entry
+		ActionEntry += "\n" + ActionEntryIndent + "}";
+	}
+
+	// Write entry to file
+	FileOutWriter->Serialize(TCHAR_TO_ANSI(*ActionEntry), ActionEntry.Len());
+}
 
 //------------------------------------------------------------------------------
-UDumpBlueprintsInfoCommandlet::UDumpBlueprintsInfoCommandlet(class FPostConstructInitializeProperties const& PCIP)
-	: Super(PCIP)
+static void DumpBlueprintInfoUtils::DumpContextualActions(uint32 Indent, UBlueprint* Blueprint, FArchive* FileOutWriter)
 {
+	UE_LOG(LogBlueprintInfoDump, Display, TEXT("%sDumping contextual info..."), *BuildIndentString(Indent, true));
+
+	FString const ContextualEntryIndent = BuildIndentString(Indent);
+	FString const BeginContextualEntry = FString::Printf(TEXT("%s\"GraphActions\" : {\n"), *ContextualEntryIndent);
+
+	FileOutWriter->Serialize(TCHAR_TO_ANSI(*BeginContextualEntry), BeginContextualEntry.Len());
+
+	TArray<UEdGraph*> BpGraphs;
+	Blueprint->GetAllGraphs(BpGraphs);
+
+	bool bIsFirstEntry = true;
+	for (UEdGraph* Graph : BpGraphs)
+	{
+		UEdGraphSchema const* GraphSchema = GetDefault<UEdGraphSchema>(Graph->Schema);
+		if ((CommandOptions.GraphFilter != GT_MAX) && (GraphSchema->GetGraphType(Graph) != CommandOptions.GraphFilter))
+		{
+			continue;
+		}
+
+		if (bIsFirstEntry)
+		{
+			FileOutWriter->Serialize(TCHAR_TO_ANSI(TEXT("\n")), 1);
+			bIsFirstEntry = false;
+		}
+		else
+		{
+			FileOutWriter->Serialize(TCHAR_TO_ANSI(TEXT(",\n")), 2);
+		}
+		DumpGraphContextActions(Indent + 1, Graph, FileOutWriter);
+	}
+
+	FString const EndContextualEntry = FString::Printf(TEXT("\n%s}"), *ContextualEntryIndent);
+	FileOutWriter->Serialize(TCHAR_TO_ANSI(*EndContextualEntry), EndContextualEntry.Len());
 }
 
 //------------------------------------------------------------------------------
-int32 UDumpBlueprintsInfoCommandlet::Main(FString const& Params)
+static void DumpBlueprintInfoUtils::DumpGraphContextActions(uint32 Indent, UEdGraph* Graph, FArchive* FileOutWriter)
 {
-	TArray<FString> Tokens, Switches;
-	ParseCommandLine(*Params, Tokens, Switches);
+	UE_LOG(LogBlueprintInfoDump, Display, TEXT("%sDumping graph: '%s'..."), *BuildIndentString(Indent, true), *Graph->GetName());
 
-	FString CommandletSaveDir = FPaths::GameSavedDir() + TEXT("Commandlets/");
-	IFileManager::Get().MakeDirectory(*CommandletSaveDir);
- 
-	FString Filename = FString::Printf(TEXT("BlueprintsInfoDump_%s.json"), FPlatformTime::StrTimestamp());
-	Filename = Filename.Replace(TEXT(" "), TEXT("_"));
-	Filename = Filename.Replace(TEXT("/"), TEXT("-"));
-	Filename = Filename.Replace(TEXT(":"), TEXT("."));
-	Filename = FPaths::GetCleanFilename(Filename);
- 
-	FString FilePath = CommandletSaveDir / *Filename;
-	FArchive* FileOut = IFileManager::Get().CreateFileWriter(*FilePath);
-	if (FileOut != nullptr)
+	FString const GraphEntryIndent = BuildIndentString(Indent);
+	FString BeginGraphEntry = FString::Printf(TEXT("%s\"%s\" : {"), *GraphEntryIndent, *Graph->GetName());
+
+	FString const NestedIndent = BuildIndentString(++Indent);
+	FString const IndentedNewline = "\n" + NestedIndent;
+	BeginGraphEntry += IndentedNewline + "\"GraphType\" : \"" + GetGraphTypeString(Graph) + "\",";
+	BeginGraphEntry += IndentedNewline + "\"GraphName\" : \"" + Graph->GetName() + "\",";
+	BeginGraphEntry += IndentedNewline + "\"GraphContextMenu\" : \n";
+
+	FileOutWriter->Serialize(TCHAR_TO_ANSI(*BeginGraphEntry), BeginGraphEntry.Len());
+
+	FBlueprintGraphActionListBuilder ActionBuilder(Graph);
+	DumpContextActionList(Indent, ActionBuilder, FileOutWriter);
+
+	if (CommandOptions.DumpFlags & BPDUMP_SelectAllObjTypes)
 	{
-		// Get class name from command line switches
-		UClass* ClassFilter = nullptr;
-		for (FString& Switch : Switches)
+		for (TObjectIterator<UClass> ClassIt; ClassIt; ++ClassIt)
 		{
-			if (Switch.StartsWith("class="))
+			UClass* Class = *ClassIt;
+			if (Class->IsChildOf(AActor::StaticClass()))
 			{
-				FString ClassSwitch, ClassName;
-				Switch.Split(TEXT("="), &ClassSwitch, &ClassName);
-				ClassFilter = FindObject<UClass>(ANY_PACKAGE, *ClassName);
+				if (AActor* LevelActor = SpawnLevelActor(Class, true))
+				{
+					UE_LOG(LogBlueprintInfoDump, Display, TEXT("%sDumping level actor actions: '%s'..."), *BuildIndentString(Indent, true), *Class->GetName());
+
+					FString ActorSelectionEntry = "," + IndentedNewline + "\"LevelActorMenu-" + Class->GetName() + "\" : \n";
+					FileOutWriter->Serialize(TCHAR_TO_ANSI(*ActorSelectionEntry), ActorSelectionEntry.Len());
+
+					DumpContextActionList(Indent, ActionBuilder, FileOutWriter);
+				}
 			}
 		}
+	}
+
+	UBlueprint* Blueprint = CastChecked<UBlueprint>(Graph->GetOuter());
+	TArray<UObjectProperty*> ComponentProperties;
+	GetComponentProperties(Blueprint, ComponentProperties);
+
+	for (UObjectProperty* Component : ComponentProperties)
+	{
+		UE_LOG(LogBlueprintInfoDump, Display, TEXT("%sDumping component actions: '%s'..."), *BuildIndentString(Indent, true), *Component->GetName());
+
+		FString SelectionContextEntry = "," + IndentedNewline + "\"ComponentContextMenu-" + Component->GetName() + "\" : \n";
+		FileOutWriter->Serialize(TCHAR_TO_ANSI(*SelectionContextEntry), SelectionContextEntry.Len());
 
-		// Build class context container
-		FString LeadingFields("{\n\t\"ClassContext\" : \"");
-		if (ClassFilter != nullptr)
+		ActionBuilder.SelectedObjects.Empty();
+		ActionBuilder.SelectedObjects.Add(Component);
+		DumpContextActionList(Indent, ActionBuilder, FileOutWriter);
+	}
+
+	FString EndGraphEntry;
+	if (CommandOptions.DumpFlags & BPDUMP_PinContextActions)
+	{
+		FString const PinActionsEntry = "," + IndentedNewline + IndentedNewline + "\"PinContextActions\" : [\n";
+		FileOutWriter->Serialize(TCHAR_TO_ANSI(*PinActionsEntry), PinActionsEntry.Len());
+
+		if (DumpPinContextActions(Indent + 1, Graph, FileOutWriter))
 		{
-			LeadingFields += ClassFilter->GetName();
+			EndGraphEntry = IndentedNewline + "]";
 		}
 		else
 		{
-			LeadingFields += "<NO_CONTEXT>";
+			EndGraphEntry = NestedIndent + "]";
 		}
-		// Build palette contents header
-		LeadingFields += "\",\n\t\"Palette\" : {";
-		
-		FileOut->Serialize(TCHAR_TO_ANSI(*LeadingFields), LeadingFields.Len());
+	}
+
+	EndGraphEntry += "\n" + GraphEntryIndent + "}";
+	FileOutWriter->Serialize(TCHAR_TO_ANSI(*EndGraphEntry), EndGraphEntry.Len());
+}
 
-		UPackage* TransientPackage = GetTransientPackage();
-		FName TempBpName = MakeUniqueObjectName(TransientPackage, UBlueprint::StaticClass(), FName(TEXT("COMMANDLET_TEMP_Blueprint")));
-		UBlueprint* TempBlueprint = FKismetEditorUtilities::CreateBlueprint(AActor::StaticClass(), TransientPackage, TempBpName, BPTYPE_Normal, UBlueprint::StaticClass(), UBlueprintGeneratedClass::StaticClass());
+//------------------------------------------------------------------------------
+static bool DumpBlueprintInfoUtils::DumpPinContextActions(uint32 Indent, UEdGraph* Graph, FArchive* FileOutWriter)
+{
+	FBlueprintGraphActionListBuilder ContextMenuBuilder(Graph);
+
+	bool bWroteToFile = false;
+	if (!CommandOptions.PinType.IsEmpty())
+	{
+		FEdGraphPinType PinType;
+		PinType.PinCategory = CommandOptions.PinType;
 
-		FBlueprintPaletteListBuilder PaletteBuilder(TempBlueprint);
 		UEdGraphSchema_K2 const* K2Schema = GetDefault<UEdGraphSchema_K2>();
-		FK2ActionMenuBuilder(K2Schema).GetPaletteActions(PaletteBuilder, ClassFilter);
+		bool const bUsePinTypeClass = (CommandOptions.DumpFlags & BPDUMP_PinTypeIsClass) != 0;
 
-		for (int32 ActionIndex = 0; ActionIndex < PaletteBuilder.GetNumActions(); ++ActionIndex)
+		UClass* TypeClass = FindObject<UClass>(ANY_PACKAGE, *CommandOptions.PinType);
+		if ((TypeClass != nullptr) && UEdGraphSchema_K2::IsAllowableBlueprintVariableType(TypeClass))
 		{
-			FGraphActionListBuilderBase::ActionGroup& Action = PaletteBuilder.GetAction(ActionIndex);
-			if (Action.Actions.Num() <= 0)
+			PinType.PinSubCategoryObject = TypeClass;
+			if (TypeClass->IsChildOf(UInterface::StaticClass()))
 			{
-				continue;
+				PinType.PinCategory = K2Schema->PC_Interface;
+			}
+			else if (TypeClass->IsChildOf(UScriptStruct::StaticClass()))
+			{
+				PinType.PinCategory = K2Schema->PC_Struct;
+				PinType.PinSubCategoryObject = TypeClass->GetDefaultObject();
+			}
+			else
+			{
+				PinType.PinCategory = K2Schema->PC_Object;
 			}
 
-			// Get action name
-			TSharedPtr<FEdGraphSchemaAction> PrimeAction = Action.Actions[0];
-			FString ActionName = PrimeAction->MenuDescription.ToString();
-			
-			// Get action category info
-			TArray<FString> MenuHierarchy;
-			Action.GetCategoryChain(MenuHierarchy);
-			bool bHasCategory = (MenuHierarchy.Num() > 0);
-			FString ActionCategory = TEXT("");
-
-			if (bHasCategory)
-			{	
-				for (FString const& SubCategory : MenuHierarchy)
-				{
-					ActionCategory += SubCategory + TEXT("|");
-				}
+			if (bUsePinTypeClass)
+			{
+				PinType.PinCategory = K2Schema->PC_Class;
+			}
+		}
+		else if (!CommandOptions.PinType.Compare("self", ESearchCase::IgnoreCase))
+		{
+			PinType.PinCategory = K2Schema->PC_Object;
+			if (bUsePinTypeClass)
+			{
+				PinType.PinCategory = K2Schema->PC_Class;
+			}
+			PinType.PinSubCategory = K2Schema->PSC_Self;
+		}
+		// @TODO: PC_Delegate, PC_MCDelegate
+
+		DumpContextualPinTypeActions(Indent, Graph, PinType, FileOutWriter);
+		FileOutWriter->Serialize(TCHAR_TO_ANSI(TEXT(",\n")), 2);
+		PinType.bIsArray = true;
+		DumpContextualPinTypeActions(Indent, Graph, PinType, FileOutWriter);
+
+		bWroteToFile = true;
+	}
+	else if (Graph->Schema->IsChildOf(UEdGraphSchema_K2::StaticClass()))
+	{
+		TArray< TSharedPtr<UEdGraphSchema_K2::FPinTypeTreeInfo> > TypeTree;
+		UEdGraphSchema_K2 const* GraphSchema = GetDefault<UEdGraphSchema_K2>(Graph->Schema);
+		GraphSchema->GetVariableTypeTree(TypeTree, /*bAllowExec =*/true, /*bAllowWildcard =*/true);
+
+		for (TSharedPtr<UEdGraphSchema_K2::FPinTypeTreeInfo> TypeInfo : TypeTree)
+		{
+			bWroteToFile |= DumpTypeTreeActions(Indent, Graph, TypeInfo, FileOutWriter);
+		}
+	}
+	else // state machine graph?
+	{
+		// look in graph for nodes and mine their pin types?
+	}
+
+	return bWroteToFile;
+}
+
+//------------------------------------------------------------------------------
+static bool DumpBlueprintInfoUtils::DumpTypeTreeActions(uint32 Indent, UEdGraph* Graph, TSharedPtr<UEdGraphSchema_K2::FPinTypeTreeInfo> PinTypeInfo, FArchive* FileOutWriter)
+{
+	FString PendingLineEnding;
+
+	bool bWroteToFile = false;
+	if (!PinTypeInfo->bReadOnly)
+	{
+		FEdGraphPinType PinType = PinTypeInfo->GetPinType(/*bForceLoadedSubCategoryObject =*/false);
+		DumpContextualPinTypeActions(Indent, Graph, PinType, FileOutWriter);
+		FileOutWriter->Serialize(TCHAR_TO_ANSI(TEXT(",\n")), 2);
+
+		PinType.bIsArray = true;
+		DumpContextualPinTypeActions(Indent, Graph, PinType, FileOutWriter);
+
+		PendingLineEnding = TEXT(",\n");
+		bWroteToFile = true;
+	}
+
+	for (int32 ChildIndex = 0; ChildIndex < PinTypeInfo->Children.Num(); ++ChildIndex)
+	{
+		TSharedPtr<UEdGraphSchema_K2::FPinTypeTreeInfo> ChildInfo = PinTypeInfo->Children[ChildIndex];
+		if (ChildInfo.IsValid())
+		{
+			if (!PendingLineEnding.IsEmpty())
+			{
+				FileOutWriter->Serialize(TCHAR_TO_ANSI(*PendingLineEnding), PendingLineEnding.Len());
+				PendingLineEnding.Empty();
+			}
+
+			if (DumpTypeTreeActions(Indent, Graph, ChildInfo, FileOutWriter))
+			{
+				bWroteToFile |= true;
+				PendingLineEnding = TEXT(",\n");
 			}
+			// else, opps... we may be left with an invalid PendingLineEnding that was serialized
+		}
+	}
+
+	return bWroteToFile;
+}
+
+//------------------------------------------------------------------------------
+static void DumpBlueprintInfoUtils::DumpContextualPinTypeActions(uint32 Indent, UEdGraph* Graph, FEdGraphPinType const& PinType, FArchive* FileOutWriter)
+{
+	FBlueprintGraphActionListBuilder ContextMenuBuilder(Graph);
+
+	UEdGraphPin* DummyPin = NewObject<UEdGraphPin>(GetTransientPackage());
+	DummyPin->PinName = DummyPin->GetName();
+	DummyPin->Direction = EGPD_Input;
+	DummyPin->PinType = PinType;
+
+	ContextMenuBuilder.FromPin = DummyPin;
+
+	DumpContextActionList(Indent, ContextMenuBuilder, FileOutWriter);
+	FileOutWriter->Serialize(TCHAR_TO_ANSI(TEXT(",\n")), 2);
+	DummyPin->Direction = EGPD_Output;
+	DumpContextActionList(Indent, ContextMenuBuilder, FileOutWriter);
+}
+
+//------------------------------------------------------------------------------
+static void DumpBlueprintInfoUtils::DumpContextActionList(uint32 Indent, FBlueprintGraphActionListBuilder ActionBuilder, FArchive* FileOutWriter)
+{
+	ActionBuilder.Empty();
+
+	FString PinTypeLog("<UNKNOWN>");
+	if (ActionBuilder.FromPin != nullptr)
+	{
+		if (ActionBuilder.FromPin->Direction == EGPD_Input)
+		{
+			PinTypeLog = "INPUT ";
+		}
+		else
+		{
+			PinTypeLog = "OUTPUT";
+		}
+		PinTypeLog += UEdGraphSchema_K2::TypeToString(ActionBuilder.FromPin->PinType);
+	}
+	UE_LOG(LogBlueprintInfoDump, Display, TEXT("%sDumping pin actions: %s"), *BuildIndentString(Indent, true), *PinTypeLog);
+
+	double MenuBuildDuration = GetContextMenuActions(ActionBuilder);
+
+	FString const ContextEntryIndent = BuildIndentString(Indent);
+	FString BeginContextEntry = FString::Printf(TEXT("%s{\n"), *ContextEntryIndent);
+
+	if (CommandOptions.DumpFlags & BPDUMP_RecordTiming)
+	{
+		BeginContextEntry += FString::Printf(TEXT("%s\"MenuBuildTime_Seconds\" : %f,\n"), *BuildIndentString(Indent + 1), MenuBuildDuration);
+	}
+	FileOutWriter->Serialize(TCHAR_TO_ANSI(*BeginContextEntry), BeginContextEntry.Len());
+
+	DumpContextInfo(Indent + 1, ActionBuilder, FileOutWriter);
+	FileOutWriter->Serialize(TCHAR_TO_ANSI(TEXT(",\n")), 2);
+	DumpActionList(Indent + 1, ActionBuilder, FileOutWriter);
+
+	FString EndContextEntry = FString::Printf(TEXT("\n%s}"), *ContextEntryIndent);
+	FileOutWriter->Serialize(TCHAR_TO_ANSI(*EndContextEntry), EndContextEntry.Len());
+}
+
+//------------------------------------------------------------------------------
+static double DumpBlueprintInfoUtils::GetContextMenuActions(FBlueprintGraphActionListBuilder& ActionBuilder)
+{
+	ActionBuilder.Empty();
+	check(ActionBuilder.CurrentGraph != nullptr);
+
+	UEdGraphSchema const* GraphSchema = GetDefault<UEdGraphSchema>(ActionBuilder.CurrentGraph->Schema);
+
+	double MenuBuildDuration = 0.0;
+	FBlueprintPaletteListBuilder PaletteBuilder(ActionBuilder.Blueprint);
+	{
+		FScopedDurationTimer DurationTimer(MenuBuildDuration);
+		GraphSchema->GetGraphContextActions(ActionBuilder);
+	}
+
+	return MenuBuildDuration;
+}
+
+//------------------------------------------------------------------------------
+static void DumpBlueprintInfoUtils::DumpContextInfo(uint32 Indent, FBlueprintGraphActionListBuilder const& ActionBuilder, FArchive* FileOutWriter)
+{
+	FString const ContextEntryIndent = BuildIndentString(Indent);
+	++Indent;
+	FString const IndentedNewline = "\n" + BuildIndentString(Indent);
+	
+	FString ContextEntry = FString::Printf(TEXT("%s\"Context\" : {%s\"Graph\"   : \""), *ContextEntryIndent, *IndentedNewline);
+	ContextEntry += GetGraphTypeString(ActionBuilder.CurrentGraph) + "\"," + IndentedNewline + "\"PinType\" : \"";
+
+	if (ActionBuilder.FromPin == nullptr)
+	{
+		ContextEntry += "<NONE>\"," + IndentedNewline + "\"PinDirection\" : \"<NONE>";
+	}
+	else
+	{
+		ContextEntry += UEdGraphSchema_K2::TypeToString(ActionBuilder.FromPin->PinType);
+
+		ContextEntry += "\"," + IndentedNewline + "\"PinDirection\" : \"";
+		if (ActionBuilder.FromPin->Direction == EGPD_Input)
+		{
+			ContextEntry += "Input";
+		}
+		else
+		{
+			check(ActionBuilder.FromPin->Direction == EGPD_Output);
+			ContextEntry += "Output";
+		}
+	}	
+	ContextEntry += "\"," + IndentedNewline + "\"SelectedObjects\" : [";
+	
+	FString const ActorEntryIndent = "\n" + BuildIndentString(Indent + 1);
+	for (UObject* Selection : ActionBuilder.SelectedObjects)
+	{
+		ContextEntry += ActorEntryIndent + "\"" + CastChecked<UObjectProperty>(Selection)->PropertyClass->GetName() + "\",";
+	}
+	if (ActionBuilder.SelectedObjects.Num() > 0)
+	{
+		ContextEntry.RemoveAt(ContextEntry.Len() - 1); // remove the trailing ','
+		ContextEntry += IndentedNewline;
+	}
+	ContextEntry += "]\n" + ContextEntryIndent + "}";
+	
+	FileOutWriter->Serialize(TCHAR_TO_ANSI(*ContextEntry), ContextEntry.Len());
+}
+
+//------------------------------------------------------------------------------
+// @TODO: should probably be a platform agnostic function
+static bool DumpBlueprintInfoUtils::DoFilesDiffer(FString const& NewFilePath, FString const& OldFilePath)
+{
+	int32 ReturnCode = -1;
+	FString const CommandStr = "fc " + NewFilePath.Replace(TEXT("/"), TEXT("\\")) + " " + OldFilePath.Replace(TEXT("/"), TEXT("\\"));
+
+	// @TODO: very platform specific :(
+	bool bSuccess = FPlatformProcess::ExecProcess(
+		TEXT("cmd.exe"), 
+		*("/C \"" + CommandStr + "\""), 
+		&ReturnCode,
+		/*OutStdOut =*/nullptr, 
+		/*OutStdErr =*/nullptr);
+
+	if (!bSuccess)
+	{
+		UE_LOG(LogBlueprintInfoDump, Error, TEXT("Failed to run diff: '%s'"), *CommandStr);
+	}
+	
+	// ReturnCode
+	// 0 : no differences?
+	// 1 : differences?
+	// 2 : couldn't find file? 
 
+	return !bSuccess || (ReturnCode != 0);
+}
 
-			// Build action key data
-			FString ActionEntry("\n\t\t\"");
-			ActionEntry += ActionCategory + ActionName + TEXT("\" : {");
+//------------------------------------------------------------------------------
+static void DumpBlueprintInfoUtils::DiffDumpFiles(FString const& NewFilePath, FString const& OldFilePath, FString const& UserDiffCmd)
+{
+	check(!NewFilePath.IsEmpty());
+	check(FPaths::FileExists(NewFilePath));
+	check(!OldFilePath.IsEmpty() || !UserDiffCmd.IsEmpty());
 
-			// Build action name data
-			ActionEntry += TEXT("\n\t\t\t\"Name\" : \"");
-			ActionEntry += ActionName;
-			ActionEntry += TEXT("\"");
+	UE_LOG(LogBlueprintInfoDump, Display, TEXT("Diffing: '%s'..."), *NewFilePath);
 
-			// Build action category key data
-			ActionEntry += TEXT(",\n\t\t\t\"Category\" : \"");
-			if (bHasCategory)
+	FString WorkingDirectory = FPaths::GetPath(NewFilePath);
+	// if we created a whole new folder for this file, then make the working 
+	// directory one up
+	bool const bSplitBlueprintsByFile = (CommandOptions.DumpFlags & BPDUMP_FilePerBlueprint) != 0;
+	if (bSplitBlueprintsByFile)
+	{
+		WorkingDirectory = WorkingDirectory / TEXT("..");
+	}
+	
+	FString QualifiedOldFilePath;
+	if (!OldFilePath.IsEmpty())
+	{
+		IFileManager& FileManager = IFileManager::Get();
+		bool const bIsLocalDirectory = FileManager.DirectoryExists(*(WorkingDirectory / *OldFilePath));
+		bool const bIsDirectoryPath  = !bIsLocalDirectory && FileManager.DirectoryExists(*OldFilePath);
+		
+		if (bIsLocalDirectory)
+		{
+			FString Filename = FPaths::GetCleanFilename(NewFilePath);
+			QualifiedOldFilePath = WorkingDirectory / *(OldFilePath / *Filename);
+		}
+		else if (bIsDirectoryPath)
+		{
+			FString Filename = FPaths::GetCleanFilename(NewFilePath);
+			QualifiedOldFilePath = OldFilePath / *Filename;
+		}
+		else
+		{
+			QualifiedOldFilePath = OldFilePath;
+			if (FPaths::GetPath(QualifiedOldFilePath).IsEmpty())
 			{
-				// Build action category data from info
-				ActionEntry += ActionCategory;
-				ActionEntry.RemoveAt(ActionEntry.Len() - 1); // remove the last '|'
+				QualifiedOldFilePath = WorkingDirectory / OldFilePath;
 			}
-			// End action category data
-			ActionEntry += TEXT("\"");
+		}
+	}
+
+	if (!FPaths::FileExists(QualifiedOldFilePath) && !QualifiedOldFilePath.IsEmpty())
+	{
+		UE_LOG(LogBlueprintInfoDump, Error, TEXT("Cannot find file '%s' to diff against"), *QualifiedOldFilePath);
+	}
+	else if (!bSplitBlueprintsByFile || QualifiedOldFilePath.IsEmpty() || DoFilesDiffer(NewFilePath, QualifiedOldFilePath))
+	{
+		FString DiffCommand = UserDiffCmd;
+		if (DiffCommand.IsEmpty())
+		{
+			// -dw : Ignore line ending and all whitespace differences
+			DiffCommand = TEXT("p4merge.exe -dw \"{2}\" \"{1}\"");
+		}
+		DiffCommand = DiffCommand.Replace(TEXT("{1}"), *NewFilePath).Replace(TEXT("{2}"), *QualifiedOldFilePath);
+
+		FString DiffArgs;
+		if (int32 ArgsIndex = DiffCommand.Find(TEXT(" ")))
+		{
+			DiffArgs = *DiffCommand + ArgsIndex + 1;
+			DiffCommand = DiffCommand.Left(ArgsIndex);
+		}
+
+		FProcHandle DiffProc = FPlatformProcess::CreateProc(*DiffCommand,
+			/*Params =*/              *DiffArgs,
+			/*bLaunchDetached =*/     true,
+			/*bLaunchHidden =*/       false,
+			/*bLaunchReallyHidden =*/ false,
+			/*OutProcessID =*/        nullptr,
+			/*PriorityModifier =*/    0,
+			*WorkingDirectory,
+			/*PipeWrite =*/           nullptr);
+
+		if (!DiffProc.IsValid())
+		{
+			UE_LOG(LogBlueprintInfoDump, Error, TEXT("Could not launch: '%s'"), *DiffCommand);
+		}
+	}
+}
+
+/*******************************************************************************
+ * UDumpBlueprintsInfoCommandlet
+ ******************************************************************************/
+
+//------------------------------------------------------------------------------
+UDumpBlueprintsInfoCommandlet::UDumpBlueprintsInfoCommandlet(class FPostConstructInitializeProperties const& PCIP)
+	: Super(PCIP)
+{
+}
+
+//------------------------------------------------------------------------------
+int32 UDumpBlueprintsInfoCommandlet::Main(FString const& Params)
+{
+	TArray<FString> Tokens, Switches;
+	ParseCommandLine(*Params, Tokens, Switches);
+
+	DumpBlueprintInfoUtils::LevelActors.Empty();
+
+	DumpBlueprintInfoUtils::CommandOptions = DumpBlueprintInfoUtils::CommandletOptions(Switches);
+	DumpBlueprintInfoUtils::CommandletOptions const& CommandOptions = DumpBlueprintInfoUtils::CommandOptions;
+
+	bool const bSplitFilesByClass = (CommandOptions.DumpFlags & DumpBlueprintInfoUtils::BPDUMP_FilePerBlueprint) != 0;
+	bool const bDiffGeneratedFile = !CommandOptions.DiffPath.IsEmpty() || !CommandOptions.DiffCommand.IsEmpty();
+
+	UClass* const SelectedObjType = DumpBlueprintInfoUtils::CommandOptions.SelectedObjectType;
+	// if the user specified that they want a level actor selected during the 
+	// dump, then spawn one and select it (extra blueprint context actions 
+	// appear in certain situations regarding selected level actors)
+	if ((SelectedObjType != nullptr) && SelectedObjType->IsChildOf(AActor::StaticClass()))
+	{
+		DumpBlueprintInfoUtils::SpawnLevelActor(SelectedObjType, /*bSelect =*/true);
+	}
+
+	FString ActiveFilePath;
+	FArchive* FileOut = nullptr;
 
-			// Get action node type info
-			UK2Node const* NodeTemplate = FK2SchemaActionUtils::ExtractNodeTemplateFromAction(PrimeAction);
-			if (NodeTemplate != nullptr)
+	// this lambda is responsible for adding closing characters to the file, and 
+	// closing out the writer (and diffing the resultant file if the user deigns us to do so)
+	auto CloseFileStream = [bDiffGeneratedFile, &ActiveFilePath, &CommandOptions](FArchive** FileOutPtr)
+	{
+		FArchive* FileOut = (*FileOutPtr);
+		if (FileOut != nullptr)
+		{
+			FileOut->Serialize(TCHAR_TO_ANSI(TEXT("\n}")), 2);
+			FileOut->Close();
+
+			if (bDiffGeneratedFile)
 			{
-				// Build action node type data
-				ActionEntry += TEXT(",\n\t\t\t\"Node\" : \"");
-				ActionEntry += NodeTemplate->GetClass()->GetPathName();
-				ActionEntry += TEXT("\"");
+				check(!ActiveFilePath.IsEmpty());
+				DumpBlueprintInfoUtils::DiffDumpFiles(ActiveFilePath, CommandOptions.DiffPath, CommandOptions.DiffCommand);
 			}
-			// Finish action entry
-			ActionEntry += TEXT("\n\t\t}");
 
-			bool bIsLastEntry = (ActionIndex == PaletteBuilder.GetNumActions()-1);
-			if (!bIsLastEntry)
+			delete FileOut;
+			FileOut = nullptr;
+		}
+	};
+
+	// this lambda is responsible for opening a file for write, and adding 
+	// opening json characters to the file (contextually tracks the active filepath as well)
+	auto OpenFileStream = [&ActiveFilePath, &CloseFileStream, &FileOut](UClass* Class)->FArchive*
+	{
+		CloseFileStream(&FileOut);
+
+		ActiveFilePath = DumpBlueprintInfoUtils::BuildDumpFilePath(Class);
+		FileOut = IFileManager::Get().CreateFileWriter(*ActiveFilePath);
+		check(FileOut != nullptr);
+		FileOut->Serialize(TCHAR_TO_ANSI(TEXT("{\n")), 2);
+
+		return FileOut;
+	};
+	
+	bool bNeedsJsonComma = false;
+	// this lambda will dump blueprint info for the specified class, if the user
+	// set -multifile, then this will also close the existing file and open a new one for this class
+	auto WriteClassInfo = [bSplitFilesByClass, &bNeedsJsonComma, &OpenFileStream, &CloseFileStream, &FileOut](UClass* Class)
+	{
+		if (bSplitFilesByClass && (FileOut != nullptr))
+		{
+			CloseFileStream(&FileOut);
+		}
+
+		if (FileOut == nullptr)
+		{
+			FileOut = OpenFileStream(Class);
+		}
+		// if we're adding all the class entries into one file, then we need to 
+		// separate them by a comma (or invalid json)
+		else if (!bSplitFilesByClass && bNeedsJsonComma)
+		{
+			FileOut->Serialize(TCHAR_TO_ANSI(TEXT(",\n")), 2);
+		}
+
+		DumpBlueprintInfoUtils::DumpInfoForClass(1, Class, FileOut);
+		bNeedsJsonComma = true;
+	};
+
+	// this lambda is used as a precursory check to ensure that the specified 
+	// class is a blueprintable type... broken into it's own lambda to save on reuse
+	auto IsInvalidBlueprintClass = [](UClass* Class)->bool
+	{
+		return Class->HasAnyClassFlags(CLASS_NewerVersionExists) || FKismetEditorUtilities::IsClassABlueprintSkeleton(Class) || !FKismetEditorUtilities::CanCreateBlueprintOfClass(Class);
+	};
+
+	UClass* const BlueprintClass = CommandOptions.BlueprintClass;
+	if (CommandOptions.DumpFlags & DumpBlueprintInfoUtils::BPDUMP_LogHelp)
+	{
+		UE_LOG(LogBlueprintInfoDump, Display, TEXT("%s"), *DumpBlueprintInfoUtils::HelpString);
+	}
+	else if (BlueprintClass != nullptr)
+	{
+		UE_LOG(LogBlueprintInfoDump, Display, TEXT("Dumping Blueprint info..."));
+		// make sure the class that the user specified is a blueprintable type
+		if (IsInvalidBlueprintClass(BlueprintClass))
+		{
+			UE_LOG(LogBlueprintInfoDump, Error, TEXT("Cannot create a blueprint from class: '%s'"), *BlueprintClass->GetName());
+			if (FileOut != nullptr)
 			{
-				// Add comma for next entry
-				ActionEntry += TEXT(",");
+				FString const InvalidClassEntry = DumpBlueprintInfoUtils::BuildIndentString(1) + "\"INVALID_BLUEPRINT_CLASS\" : \"" + BlueprintClass->GetName() + "\"";
+				FileOut->Serialize(TCHAR_TO_ANSI(*InvalidClassEntry), InvalidClassEntry.Len());
 			}
-			// Write entry to file
-			FileOut->Serialize(TCHAR_TO_ANSI(*ActionEntry), ActionEntry.Len());
 		}
-		// Close remaining entries and write to file
-		FileOut->Serialize(TCHAR_TO_ANSI(TEXT("\n\t}\n}")), 5);
- 
-		FileOut->Close();
-		delete FileOut;
-		FileOut = nullptr;
+		else
+		{
+			WriteClassInfo(BlueprintClass);
+		}
+	}
+	// if the user didn't specify a class, then we take that to mean dump ALL the classes!
+	else
+	{
+		UE_LOG(LogBlueprintInfoDump, Display, TEXT("Dumping Blueprint info..."));
+		for (TObjectIterator<UClass> ClassIt; ClassIt; ++ClassIt)
+		{
+			UClass* Class = *ClassIt;
+			if (IsInvalidBlueprintClass(Class))
+			{
+				continue;
+			}
+			WriteClassInfo(Class);
+		}
 	}
+
+	CloseFileStream(&FileOut);
 	return 0;
-}
\ No newline at end of file
+}
+
+
+
+
