Fixed bind pose issue where relative matrices don't match
- now we recreate missing bind pose

diff --git a/Engine/Source/Editor/UnrealEd/Private/Fbx/FbxSkeletalMeshImport.cpp b/Engine/Source/Editor/UnrealEd/Private/Fbx/FbxSkeletalMeshImport.cpp
--- a/Engine/Source/Editor/UnrealEd/Private/Fbx/FbxSkeletalMeshImport.cpp
+++ b/Engine/Source/Editor/UnrealEd/Private/Fbx/FbxSkeletalMeshImport.cpp
@@ -649,6 +649,97 @@ void UnFbx::FFbxImporter::BuildSkeletonSystem(TArray<FbxCluster*>& ClusterArray,
 	}
 }
 
+bool UnFbx::FFbxImporter::RetrievePoseFromBindPose(const TArray<FbxNode*>& NodeArray, FbxArray<FbxPose*> & PoseArray) const
+{
+	const int32 PoseCount = Scene->GetPoseCount();
+	for(int32 PoseIndex = 0; PoseIndex < PoseCount; PoseIndex++)
+	{
+		FbxPose* CurrentPose = Scene->GetPose(PoseIndex);
+
+		// current pose is bind pose, 
+		if(CurrentPose && CurrentPose->IsBindPose())
+		{
+			// IsValidBindPose doesn't work reliably
+			// It checks all the parent chain(regardless root given), and if the parent doesn't have correct bind pose, it fails
+			// It causes more false positive issues than the real issue we have to worry about
+			// If you'd like to try this, set CHECK_VALID_BIND_POSE to 1, and try the error message
+			// when Autodesk fixes this bug, then we might be able to re-open this
+			FString PoseName = CurrentPose->GetName();
+			// all error report status
+			FbxStatus Status;
+
+			// it does not make any difference of checking with different node
+			// it is possible pose 0 -> node array 2, but isValidBindPose function returns true even with node array 0
+			for(auto Current : NodeArray)
+			{
+				FString CurrentName = Current->GetName();
+				NodeList pMissingAncestors, pMissingDeformers, pMissingDeformersAncestors, pWrongMatrices;
+
+				if(CurrentPose->IsValidBindPoseVerbose(Current, pMissingAncestors, pMissingDeformers, pMissingDeformersAncestors, pWrongMatrices, 0.0001, &Status))
+				{
+					PoseArray.Add(CurrentPose);
+					UE_LOG(LogFbx, Log, TEXT("Valid bind pose for Pose (%s) - %s"), *PoseName, *FString(Current->GetName()));
+					break;
+				}
+				else
+				{
+					// first try to fix up
+					// add missing ancestors
+					for(int i = 0; i < pMissingAncestors.GetCount(); i++)
+					{
+						FbxAMatrix mat = pMissingAncestors.GetAt(i)->EvaluateGlobalTransform(FBXSDK_TIME_ZERO);
+						CurrentPose->Add(pMissingAncestors.GetAt(i), mat);
+					}
+
+					pMissingAncestors.Clear();
+					pMissingDeformers.Clear();
+					pMissingDeformersAncestors.Clear();
+					pWrongMatrices.Clear();
+
+					// check it again
+					if(CurrentPose->IsValidBindPose(Current))
+					{
+						PoseArray.Add(CurrentPose);
+						UE_LOG(LogFbx, Log, TEXT("Valid bind pose for Pose (%s) - %s"), *PoseName, *FString(Current->GetName()));
+						break;
+					}
+					else
+					{
+						// first try to find parent who is null group and see if you can try test it again
+						FbxNode * ParentNode = Current->GetParent();
+						while(ParentNode)
+						{
+							FbxNodeAttribute* Attr = ParentNode->GetNodeAttribute();
+							if(Attr && Attr->GetAttributeType() == FbxNodeAttribute::eNull)
+							{
+								// found it 
+								break;
+							}
+
+							// find next parent
+							ParentNode = ParentNode->GetParent();
+						}
+
+						if(ParentNode && CurrentPose->IsValidBindPose(ParentNode))
+						{
+							PoseArray.Add(CurrentPose);
+							UE_LOG(LogFbx, Log, TEXT("Valid bind pose for Pose (%s) - %s"), *PoseName, *FString(Current->GetName()));
+							break;
+						}
+						else
+						{
+							FString ErrorString = Status.GetErrorString();
+							UE_LOG(LogFbx, Warning, TEXT("Not valid bind pose for Pose (%s) - Node %s : %s"), *PoseName, *FString(Current->GetName()), *ErrorString);
+						}
+					}
+				}
+			}
+		}
+	}
+
+	return (PoseArray.Size() > 0);
+}
+
 bool UnFbx::FFbxImporter::ImportBone(TArray<FbxNode*>& NodeArray, FSkeletalMeshImportData &ImportData, TArray<FbxNode*> &SortedLinks, bool& bOutDiffPose, bool bDisableMissingBindPoseWarning, bool & bUseTime0AsRefPose)
 {
 	bOutDiffPose = false;
@@ -706,87 +797,32 @@ bool UnFbx::FFbxImporter::ImportBone(TArray<FbxNode*>& NodeArray, FSkeletalMeshI
 		}
 
 		// get bind pose
-		const int32 PoseCount = Scene->GetPoseCount();
-		for (int32 PoseIndex = 0; PoseIndex < PoseCount; PoseIndex++)
+		if(RetrievePoseFromBindPose(NodeArray, PoseArray) == false)
 		{
-			FbxPose* CurrentPose = Scene->GetPose(PoseIndex);
-
-			// current pose is bind pose, 
-			if (CurrentPose && CurrentPose->IsBindPose())
+			UE_LOG(LogFbx, Warning, TEXT("Getting valid bind pose failed. Try to recreate bind pose"));
+			// if failed, delete bind pose, and retry.
+			const int32 PoseCount = Scene->GetPoseCount();
+			for(int32 PoseIndex = PoseCount-1; PoseIndex >= 0; --PoseIndex)
 			{
-				// IsValidBindPose doesn't work reliably
-				// It checks all the parent chain(regardless root given), and if the parent doesn't have correct bind pose, it fails
-				// It causes more false positive issues than the real issue we have to worry about
-				// If you'd like to try this, set CHECK_VALID_BIND_POSE to 1, and try the error message
-				// when Autodesk fixes this bug, then we might be able to re-open this
-				FString PoseName = CurrentPose->GetName();
-				// all error report status
-				FbxStatus Status;
-				FbxUserNotification Notification(SdkManager, FbxString(), FbxString());
-				Notification.InitAccumulator();
-
-				// it does not make any difference of checking with different node
-				// it is possible pose 0 -> node array 2, but isValidBindPose function returns true even with node array 0
-				for (auto Current : NodeArray)
-				{
-					FString CurrentName = Current->GetName();
-					if(CurrentPose->IsValidBindPoseVerbose(Current, &Notification, 0.0001, &Status))
-					{
-						PoseArray.Add(CurrentPose);
-						UE_LOG(LogFbx, Warning, TEXT("Valid bind pose for Pose (%s) - %s"), *PoseName, *FString(Current->GetName()));
-						break;
-					}
-					else
-					{
-						// first try to find parent who is null group and see if you can try test it again
-						FbxNode * ParentNode = Current->GetParent();
-						while (ParentNode)
-						{
-							FbxNodeAttribute* Attr = ParentNode->GetNodeAttribute();
-							if ( Attr && Attr->GetAttributeType() == FbxNodeAttribute::eNull )
-							{
-								// found it 
-								break;
-							}
-
-							// find next parent
-							ParentNode = ParentNode->GetParent();
-						}
-
-						if (ParentNode && CurrentPose->IsValidBindPose(ParentNode))
-						{
-							PoseArray.Add(CurrentPose);
-							UE_LOG(LogFbx, Warning, TEXT("Valid bind pose for Pose (%s) - %s"), *PoseName, *FString(Current->GetName()));
-							break;
-						}
-						else
-						{
-							FString ErrorString = Status.GetErrorString();
+				FbxPose* CurrentPose = Scene->GetPose(PoseIndex);
 
-							UE_LOG(LogFbx, Warning, TEXT("Not valid bind pose for Pose (%s) - Node %s : %s"), *PoseName, *FString(Current->GetName()), *ErrorString);
-							for(int32 NotificationIndex=0; NotificationIndex<Notification.GetNbEntries(); ++NotificationIndex)
-							{
-								const FbxAccumulatorEntry* Entry = Notification.GetEntryAt(NotificationIndex);
-								if(Entry)
-								{
-									FString Name = (const char*)(Entry->GetName());
-									FString Description = (const char*)(Entry->GetDescription());
-									FString Detail;
-
-									for(int32 DetailCount =0; DetailCount< Entry->GetDetailsCount(); ++DetailCount)
-									{
-										Detail += FString((const char*)(*Entry->GetDetail(DetailCount)));
-									}
-
-									UE_LOG(LogFbx, Warning, TEXT("\t(%d) %s : %s (%s)"), NotificationIndex+1, *Name, *Description, *Detail);
-								}
-							}
-
-							Notification.ClearAccumulator();
-						}
-					}
+				// current pose is bind pose, 
+				if(CurrentPose && CurrentPose->IsBindPose())
+				{
+					Scene->RemovePose(PoseIndex);
+					CurrentPose->Destroy();
 				}
 			}
+
+			SdkManager->CreateMissingBindPoses(Scene);
+			if ( RetrievePoseFromBindPose(NodeArray, PoseArray) == false)
+			{
+				UE_LOG(LogFbx, Warning, TEXT("Recreating bind pose failed."));
+			}
+			else
+			{
+				UE_LOG(LogFbx, Warning, TEXT("Recreating bind pose succeeded."));
+			}
 		}
 
 		// recurse through skeleton and build ordered table
diff --git a/Engine/Source/Editor/UnrealEd/Private/FbxImporter.h b/Engine/Source/Editor/UnrealEd/Private/FbxImporter.h
--- a/Engine/Source/Editor/UnrealEd/Private/FbxImporter.h
+++ b/Engine/Source/Editor/UnrealEd/Private/FbxImporter.h
@@ -967,6 +967,14 @@ protected:
 	 */
 	bool IsValidAnimationData(TArray<FbxNode*>& SortedLinks, TArray<FbxNode*>& NodeArray, int32 & ValidTakeCount);
 
+	/**
+	 * Retrieve pose array from bind pose
+	 *
+	 * Iterate through Scene:Poses, and find valid bind pose for NodeArray, and return those Pose if valid
+	 *
+	 */
+	bool RetrievePoseFromBindPose(const TArray<FbxNode*>& NodeArray, FbxArray<FbxPose*> & PoseArray) const;
+
 public:
 	/** Import and set up animation related data from mesh **/
 	void SetupAnimationDataFromMesh(USkeletalMesh * SkeletalMesh, UObject * InParent, TArray<FbxNode*>& NodeArray, UFbxAnimSequenceImportData* ImportData, const FString & Filename);
