UE4: Fixed an issue where inactive worlds were incorrectly informing the engine and causing actors to spawn too early.

diff --git a/Engine/Source/Editor/UnrealEd/Classes/Factories/WorldFactory.h b/Engine/Source/Editor/UnrealEd/Classes/Factories/WorldFactory.h
--- a/Engine/Source/Editor/UnrealEd/Classes/Factories/WorldFactory.h
+++ b/Engine/Source/Editor/UnrealEd/Classes/Factories/WorldFactory.h
@@ -12,6 +12,7 @@ class UWorldFactory : public UFactory
 	GENERATED_UCLASS_BODY()
 
 	TEnumAsByte<EWorldType::Type> WorldType;
+	bool bInformEngineOfWorld;
 
 	// Begin UFactory Interface
 	virtual bool ConfigureProperties() override;
diff --git a/Engine/Source/Editor/UnrealEd/Private/EditorServer.cpp b/Engine/Source/Editor/UnrealEd/Private/EditorServer.cpp
--- a/Engine/Source/Editor/UnrealEd/Private/EditorServer.cpp
+++ b/Engine/Source/Editor/UnrealEd/Private/EditorServer.cpp
@@ -1831,6 +1831,7 @@ UWorld* UEditorEngine::NewMap()
 	// Create a new world
 	UWorldFactory* Factory = ConstructObject<UWorldFactory>(UWorldFactory::StaticClass());
 	Factory->WorldType = EWorldType::Editor;
+	Factory->bInformEngineOfWorld = true;
 	UPackage* Pkg = CreatePackage( NULL, NULL );
 	EObjectFlags Flags = RF_Public | RF_Standalone;
 	UWorld* NewWorld = CastChecked<UWorld>(Factory->FactoryCreateNew(UWorld::StaticClass(), Pkg, TEXT("NewWorld"), Flags, NULL, GWarn));
diff --git a/Engine/Source/Editor/UnrealEd/Private/Factories/WorldFactory.cpp b/Engine/Source/Editor/UnrealEd/Private/Factories/WorldFactory.cpp
--- a/Engine/Source/Editor/UnrealEd/Private/Factories/WorldFactory.cpp
+++ b/Engine/Source/Editor/UnrealEd/Private/Factories/WorldFactory.cpp
@@ -10,6 +10,7 @@ UWorldFactory::UWorldFactory(const class FPostConstructInitializeProperties& PCI
 	bCreateNew = FParse::Param(FCommandLine::Get(), TEXT("WorldAssets"));
 	SupportedClass = UWorld::StaticClass();
 	WorldType = EWorldType::Inactive;
+	bInformEngineOfWorld = false;
 }
 
 bool UWorldFactory::ConfigureProperties()
@@ -21,7 +22,6 @@ UObject* UWorldFactory::FactoryCreateNew(UClass* Class, UObject* InParent, FName
 {
 	// Create a new world.
 	const bool bAddToRoot = false;
-	const bool bInformEngineOfWorld = true;
 	UWorld* NewWorld = UWorld::CreateWorld(WorldType, bInformEngineOfWorld, Name, Cast<UPackage>(InParent), bAddToRoot);
 	NewWorld->SetFlags(Flags);
 
