Added usage comment to FEditorModeRegistry

diff --git a/Engine/Source/Editor/UnrealEd/Public/EditorModeRegistry.h b/Engine/Source/Editor/UnrealEd/Public/EditorModeRegistry.h
--- a/Engine/Source/Editor/UnrealEd/Public/EditorModeRegistry.h
+++ b/Engine/Source/Editor/UnrealEd/Public/EditorModeRegistry.h
@@ -2,6 +2,27 @@
 
 #pragma once
 
+/**
+ *	Class responsible for maintaining a list of registered editor mode types.
+ *
+ *	Example usage:
+ *
+ *	Register your mode type with:
+ *		FEditorModeRegistry::Get().RegisterMode<FMyEditorMode>( FName( TEXT("MyEditorMode") ) );
+ *	or:
+ *		class FMyEditorModeFactory : public IEditorModeFactory
+ *		{
+ *			virtual void OnSelectionChanged( FEditorModeTools& Tools, UObject* ItemUndergoingChange ) const override;
+ *			virtual FEditorModeInfo GetModeInfo() const override;
+ *			virtual TSharedRef<FEdMode> CreateMode() const override;
+ *		};
+ *		TSharedRef<FMyEditorModeFactory> Factory = MakeShareable( new FMyEditorModeFactory );
+ *		FEditorModeRegistry::Get().RegisterMode( FName( TEXT("MyEditorMode") ), Factory );
+ *
+ *	Unregister your mode when it is no longer available like so (this will prompt the destruction of any existing modes of this type):
+ *		FEditorModeRegistry::Get().UnregisterMode( FName( TEXT("MyEditorMode") ) );
+ */
+
 // Required forward declarations
 class FEdMode;
 
