Allow editing of thumbnail info through details panel/property matrix

diff --git a/Engine/Source/Editor/UnrealEd/Classes/ThumbnailRendering/SceneThumbnailInfo.h b/Engine/Source/Editor/UnrealEd/Classes/ThumbnailRendering/SceneThumbnailInfo.h
--- a/Engine/Source/Editor/UnrealEd/Classes/ThumbnailRendering/SceneThumbnailInfo.h
+++ b/Engine/Source/Editor/UnrealEd/Classes/ThumbnailRendering/SceneThumbnailInfo.h
@@ -14,14 +14,14 @@ class USceneThumbnailInfo : public UThumbnailInfo
 	GENERATED_UCLASS_BODY()
 
 	/** The pitch of the orbit camera around the asset */
-	UPROPERTY()
+	UPROPERTY(EditAnywhere, Category=Thumbnail)
 	float OrbitPitch;
 
 	/** The yaw of the orbit camera around the asset */
-	UPROPERTY()
+	UPROPERTY(EditAnywhere, Category=Thumbnail)
 	float OrbitYaw;
 
 	/** The offset from the bounds sphere distance from the asset */
-	UPROPERTY()
+	UPROPERTY(EditAnywhere, Category=Thumbnail)
 	float OrbitZoom;
 };
diff --git a/Engine/Source/Editor/UnrealEd/Classes/ThumbnailRendering/SceneThumbnailInfoWithPrimitive.h b/Engine/Source/Editor/UnrealEd/Classes/ThumbnailRendering/SceneThumbnailInfoWithPrimitive.h
--- a/Engine/Source/Editor/UnrealEd/Classes/ThumbnailRendering/SceneThumbnailInfoWithPrimitive.h
+++ b/Engine/Source/Editor/UnrealEd/Classes/ThumbnailRendering/SceneThumbnailInfoWithPrimitive.h
@@ -14,10 +14,10 @@ class USceneThumbnailInfoWithPrimitive : public USceneThumbnailInfo
 	GENERATED_UCLASS_BODY()
 
 	/** The type of primitive used in this thumbnail */
-	UPROPERTY()
+	UPROPERTY(EditAnywhere, Category=Thumbnail)
 	TEnumAsByte<EThumbnailPrimType> PrimitiveType;
 
 	/** The custom mesh used when the primitive type is TPT_None */
-	UPROPERTY()
+	UPROPERTY(EditAnywhere, Category=Thumbnail)
 	FStringAssetReference PreviewMesh;
 };
diff --git a/Engine/Source/Runtime/Engine/Classes/Engine/Blueprint.h b/Engine/Source/Runtime/Engine/Classes/Engine/Blueprint.h
--- a/Engine/Source/Runtime/Engine/Classes/Engine/Blueprint.h
+++ b/Engine/Source/Runtime/Engine/Classes/Engine/Blueprint.h
@@ -393,7 +393,7 @@ protected:
 	TWeakObjectPtr< class UWorld > CurrentWorldBeingDebugged;
 public:
 	/** Information for thumbnail rendering */
-	UPROPERTY()
+	UPROPERTY(VisibleAnywhere, EditInline, Category=Thumbnail)
 	class UThumbnailInfo* ThumbnailInfo;
 
 	/** The blueprint is currently compiled */
diff --git a/Engine/Source/Runtime/Engine/Classes/Engine/SkeletalMesh.h b/Engine/Source/Runtime/Engine/Classes/Engine/SkeletalMesh.h
--- a/Engine/Source/Runtime/Engine/Classes/Engine/SkeletalMesh.h
+++ b/Engine/Source/Runtime/Engine/Classes/Engine/SkeletalMesh.h
@@ -545,7 +545,7 @@ public:
 	FString SourceFileTimestamp_DEPRECATED;
 
 	/** Information for thumbnail rendering */
-	UPROPERTY()
+	UPROPERTY(VisibleAnywhere, EditInline, Category=Thumbnail)
 	class UThumbnailInfo* ThumbnailInfo;
 
 	/** Optimization settings used to simplify LODs of this mesh. */
diff --git a/Engine/Source/Runtime/Engine/Classes/Engine/StaticMesh.h b/Engine/Source/Runtime/Engine/Classes/Engine/StaticMesh.h
--- a/Engine/Source/Runtime/Engine/Classes/Engine/StaticMesh.h
+++ b/Engine/Source/Runtime/Engine/Classes/Engine/StaticMesh.h
@@ -415,7 +415,7 @@ class UStaticMesh : public UObject, public IInterface_CollisionDataProvider, pub
 	FString SourceFileTimestamp_DEPRECATED;
 
 	/** Information for thumbnail rendering */
-	UPROPERTY()
+	UPROPERTY(VisibleAnywhere, EditInline, Category=Thumbnail)
 	class UThumbnailInfo* ThumbnailInfo;
 
 	/** The stored camera position to use as a default for the static mesh editor */
diff --git a/Engine/Source/Runtime/Engine/Classes/Materials/MaterialFunction.h b/Engine/Source/Runtime/Engine/Classes/Materials/MaterialFunction.h
--- a/Engine/Source/Runtime/Engine/Classes/Materials/MaterialFunction.h
+++ b/Engine/Source/Runtime/Engine/Classes/Materials/MaterialFunction.h
@@ -98,7 +98,7 @@ public:
 	uint32 CombinedOutputTypes;
 
 	/** Information for thumbnail rendering */
-	UPROPERTY()
+	UPROPERTY(VisibleAnywhere, EditInline, Category=Thumbnail)
 	class UThumbnailInfo* ThumbnailInfo;
 #endif
 };
diff --git a/Engine/Source/Runtime/Engine/Classes/Materials/MaterialInterface.h b/Engine/Source/Runtime/Engine/Classes/Materials/MaterialInterface.h
--- a/Engine/Source/Runtime/Engine/Classes/Materials/MaterialInterface.h
+++ b/Engine/Source/Runtime/Engine/Classes/Materials/MaterialInterface.h
@@ -118,7 +118,7 @@ public:
 	FStringAssetReference PreviewMesh;
 
 	/** Information for thumbnail rendering */
-	UPROPERTY()
+	UPROPERTY(VisibleAnywhere, EditInline, Category=Thumbnail)
 	class UThumbnailInfo* ThumbnailInfo;
 
 private:
