Fixes some potential issues with Find-in-Blueprints running during cooks and other commandlets which is not desired.

diff --git a/Engine/Source/Editor/UnrealEd/Private/Kismet2/BlueprintEditorUtils.cpp b/Engine/Source/Editor/UnrealEd/Private/Kismet2/BlueprintEditorUtils.cpp
--- a/Engine/Source/Editor/UnrealEd/Private/Kismet2/BlueprintEditorUtils.cpp
+++ b/Engine/Source/Editor/UnrealEd/Private/Kismet2/BlueprintEditorUtils.cpp
@@ -1084,7 +1084,8 @@ UClass* FBlueprintEditorUtils::RegenerateBlueprintClass(UBlueprint* Blueprint, U
 		}
 
 #if WITH_EDITOR
-		if(GEditor)
+		// Do not want to run this code without the editor present nor when running commandlets.
+		if(GEditor && GIsEditor && !IsRunningCommandlet())
 		{
 			FFindInBlueprintSearchManager::Get().AddOrUpdateBlueprintSearchMetadata(Blueprint);
 		}
diff --git a/Engine/Source/Editor/UnrealEd/Private/Kismet2/Kismet2.cpp b/Engine/Source/Editor/UnrealEd/Private/Kismet2/Kismet2.cpp
--- a/Engine/Source/Editor/UnrealEd/Private/Kismet2/Kismet2.cpp
+++ b/Engine/Source/Editor/UnrealEd/Private/Kismet2/Kismet2.cpp
@@ -309,7 +309,8 @@ void FKismetEditorUtilities::CompileBlueprint(UBlueprint* BlueprintObj, bool bIs
 	// weren't any changes to save before, there shouldn't be after
 	bool const bStartedWithUnsavedChanges = (BlueprintPackage != NULL) ? BlueprintPackage->IsDirty() : true;
 #if WITH_EDITOR
-	if(GEditor)
+	// Do not want to run this code without the editor present nor when running commandlets.
+	if(GEditor && GIsEditor && !IsRunningCommandlet())
 	{
 		// We do not want to regenerate a search Guid during loads, nothing has changed in the Blueprint and it is cached elsewhere
 		if(!bIsRegeneratingOnLoad)
