UMG - ESlateCheckBoxState is now an actual Enum.  This required updating a lot of code that was using auto conversion from bool to int.  Fixing a lot of issues with garbage collecting and circular references, old button/border widgets will be broken by this checkin.  No longer crashing on exit of PIE.

diff --git a/Engine/Plugins/Editor/SpeedTreeImporter/Source/SpeedTreeImporter/Private/SpeedTreeImportFactory.cpp b/Engine/Plugins/Editor/SpeedTreeImporter/Source/SpeedTreeImporter/Private/SpeedTreeImportFactory.cpp
--- a/Engine/Plugins/Editor/SpeedTreeImporter/Source/SpeedTreeImporter/Private/SpeedTreeImportFactory.cpp
+++ b/Engine/Plugins/Editor/SpeedTreeImporter/Source/SpeedTreeImporter/Private/SpeedTreeImportFactory.cpp
@@ -92,7 +92,7 @@ public:
 	TSharedRef<SWidget> CreateOption(TSharedPtr<SCheckBox>& CheckBox, const FText& NameText, bool bChecked)
 	{
 		return SAssignNew(CheckBox, SCheckBox)
-		.IsChecked(bChecked)
+		.IsChecked(bChecked ? ESlateCheckBoxState::Checked : ESlateCheckBoxState::Unchecked)
 		.OnCheckStateChanged(this, &SSpeedTreeImportOptions::OnOptionModified)
 		.Content()
 		[
diff --git a/Engine/Plugins/Messaging/MessagingDebugger/Source/MessagingDebugger/Private/Widgets/Breakpoints/SMessagingBreakpointsTableRow.h b/Engine/Plugins/Messaging/MessagingDebugger/Source/MessagingDebugger/Private/Widgets/Breakpoints/SMessagingBreakpointsTableRow.h
--- a/Engine/Plugins/Messaging/MessagingDebugger/Source/MessagingDebugger/Private/Widgets/Breakpoints/SMessagingBreakpointsTableRow.h
+++ b/Engine/Plugins/Messaging/MessagingDebugger/Source/MessagingDebugger/Private/Widgets/Breakpoints/SMessagingBreakpointsTableRow.h
@@ -56,7 +56,7 @@ public:
 				+ SHorizontalBox::Slot()
 					[
 						SNew(SCheckBox)
-							.IsChecked(true)
+							.IsChecked(ESlateCheckBoxState::Checked)
 					]
 
 				+ SHorizontalBox::Slot()
diff --git a/Engine/Source/Developer/AssetTools/Private/AssetTypeActions/AssetTypeActions_SkeletalMesh.cpp b/Engine/Source/Developer/AssetTools/Private/AssetTypeActions/AssetTypeActions_SkeletalMesh.cpp
--- a/Engine/Source/Developer/AssetTools/Private/AssetTypeActions/AssetTypeActions_SkeletalMesh.cpp
+++ b/Engine/Source/Developer/AssetTools/Private/AssetTypeActions/AssetTypeActions_SkeletalMesh.cpp
@@ -204,7 +204,7 @@ public:
 	 */
 	ESlateCheckBoxState::Type IsCheckboxChecked( int32 ButtonId ) const
 	{
-		return CheckBoxInfoMap.FindChecked(ButtonId).bUsed;
+		return CheckBoxInfoMap.FindChecked(ButtonId).bUsed ? ESlateCheckBoxState::Checked : ESlateCheckBoxState::Unchecked;
 	}
 
 	/**
diff --git a/Engine/Source/Developer/AutomationWindow/Private/SAutomationTestItem.cpp b/Engine/Source/Developer/AutomationWindow/Private/SAutomationTestItem.cpp
--- a/Engine/Source/Developer/AutomationWindow/Private/SAutomationTestItem.cpp
+++ b/Engine/Source/Developer/AutomationWindow/Private/SAutomationTestItem.cpp
@@ -367,7 +367,7 @@ FText SAutomationTestItem::GetTestToolTip( int32 ClusterIndex ) const
 
 ESlateCheckBoxState::Type SAutomationTestItem::IsTestEnabled() const
 {
-	return TestStatus->IsEnabled();
+	return TestStatus->IsEnabled() ? ESlateCheckBoxState::Checked : ESlateCheckBoxState::Unchecked;
 }
 
 
diff --git a/Engine/Source/Editor/FoliageEdit/Private/SFoliageEditMeshDisplayItem.cpp b/Engine/Source/Editor/FoliageEdit/Private/SFoliageEditMeshDisplayItem.cpp
--- a/Engine/Source/Editor/FoliageEdit/Private/SFoliageEditMeshDisplayItem.cpp
+++ b/Engine/Source/Editor/FoliageEdit/Private/SFoliageEditMeshDisplayItem.cpp
@@ -2119,7 +2119,7 @@ float SFoliageEditMeshDisplayItem::GetScaleXMax() const
 
 ESlateCheckBoxState::Type SFoliageEditMeshDisplayItem::IsScaleXLockedChecked() const
 {
-	return FoliageSettingsPtr->LockScaleX;
+	return FoliageSettingsPtr->LockScaleX ? ESlateCheckBoxState::Checked : ESlateCheckBoxState::Unchecked;
 }
 
 void SFoliageEditMeshDisplayItem::OnScaleXLocked(ESlateCheckBoxState::Type InState)
@@ -2134,7 +2134,7 @@ void SFoliageEditMeshDisplayItem::OnScaleYReapply(ESlateCheckBoxState::Type InSt
 
 ESlateCheckBoxState::Type SFoliageEditMeshDisplayItem::IsScaleYReapplyChecked() const
 {
-	return FoliageSettingsPtr->ReapplyScaleY? ESlateCheckBoxState::Checked : ESlateCheckBoxState::Unchecked;
+	return FoliageSettingsPtr->ReapplyScaleY ? ESlateCheckBoxState::Checked : ESlateCheckBoxState::Unchecked;
 }
 
 void SFoliageEditMeshDisplayItem::OnScaleYMinChanged(float InValue)
@@ -2163,7 +2163,7 @@ float SFoliageEditMeshDisplayItem::GetScaleYMax() const
 
 ESlateCheckBoxState::Type SFoliageEditMeshDisplayItem::IsScaleYLockedChecked() const
 {
-	return FoliageSettingsPtr->LockScaleY;
+	return FoliageSettingsPtr->LockScaleY ? ESlateCheckBoxState::Checked : ESlateCheckBoxState::Unchecked;
 }
 
 void SFoliageEditMeshDisplayItem::OnScaleYLocked(ESlateCheckBoxState::Type InState)
@@ -2207,7 +2207,7 @@ float SFoliageEditMeshDisplayItem::GetScaleZMax() const
 
 ESlateCheckBoxState::Type SFoliageEditMeshDisplayItem::IsScaleZLockedChecked() const
 {
-	return FoliageSettingsPtr->LockScaleZ;
+	return FoliageSettingsPtr->LockScaleZ ? ESlateCheckBoxState::Checked : ESlateCheckBoxState::Unchecked;
 }
 
 void SFoliageEditMeshDisplayItem::OnScaleZLocked(ESlateCheckBoxState::Type InState)
diff --git a/Engine/Source/Editor/GameProjectGeneration/Private/SProjectBrowser.cpp b/Engine/Source/Editor/GameProjectGeneration/Private/SProjectBrowser.cpp
--- a/Engine/Source/Editor/GameProjectGeneration/Private/SProjectBrowser.cpp
+++ b/Engine/Source/Editor/GameProjectGeneration/Private/SProjectBrowser.cpp
@@ -212,7 +212,7 @@ void SProjectBrowser::Construct( const FArguments& InArgs )
 			.Padding(4, 0)
 			[
 				SNew(SCheckBox)			
-				.IsChecked(GEditor->GetGameAgnosticSettings().bLoadTheMostRecentlyLoadedProjectAtStartup)
+				.IsChecked(GEditor->GetGameAgnosticSettings().bLoadTheMostRecentlyLoadedProjectAtStartup ? ESlateCheckBoxState::Checked : ESlateCheckBoxState::Unchecked)
 				.OnCheckStateChanged(this, &SProjectBrowser::OnAutoloadLastProjectChanged)
 				.Content()
 				[
diff --git a/Engine/Source/Editor/Kismet/Private/BlueprintDetailsCustomization.cpp b/Engine/Source/Editor/Kismet/Private/BlueprintDetailsCustomization.cpp
--- a/Engine/Source/Editor/Kismet/Private/BlueprintDetailsCustomization.cpp
+++ b/Engine/Source/Editor/Kismet/Private/BlueprintDetailsCustomization.cpp
@@ -613,7 +613,7 @@ TSharedRef<ITableRow> FBlueprintVarActionDetails::OnGenerateWidgetForPropertyLis
 				.AutoWidth()
 			[
 				SNew(SCheckBox)
-					.IsChecked(true)
+					.IsChecked(true ? ESlateCheckBoxState::Checked : ESlateCheckBoxState::Unchecked)
 					.IsEnabled(false)
 			]
 		];
@@ -3325,13 +3325,13 @@ bool FBaseBlueprintGraphActionDetails::IsPinNameUnique(const FString& TestName)
 	UK2Node_EditablePinBase * FunctionResultNode = FunctionResultNodePtr.Get();
 
 	// Check the built in pins
- 	for( TArray<UEdGraphPin*>::TIterator it(FunctionEntryNode->Pins); it; ++it )
- 	{
+	for( TArray<UEdGraphPin*>::TIterator it(FunctionEntryNode->Pins); it; ++it )
+	{
 		if( (*it)->PinName == TestName )
- 		{
- 			return false;
- 		}
- 	}
+		{
+			return false;
+		}
+	}
 
 	// If there is a result node, that isn't the same as the entry node, check it as well
 	if( FunctionResultNode && (FunctionResultNode != FunctionEntryNode) )
diff --git a/Engine/Source/Editor/KismetWidgets/Private/SPinTypeSelector.cpp b/Engine/Source/Editor/KismetWidgets/Private/SPinTypeSelector.cpp
--- a/Engine/Source/Editor/KismetWidgets/Private/SPinTypeSelector.cpp
+++ b/Engine/Source/Editor/KismetWidgets/Private/SPinTypeSelector.cpp
@@ -109,14 +109,14 @@ void SPinTypeSelector::Construct(const FArguments& InArgs, FGetPinTypeTree GetPi
 
 FString SPinTypeSelector::GetTypeDescription() const
 {
- 	if( TargetPinType.Get().PinSubCategoryObject != NULL )
- 	{
- 		return TargetPinType.Get().PinSubCategoryObject->GetName();
- 	}
- 	else
- 	{
- 		return TargetPinType.Get().PinCategory;
- 	}
+	if( TargetPinType.Get().PinSubCategoryObject != NULL )
+	{
+		return TargetPinType.Get().PinSubCategoryObject->GetName();
+	}
+	else
+	{
+		return TargetPinType.Get().PinCategory;
+	}
 }
 
 const FSlateBrush* SPinTypeSelector::GetTypeIconImage() const
@@ -131,7 +131,7 @@ FSlateColor SPinTypeSelector::GetTypeIconColor() const
 
 ESlateCheckBoxState::Type SPinTypeSelector::IsArrayChecked() const
 {
-	return TargetPinType.Get().bIsArray;
+	return TargetPinType.Get().bIsArray ? ESlateCheckBoxState::Checked : ESlateCheckBoxState::Unchecked;
 }
 
 void SPinTypeSelector::OnArrayCheckStateChanged(ESlateCheckBoxState::Type NewState)
@@ -190,8 +190,8 @@ void SPinTypeSelector::OnTypeSelectionChanged(FPinTypeTreeItem Selection, ESelec
 	}
 
 	// Only handle selection for non-read only items, since STreeViewItem doesn't actually support read-only
- 	if( Selection.IsValid() )
- 	{
+	if( Selection.IsValid() )
+	{
 		if( !Selection->bReadOnly )
 		{
 			const FScopedTransaction Transaction( LOCTEXT("ChangeParam", "Change Paramater Type") );
@@ -230,7 +230,7 @@ void SPinTypeSelector::OnTypeSelectionChanged(FPinTypeTreeItem Selection, ESelec
 				}
 			}
 		}
- 	}
+	}
 }
 
 void SPinTypeSelector::GetTypeChildren(FPinTypeTreeItem InItem, TArray<FPinTypeTreeItem>& OutChildren)
diff --git a/Engine/Source/Editor/LevelEditor/Private/SSurfaceProperties.cpp b/Engine/Source/Editor/LevelEditor/Private/SSurfaceProperties.cpp
--- a/Engine/Source/Editor/LevelEditor/Private/SSurfaceProperties.cpp
+++ b/Engine/Source/Editor/LevelEditor/Private/SSurfaceProperties.cpp
@@ -941,10 +941,10 @@ ESlateCheckBoxState::Type SSurfaceProperties::IsUsingNegativePanning( TextureCoo
 {
 	if (Channel == UChannel)
 	{
-		return bUseNegativePanningU;
+		return bUseNegativePanningU ? ESlateCheckBoxState::Checked : ESlateCheckBoxState::Unchecked;
 	}
 
-	return bUseNegativePanningV;
+	return bUseNegativePanningV ? ESlateCheckBoxState::Checked : ESlateCheckBoxState::Unchecked;
 }
 
 void SSurfaceProperties::OnTogglePanningDirection( ESlateCheckBoxState::Type NewState, TextureCoordChannel Channel )
@@ -960,7 +960,7 @@ const FSlateBrush* SSurfaceProperties::GetToggleRotationDirectionImage() const
 
 ESlateCheckBoxState::Type SSurfaceProperties::IsUsingNegativeRotation() const
 {
-	return bUseNegativeRotation;
+	return bUseNegativeRotation ? ESlateCheckBoxState::Checked : ESlateCheckBoxState::Unchecked;
 }
 
 void SSurfaceProperties::OnToggleRotationDirection( ESlateCheckBoxState::Type NewState )
diff --git a/Engine/Source/Editor/MeshPaint/Private/SMeshPaint.cpp b/Engine/Source/Editor/MeshPaint/Private/SMeshPaint.cpp
--- a/Engine/Source/Editor/MeshPaint/Private/SMeshPaint.cpp
+++ b/Engine/Source/Editor/MeshPaint/Private/SMeshPaint.cpp
@@ -2264,12 +2264,13 @@ void SMeshPaint::OnPaintWeightChanged(int32 InWeightIndex)
 
 ESlateCheckBoxState::Type SMeshPaint::IsWriteColorChannelChecked(EMeshPaintWriteColorChannels::Type CheckBoxInfo) const
 {
-	bool bIsRedAndChecked = (CheckBoxInfo == EMeshPaintWriteColorChannels::Red && FMeshPaintSettings::Get().bWriteRed);
-	bool bIsGreenAndChecked = (CheckBoxInfo == EMeshPaintWriteColorChannels::Green && FMeshPaintSettings::Get().bWriteGreen);
-	bool bIsBlueAndChecked = (CheckBoxInfo == EMeshPaintWriteColorChannels::Blue && FMeshPaintSettings::Get().bWriteBlue);
-	bool bIsAlphaAndChecked = (CheckBoxInfo == EMeshPaintWriteColorChannels::Alpha && FMeshPaintSettings::Get().bWriteAlpha);
+	const bool bIsRedAndChecked = ( CheckBoxInfo == EMeshPaintWriteColorChannels::Red && FMeshPaintSettings::Get().bWriteRed );
+	const bool bIsGreenAndChecked = ( CheckBoxInfo == EMeshPaintWriteColorChannels::Green && FMeshPaintSettings::Get().bWriteGreen );
+	const bool bIsBlueAndChecked = ( CheckBoxInfo == EMeshPaintWriteColorChannels::Blue && FMeshPaintSettings::Get().bWriteBlue );
+	const bool bIsAlphaAndChecked = (CheckBoxInfo == EMeshPaintWriteColorChannels::Alpha && FMeshPaintSettings::Get().bWriteAlpha);
 
-	return (bIsRedAndChecked || bIsGreenAndChecked || bIsBlueAndChecked || bIsAlphaAndChecked);
+	const bool bIsChecked = ( bIsRedAndChecked || bIsGreenAndChecked || bIsBlueAndChecked || bIsAlphaAndChecked );
+	return bIsChecked ? ESlateCheckBoxState::Checked : ESlateCheckBoxState::Unchecked;
 }
 
 void SMeshPaint::OnWriteColorChannelChanged(ESlateCheckBoxState::Type InNewValue, EMeshPaintWriteColorChannels::Type CheckBoxInfo)
diff --git a/Engine/Source/Editor/PListEditor/Private/PListNodeBoolean.cpp b/Engine/Source/Editor/PListEditor/Private/PListNodeBoolean.cpp
--- a/Engine/Source/Editor/PListEditor/Private/PListNodeBoolean.cpp
+++ b/Engine/Source/Editor/PListEditor/Private/PListNodeBoolean.cpp
@@ -113,7 +113,7 @@ TSharedRef<SWidget> FPListNodeBoolean::GenerateWidgetForColumn(const FName& Colu
 			.FillWidth(1.0f)
 			[
 				SAssignNew(ValueCheckBox, SCheckBox)
-				.IsChecked(bValue)
+				.IsChecked(bValue ? ESlateCheckBoxState::Checked : ESlateCheckBoxState::Unchecked)
 				.OnCheckStateChanged(this, &FPListNodeBoolean::OnValueChanged)
 			]
 		];
diff --git a/Engine/Source/Editor/Persona/Private/SAnimationBlendSpaceBase.cpp b/Engine/Source/Editor/Persona/Private/SAnimationBlendSpaceBase.cpp
--- a/Engine/Source/Editor/Persona/Private/SAnimationBlendSpaceBase.cpp
+++ b/Engine/Source/Editor/Persona/Private/SAnimationBlendSpaceBase.cpp
@@ -1236,7 +1236,7 @@ TSharedRef<SWidget> SBlendSpaceEditorBase::MakeDisplayOptionsBox() const
 	DisplayOptionsPanel->AddCheckBoxSlot()
 	[
 		SNew(SCheckBox) 
-		.IsChecked( true )
+		.IsChecked( true ? ESlateCheckBoxState::Checked : ESlateCheckBoxState::Unchecked )
 		.OnCheckStateChanged( this, &SBlendSpaceEditorBase::ShowToolTip_OnIsCheckedChanged)
 		[
 			SNew(STextBlock)
diff --git a/Engine/Source/Editor/PhAT/Private/SPhATNewAssetDlg.cpp b/Engine/Source/Editor/PhAT/Private/SPhATNewAssetDlg.cpp
--- a/Engine/Source/Editor/PhAT/Private/SPhATNewAssetDlg.cpp
+++ b/Engine/Source/Editor/PhAT/Private/SPhATNewAssetDlg.cpp
@@ -76,7 +76,7 @@ void SPhATNewAssetDlg::Construct(const FArguments& InArgs)
 				+SUniformGridPanel::Slot(1, 1)
 				[
 					SNew(SCheckBox)
-					.IsChecked(NewBodyData->bAlignDownBone)
+					.IsChecked(NewBodyData->bAlignDownBone ? ESlateCheckBoxState::Checked : ESlateCheckBoxState::Unchecked)
 					.OnCheckStateChanged(this, &SPhATNewAssetDlg::OnToggleOrientAlongBone)
 				]
 				+SUniformGridPanel::Slot(0, 2)
@@ -116,7 +116,7 @@ void SPhATNewAssetDlg::Construct(const FArguments& InArgs)
 				+SUniformGridPanel::Slot(1, 4)
 				[
 					SNew(SCheckBox)
-					.IsChecked(NewBodyData->bCreateJoints)
+					.IsChecked(NewBodyData->bCreateJoints ? ESlateCheckBoxState::Checked : ESlateCheckBoxState::Unchecked)
 					.OnCheckStateChanged(this, &SPhATNewAssetDlg::OnToggleCreateJoints)
 				]
 				+SUniformGridPanel::Slot(0, 5)
@@ -142,7 +142,7 @@ void SPhATNewAssetDlg::Construct(const FArguments& InArgs)
 				+SUniformGridPanel::Slot(1, 6)
 				[
 					SNew(SCheckBox)
-					.IsChecked(NewBodyData->bWalkPastSmall)
+					.IsChecked(NewBodyData->bWalkPastSmall ? ESlateCheckBoxState::Checked : ESlateCheckBoxState::Unchecked)
 					.OnCheckStateChanged(this, &SPhATNewAssetDlg::OnToggleWalkPastSmallBones)
 				]
 				+SUniformGridPanel::Slot(0, 7)
@@ -154,7 +154,7 @@ void SPhATNewAssetDlg::Construct(const FArguments& InArgs)
 				+SUniformGridPanel::Slot(1, 7)
 				[
 					SNew(SCheckBox)
-					.IsChecked(NewBodyData->bBodyForAll)
+					.IsChecked(NewBodyData->bBodyForAll ? ESlateCheckBoxState::Checked : ESlateCheckBoxState::Unchecked)
 					.OnCheckStateChanged(this, &SPhATNewAssetDlg::OnToggleCreateBodyForAllBones)
 				]
 				// Add in the UI options for the max hulls and the max verts on the hulls
diff --git a/Engine/Source/Editor/SourceControlWindows/Private/SSourceControlSubmit.cpp b/Engine/Source/Editor/SourceControlWindows/Private/SSourceControlSubmit.cpp
--- a/Engine/Source/Editor/SourceControlWindows/Private/SSourceControlSubmit.cpp
+++ b/Engine/Source/Editor/SourceControlWindows/Private/SSourceControlSubmit.cpp
@@ -84,7 +84,7 @@ private:
 	/** The check status of the item. */
 	ESlateCheckBoxState::Type IsChecked() const
 	{
-		return SubmitItemData->bIsChecked;
+		return SubmitItemData->bIsChecked ? ESlateCheckBoxState::Checked : ESlateCheckBoxState::Unchecked;
 	}
 
 	/** Changes the check status of the item .*/
diff --git a/Engine/Source/Editor/TaskBrowser/Private/STaskSettings.cpp b/Engine/Source/Editor/TaskBrowser/Private/STaskSettings.cpp
--- a/Engine/Source/Editor/TaskBrowser/Private/STaskSettings.cpp
+++ b/Engine/Source/Editor/TaskBrowser/Private/STaskSettings.cpp
@@ -168,7 +168,7 @@ void STaskSettings::Construct(const FArguments& InArgs)
 							.AutoWidth()
 							[
 								SAssignNew(Autoconnect, SCheckBox)
-								.IsChecked( TBSettings.bAutoConnectAtStartup )
+								.IsChecked( TBSettings.bAutoConnectAtStartup ? ESlateCheckBoxState::Checked : ESlateCheckBoxState::Unchecked )
 							]
 							+SHorizontalBox::Slot()
 							.AutoWidth()
@@ -183,7 +183,7 @@ void STaskSettings::Construct(const FArguments& InArgs)
 							.AutoWidth()
 							[
 								SAssignNew(SingleSignOn, SCheckBox)
-								.IsChecked( TBSettings.bUseSingleSignOn )
+								.IsChecked( TBSettings.bUseSingleSignOn ? ESlateCheckBoxState::Checked : ESlateCheckBoxState::Unchecked )
 								.OnCheckStateChanged( this, &STaskSettings::OnSingleSignOnChanged )
 							]
 							+SHorizontalBox::Slot()
diff --git a/Engine/Source/Editor/UMGEditor/Private/SUMGEditorTreeItem.cpp b/Engine/Source/Editor/UMGEditor/Private/SUMGEditorTreeItem.cpp
--- a/Engine/Source/Editor/UMGEditor/Private/SUMGEditorTreeItem.cpp
+++ b/Engine/Source/Editor/UMGEditor/Private/SUMGEditorTreeItem.cpp
@@ -31,6 +31,7 @@ void SUMGEditorTreeItem::Construct(const FArguments& InArgs, const TSharedRef< S
 		[
 			SNew(SHorizontalBox)
 
+			// Widget icon
 			+ SHorizontalBox::Slot()
 			.AutoWidth()
 			[
@@ -39,6 +40,7 @@ void SUMGEditorTreeItem::Construct(const FArguments& InArgs, const TSharedRef< S
 				.Image(Item->GetEditorIcon())
 			]
 
+			// Name of the widget
 			+ SHorizontalBox::Slot()
 			.FillWidth(1.0f)
 			.Padding(2, 0, 0, 0)
diff --git a/Engine/Source/Editor/UMGEditor/Private/WidgetBlueprintEditorUtils.h b/Engine/Source/Editor/UMGEditor/Private/WidgetBlueprintEditorUtils.h
--- a/Engine/Source/Editor/UMGEditor/Private/WidgetBlueprintEditorUtils.h
+++ b/Engine/Source/Editor/UMGEditor/Private/WidgetBlueprintEditorUtils.h
@@ -8,5 +8,5 @@
 class FWidgetBlueprintEditorUtils
 {
 public:
-	static bool RenameWidget(TSharedRef<FWidgetBlueprintEditor> BlueprintEditor, const FName& OldName, const FName& NewName);
+	static bool RenameWidget(TSharedRef<class FWidgetBlueprintEditor> BlueprintEditor, const FName& OldName, const FName& NewName);
 };
diff --git a/Engine/Source/Editor/UnrealEd/Private/Fbx/FbxOptionWindow.cpp b/Engine/Source/Editor/UnrealEd/Private/Fbx/FbxOptionWindow.cpp
--- a/Engine/Source/Editor/UnrealEd/Private/Fbx/FbxOptionWindow.cpp
+++ b/Engine/Source/Editor/UnrealEd/Private/Fbx/FbxOptionWindow.cpp
@@ -536,7 +536,7 @@ TSharedRef<SWidget> SFbxOptionWindow::ConstructMiscOption()
 	NewBox->AddSlot().AutoHeight() .Padding(2)
 	[
 		SNew(SCheckBox)
-		.IsChecked(ImportUI->bOverrideFullName? true: false)
+		.IsChecked(ImportUI->bOverrideFullName ? ESlateCheckBoxState::Checked : ESlateCheckBoxState::Unchecked)
 		.OnCheckStateChanged(this, &SFbxOptionWindow::SetGeneral_OverrideFullName)
 		.IsEnabled( !bIsObjFormat )
 		[
@@ -548,7 +548,7 @@ TSharedRef<SWidget> SFbxOptionWindow::ConstructMiscOption()
 	NewBox->AddSlot().AutoHeight() .Padding(2)
 	[
 		SNew(SCheckBox)
-		.IsChecked(ImportUI->bConvertScene? true: false)
+		.IsChecked(ImportUI->bConvertScene ? ESlateCheckBoxState::Checked : ESlateCheckBoxState::Unchecked)
 		.OnCheckStateChanged(this, &SFbxOptionWindow::SetGeneral_ConvertScene)
 		.IsEnabled( !bIsObjFormat )
 		[
@@ -663,7 +663,7 @@ TSharedRef<SWidget> SFbxOptionWindow::ConstructStaticMeshAdvanced()
 	NewBox->AddSlot().AutoHeight().Padding(2)
 	[
 		SNew(SCheckBox)
-		.IsChecked(ImportUI->bCombineMeshes ? true : false)
+		.IsChecked(ImportUI->bCombineMeshes ? ESlateCheckBoxState::Checked : ESlateCheckBoxState::Unchecked)
 		.OnCheckStateChanged(this, &SFbxOptionWindow::SetStaticMesh_CombineMeshes)
 		.Content()
 		[
@@ -676,7 +676,7 @@ TSharedRef<SWidget> SFbxOptionWindow::ConstructStaticMeshAdvanced()
 	NewBox->AddSlot().AutoHeight() .Padding(2)
 	[
 		SNew(SCheckBox)
-		.IsChecked(ImportUI->StaticMeshImportData->bImportMeshLODs? true: false)
+		.IsChecked(ImportUI->StaticMeshImportData->bImportMeshLODs ? ESlateCheckBoxState::Checked : ESlateCheckBoxState::Unchecked)
 		.OnCheckStateChanged(this, &SFbxOptionWindow::SetStaticMesh_ImportMeshLODs)
 		.IsEnabled( !bIsObjFormat )
 		[
@@ -688,7 +688,7 @@ TSharedRef<SWidget> SFbxOptionWindow::ConstructStaticMeshAdvanced()
 	NewBox->AddSlot().AutoHeight() .Padding(2)
 	[
 		SNew(SCheckBox)
-		.IsChecked(ImportUI->StaticMeshImportData->bReplaceVertexColors? true: false)
+		.IsChecked(ImportUI->StaticMeshImportData->bReplaceVertexColors ? ESlateCheckBoxState::Checked : ESlateCheckBoxState::Unchecked)
 		.OnCheckStateChanged(this, &SFbxOptionWindow::SetStaticMesh_ReplaceVertexColor)
 		.IsEnabled( !bIsObjFormat )
 		[
@@ -700,7 +700,7 @@ TSharedRef<SWidget> SFbxOptionWindow::ConstructStaticMeshAdvanced()
 	NewBox->AddSlot().AutoHeight() .Padding(2)
 	[
 		SNew(SCheckBox)
-		.IsChecked(ImportUI->StaticMeshImportData->bRemoveDegenerates? true: false)
+		.IsChecked(ImportUI->StaticMeshImportData->bRemoveDegenerates ? ESlateCheckBoxState::Checked : ESlateCheckBoxState::Unchecked)
 		.OnCheckStateChanged(this, &SFbxOptionWindow::SetStaticMesh_RemoveDegenerates)
 		[
 			SNew(STextBlock)
@@ -711,7 +711,7 @@ TSharedRef<SWidget> SFbxOptionWindow::ConstructStaticMeshAdvanced()
 	NewBox->AddSlot().AutoHeight() .Padding(2)
 	[
 		SNew(SCheckBox)
-		.IsChecked(ImportUI->StaticMeshImportData->bOneConvexHullPerUCX? true: false)
+		.IsChecked(ImportUI->StaticMeshImportData->bOneConvexHullPerUCX ? ESlateCheckBoxState::Checked : ESlateCheckBoxState::Unchecked)
 		.OnCheckStateChanged(this, &SFbxOptionWindow::SetStaticMesh_OneConvexHullPerUCX)
 		.IsEnabled( !bIsObjFormat )
 		[
@@ -738,7 +738,7 @@ TSharedRef<SWidget> SFbxOptionWindow::ConstructSkeletalMeshBasic()
 	NewBox->AddSlot().AutoHeight() .Padding(2)
 	[
 		SNew(SCheckBox)
-		.IsChecked(ImportUI->SkeletalMeshImportData->bImportMorphTargets? true: false)
+		.IsChecked(ImportUI->SkeletalMeshImportData->bImportMorphTargets ? ESlateCheckBoxState::Checked : ESlateCheckBoxState::Unchecked)
 		.OnCheckStateChanged(this, &SFbxOptionWindow::SetSkeletalMesh_ImportMorphTargets)
 		[
 			SNew(STextBlock)
@@ -754,7 +754,7 @@ TSharedRef<SWidget> SFbxOptionWindow::ConstructSkeletalMeshBasic()
 		.AutoWidth()
 		[
 			SNew(SCheckBox)
-			.IsChecked(ImportUI->bImportAnimations? true: false)
+			.IsChecked(ImportUI->bImportAnimations ? ESlateCheckBoxState::Checked : ESlateCheckBoxState::Unchecked)
 			.OnCheckStateChanged(this, &SFbxOptionWindow::SetSkeletalMesh_ImportAnimation)
 			[
 				SNew(STextBlock)
@@ -778,7 +778,7 @@ TSharedRef<SWidget> SFbxOptionWindow::ConstructSkeletalMeshBasic()
 	NewBox->AddSlot().AutoHeight() .Padding(2)
 	[
 		SNew(SCheckBox)
-		.IsChecked(ImportUI->SkeletalMeshImportData->bUpdateSkeletonReferencePose? true: false)
+		.IsChecked(ImportUI->SkeletalMeshImportData->bUpdateSkeletonReferencePose ? ESlateCheckBoxState::Checked : ESlateCheckBoxState::Unchecked)
 		.OnCheckStateChanged(this, &SFbxOptionWindow::SetSkeletalMesh_UpdateSkeletonRefPose)
 		[
 			SNew(STextBlock)
@@ -804,7 +804,7 @@ TSharedRef<SWidget> SFbxOptionWindow::ConstructSkeletalMeshAdvanced()
 	NewBox->AddSlot().AutoHeight() .Padding(2)
 	[
 		SNew(SCheckBox)
-		.IsChecked(ImportUI->SkeletalMeshImportData->bImportMeshLODs? true: false)
+		.IsChecked(ImportUI->SkeletalMeshImportData->bImportMeshLODs ? ESlateCheckBoxState::Checked : ESlateCheckBoxState::Unchecked)
 		.OnCheckStateChanged(this, &SFbxOptionWindow::SetSkeletalMesh_ImportMeshLODs)
 		[
 			SNew(STextBlock)
@@ -815,7 +815,7 @@ TSharedRef<SWidget> SFbxOptionWindow::ConstructSkeletalMeshAdvanced()
 	NewBox->AddSlot().AutoHeight() .Padding(2)
 	[
 		SNew(SCheckBox)
-		.IsChecked(ImportUI->bImportRigidMesh? true: false)
+		.IsChecked(ImportUI->bImportRigidMesh ? ESlateCheckBoxState::Checked : ESlateCheckBoxState::Unchecked)
 		.OnCheckStateChanged(this, &SFbxOptionWindow::SetSkeletalMesh_ImportRigidMesh)
 		[
 			SNew(STextBlock)
@@ -826,7 +826,7 @@ TSharedRef<SWidget> SFbxOptionWindow::ConstructSkeletalMeshAdvanced()
 	NewBox->AddSlot().AutoHeight() .Padding(2)
 	[
 		SNew(SCheckBox)
-		.IsChecked(ImportUI->SkeletalMeshImportData->bUseT0AsRefPose? true: false)
+		.IsChecked(ImportUI->SkeletalMeshImportData->bUseT0AsRefPose ? ESlateCheckBoxState::Checked : ESlateCheckBoxState::Unchecked)
 		.OnCheckStateChanged(this, &SFbxOptionWindow::SetSkeletalMesh_UseT0AsRefPose)
 		[
 			SNew(STextBlock)
@@ -837,7 +837,7 @@ TSharedRef<SWidget> SFbxOptionWindow::ConstructSkeletalMeshAdvanced()
 	NewBox->AddSlot().AutoHeight() .Padding(2)
 	[
 		SNew(SCheckBox)
-		.IsChecked(ImportUI->SkeletalMeshImportData->bPreserveSmoothingGroups? true: false)
+		.IsChecked(ImportUI->SkeletalMeshImportData->bPreserveSmoothingGroups ? ESlateCheckBoxState::Checked : ESlateCheckBoxState::Unchecked)
 		.OnCheckStateChanged(this, &SFbxOptionWindow::SetSkeletalMesh_ReserveSmoothingGroups)
 		[
 			SNew(STextBlock)
@@ -849,7 +849,7 @@ TSharedRef<SWidget> SFbxOptionWindow::ConstructSkeletalMeshAdvanced()
 	NewBox->AddSlot().AutoHeight() .Padding(2)
 	[
 		SNew(SCheckBox)
-		.IsChecked(ImportUI->SkeletalMeshImportData->bKeepOverlappingVertices? true: false)
+		.IsChecked(ImportUI->SkeletalMeshImportData->bKeepOverlappingVertices ? ESlateCheckBoxState::Checked : ESlateCheckBoxState::Unchecked)
 		.OnCheckStateChanged(this, &SFbxOptionWindow::SetSkeletalMesh_KeepOverlappingVertices)
 		[
 			SNew(STextBlock)
@@ -859,7 +859,7 @@ TSharedRef<SWidget> SFbxOptionWindow::ConstructSkeletalMeshAdvanced()
 	NewBox->AddSlot().AutoHeight() .Padding(2)
 	[
 		SNew(SCheckBox)
-		.IsChecked(ImportUI->SkeletalMeshImportData->bImportMeshesInBoneHierarchy? true: false)
+		.IsChecked(ImportUI->SkeletalMeshImportData->bImportMeshesInBoneHierarchy ? ESlateCheckBoxState::Checked : ESlateCheckBoxState::Unchecked)
 		.OnCheckStateChanged(this, &SFbxOptionWindow::SetSkeletalMesh_ImportMeshesInBoneHierarchy)
 		[
 			SNew(STextBlock)
@@ -877,7 +877,7 @@ TSharedRef<SWidget> SFbxOptionWindow::ConstructSkeletalMeshAdvanced()
 	NewBox->AddSlot().AutoHeight() .Padding(2)
 	[
 		SNew(SCheckBox)
-		.IsChecked(ImportUI->bCreatePhysicsAsset? true: false)
+		.IsChecked(ImportUI->bCreatePhysicsAsset ? ESlateCheckBoxState::Checked : ESlateCheckBoxState::Unchecked)
 		.OnCheckStateChanged(this, &SFbxOptionWindow::SetSkeletalMesh_CreatePhysicsAsset)
 		.IsEnabled(!bForceImportType)
 		[
@@ -932,7 +932,7 @@ TSharedRef<SWidget> SFbxOptionWindow::ConstructMaterialOption()
 	NewBox->AddSlot().AutoHeight() .Padding(2)
 	[
 		SNew(SCheckBox)
-		.IsChecked(ImportUI->bImportMaterials? true: false)
+		.IsChecked(ImportUI->bImportMaterials ? ESlateCheckBoxState::Checked : ESlateCheckBoxState::Unchecked)
 		.OnCheckStateChanged(this, &SFbxOptionWindow::SetMaterial_ImportMaterials)
 		.IsEnabled( !bIsObjFormat )
 		[
@@ -944,7 +944,7 @@ TSharedRef<SWidget> SFbxOptionWindow::ConstructMaterialOption()
 	NewBox->AddSlot().AutoHeight() .Padding(2)
 	[
 		SNew(SCheckBox)
-		.IsChecked(ImportUI->bImportTextures? true: false)
+		.IsChecked(ImportUI->bImportTextures ? ESlateCheckBoxState::Checked : ESlateCheckBoxState::Unchecked)
 		.OnCheckStateChanged(this, &SFbxOptionWindow::SetMaterial_ImportTextures)
 		.IsEnabled( !bIsObjFormat )
 		[
@@ -956,7 +956,7 @@ TSharedRef<SWidget> SFbxOptionWindow::ConstructMaterialOption()
 	NewBox->AddSlot().AutoHeight() .Padding(2)
 	[
 		SNew(SCheckBox)
-		.IsChecked(ImportUI->TextureImportData->bInvertNormalMaps? true: false)
+		.IsChecked(ImportUI->TextureImportData->bInvertNormalMaps ? ESlateCheckBoxState::Checked : ESlateCheckBoxState::Unchecked)
 		.OnCheckStateChanged(this, &SFbxOptionWindow::SetMaterial_InvertNormalMaps)
 		.IsEnabled( !bIsObjFormat )
 		[
@@ -975,7 +975,7 @@ TSharedRef<SWidget> SFbxOptionWindow::ConstructAnimationOption()
 	NewBox->AddSlot().AutoHeight() .Padding(2, 4)
 	[
 		SNew(SCheckBox)
-		.IsChecked(ImportUI->bUseDefaultSampleRate? true: false)
+		.IsChecked(ImportUI->bUseDefaultSampleRate ? ESlateCheckBoxState::Checked : ESlateCheckBoxState::Unchecked)
 		.OnCheckStateChanged(this, &SFbxOptionWindow::SetSkeletalMesh_UseDefaultSampleRate)
 		[
 			SNew(STextBlock)
@@ -986,7 +986,7 @@ TSharedRef<SWidget> SFbxOptionWindow::ConstructAnimationOption()
 	NewBox->AddSlot().AutoHeight() .Padding(2)
 	[
 		SNew(SCheckBox)
-		.IsChecked(ImportUI->bPreserveLocalTransform? true: false)
+		.IsChecked(ImportUI->bPreserveLocalTransform ? ESlateCheckBoxState::Checked : ESlateCheckBoxState::Unchecked)
 		.OnCheckStateChanged(this, &SFbxOptionWindow::SetAnimation_ReserveLocalTransform)
 		[
 			SNew(STextBlock)
diff --git a/Engine/Source/Runtime/Engine/Private/GameViewportClient.cpp b/Engine/Source/Runtime/Engine/Private/GameViewportClient.cpp
--- a/Engine/Source/Runtime/Engine/Private/GameViewportClient.cpp
+++ b/Engine/Source/Runtime/Engine/Private/GameViewportClient.cpp
@@ -155,6 +155,7 @@ void UGameViewportClient::BeginDestroy()
 void UGameViewportClient::DetachViewportClient()
 {
 	ViewportConsole = NULL;
+	RemoveAllViewportWidgets();
 	RemoveFromRoot();
 }
 
diff --git a/Engine/Source/Runtime/Slate/Private/Widgets/Input/SCheckBox.cpp b/Engine/Source/Runtime/Slate/Private/Widgets/Input/SCheckBox.cpp
--- a/Engine/Source/Runtime/Slate/Private/Widgets/Input/SCheckBox.cpp
+++ b/Engine/Source/Runtime/Slate/Private/Widgets/Input/SCheckBox.cpp
@@ -293,6 +293,10 @@ const FSlateBrush* SCheckBox::OnGetCheckImage() const
 }
 
 
+ESlateCheckBoxState::Type SCheckBox::GetCheckedState() const
+{
+	return IsCheckboxChecked.Get();
+}
 
 /**
  * Toggles the checked state for this check box, fire events as needed
@@ -326,6 +330,11 @@ void SCheckBox::ToggleCheckedState()
 	}
 }
 
+void SCheckBox::SetIsChecked(TAttribute<ESlateCheckBoxState::Type> InIsChecked)
+{
+	IsCheckboxChecked = InIsChecked;
+}
+
 void SCheckBox::PlayCheckedSound() const
 {
 	FSlateApplication::Get().PlaySound( CheckedSound );
diff --git a/Engine/Source/Runtime/Slate/Public/Widgets/Input/SCheckBox.h b/Engine/Source/Runtime/Slate/Public/Widgets/Input/SCheckBox.h
--- a/Engine/Source/Runtime/Slate/Public/Widgets/Input/SCheckBox.h
+++ b/Engine/Source/Runtime/Slate/Public/Widgets/Input/SCheckBox.h
@@ -2,21 +2,22 @@
 
 #pragma once
 
+#include "SCheckBox.generated.h"
 
 /** Current state of the check box */
+UENUM()
 namespace ESlateCheckBoxState
 {
-	typedef uint8 Type;
-	
+	enum Type
+	{
 		/** Unchecked */
-	const Type Unchecked = 0;
-
+		Unchecked,
 		/** Checked */
-	const Type Checked = 1;
-
+		Checked,
 		/** Neither checked nor unchecked */
-	const Type Undetermined = 2;
-};
+		Undetermined
+	};
+}
 
 
 /** Delegate that is executed when the check box state changes */
@@ -153,6 +154,9 @@ public:
 		return (IsCheckboxChecked.Get() == ESlateCheckBoxState::Checked);
 	}
 
+	/** @return The current checked state of the checkbox. */
+	ESlateCheckBoxState::Type GetCheckedState() const;
+
 	/**
 	 * Returns true if this button is currently pressed
 	 *
@@ -167,6 +171,9 @@ public:
 	 * Toggles the checked state for this check box, fire events as needed
 	 */
 	void ToggleCheckedState();
+
+	/** See the IsChecked attribute */
+	void SetIsChecked(TAttribute<ESlateCheckBoxState::Type> InIsChecked);
 	
 	/** See the Style attribute */
 	void SetStyle(const FCheckBoxStyle* InStyle);
diff --git a/Engine/Source/Runtime/UMG/Private/Components/Border.cpp b/Engine/Source/Runtime/UMG/Private/Components/Border.cpp
--- a/Engine/Source/Runtime/UMG/Private/Components/Border.cpp
+++ b/Engine/Source/Runtime/UMG/Private/Components/Border.cpp
@@ -32,7 +32,7 @@ UBorder::UBorder(const FPostConstructInitializeProperties& PCIP)
 TSharedRef<SWidget> UBorder::RebuildWidget()
 {
 	MyBorder = SNew(SBorder);
-	MyBorder->SetContent(ContentSlot->Content ? ContentSlot->Content->GetWidget() : SNullWidget::NullWidget);
+	MyBorder->SetContent(GetContentSlot()->Content ? GetContentSlot()->Content->GetWidget() : SNullWidget::NullWidget);
 
 	return MyBorder.ToSharedRef();
 }
diff --git a/Engine/Source/Runtime/UMG/Private/Components/Button.cpp b/Engine/Source/Runtime/UMG/Private/Components/Button.cpp
--- a/Engine/Source/Runtime/UMG/Private/Components/Button.cpp
+++ b/Engine/Source/Runtime/UMG/Private/Components/Button.cpp
@@ -29,7 +29,7 @@ UButton::UButton(const FPostConstructInitializeProperties& PCIP)
 TSharedRef<SWidget> UButton::RebuildWidget()
 {
 	MyButton = SNew(SButton);
-	MyButton->SetContent(ContentSlot->Content ? ContentSlot->Content->GetWidget() : SNullWidget::NullWidget);
+	MyButton->SetContent(GetContentSlot()->Content ? GetContentSlot()->Content->GetWidget() : SNullWidget::NullWidget);
 	
 	return MyButton.ToSharedRef();
 }
diff --git a/Engine/Source/Runtime/UMG/Private/Components/CheckBox.cpp b/Engine/Source/Runtime/UMG/Private/Components/CheckBox.cpp
--- a/Engine/Source/Runtime/UMG/Private/Components/CheckBox.cpp
+++ b/Engine/Source/Runtime/UMG/Private/Components/CheckBox.cpp
@@ -12,7 +12,7 @@ UCheckBox::UCheckBox(const FPostConstructInitializeProperties& PCIP)
 {
 	SCheckBox::FArguments CheckBoxDefaults;
 
-	bIsChecked = false;
+	CheckedState = ESlateCheckBoxState::Unchecked;
 
 	HorizontalAlignment = CheckBoxDefaults._HAlign;
 	Padding = CheckBoxDefaults._Padding.Get();
@@ -51,7 +51,6 @@ TSharedRef<SWidget> UCheckBox::RebuildWidget()
 	MyCheckbox = SNew(SCheckBox)
 		.Style(StylePtr)
 		.OnCheckStateChanged( BIND_UOBJECT_DELEGATE(FOnCheckStateChanged, SlateOnCheckStateChangedCallback) )
-		.IsChecked( BIND_UOBJECT_ATTRIBUTE(ESlateCheckBoxState::Type, GetCheckState) )
 		.HAlign( HorizontalAlignment )
 		.Padding( Padding )
 		.ForegroundColor( ForegroundColor )
@@ -67,6 +66,8 @@ TSharedRef<SWidget> UCheckBox::RebuildWidget()
 void UCheckBox::SyncronizeProperties()
 {
 	Super::SyncronizeProperties();
+
+	MyCheckbox->SetIsChecked( OPTIONAL_BINDING(ESlateCheckBoxState::Type, CheckedState) );
 	
 	MyCheckbox->SetUncheckedImage(UncheckedImage ? &UncheckedImage->Brush : nullptr);
 	MyCheckbox->SetUncheckedHoveredImage(UncheckedHoveredImage ? &UncheckedHoveredImage->Brush : nullptr);
@@ -83,17 +84,50 @@ void UCheckBox::SyncronizeProperties()
 
 bool UCheckBox::IsPressed() const
 {
-	return MyCheckbox->IsPressed();
+	if ( MyCheckbox.IsValid() )
+	{
+		return MyCheckbox->IsPressed();
+	}
+
+	return false;
 }
 
 bool UCheckBox::IsChecked() const
 {
-	return MyCheckbox->IsChecked();
+	if ( MyCheckbox.IsValid() )
+	{
+		return MyCheckbox->IsChecked();
+	}
+
+	return ( CheckedState == ESlateCheckBoxState::Checked );
+}
+
+ESlateCheckBoxState::Type UCheckBox::GetCheckedState() const
+{
+	if ( MyCheckbox.IsValid() )
+	{
+		return MyCheckbox->GetCheckedState();
+	}
+
+	return CheckedState;
+}
+
+void UCheckBox::SetIsChecked(bool InIsChecked)
+{
+	CheckedState = InIsChecked ? ESlateCheckBoxState::Checked : ESlateCheckBoxState::Unchecked;
+	if ( MyCheckbox.IsValid() )
+	{
+		MyCheckbox->SetIsChecked(OPTIONAL_BINDING(ESlateCheckBoxState::Type, CheckedState));
+	}
 }
 
-ESlateCheckBoxState::Type UCheckBox::GetCheckState() const
+void UCheckBox::SetCheckedState(ESlateCheckBoxState::Type InCheckedState)
 {
-	return bIsChecked ? ESlateCheckBoxState::Checked : ESlateCheckBoxState::Unchecked;
+	CheckedState = InCheckedState;
+	if ( MyCheckbox.IsValid() )
+	{
+		MyCheckbox->SetIsChecked(OPTIONAL_BINDING(ESlateCheckBoxState::Type, CheckedState));
+	}
 }
 
 void UCheckBox::SlateOnCheckStateChangedCallback(ESlateCheckBoxState::Type NewState)
@@ -107,7 +141,7 @@ void UCheckBox::SlateOnCheckStateChangedCallback(ESlateCheckBoxState::Type NewSt
 	}
 	else
 	{
-		bIsChecked = bWantsToBeChecked;
+		CheckedState = NewState;
 	}
 }
 
diff --git a/Engine/Source/Runtime/UMG/Private/Components/ContentWidget.cpp b/Engine/Source/Runtime/UMG/Private/Components/ContentWidget.cpp
--- a/Engine/Source/Runtime/UMG/Private/Components/ContentWidget.cpp
+++ b/Engine/Source/Runtime/UMG/Private/Components/ContentWidget.cpp
@@ -15,50 +15,54 @@
 UContentWidget::UContentWidget(const FPostConstructInitializeProperties& PCIP)
 	: Super(PCIP)
 {
+	ChildSlot = PCIP.CreateDefaultSubobject<UPanelSlot>(this, TEXT("ChildSlot"));
+	ChildSlot.Get()->SetFlags(RF_Transactional);
+	ChildSlot.Get()->Parent = this;
 }
 
-void UContentWidget::PostInitProperties()
+UPanelSlot* UContentWidget::GetContentSlot() const
 {
-	Super::PostInitProperties();
+	//if ( ContentSlot == NULL )
+	//{
+	//	UContentWidget* MutableThis = const_cast<UContentWidget*>( this );
+	//	MutableThis->ContentSlot = ConstructObject<UPanelSlot>(UPanelSlot::StaticClass(), MutableThis);
 
-	if ( ContentSlot == NULL )
-	{
-		ContentSlot = ConstructObject<UPanelSlot>(UPanelSlot::StaticClass(), this);
-	}
+	//	MutableThis->ContentSlot->SetFlags(RF_Transactional);
+	//	//MutableThis->ContentSlot->SetFlags(RF_DefaultSubObject);
+	//	MutableThis->ContentSlot->Parent = MutableThis;
+	//}
 
-	ContentSlot->SetFlags(RF_Transactional);
-	ContentSlot->SetFlags(RF_DefaultSubObject);
-	ContentSlot->Parent = this;
+	return ChildSlot.Get();
 }
 
 UWidget* UContentWidget::GetContent()
 {
-	return ContentSlot->Content;
+	return GetContentSlot()->Content;
 }
 
 void UContentWidget::SetContent(UWidget* InContent)
 {
-	ContentSlot->Content = InContent;
+	GetContentSlot()->Content = InContent;
 
 	if ( InContent )
 	{
-		InContent->Slot = ContentSlot;
+		InContent->Slot = GetContentSlot();
 	}
 }
 
 int32 UContentWidget::GetChildrenCount() const
 {
-	return ContentSlot->Content != NULL ? 1 : 0;
+	return GetContentSlot()->Content != NULL ? 1 : 0;
 }
 
 UWidget* UContentWidget::GetChildAt(int32 Index) const
 {
-	return ContentSlot->Content;
+	return GetContentSlot()->Content;
 }
 
 bool UContentWidget::AddChild(UWidget* InContent, FVector2D Position)
 {
-	if ( ContentSlot->Content == NULL )
+	if ( GetContentSlot()->Content == NULL )
 	{
 		SetContent(InContent);
 		return true;
@@ -69,7 +73,7 @@ bool UContentWidget::AddChild(UWidget* InContent, FVector2D Position)
 
 bool UContentWidget::RemoveChild(UWidget* Child)
 {
-	if ( ContentSlot->Content == Child )
+	if ( GetContentSlot()->Content == Child )
 	{
 		SetContent(NULL);
 		return true;
diff --git a/Engine/Source/Runtime/UMG/Private/Components/Widget.cpp b/Engine/Source/Runtime/UMG/Private/Components/Widget.cpp
--- a/Engine/Source/Runtime/UMG/Private/Components/Widget.cpp
+++ b/Engine/Source/Runtime/UMG/Private/Components/Widget.cpp
@@ -16,7 +16,7 @@ UWidget::UWidget(const FPostConstructInitializeProperties& PCIP)
 
 bool UWidget::GetIsEnabled() const
 {
-	return MyWidget.IsValid() ? MyWidget->IsEnabled() : bIsEnabled;
+	return MyWidget.IsValid() ? MyWidget.Pin()->IsEnabled() : bIsEnabled;
 }
 
 void UWidget::SetIsEnabled(bool bInIsEnabled)
@@ -24,7 +24,7 @@ void UWidget::SetIsEnabled(bool bInIsEnabled)
 	bIsEnabled = bInIsEnabled;
 	if ( MyWidget.IsValid() )
 	{
-		MyWidget->SetEnabled(bInIsEnabled);
+		MyWidget.Pin()->SetEnabled(bInIsEnabled);
 	}
 }
 
@@ -32,7 +32,7 @@ TEnumAsByte<ESlateVisibility::Type> UWidget::GetVisibility()
 {
 	if ( MyWidget.IsValid() )
 	{
-		return UWidget::ConvertRuntimeToSerializedVisiblity(MyWidget->GetVisibility());
+		return UWidget::ConvertRuntimeToSerializedVisiblity(MyWidget.Pin()->GetVisibility());
 	}
 
 	return Visiblity;
@@ -44,7 +44,7 @@ void UWidget::SetVisibility(TEnumAsByte<ESlateVisibility::Type> InVisibility)
 
 	if ( MyWidget.IsValid() )
 	{
-		MyWidget->SetVisibility(UWidget::ConvertSerializedVisibilityToRuntime(InVisibility));
+		return MyWidget.Pin()->SetVisibility(UWidget::ConvertSerializedVisibilityToRuntime(InVisibility));
 	}
 }
 
@@ -54,13 +54,18 @@ void UWidget::SetToolTipText(const FText& InToolTipText)
 
 	if ( MyWidget.IsValid() )
 	{
-		MyWidget->SetToolTipText(InToolTipText);
+		return MyWidget.Pin()->SetToolTipText(InToolTipText);
 	}
 }
 
 bool UWidget::IsHovered() const
 {
-	return MyWidget->IsHovered();
+	if ( MyWidget.IsValid() )
+	{
+		return MyWidget.Pin()->IsHovered();
+	}
+
+	return false;
 }
 
 UWidget* UWidget::GetParent() const
@@ -75,6 +80,8 @@ UWidget* UWidget::GetParent() const
 
 TSharedRef<SWidget> UWidget::GetWidget() const
 {
+	TSharedPtr<SWidget> SafeWidget;
+
 	if ( !MyWidget.IsValid() )
 	{
 		// We lie a bit about this being a const function.  If this is the first call it's not really const
@@ -83,11 +90,16 @@ TSharedRef<SWidget> UWidget::GetWidget() const
 		// just blow away the const here.
 		UWidget* MutableThis = const_cast<UWidget*>(this);
 		
-		MyWidget = MutableThis->RebuildWidget();
+		SafeWidget = MutableThis->RebuildWidget();
+		MyWidget = SafeWidget;
 		MutableThis->SyncronizeProperties();
 	}
+	else
+	{
+		SafeWidget = MyWidget.Pin();
+	}
 
-	return MyWidget.ToSharedRef();
+	return SafeWidget.ToSharedRef();
 }
 
 TSharedRef<SWidget> UWidget::BuildDesignTimeWidget(TSharedRef<SWidget> WrapWidget)
@@ -186,13 +198,13 @@ TSharedRef<SWidget> UWidget::RebuildWidget()
 
 void UWidget::SyncronizeProperties()
 {
-	MyWidget->SetEnabled(OPTIONAL_BINDING(bool, bIsEnabled));
-	MyWidget->SetVisibility(OPTIONAL_BINDING_CONVERT(ESlateVisibility::Type, Visiblity, EVisibility, ConvertVisibility));
+	MyWidget.Pin()->SetEnabled(OPTIONAL_BINDING(bool, bIsEnabled));
+	MyWidget.Pin()->SetVisibility(OPTIONAL_BINDING_CONVERT(ESlateVisibility::Type, Visiblity, EVisibility, ConvertVisibility));
 	//MyWidget->SetCursor(OPTIONAL_BINDING(EMouseCursor)
 
 	if ( !ToolTipText.IsEmpty() )
 	{
-			MyWidget->SetToolTipText(OPTIONAL_BINDING(FText, ToolTipText));
+		MyWidget.Pin()->SetToolTipText(OPTIONAL_BINDING(FText, ToolTipText));
 	}
 }
 
diff --git a/Engine/Source/Runtime/UMG/Private/Slate/SObjectWidget.cpp b/Engine/Source/Runtime/UMG/Private/Slate/SObjectWidget.cpp
--- a/Engine/Source/Runtime/UMG/Private/Slate/SObjectWidget.cpp
+++ b/Engine/Source/Runtime/UMG/Private/Slate/SObjectWidget.cpp
@@ -12,8 +12,14 @@ void SObjectWidget::Construct(const FArguments& InArgs, UUserWidget* InWidgetObj
 	];
 }
 
+SObjectWidget::~SObjectWidget(void)
+{
+	printf("");
+}
+
 void SObjectWidget::AddReferencedObjects(FReferenceCollector& Collector)
 {
+	printf("");
 	Collector.AddReferencedObject(WidgetObject);
 }
 
diff --git a/Engine/Source/Runtime/UMG/Private/UserWidget.cpp b/Engine/Source/Runtime/UMG/Private/UserWidget.cpp
--- a/Engine/Source/Runtime/UMG/Private/UserWidget.cpp
+++ b/Engine/Source/Runtime/UMG/Private/UserWidget.cpp
@@ -186,6 +186,8 @@ TSharedRef<SWidget> UUserWidget::RebuildWidget()
 {
 	WidgetToComponent.Reset();
 	
+	TSharedPtr<SWidget> UserRootWidget;
+
 	// Add the first component to the root of the widget surface.
 	if ( Components.Num() > 0 )
 	{
@@ -205,7 +207,10 @@ TSharedRef<SWidget> UUserWidget::RebuildWidget()
 		WidgetToComponent.Add(Widget, Handle);
 	}
 
-	return MakeWidget();
+	return SNew(SObjectWidget, this)
+		[
+			UserRootWidget.ToSharedRef()
+		];
 }
 
 TSharedPtr<SWidget> UUserWidget::GetWidgetFromName(const FString& Name) const
@@ -239,18 +244,8 @@ UWidget* UUserWidget::GetHandleFromName(const FString& Name) const
 	return NULL;
 }
 
-TSharedRef<SWidget> UUserWidget::MakeWidget()
-{
-	return SNew(SObjectWidget, this)
-	[
-		UserRootWidget.ToSharedRef()
-	];
-}
-
 TSharedRef<SWidget> UUserWidget::MakeFullScreenWidget()
 {
-	GetWidget();
-	
 	if ( bAbsoluteLayout )
 	{
 		return SNew(SCanvas)
@@ -261,7 +256,7 @@ TSharedRef<SWidget> UUserWidget::MakeFullScreenWidget()
 			.VAlign(VerticalAlignment)
 			.HAlign(HorizontalAlignment)
 			[
-				MakeWidget()
+				GetWidget()
 			];
 	}
 	else
@@ -273,7 +268,7 @@ TSharedRef<SWidget> UUserWidget::MakeFullScreenWidget()
 			.HAlign(HorizontalAlignment)
 			.VAlign(VerticalAlignment)
 			[
-				MakeWidget()
+				GetWidget()
 			];
 
 		NewSlot.SizeParam = UWidget::ConvertSerializedSizeParamToRuntime(Size);
diff --git a/Engine/Source/Runtime/UMG/Private/WidgetBlueprintGeneratedClass.cpp b/Engine/Source/Runtime/UMG/Private/WidgetBlueprintGeneratedClass.cpp
--- a/Engine/Source/Runtime/UMG/Private/WidgetBlueprintGeneratedClass.cpp
+++ b/Engine/Source/Runtime/UMG/Private/WidgetBlueprintGeneratedClass.cpp
@@ -22,7 +22,7 @@ void UWidgetBlueprintGeneratedClass::PostInitProperties()
 		WidgetTree = ConstructObject<UWidgetTree>(UWidgetTree::StaticClass(), this);
 	}
 
-	WidgetTree->SetFlags(RF_DefaultSubObject);
+	//WidgetTree->SetFlags(RF_DefaultSubObject);
 	WidgetTree->SetFlags(RF_Transactional);
 }
 
diff --git a/Engine/Source/Runtime/UMG/Public/Blueprint/UserWidget.h b/Engine/Source/Runtime/UMG/Public/Blueprint/UserWidget.h
--- a/Engine/Source/Runtime/UMG/Public/Blueprint/UserWidget.h
+++ b/Engine/Source/Runtime/UMG/Public/Blueprint/UserWidget.h
@@ -249,9 +249,6 @@ class UMG_API UUserWidget : public UWidget
 	/** @returns The UObject wrapper for a given SWidget */
 	UWidget* GetWidgetHandle(TSharedRef<SWidget> InWidget);
 
-	/** Creates this widget and all children and returns them. */
-	TSharedRef<SWidget> MakeWidget();
-
 	/** Creates a fullscreen host widget, that wraps this widget. */
 	TSharedRef<SWidget> MakeFullScreenWidget();
 
@@ -268,7 +265,6 @@ protected:
 	virtual TSharedRef<SWidget> RebuildWidget() override;
 
 private:
-	TSharedPtr<SWidget> UserRootWidget;
 	TMap< TWeakPtr<SWidget>, TWeakObjectPtr<UWidget> > WidgetToComponent;
 
 	TWeakPtr<SWidget> FullScreenWidget;
diff --git a/Engine/Source/Runtime/UMG/Public/Components/CheckBox.h b/Engine/Source/Runtime/UMG/Public/Components/CheckBox.h
--- a/Engine/Source/Runtime/UMG/Public/Components/CheckBox.h
+++ b/Engine/Source/Runtime/UMG/Public/Components/CheckBox.h
@@ -53,13 +53,13 @@ protected:
 	UPROPERTY(EditDefaultsOnly, Category="Style", meta=( DisplayThumbnail = "true" ), AdvancedDisplay)
 	USlateBrushAsset* UndeterminedPressedImage;
 
-	/** Called when the checked state has changed */
-	UPROPERTY(BlueprintAssignable)
-	FOnCheckBoxComponentStateChanged OnCheckStateChanged;
-
 	/** Whether the check box is currently in a checked state */
 	UPROPERTY(EditDefaultsOnly, Category=Appearance)
-	bool bIsChecked;
+	TEnumAsByte<ESlateCheckBoxState::Type> CheckedState;
+
+	/** A bindable delegate for the IsChecked. */
+	UPROPERTY()
+	FGetCheckBoxState CheckedStateDelegate;
 
 	/** How the content of the toggle button should align within the given space */
 	UPROPERTY(EditDefaultsOnly, Category=Appearance)
@@ -90,7 +90,11 @@ protected:
 
 	UPROPERTY(EditDefaultsOnly, Category="Sound")
 	FSlateSound HoveredSound;
-	
+
+	/** Called when the checked state has changed */
+	UPROPERTY(BlueprintAssignable)
+	FOnCheckBoxComponentStateChanged OnCheckStateChanged;
+
 	/** Returns true if this button is currently pressed */
 	UFUNCTION(BlueprintCallable, Category="Widget")
 	bool IsPressed() const;
@@ -98,6 +102,18 @@ protected:
 	/** Returns true if the checkbox is currently checked */
 	UFUNCTION(BlueprintCallable, Category="Widget")
 	bool IsChecked() const;
+
+	/** @return the full current checked state. */
+	UFUNCTION(BlueprintCallable, Category="Widget")
+	ESlateCheckBoxState::Type GetCheckedState() const;
+
+	/** Sets the checked state. */
+	UFUNCTION(BlueprintCallable, Category="Widget")
+	void SetIsChecked(bool InIsChecked);
+
+	/** Sets the checked state. */
+	UFUNCTION(BlueprintCallable, Category="Widget")
+	void SetCheckedState(ESlateCheckBoxState::Type InCheckedState);
 	
 	// UWidget interface
 	virtual void SyncronizeProperties() override;
@@ -112,7 +128,6 @@ protected:
 	virtual TSharedRef<SWidget> RebuildWidget() override;
 	// End of UWidget interface
 
-	ESlateCheckBoxState::Type GetCheckState() const;
 	void SlateOnCheckStateChangedCallback(ESlateCheckBoxState::Type NewState);
 	
 protected:
diff --git a/Engine/Source/Runtime/UMG/Public/Components/ContentWidget.h b/Engine/Source/Runtime/UMG/Public/Components/ContentWidget.h
--- a/Engine/Source/Runtime/UMG/Public/Components/ContentWidget.h
+++ b/Engine/Source/Runtime/UMG/Public/Components/ContentWidget.h
@@ -23,12 +23,11 @@ class UMG_API UContentWidget : public UPanelWidget
 	virtual void ReplaceChildAt(int32 Index, UWidget* Child) override;
 	// End UPanelWidget
 
-	// UObject
-	virtual void PostInitProperties() override;
-	// End UObject
-
 protected:
+	UPanelSlot* GetContentSlot() const;
+
+private:
 
-	UPROPERTY()
-	UPanelSlot* ContentSlot;
+	UPROPERTY(Instanced)
+	TSubobjectPtr<UPanelSlot> ChildSlot;
 };
diff --git a/Engine/Source/Runtime/UMG/Public/Components/Widget.h b/Engine/Source/Runtime/UMG/Public/Components/Widget.h
--- a/Engine/Source/Runtime/UMG/Public/Components/Widget.h
+++ b/Engine/Source/Runtime/UMG/Public/Components/Widget.h
@@ -35,6 +35,7 @@ public:
 	DECLARE_DYNAMIC_DELEGATE_RetVal(ESlateVisibility::Type, FGetSlateVisibility);
 	DECLARE_DYNAMIC_DELEGATE_RetVal(EMouseCursor::Type, FGetMouseCursor);
 	DECLARE_DYNAMIC_DELEGATE_RetVal(USlateBrushAsset*, FGetSlateBrushAsset);
+	DECLARE_DYNAMIC_DELEGATE_RetVal(ESlateCheckBoxState::Type, FGetCheckBoxState);
 
 	DECLARE_DYNAMIC_DELEGATE_RetVal_OneParam(UWidget*, FGenerateWidgetUObject, UObject*, Item);
 
@@ -123,7 +124,11 @@ public:
 	UFUNCTION(BlueprintCallable, Category="Widget")
 	UWidget* GetParent() const;
 
-	/** Gets the underlying Slate SWidget for this UWidget. */
+	/**
+	 * Gets the underlying slate widget or constructs it if it doesn't exist.  This function is
+	 * virtual however, you should not inherit this function unless you're very aware of what you're
+	 * doing.  Normal derived versions should only ever override RebuildWidget.
+	 */
 	TSharedRef<SWidget> GetWidget() const;
 
 	/**
@@ -189,7 +194,7 @@ protected:
 
 protected:
 	/** The underlying SWidget. */
-	mutable TSharedPtr<SWidget> MyWidget;
+	mutable TWeakPtr<SWidget> MyWidget;
 	
 	/** Is this widget being displayed on a designer surface */
 	bool bDesignTime;
diff --git a/Engine/Source/Runtime/UMG/Public/Slate/SObjectWidget.h b/Engine/Source/Runtime/UMG/Public/Slate/SObjectWidget.h
--- a/Engine/Source/Runtime/UMG/Public/Slate/SObjectWidget.h
+++ b/Engine/Source/Runtime/UMG/Public/Slate/SObjectWidget.h
@@ -15,9 +15,13 @@ class SObjectWidget : public SCompoundWidget, public FGCObject
 		SLATE_DEFAULT_SLOT(FArguments, Content)
 	SLATE_END_ARGS()
 
+	virtual ~SObjectWidget(void);
+
 	void Construct(const FArguments& InArgs, class UUserWidget* InWidgetObject);
 
+	// FGCObject interface
 	virtual void AddReferencedObjects(FReferenceCollector& Collector) override;
+	// End of FGCObject interface
 
 	virtual int32 OnPaint(const FGeometry& AllottedGeometry, const FSlateRect& MyClippingRect, FSlateWindowElementList& OutDrawElements, int32 LayerId, const FWidgetStyle& InWidgetStyle, bool bParentEnabled) const override;
 
