Editor.Performence.FrameRate is no longer triggered while the editor has no project loaded

TTP# 337729 - EDITOR:ANALYTICS: Editor.Performence.FrameRate should only fire once the editor launches.

diff --git a/Engine/Source/Editor/UnrealEd/Private/UnrealEdMisc.cpp b/Engine/Source/Editor/UnrealEd/Private/UnrealEdMisc.cpp
--- a/Engine/Source/Editor/UnrealEd/Private/UnrealEdMisc.cpp
+++ b/Engine/Source/Editor/UnrealEd/Private/UnrealEdMisc.cpp
@@ -1199,6 +1199,12 @@ void FUnrealEdMisc::TickPerformanceSurvey()
 		return;
 	}
 
+	// Don't run if we've not yet loaded a project
+	if( !FApp::HasGameName() )
+	{
+		return;
+	}
+
 	// Before beginning the survey wait for the asset registry to load and make sure Slate is ready
 	FAssetRegistryModule& AssetRegistryModule = FModuleManager::LoadModuleChecked<FAssetRegistryModule>(TEXT("AssetRegistry"));
 	if (AssetRegistryModule.Get().IsLoadingAssets() || !FSlateApplication::IsInitialized())
