UDS Renaming: When a UDS's variable is renamed it's editor-only user-friendly name will change, but 'inner' property name wont.
Pros: when a uds is used in a BP with a modified default value, the value remains unchanged while renaming (https://udn.unrealengine.com/questions/211412/issue-with-renaming-a-variable-of-a-structure.html)
Cons: The property name in the structure is no more related with the user name - harder debugging, less clear logs.

diff --git a/Engine/Source/Editor/UnrealEd/Private/Kismet2/StructureEditorUtils.cpp b/Engine/Source/Editor/UnrealEd/Private/Kismet2/StructureEditorUtils.cpp
--- a/Engine/Source/Editor/UnrealEd/Private/Kismet2/StructureEditorUtils.cpp
+++ b/Engine/Source/Editor/UnrealEd/Private/Kismet2/StructureEditorUtils.cpp
@@ -280,25 +280,17 @@ bool FStructureEditorUtils::RenameVariable(UUserDefinedStruct* Struct, FGuid Var
 {
 	if (Struct)
 	{
-		const FName NewNameBase(*NewDisplayNameStr);
-		if (NewNameBase.IsValidXName(INVALID_OBJECTNAME_CHARACTERS) &&
+		auto VarDesc = GetVarDescByGuid(Struct, VarGuid);
+		if (VarDesc &&
+			FName(*NewDisplayNameStr).IsValidXName(INVALID_OBJECTNAME_CHARACTERS) &&
 			IsUniqueVariableDisplayName(Struct, NewDisplayNameStr))
 		{
-			const FName NewName = FMemberVariableNameHelper::Generate(Struct, NewDisplayNameStr);
-			if (NULL == GetVarDesc(Struct).FindByPredicate(FFindByNameHelper<FStructVariableDescription>(NewName)))
-			{
-				auto VarDesc = GetVarDescByGuid(Struct, VarGuid);
-				if (VarDesc)
-				{
-					const FScopedTransaction Transaction(LOCTEXT("RenameVariable", "Rename Variable"));
-					ModifyStructData(Struct);
+			const FScopedTransaction Transaction(LOCTEXT("RenameVariable", "Rename Variable"));
+			ModifyStructData(Struct);
 
-					VarDesc->FriendlyName = NewDisplayNameStr;
-					VarDesc->VarName = NewName;
-					OnStructureChanged(Struct);
-					return true;
-				}
-			}
+			VarDesc->FriendlyName = NewDisplayNameStr;
+			OnStructureChanged(Struct);
+			return true;
 		}
 	}
 	return false;
