Locking a viewport to a camera actor now correctly uses the camera's properties

I've made a few changes here to make the code more consistent for the various ways of using a viewport to view an actor's position:
    * Camera properties are now taken off any UCameraComponent, not specifically limited to ACameraActors.
    * Calling FLevelEditorViewportClient::SetActorLock will now lock the viewport to the specified actor. If the actor has a camera component the viewport will inherit this camera's properties. The locked actor will also move with the viewport position if in perspective mode (as before).
    * ControllingActor and PostProcessingActor have been removed from FLevelEditorViewportClient. Both were used for PIP camera previews which now work by locking a viewport to their desired camera.
    * The inheriting of camera properties is controlled by the viewport "Camera View" setting shown on the quick settings viewport toolbar under the "Lock Viewport to Actor sub-menu.

This addresses TTP#337559 - feature request - improve viewport functionality for locking to camera

diff --git a/Engine/Source/Developer/MovieSceneTools/Private/DirectorTrackEditor.cpp b/Engine/Source/Developer/MovieSceneTools/Private/DirectorTrackEditor.cpp
--- a/Engine/Source/Developer/MovieSceneTools/Private/DirectorTrackEditor.cpp
+++ b/Engine/Source/Developer/MovieSceneTools/Private/DirectorTrackEditor.cpp
@@ -174,7 +174,7 @@ FShotSection::FShotSection( TSharedPtr<ISequencer> InSequencer, TSharedPtr<FShot
 		InternalViewportClient->bDrawAxes = false;
 		InternalViewportClient->EngineShowFlags = FEngineShowFlags(ESFIM_Game);
 		InternalViewportClient->EngineShowFlags.DisableAdvancedFeatures();
-		InternalViewportClient->SetControllingActor(Camera.Get());
+		InternalViewportClient->SetActorLock(Camera.Get());
 
 		InternalViewportScene = MakeShareable(new FSceneViewport(InternalViewportClient.Get(), NULL));
 		InternalViewportClient->Viewport = InternalViewportScene.Get();
@@ -345,7 +345,7 @@ void FShotSection::RegenerateViewportThumbnails(const FIntPoint& Size)
 void FShotSection::DrawViewportThumbnail(TSharedPtr<FShotThumbnail> ShotThumbnail)
 {
 	Sequencer.Pin()->SetGlobalTime(ShotThumbnail->GetTime());
-	InternalViewportClient->PushControllingActorDataToViewportClient();
+	InternalViewportClient->UpdateViewForLockedActor();
 			
 	GWorld->SendAllEndOfFrameUpdates();
 	InternalViewportScene->Draw(false);
diff --git a/Engine/Source/Editor/LevelEditor/Private/LevelViewportActions.cpp b/Engine/Source/Editor/LevelEditor/Private/LevelViewportActions.cpp
--- a/Engine/Source/Editor/LevelEditor/Private/LevelViewportActions.cpp
+++ b/Engine/Source/Editor/LevelEditor/Private/LevelViewportActions.cpp
@@ -26,7 +26,8 @@ void FLevelViewportCommands::RegisterCommands()
 	UI_COMMAND( ToggleMaximize, "Maximize Viewport", "Toggles the Maximize state of the current viewport", EUserInterfaceActionType::ToggleButton, FInputGesture() );
 	UI_COMMAND( ToggleGameView, "Game View", "Toggles game view.  Game view shows the scene as it appears in game", EUserInterfaceActionType::ToggleButton, FInputGesture( EKeys::G ) );
 	UI_COMMAND( ToggleImmersive, "Immersive Mode", "Switches this viewport between immersive mode and regular mode", EUserInterfaceActionType::ToggleButton, FInputGesture( EKeys::F11 ) );
-	
+	UI_COMMAND( ToggleLockedCameraView, "Camera View", "Toggles showing the exact camera view when locking this viewport to a camera", EUserInterfaceActionType::ToggleButton, FInputGesture( EModifierKey::Control | EModifierKey::Shift, EKeys::C ) );
+
 	UI_COMMAND( CreateCamera, "Create Camera Here", "Creates a new camera actor at the current location of this viewport's camera", EUserInterfaceActionType::Button, FInputGesture() );
 	UI_COMMAND( HighResScreenshot, "High Resolution Screenshot...", "Opens the control panel for high resolution screenshots", EUserInterfaceActionType::Button, FInputGesture() );
 	
diff --git a/Engine/Source/Editor/LevelEditor/Private/SLevelViewport.cpp b/Engine/Source/Editor/LevelEditor/Private/SLevelViewport.cpp
--- a/Engine/Source/Editor/LevelEditor/Private/SLevelViewport.cpp
+++ b/Engine/Source/Editor/LevelEditor/Private/SLevelViewport.cpp
@@ -766,17 +766,6 @@ void SLevelViewport::Tick( const FGeometry& AllottedGeometry, const double InCur
 		PIEOverlayAnim.Play();
 	}
 
-	// Check to see if the locked actor wants to override the camera settings
-	if( LevelViewportClient->GetActiveActorLock().IsValid() )
-	{
-		AActor* LockedActor = LevelViewportClient->GetActiveActorLock().Get();
-		FMinimalViewInfo CameraInfo;
-		if( GetCameraInformationFromActor(LockedActor, /*out*/ CameraInfo) )
-		{
-			LevelViewportClient->ViewFOV = CameraInfo.FOV;
-		}
-	}
-
 	// Update actor preview viewports, if we have any
 	UpdateActorPreviewViewports();
 
@@ -1030,6 +1019,13 @@ void SLevelViewport::BindOptionCommands( FUICommandList& CommandList )
 		FCanExecuteAction()
 		);
 
+	CommandList.MapAction(
+		ViewportActions.ToggleLockedCameraView,
+		FExecuteAction::CreateSP(this, &SLevelViewport::ToggleLockedCameraView),
+		FCanExecuteAction(),
+		FIsActionChecked::CreateSP( this, &SLevelViewport::IsLockedCameraViewEnabled )
+		);
+
 	// Map each bookmark action
 	for( int32 BookmarkIndex = 0; BookmarkIndex < AWorldSettings::MAX_BOOKMARK_NUMBER; ++BookmarkIndex )
 	{
@@ -1974,6 +1970,16 @@ bool SLevelViewport::IsAnyActorLocked() const
 	return LevelViewportClient->IsAnyActorLocked();
 }
 
+void SLevelViewport::ToggleLockedCameraView()
+{
+	LevelViewportClient->bLockedCameraView = !LevelViewportClient->bLockedCameraView;
+}
+
+bool SLevelViewport::IsLockedCameraViewEnabled() const
+{
+	return LevelViewportClient->bLockedCameraView;
+}
+
 void SLevelViewport::FindSelectedInLevelScript()
 {
 	GUnrealEd->FindSelectedActorsInLevelScript();
@@ -2825,8 +2831,8 @@ void SLevelViewport::PreviewActors( const TArray< AActor* >& ActorsToPreview )
 
 				// Push actor transform to view.  From here on out, this will happen automatically in FLevelEditorViewportClient::Tick.
 				// The reason we allow the viewport client to update this is to avoid off-by-one-frame issues when dragging actors around.
-				ActorPreviewLevelViewportClient->SetControllingActor( CurActor );
-				ActorPreviewLevelViewportClient->PushControllingActorDataToViewportClient();
+				ActorPreviewLevelViewportClient->SetActorLock( CurActor );
+				ActorPreviewLevelViewportClient->UpdateViewForLockedActor();
 			}
 
 			TSharedPtr< SActorPreview > ActorPreviewWidget = SNew(SActorPreview)
diff --git a/Engine/Source/Editor/LevelEditor/Private/SLevelViewportToolBar.cpp b/Engine/Source/Editor/LevelEditor/Private/SLevelViewportToolBar.cpp
--- a/Engine/Source/Editor/LevelEditor/Private/SLevelViewportToolBar.cpp
+++ b/Engine/Source/Editor/LevelEditor/Private/SLevelViewportToolBar.cpp
@@ -442,6 +442,8 @@ static void OnGenerateActorLockingMenu( FMenuBuilder& MenuBuilder, TWeakPtr<SLev
 	const FLevelViewportCommands& Actions = FLevelViewportCommands::Get();
 	FLevelEditorViewportClient& ViewportClient = Viewport.Pin()->GetLevelViewportClient();
 
+	MenuBuilder.AddMenuEntry(Actions.ToggleLockedCameraView);
+
 	bool IsLocked = false;
 	if (ViewportClient.GetActiveActorLock().IsValid())
 	{
diff --git a/Engine/Source/Editor/LevelEditor/Public/LevelViewportActions.h b/Engine/Source/Editor/LevelEditor/Public/LevelViewportActions.h
--- a/Engine/Source/Editor/LevelEditor/Public/LevelViewportActions.h
+++ b/Engine/Source/Editor/LevelEditor/Public/LevelViewportActions.h
@@ -68,6 +68,9 @@ public:
 	/** Toggles immersive mode in the viewport */
 	TSharedPtr< FUICommandInfo > ToggleImmersive;
 
+	/** Toggles showing the exact camera view when locking a viewport to a camera */
+	TSharedPtr< FUICommandInfo > ToggleLockedCameraView;
+
 	/** Toggles maximize mode in the viewport */
 	TSharedPtr< FUICommandInfo > ToggleMaximize;
 
diff --git a/Engine/Source/Editor/LevelEditor/Public/SLevelViewport.h b/Engine/Source/Editor/LevelEditor/Public/SLevelViewport.h
--- a/Engine/Source/Editor/LevelEditor/Public/SLevelViewport.h
+++ b/Engine/Source/Editor/LevelEditor/Public/SLevelViewport.h
@@ -190,6 +190,16 @@ public:
 	 */
 	bool IsSelectedActorLocked() const;
 
+	/**
+	 * Toggles enabling the exact camera view when locking a viewport to a camera
+	 */
+	void ToggleLockedCameraView();
+
+	/**
+	 * Check whether locked camera view is enabled
+	 */
+	bool IsLockedCameraViewEnabled() const;
+
 	/**
 	 * @return the fixed width that a column returned by CreateActorLockSceneOutlinerColumn expects to be
 	 */
diff --git a/Engine/Source/Editor/Matinee/Private/MatineeTools.cpp b/Engine/Source/Editor/Matinee/Private/MatineeTools.cpp
--- a/Engine/Source/Editor/Matinee/Private/MatineeTools.cpp
+++ b/Engine/Source/Editor/Matinee/Private/MatineeTools.cpp
@@ -3883,11 +3883,8 @@ void FMatinee::UpdateLevelViewport(AActor* InActor, FLevelEditorViewportClient*
 
 		InViewportClient->bEditorCameraCut = false;
 	}
-	//@TODO: CAMERA: Unify this code path with the code in EditorViewportClient when viewing thru an actor that contains a camera component (e.g., SActorPreview)
-
-	// If viewing through a camera - PP settings of camera.
-	InViewportClient->SetPostprocessCameraActor(Cam);
 
+	// Set the actor lock.
 	InViewportClient->SetMatineeActorLock(InActor);
 
 	// If viewing through a camera - enforce aspect ratio.
@@ -3948,7 +3945,7 @@ void FMatinee::RestoreLevelViewports()
 			FLevelEditorViewportClient* LevelVC = GEditor->LevelViewportClients[ SavedData.ViewIndex ];
 			if ( LevelVC && LevelVC->IsPerspective() && LevelVC->AllowMatineePreview() )
 			{
-				LevelVC->SetPostprocessCameraActor(0);
+				LevelVC->SetMatineeActorLock( nullptr );
 				LevelVC->SetViewRotation( SavedData.ViewRotation );
 				LevelVC->SetViewLocation( SavedData.ViewLocation );				
 			}
diff --git a/Engine/Source/Editor/Sequencer/Private/Sequencer.cpp b/Engine/Source/Editor/Sequencer/Private/Sequencer.cpp
--- a/Engine/Source/Editor/Sequencer/Private/Sequencer.cpp
+++ b/Engine/Source/Editor/Sequencer/Private/Sequencer.cpp
@@ -544,35 +544,7 @@ void FSequencer::UpdateViewports(AActor* ActorToViewThrough) const
 	{
 		if(LevelVC && LevelVC->IsPerspective() && LevelVC->AllowMatineePreview())
 		{
-			if(ActorToViewThrough)
-			{
-				LevelVC->SetViewLocation(ActorToViewThrough->GetActorLocation());
-				LevelVC->SetViewRotation(ActorToViewThrough->GetActorRotation());
-			}
-			else
-			{
-				LevelVC->ViewFOV = LevelVC->FOVAngle;
-			}
-
-			ACameraActor* Camera = Cast<ACameraActor>(ActorToViewThrough);
-
-			// If viewing through a camera - PP settings of camera.
-			LevelVC->SetPostprocessCameraActor(Camera);
-
-			if(Camera)
-			{
-				LevelVC->ViewFOV = LevelVC->FOVAngle = Camera->CameraComponent->FieldOfView;
-				// If the Camera's aspect ratio is zero, put a more reasonable default here - this at least stops it from crashing
-				// nb. the AspectRatio will be reported as a Map Check Warning
-				if(Camera->CameraComponent->AspectRatio == 0)
-				{
-					LevelVC->AspectRatio = 1.7f;
-				}
-				else
-				{
-					LevelVC->AspectRatio = Camera->CameraComponent->AspectRatio;
-				}
-			}
+			LevelVC->SetMatineeActorLock(ActorToViewThrough);
 		}
 	}
 }
diff --git a/Engine/Source/Editor/UnrealEd/Private/LevelEditorViewport.cpp b/Engine/Source/Editor/UnrealEd/Private/LevelEditorViewport.cpp
--- a/Engine/Source/Editor/UnrealEd/Private/LevelEditorViewport.cpp
+++ b/Engine/Source/Editor/UnrealEd/Private/LevelEditorViewport.cpp
@@ -1771,13 +1771,12 @@ FLevelEditorViewportClient::FLevelEditorViewportClient()
 	, bDuplicateActorsOnNextDrag( false )
 	, bDuplicateActorsInProgress( false )
 	, bIsTrackingBrushModification( false )
-	, ControllingActor(NULL)
 	, SpriteCategoryVisibility()
-	, PostprocessCameraActor(NULL)
 	, World(NULL)
 	, TrackingTransaction()
 	, DropPreviewMouseX(0)
 	, DropPreviewMouseY(0)
+	, bLockedCameraView(true)
 	, bWasControlledByOtherViewport(false)
 {
 	// By default a level editor viewport is pointed to the editor world
@@ -1922,24 +1921,9 @@ FSceneView* FLevelEditorViewportClient::CalcSceneView(FSceneViewFamily* ViewFami
 
 }
 
-void FLevelEditorViewportClient::SetPostprocessCameraActor(ACameraActor* InPostprocessCameraActor)
-{
-	PostprocessCameraActor = InPostprocessCameraActor;
-}
-
 ELevelViewportType FLevelEditorViewportClient::GetViewportType() const
 {
-	UCameraComponent* ActiveCameraComponent = NULL;
-
-	if (AActor* TestActor = GetActiveActorLock().Get())
-	{
-		ActiveCameraComponent = TestActor->FindComponentByClass<UCameraComponent>();
-	}
-
-	if (ControllingActor.Get() != NULL)
-	{
-		ActiveCameraComponent = ControllingActor.Get()->FindComponentByClass<UCameraComponent>();
-	}
+	const UCameraComponent* ActiveCameraComponent = GetCameraComponentForView();
 	
 	if (ActiveCameraComponent != NULL)
 	{
@@ -1953,9 +1937,10 @@ ELevelViewportType FLevelEditorViewportClient::GetViewportType() const
 
 void FLevelEditorViewportClient::OverridePostProcessSettings( FSceneView& View )
 {
-	if(PostprocessCameraActor.IsValid() && PostprocessCameraActor.Get()->CameraComponent.IsValid())
+	const UCameraComponent* CameraComponent = GetCameraComponentForView();
+	if (CameraComponent)
 	{
-		View.OverridePostProcessSettings(PostprocessCameraActor.Get()->CameraComponent->PostProcessSettings, PostprocessCameraActor.Get()->CameraComponent->PostProcessBlendWeight);
+		View.OverridePostProcessSettings(CameraComponent->PostProcessSettings, CameraComponent->PostProcessBlendWeight);
 	}
 }
 
@@ -2218,18 +2203,22 @@ void FLevelEditorViewportClient::Tick(float DeltaTime)
 		GPerspViewMatrix = View->ViewMatrices.ViewMatrix;
 	}
 
-
-	// If we have a controlling actor, then go ahead and push location and other data from that actor to
-	// this view
-	PushControllingActorDataToViewportClient();
+	UpdateViewForLockedActor();
 }
 
 
-void FLevelEditorViewportClient::PushControllingActorDataToViewportClient()
+void FLevelEditorViewportClient::UpdateViewForLockedActor()
 {
+	// We can't be locked to a matinee actor if this viewport doesn't allow matinee control
+	if ( !bAllowMatineePreview && ActorLockedByMatinee.IsValid() )
+	{
+		ActorLockedByMatinee = nullptr;
+	}
+
 	bUseControllingActorViewInfo = false;
+	ControllingActorViewInfo = FMinimalViewInfo();
 
-	AActor* Actor = ControllingActor.Get();
+	const AActor* Actor = ActorLockedByMatinee.IsValid() ? ActorLockedByMatinee.Get() : ActorLockedToCamera.Get();
 	if( Actor != NULL )
 	{
 		// Update transform
@@ -2247,21 +2236,21 @@ void FLevelEditorViewportClient::PushControllingActorDataToViewportClient()
 			SetViewRotation( Actor->GetRootComponent()->RelativeRotation );
 		}
 
-		// If this is a camera actor, then inherit some other settings!
-		UCameraComponent* CameraComponent = Actor->FindComponentByClass<UCameraComponent>();
-		if( CameraComponent != NULL )
+		if( bLockedCameraView )
 		{
-			bUseControllingActorViewInfo = true;
-			CameraComponent->GetCameraView(0.0f, ControllingActorViewInfo);
-
-			//@TODO: CAMERA: Not copying nearly enough information here!
-			ViewFOV = ControllingActorViewInfo.FOV;
-			AspectRatio = ControllingActorViewInfo.AspectRatio;
-			SetViewLocation(ControllingActorViewInfo.Location);
-			SetViewRotation(ControllingActorViewInfo.Rotation);
-
-			// Tell the viewport to use post-process settings from the camera
-			SetPostprocessCameraActor( Cast<ACameraActor>(Actor) );
+			// If this is a camera actor, then inherit some other settings
+			UCameraComponent* CameraComponent = Actor->FindComponentByClass<UCameraComponent>();
+			if( CameraComponent != NULL )
+			{
+				bUseControllingActorViewInfo = true;
+				CameraComponent->GetCameraView(0.0f, ControllingActorViewInfo);
+
+				// Post processing is handled by OverridePostProcessingSettings
+				ViewFOV = ControllingActorViewInfo.FOV;
+				AspectRatio = ControllingActorViewInfo.AspectRatio;
+				SetViewLocation(ControllingActorViewInfo.Location);
+				SetViewRotation(ControllingActorViewInfo.Rotation);
+			}
 		}
 	}
 }
@@ -3939,29 +3928,17 @@ void FLevelEditorViewportClient::CheckHoveredHitProxy( HHitProxy* HoveredHitProx
 
 bool FLevelEditorViewportClient::GetActiveSafeFrame(float& OutAspectRatio) const
 {
-	ACameraActor* LockedCamera = NULL;
-
 	if (!IsOrtho())
 	{
-		ACameraActor* Camera = NULL;
-
-		if (ActorLockedByMatinee.IsValid())
-		{
-			Camera = Cast<ACameraActor>(ActorLockedByMatinee.Get());
-		}
-		else if (ActorLockedToCamera.IsValid())
+		const UCameraComponent* CameraComponent = GetCameraComponentForView();
+		if (CameraComponent && CameraComponent->bConstrainAspectRatio)
 		{
-			Camera = Cast<ACameraActor>(ActorLockedToCamera.Get());
-		}
-
-		if (Camera != NULL && Camera->CameraComponent->bConstrainAspectRatio)
-		{
-			LockedCamera = Camera;
-			OutAspectRatio = Camera->CameraComponent->AspectRatio;
+			OutAspectRatio = CameraComponent->AspectRatio;
+			return true;
 		}
 	}
 
-	return LockedCamera != NULL;
+	return false;
 }
 
 void FLevelEditorViewportClient::SetCurrentWidgetAxis( EAxisList::Type NewAxis )
@@ -4418,18 +4395,6 @@ FLinearColor FLevelEditorViewportClient::GetBackgroundColor() const
 	return IsPerspective() ? GEditor->C_WireBackground : GEditor->C_OrthoBackground;
 }
 
-bool FLevelEditorViewportClient::IsAspectRatioConstrained() const
-{
-	if (PostprocessCameraActor.IsValid())
-	{
-		return PostprocessCameraActor.Get()->CameraComponent->bConstrainAspectRatio;
-	}
-	else
-	{
-		return FEditorViewportClient::IsAspectRatioConstrained();
-	}
-}
-
 int32 FLevelEditorViewportClient::GetCameraSpeedSetting() const
 {
 	return GetDefault<ULevelEditorViewportSettings>()->CameraSpeed;
diff --git a/Engine/Source/Editor/UnrealEd/Public/LevelEditorViewport.h b/Engine/Source/Editor/UnrealEd/Public/LevelEditorViewport.h
--- a/Engine/Source/Editor/UnrealEd/Public/LevelEditorViewport.h
+++ b/Engine/Source/Editor/UnrealEd/Public/LevelEditorViewport.h
@@ -158,7 +158,6 @@ public:
 	virtual void SetupViewForRendering( FSceneViewFamily& ViewFamily, FSceneView& View ) override;
 	virtual FSceneInterface* GetScene() const override;
 	virtual FLinearColor GetBackgroundColor() const override;
-	virtual bool IsAspectRatioConstrained() const override;
 	virtual int32 GetCameraSpeedSetting() const override;
 	virtual void SetCameraSpeedSetting(int32 SpeedSetting) override;
 	virtual void ReceivedFocus(FViewport* Viewport) override;
@@ -218,25 +217,9 @@ public:
 	bool IsVolumeVisibleInViewport( const AActor& VolumeActor ) const;
 
 	/**
-	 * Sets the actor that should "drive" this view's location and other settings
-	 *
-	 * @param	Actor	The actor that will push view settings to this view
-	 */
-	void SetControllingActor( const AActor* Actor )
-	{
-		ControllingActor = Actor;
-	}
-
-	/** Called every time the viewport is ticked to update the view based on an actor that is
-		driving the location, FOV and other settings for this view.  This is used for live camera PIP
-		preview within editor viewports */
-	void PushControllingActorDataToViewportClient();
-
-
-	/**
-	 * Allows to set the postprocess camera actor to unset it (0)
+	 * Updates or resets view properties such as aspect ratio, FOV, location etc to match that of any actor we are locked to
 	 */
-	void SetPostprocessCameraActor(ACameraActor* InPostprocessCameraActor);
+	void UpdateViewForLockedActor();
 
 	/**
 	 * Returns the horizontal axis for this viewport.
@@ -455,6 +438,25 @@ public:
 		}
 		return ActorLockedToCamera;
 	}
+	
+	/** 
+	 * Find the camera component that is driving this viewport, in the following order of preference:
+	 *		1. Matinee locked actor
+	 *		2. User actor lock (if (bLockedCameraView is true)
+	 * 
+	 * @return  Pointer to a camera component to use for this viewport's view
+	 */
+	UCameraComponent* GetCameraComponentForView() const
+	{
+		const AActor* LockedActor = ActorLockedByMatinee.Get();
+
+		if (!LockedActor && bLockedCameraView)
+		{
+			LockedActor = ActorLockedToCamera.Get();
+		}
+
+		return LockedActor ? LockedActor->FindComponentByClass<UCameraComponent>() : nullptr;
+	}
 
 	/** 
 	 * Set the actor lock. This is the actor locked to the viewport via the viewport menus.
@@ -773,21 +775,16 @@ public:
 	 */
 	bool					bIsTrackingBrushModification;
 
+	/** True if this viewport is to change its view (aspect ratio, post processing, FOV etc) to match that of the currently locked camera, if applicable */
+	bool					bLockedCameraView;
+
 private:
 	/** The actors that are currently being placed in the viewport via dragging */
 	static TArray< TWeakObjectPtr< AActor > > DropPreviewActors;
 
-	/** Optional actor that is 'controlling' this view.  That is, the view's location, rotation, FOV and other settings
-		should be pushed to this caInera every frame.  This is used by the editor's live viewport camera PIP feature */
-	TWeakObjectPtr< AActor > ControllingActor;
-
 	/** Bit array representing the visibility of every sprite category in the current viewport */
 	TBitArray<>	SpriteCategoryVisibility;
 
-
-	/** Valid if there is a camera available that can override the post process setting */
-	TWeakObjectPtr<class ACameraActor> PostprocessCameraActor;
-
 	UWorld* World;
 
 	FTrackingTransaction TrackingTransaction;
@@ -799,11 +796,13 @@ private:
 	/** If this view was controlled by another view this/last frame, don't update itself */
 	bool bWasControlledByOtherViewport;
 
-	/** When the viewpoint is locked to an actor this references the actor, invalid if not locked (replaces bLockSelectedToCamera) */
-	TWeakObjectPtr<AActor>	ActorLockedToCamera;
-
-	/** When the viewpoint is locked to an actor (by Matinee) this references the actor, invalid if not locked */
+	/**
+	 * When locked to an actor this view will be positioned in the same location and rotation as the actor.
+	 * If the actor has a camera component the view will also inherit camera settings such as aspect ratio, FOV, post processing settings, and the like.
+	 * A viewport locked to an actor by Matinee will always take precedent over any other.
+	 */
 	TWeakObjectPtr<AActor>	ActorLockedByMatinee;
+	TWeakObjectPtr<AActor>	ActorLockedToCamera;
 
 	/** Data needed to display perframe stat tracking when STAT UNIT is enabled */
 	FStatUnitData*			StatUnitData;
