addition of analytics for File/Cook, File/Package, and Launch On
TTP337717
#ue4
#editor

diff --git a/Engine/Source/Developer/LauncherServices/Private/Launcher/LauncherTask.h b/Engine/Source/Developer/LauncherServices/Private/Launcher/LauncherTask.h
--- a/Engine/Source/Developer/LauncherServices/Private/Launcher/LauncherTask.h
+++ b/Engine/Source/Developer/LauncherServices/Private/Launcher/LauncherTask.h
@@ -297,8 +297,8 @@ private:
 	FRunnableThread* Thread;
 
 	// message delegates
-	FOnStageStartedDelegate TaskStarted;
-	FOnStageCompletedDelegate TaskCompleted;
+	FOnTaskStartedDelegate TaskStarted;
+	FOnTaskCompletedDelegate TaskCompleted;
 
 protected:
 
diff --git a/Engine/Source/Developer/LauncherServices/Private/Launcher/LauncherWorker.cpp b/Engine/Source/Developer/LauncherServices/Private/Launcher/LauncherWorker.cpp
--- a/Engine/Source/Developer/LauncherServices/Private/Launcher/LauncherWorker.cpp
+++ b/Engine/Source/Developer/LauncherServices/Private/Launcher/LauncherWorker.cpp
@@ -30,6 +30,8 @@ uint32 FLauncherWorker::Run( )
 {
 	FString Line;
 
+	LaunchStartTime = FPlatformTime::Seconds();
+
 	// wait for tasks to be completed
 	while (Status == ELauncherWorkerStatus::Busy)
 	{
@@ -101,11 +103,11 @@ uint32 FLauncherWorker::Run( )
 	if (Status == ELauncherWorkerStatus::Canceling)
 	{
 		Status = ELauncherWorkerStatus::Canceled;
-		LaunchCanceled.Broadcast();
+		LaunchCanceled.Broadcast(FPlatformTime::Seconds() - LaunchStartTime);
 	}
 	else
 	{
-		LaunchCompleted.Broadcast(TaskChain->Succeeded());
+		LaunchCompleted.Broadcast(TaskChain->Succeeded(), FPlatformTime::Seconds() - LaunchStartTime);
 	}
 
 	FPlatformProcess::ClosePipe(ReadPipe, WritePipe);
@@ -165,13 +167,14 @@ int32 FLauncherWorker::GetTasks( TArray<ILauncherTaskPtr>& OutTasks ) const
 
 void FLauncherWorker::OnTaskStarted(const FString& TaskName)
 {
+	StageStartTime = FPlatformTime::Seconds();
 	StageStarted.Broadcast(TaskName);
 }
 
 
 void FLauncherWorker::OnTaskCompleted(const FString& TaskName)
 {
-	StageCompleted.Broadcast(TaskName);
+	StageCompleted.Broadcast(TaskName, FPlatformTime::Seconds() - StageStartTime);
 }
 
 
diff --git a/Engine/Source/Developer/LauncherServices/Private/Launcher/LauncherWorker.h b/Engine/Source/Developer/LauncherServices/Private/Launcher/LauncherWorker.h
--- a/Engine/Source/Developer/LauncherServices/Private/Launcher/LauncherWorker.h
+++ b/Engine/Source/Developer/LauncherServices/Private/Launcher/LauncherWorker.h
@@ -105,6 +105,9 @@ private:
 	void* ReadPipe;
 	void* WritePipe;
 
+	double StageStartTime;
+	double LaunchStartTime;
+
 	// message delegate
 	FOutputMessageReceivedDelegate OutputMessageReceived;
 	FOnStageStartedDelegate StageStarted;
diff --git a/Engine/Source/Developer/LauncherServices/Public/Interfaces/ILauncherWorker.h b/Engine/Source/Developer/LauncherServices/Public/Interfaces/ILauncherWorker.h
--- a/Engine/Source/Developer/LauncherServices/Public/Interfaces/ILauncherWorker.h
+++ b/Engine/Source/Developer/LauncherServices/Public/Interfaces/ILauncherWorker.h
@@ -32,13 +32,13 @@ DECLARE_MULTICAST_DELEGATE_OneParam(FOutputMessageReceivedDelegate, const FStrin
 DECLARE_MULTICAST_DELEGATE_OneParam(FOnStageStartedDelegate, const FString&);
 
 /** Delegate used to notify when a stage ends */
-DECLARE_MULTICAST_DELEGATE_OneParam(FOnStageCompletedDelegate, const FString&);
+DECLARE_MULTICAST_DELEGATE_TwoParams(FOnStageCompletedDelegate, const FString&, double);
 
 /** Delegate used to notify when the launch is complete */
-DECLARE_MULTICAST_DELEGATE_OneParam(FOnLaunchCompletedDelegate, bool);
+DECLARE_MULTICAST_DELEGATE_TwoParams(FOnLaunchCompletedDelegate, bool, double);
 
 /** Delegate used to notify when the launch was canceled */
-DECLARE_MULTICAST_DELEGATE(FOnLaunchCanceledDelegate);
+DECLARE_MULTICAST_DELEGATE_OneParam(FOnLaunchCanceledDelegate, double);
 
 
 /**
diff --git a/Engine/Source/Editor/MainFrame/MainFrame.Build.cs b/Engine/Source/Editor/MainFrame/MainFrame.Build.cs
--- a/Engine/Source/Editor/MainFrame/MainFrame.Build.cs
+++ b/Engine/Source/Editor/MainFrame/MainFrame.Build.cs
@@ -20,6 +20,7 @@ public class MainFrame : ModuleRules
 				"CoreUObject",
 				"CrashTracker",
 				"Engine",
+				"EngineSettings",
                 "InputCore",
 				"RHI",
 				"ShaderCore",
diff --git a/Engine/Source/Editor/MainFrame/Private/Frame/MainFrameActions.cpp b/Engine/Source/Editor/MainFrame/Private/Frame/MainFrameActions.cpp
--- a/Engine/Source/Editor/MainFrame/Private/Frame/MainFrameActions.cpp
+++ b/Engine/Source/Editor/MainFrame/Private/Frame/MainFrameActions.cpp
@@ -993,13 +993,32 @@ void FMainFrameActionCallbacks::CreateUatTask( const FString& CommandLine, const
 		return;
 	}
 
+	FGameProjectGenerationModule& GameProjectModule = FModuleManager::LoadModuleChecked<FGameProjectGenerationModule>(TEXT("GameProjectGeneration"));
+	bool bProjectHasCode = GameProjectModule.Get().GetProjectCodeFileCount() > 0;
+
+	FString EventName = (CommandLine.Contains(TEXT("-package")) ? TEXT("Editor.Package") : TEXT("Editor.Cook"));
+	if( FEngineAnalytics::IsAvailable() )
+	{
+		const UGeneralProjectSettings& ProjectSettings = *GetDefault<UGeneralProjectSettings>();
+		TArray<FAnalyticsEventAttribute> ParamArray;
+		ParamArray.Add(FAnalyticsEventAttribute(TEXT("ProjectID"), ProjectSettings.ProjectID.ToString()));
+		ParamArray.Add(FAnalyticsEventAttribute(TEXT("Platform"), PlatformDisplayName.ToString()));
+		ParamArray.Add(FAnalyticsEventAttribute(TEXT("ProjectType"), bProjectHasCode ? TEXT("C++ Code") : TEXT("Content Only")));
+
+		FEngineAnalytics::GetProvider().RecordEvent( EventName + TEXT(".Start"), ParamArray );
+	}
+
 	NotificationItem->SetCompletionState(SNotificationItem::CS_Pending);
 
 	// launch the packager
 	TWeakPtr<SNotificationItem> NotificationItemPtr(NotificationItem);
 
-	UatProcess->OnCanceled().BindStatic(&FMainFrameActionCallbacks::HandleUatProcessCanceled, NotificationItemPtr, PlatformDisplayName, TaskShortName);
-	UatProcess->OnCompleted().BindStatic(&FMainFrameActionCallbacks::HandleUatProcessCompleted, NotificationItemPtr, PlatformDisplayName, TaskShortName);
+	EventData Data;
+	Data.StartTime = FPlatformTime::Seconds();
+	Data.EventName = EventName;
+	Data.bProjectHasCode = bProjectHasCode;
+	UatProcess->OnCanceled().BindStatic(&FMainFrameActionCallbacks::HandleUatProcessCanceled, NotificationItemPtr, PlatformDisplayName, TaskShortName, Data);
+	UatProcess->OnCompleted().BindStatic(&FMainFrameActionCallbacks::HandleUatProcessCompleted, NotificationItemPtr, PlatformDisplayName, TaskShortName, Data);
 	UatProcess->OnOutput().BindStatic(&FMainFrameActionCallbacks::HandleUatProcessOutput, NotificationItemPtr, PlatformDisplayName, TaskShortName);
 
 	if (UatProcess->Launch())
@@ -1080,7 +1099,7 @@ void FMainFrameActionCallbacks::HandleUatCancelButtonClicked( TSharedPtr<FMonito
 }
 
 
-void FMainFrameActionCallbacks::HandleUatProcessCanceled( TWeakPtr<SNotificationItem> NotificationItemPtr, FText PlatformDisplayName, FText TaskName )
+void FMainFrameActionCallbacks::HandleUatProcessCanceled( TWeakPtr<SNotificationItem> NotificationItemPtr, FText PlatformDisplayName, FText TaskName, EventData Event )
 {
 	FFormatNamedArguments Arguments;
 	Arguments.Add(TEXT("Platform"), PlatformDisplayName);
@@ -1092,11 +1111,22 @@ void FMainFrameActionCallbacks::HandleUatProcessCanceled( TWeakPtr<SNotification
 		FText::Format(LOCTEXT("UatProcessFailedNotification", "{TaskName} canceled!"), Arguments)
 	);
 
+	if( FEngineAnalytics::IsAvailable() )
+	{
+		const UGeneralProjectSettings& ProjectSettings = *GetDefault<UGeneralProjectSettings>();
+		TArray<FAnalyticsEventAttribute> ParamArray;
+		ParamArray.Add(FAnalyticsEventAttribute(TEXT("ProjectID"), ProjectSettings.ProjectID.ToString()));
+		ParamArray.Add(FAnalyticsEventAttribute(TEXT("Platform"), PlatformDisplayName.ToString()));
+		ParamArray.Add(FAnalyticsEventAttribute(TEXT("ProjectType"), Event.bProjectHasCode ? TEXT("C++ Code") : TEXT("Content Only")));
+		ParamArray.Add(FAnalyticsEventAttribute(TEXT("Time"), FPlatformTime::Seconds() - Event.StartTime));
+
+		FEngineAnalytics::GetProvider().RecordEvent( Event.EventName + TEXT(".Canceled"), ParamArray );
+	}
 //	FMessageLog("PackagingResults").Warning(FText::Format(LOCTEXT("UatProcessCanceledMessageLog", "{TaskName} for {Platform} canceled by user"), Arguments));
 }
 
 
-void FMainFrameActionCallbacks::HandleUatProcessCompleted( int32 ReturnCode, TWeakPtr<SNotificationItem> NotificationItemPtr, FText PlatformDisplayName, FText TaskName )
+void FMainFrameActionCallbacks::HandleUatProcessCompleted( int32 ReturnCode, TWeakPtr<SNotificationItem> NotificationItemPtr, FText PlatformDisplayName, FText TaskName, EventData Event )
 {
 	FFormatNamedArguments Arguments;
 	Arguments.Add(TEXT("Platform"), PlatformDisplayName);
@@ -1110,6 +1140,18 @@ void FMainFrameActionCallbacks::HandleUatProcessCompleted( int32 ReturnCode, TWe
 			FText::Format(LOCTEXT("UatProcessSucceededNotification", "{TaskName} complete!"), Arguments)
 		);
 
+		if( FEngineAnalytics::IsAvailable() )
+		{
+			const UGeneralProjectSettings& ProjectSettings = *GetDefault<UGeneralProjectSettings>();
+			TArray<FAnalyticsEventAttribute> ParamArray;
+			ParamArray.Add(FAnalyticsEventAttribute(TEXT("ProjectID"), ProjectSettings.ProjectID.ToString()));
+			ParamArray.Add(FAnalyticsEventAttribute(TEXT("Platform"), PlatformDisplayName.ToString()));
+			ParamArray.Add(FAnalyticsEventAttribute(TEXT("ProjectType"), Event.bProjectHasCode ? TEXT("C++ Code") : TEXT("Content Only")));
+			ParamArray.Add(FAnalyticsEventAttribute(TEXT("Time"), FPlatformTime::Seconds() - Event.StartTime));
+
+			FEngineAnalytics::GetProvider().RecordEvent( Event.EventName + TEXT(".Completed"), ParamArray );
+		}
+
 //		FMessageLog("PackagingResults").Info(FText::Format(LOCTEXT("UatProcessSuccessMessageLog", "{TaskName} for {Platform} completed successfully"), Arguments));
 	}
 	else
@@ -1120,6 +1162,18 @@ void FMainFrameActionCallbacks::HandleUatProcessCompleted( int32 ReturnCode, TWe
 			FText::Format(LOCTEXT("PackagerFailedNotification", "{TaskName} failed!"), Arguments)
 		);
 
+		if( FEngineAnalytics::IsAvailable() )
+		{
+			const UGeneralProjectSettings& ProjectSettings = *GetDefault<UGeneralProjectSettings>();
+			TArray<FAnalyticsEventAttribute> ParamArray;
+			ParamArray.Add(FAnalyticsEventAttribute(TEXT("ProjectID"), ProjectSettings.ProjectID.ToString()));
+			ParamArray.Add(FAnalyticsEventAttribute(TEXT("Platform"), PlatformDisplayName.ToString()));
+			ParamArray.Add(FAnalyticsEventAttribute(TEXT("ProjectType"), Event.bProjectHasCode ? TEXT("C++ Code") : TEXT("Content Only")));
+			ParamArray.Add(FAnalyticsEventAttribute(TEXT("Time"), FPlatformTime::Seconds() - Event.StartTime));
+
+			FEngineAnalytics::GetProvider().RecordEvent( Event.EventName + TEXT(".Completed"), ParamArray );
+		}
+
 //		FMessageLog("PackagingResults").Info(FText::Format(LOCTEXT("UatProcessFailedMessageLog", "{TaskName} for {Platform} failed"), Arguments));
 	}
 }
diff --git a/Engine/Source/Editor/MainFrame/Private/Frame/MainFrameActions.h b/Engine/Source/Editor/MainFrame/Private/Frame/MainFrameActions.h
--- a/Engine/Source/Editor/MainFrame/Private/Frame/MainFrameActions.h
+++ b/Engine/Source/Editor/MainFrame/Private/Frame/MainFrameActions.h
@@ -247,6 +247,13 @@ protected:
 
 private:
 
+	struct EventData
+	{
+		FString EventName;
+		bool bProjectHasCode;
+		double StartTime;
+	};
+
 	// Handles clicking the packager notification item's Cancel button.
 	static void HandleUatCancelButtonClicked( TSharedPtr<FMonitoredProcess> PackagerProcess );
 
@@ -254,10 +261,10 @@ private:
 	static void HandleUatHyperlinkNavigate( );
 
 	// Handles canceled packager processes.
-	static void HandleUatProcessCanceled( TWeakPtr<class SNotificationItem> NotificationItemPtr, FText PlatformDisplayName, FText TaskName );
+	static void HandleUatProcessCanceled( TWeakPtr<class SNotificationItem> NotificationItemPtr, FText PlatformDisplayName, FText TaskName, EventData Event );
 
 	// Handles the completion of a packager process.
-	static void HandleUatProcessCompleted( int32 ReturnCode, TWeakPtr<class SNotificationItem> NotificationItemPtr, FText PlatformDisplayName, FText TaskName );
+	static void HandleUatProcessCompleted( int32 ReturnCode, TWeakPtr<class SNotificationItem> NotificationItemPtr, FText PlatformDisplayName, FText TaskName, EventData Event );
 
 	// Handles packager process output.
 	static void HandleUatProcessOutput( FString Output, TWeakPtr<class SNotificationItem> NotificationItemPtr, FText PlatformDisplayName, FText TaskName );
diff --git a/Engine/Source/Editor/UnrealEd/Classes/Editor/EditorEngine.h b/Engine/Source/Editor/UnrealEd/Classes/Editor/EditorEngine.h
--- a/Engine/Source/Editor/UnrealEd/Classes/Editor/EditorEngine.h
+++ b/Engine/Source/Editor/UnrealEd/Classes/Editor/EditorEngine.h
@@ -2585,6 +2585,12 @@ protected:
 
 	/** Destroy any online subsystems generated by PIE */
 	void CleanupPIEOnlineSessions(TArray<FName> OnlineIdentifiers);
+
+	// launch on callbacks
+	void HandleStageStarted(const FString& InStage, TWeakPtr<SNotificationItem> NotificationItemPtr);
+	void HandleStageCompleted(const FString& InStage, double StageTime, bool bHasCode, TWeakPtr<SNotificationItem> NotificationItemPtr);
+	void HandleLaunchCanceled(double TotalTime, bool bHasCode, TWeakPtr<SNotificationItem> NotificationItemPtr);
+	void HandleLaunchCompleted(bool Succeeded, double TotalTime, bool bHasCode, TWeakPtr<SNotificationItem> NotificationItemPtr);
 };
 
 
diff --git a/Engine/Source/Editor/UnrealEd/Private/PlayLevel.cpp b/Engine/Source/Editor/UnrealEd/Private/PlayLevel.cpp
--- a/Engine/Source/Editor/UnrealEd/Private/PlayLevel.cpp
+++ b/Engine/Source/Editor/UnrealEd/Private/PlayLevel.cpp
@@ -27,6 +27,8 @@
 #include "GameProjectGenerationModule.h"
 #include "SourceCodeNavigation.h"
 #include "PhysicsPublic.h"
+#include "Runtime/Analytics/Analytics/Public/Interfaces/IAnalyticsProvider.h"
+#include "EngineAnalytics.h"
 
 #include "Online.h"
 
@@ -1133,14 +1135,14 @@ private:
 };
 
 
-static void HandleStageStarted(const FString& InStage, FString DeviceName, TWeakPtr<SNotificationItem> NotificationItemPtr)
+void UEditorEngine::HandleStageStarted(const FString& InStage, TWeakPtr<SNotificationItem> NotificationItemPtr)
 {
 	FFormatNamedArguments Arguments;
 	Arguments.Add(TEXT("StageName"), FText::FromString(InStage));
-	Arguments.Add(TEXT("DeviceName"), FText::FromString(DeviceName));
+	Arguments.Add(TEXT("DeviceName"), FText::FromString(PlayUsingLauncherDeviceName));
 
 	FText NotificationText;
-	if(DeviceName.Len() == 0)
+	if(PlayUsingLauncherDeviceName.Len() == 0)
 	{
 		NotificationText = FText::Format(LOCTEXT("LauncherTaskStageNotificationNoDevice", "{StageName} for launch..."), Arguments);
 	}
@@ -1151,7 +1153,24 @@ static void HandleStageStarted(const FString& InStage, FString DeviceName, TWeak
 	NotificationItemPtr.Pin()->SetText(NotificationText);
 }
 
-static void HandleLaunchCanceled(bool* bPlayUsingLauncher, TWeakPtr<SNotificationItem> NotificationItemPtr)
+void UEditorEngine::HandleStageCompleted(const FString& InStage, double StageTime, bool bHasCode, TWeakPtr<SNotificationItem> NotificationItemPtr)
+{
+	// analytics for launch on
+	if( FEngineAnalytics::IsAvailable() )
+	{
+		const UGeneralProjectSettings& ProjectSettings = *GetDefault<UGeneralProjectSettings>();
+		TArray<FAnalyticsEventAttribute> ParamArray;
+		ParamArray.Add(FAnalyticsEventAttribute(TEXT("ProjectID"), ProjectSettings.ProjectID.ToString()));
+		ParamArray.Add(FAnalyticsEventAttribute(TEXT("Platform"), PlayUsingLauncherDeviceId.Left(PlayUsingLauncherDeviceId.Find(TEXT("@")))));
+		ParamArray.Add(FAnalyticsEventAttribute(TEXT("ProjectType"), bHasCode ? TEXT("C++ Code") : TEXT("Content Only")));
+		ParamArray.Add(FAnalyticsEventAttribute(TEXT("Time"), StageTime));
+		ParamArray.Add(FAnalyticsEventAttribute(TEXT("StageName"), InStage));
+
+		FEngineAnalytics::GetProvider().RecordEvent( TEXT( "Editor.LaunchOn.StageComplete" ), ParamArray );
+	}
+}
+
+void UEditorEngine::HandleLaunchCanceled(double TotalTime, bool bHasCode, TWeakPtr<SNotificationItem> NotificationItemPtr)
 {
 	TGraphTask<FLauncherNotificationTask>::CreateTask().ConstructAndDispatchWhenReady(
 		NotificationItemPtr,
@@ -1159,10 +1178,23 @@ static void HandleLaunchCanceled(bool* bPlayUsingLauncher, TWeakPtr<SNotificatio
 		LOCTEXT("LaunchtaskFailedNotification", "Launch canceled!")
 	);
 
-	*bPlayUsingLauncher = false;	
+	// analytics for launch on
+	if( FEngineAnalytics::IsAvailable() )
+	{
+		const UGeneralProjectSettings& ProjectSettings = *GetDefault<UGeneralProjectSettings>();
+		TArray<FAnalyticsEventAttribute> ParamArray;
+		ParamArray.Add(FAnalyticsEventAttribute(TEXT("ProjectID"), ProjectSettings.ProjectID.ToString()));
+		ParamArray.Add(FAnalyticsEventAttribute(TEXT("Platform"), PlayUsingLauncherDeviceId.Left(PlayUsingLauncherDeviceId.Find(TEXT("@")))));
+		ParamArray.Add(FAnalyticsEventAttribute(TEXT("ProjectType"), bHasCode ? TEXT("C++ Code") : TEXT("Content Only")));
+		ParamArray.Add(FAnalyticsEventAttribute(TEXT("Time"), TotalTime));
+
+		FEngineAnalytics::GetProvider().RecordEvent( TEXT( "Editor.LaunchOn.Canceled" ), ParamArray );
+	}
+
+	bPlayUsingLauncher = false;	
 }
 
-static void HandleLaunchCompleted(bool Succeeded, bool* bPlayUsingLauncher, TWeakPtr<SNotificationItem> NotificationItemPtr)
+void UEditorEngine::HandleLaunchCompleted(bool Succeeded, double TotalTime, bool bHasCode, TWeakPtr<SNotificationItem> NotificationItemPtr)
 {
 	if (Succeeded)
 	{
@@ -1171,6 +1203,19 @@ static void HandleLaunchCompleted(bool Succeeded, bool* bPlayUsingLauncher, TWea
 			SNotificationItem::CS_Success,
 			LOCTEXT("LauncherTaskCompleted", "Launch complete!!")
 		);
+
+		// analytics for launch on
+		if( FEngineAnalytics::IsAvailable() )
+		{
+			const UGeneralProjectSettings& ProjectSettings = *GetDefault<UGeneralProjectSettings>();
+			TArray<FAnalyticsEventAttribute> ParamArray;
+			ParamArray.Add(FAnalyticsEventAttribute(TEXT("ProjectID"), ProjectSettings.ProjectID.ToString()));
+			ParamArray.Add(FAnalyticsEventAttribute(TEXT("Platform"), PlayUsingLauncherDeviceId.Left(PlayUsingLauncherDeviceId.Find(TEXT("@")))));
+			ParamArray.Add(FAnalyticsEventAttribute(TEXT("ProjectType"), bHasCode ? TEXT("C++ Code") : TEXT("Content Only")));
+			ParamArray.Add(FAnalyticsEventAttribute(TEXT("Time"), TotalTime));
+
+			FEngineAnalytics::GetProvider().RecordEvent( TEXT( "Editor.LaunchOn.Completed" ), ParamArray );
+		}
 	}
 	else
 	{
@@ -1179,8 +1224,21 @@ static void HandleLaunchCompleted(bool Succeeded, bool* bPlayUsingLauncher, TWea
 			SNotificationItem::CS_Fail,
 			LOCTEXT("LauncherTaskFailed", "Launch failed!!")
 			);
+
+		// analytics for launch on
+		if( FEngineAnalytics::IsAvailable() )
+		{
+			const UGeneralProjectSettings& ProjectSettings = *GetDefault<UGeneralProjectSettings>();
+			TArray<FAnalyticsEventAttribute> ParamArray;
+			ParamArray.Add(FAnalyticsEventAttribute(TEXT("ProjectID"), ProjectSettings.ProjectID.ToString()));
+			ParamArray.Add(FAnalyticsEventAttribute(TEXT("Platform"), PlayUsingLauncherDeviceId.Left(PlayUsingLauncherDeviceId.Find(TEXT("@")))));
+			ParamArray.Add(FAnalyticsEventAttribute(TEXT("ProjectType"), bHasCode ? TEXT("C++ Code") : TEXT("Content Only")));
+			ParamArray.Add(FAnalyticsEventAttribute(TEXT("Time"), TotalTime));
+
+			FEngineAnalytics::GetProvider().RecordEvent( TEXT( "Editor.LaunchOn.Failed" ), ParamArray );
+		}
 	}
-	*bPlayUsingLauncher = false;
+	bPlayUsingLauncher = false;
 }
 
 static void HandleHyperlinkNavigate()
@@ -1303,6 +1361,18 @@ void UEditorEngine::PlayUsingLauncher()
 			return;
 		}
 
+		// analytics for launch on
+		if( FEngineAnalytics::IsAvailable() )
+		{
+			const UGeneralProjectSettings& ProjectSettings = *GetDefault<UGeneralProjectSettings>();
+			TArray<FAnalyticsEventAttribute> ParamArray;
+			ParamArray.Add(FAnalyticsEventAttribute(TEXT("ProjectID"), ProjectSettings.ProjectID.ToString()));
+			ParamArray.Add(FAnalyticsEventAttribute(TEXT("Platform"), PlayUsingLauncherDeviceId.Left(PlayUsingLauncherDeviceId.Find(TEXT("@")))));
+			ParamArray.Add(FAnalyticsEventAttribute(TEXT("ProjectType"), bHasCode ? TEXT("C++ Code") : TEXT("Content Only")));
+
+			FEngineAnalytics::GetProvider().RecordEvent( TEXT( "Editor.LaunchOn.Started" ), ParamArray );
+		}
+
 		NotificationItem->SetCompletionState(SNotificationItem::CS_Pending);
 
 		TWeakPtr<SNotificationItem> NotificationItemPtr(NotificationItem);
@@ -1310,9 +1380,10 @@ void UEditorEngine::PlayUsingLauncher()
 		{
 			GEditor->PlayEditorSound(TEXT("/Engine/EditorSounds/Notifications/CompileStart_Cue.CompileStart_Cue"));
 			LauncherWorker->OnOutputReceived().AddStatic(HandleOutputReceived);
-			LauncherWorker->OnStageStarted().AddStatic(HandleStageStarted, PlayUsingLauncherDeviceName, NotificationItemPtr);
-			LauncherWorker->OnCompleted().AddStatic(HandleLaunchCompleted, &bPlayUsingLauncher, NotificationItemPtr);
-			LauncherWorker->OnCanceled().AddStatic(HandleLaunchCanceled, &bPlayUsingLauncher, NotificationItemPtr);
+			LauncherWorker->OnStageStarted().AddUObject(this, &UEditorEngine::HandleStageStarted, NotificationItemPtr);
+			LauncherWorker->OnStageCompleted().AddUObject(this, &UEditorEngine::HandleStageCompleted, bHasCode, NotificationItemPtr);
+			LauncherWorker->OnCompleted().AddUObject(this, &UEditorEngine::HandleLaunchCompleted, bHasCode, NotificationItemPtr);
+			LauncherWorker->OnCanceled().AddUObject(this, &UEditorEngine::HandleLaunchCanceled, bHasCode, NotificationItemPtr);
 		}
 		else
 		{
@@ -1323,6 +1394,19 @@ void UEditorEngine::PlayUsingLauncher()
 			NotificationItem->SetCompletionState(SNotificationItem::CS_Fail);
 			NotificationItem->ExpireAndFadeout();
 			bPlayUsingLauncher = false;
+
+			// analytics for launch on
+			if( FEngineAnalytics::IsAvailable() )
+			{
+				const UGeneralProjectSettings& ProjectSettings = *GetDefault<UGeneralProjectSettings>();
+				TArray<FAnalyticsEventAttribute> ParamArray;
+				ParamArray.Add(FAnalyticsEventAttribute(TEXT("ProjectID"), ProjectSettings.ProjectID.ToString()));
+				ParamArray.Add(FAnalyticsEventAttribute(TEXT("Platform"), PlayUsingLauncherDeviceId.Left(PlayUsingLauncherDeviceId.Find(TEXT("@")))));
+				ParamArray.Add(FAnalyticsEventAttribute(TEXT("ProjectType"), bHasCode ? TEXT("C++ Code") : TEXT("Content Only")));
+				ParamArray.Add(FAnalyticsEventAttribute(TEXT("Time"), 0.0f));
+
+				FEngineAnalytics::GetProvider().RecordEvent( TEXT( "Editor.LaunchOn.Failed" ), ParamArray );
+			}
 		}
 	}
 }
