Fixed material flattening when material has static switch parameters

diff --git a/Engine/Source/Editor/UnrealEd/Private/EditorExporters.cpp b/Engine/Source/Editor/UnrealEd/Private/EditorExporters.cpp
--- a/Engine/Source/Editor/UnrealEd/Private/EditorExporters.cpp
+++ b/Engine/Source/Editor/UnrealEd/Private/EditorExporters.cpp
@@ -699,21 +699,57 @@ class FExportMaterialProxy : public FMaterial, public FMaterialRenderProxy
 {
 public:
 	FExportMaterialProxy()
-	: FMaterial()
+		: FMaterial()
 	{
 		SetQualityLevelProperties(EMaterialQualityLevel::High, false, GRHIFeatureLevel);
 	}
 
 	FExportMaterialProxy(UMaterialInterface* InMaterialInterface, EMaterialProperty InPropertyToCompile)
-	: FMaterial()
-	, MaterialInterface(InMaterialInterface)
-	, PropertyToCompile(InPropertyToCompile)
+		: FMaterial()
+		, MaterialInterface(InMaterialInterface)
+		, PropertyToCompile(InPropertyToCompile)
 	{
 		SetQualityLevelProperties(EMaterialQualityLevel::High, false, GRHIFeatureLevel);
 		Material = InMaterialInterface->GetMaterial();
 		Material->AppendReferencedTextures(ReferencedTextures);
 		FPlatformMisc::CreateGuid(Id);
-		CacheShaders(GRHIShaderPlatform, true);
+		
+		// Have to properly handle compilation of static switches in MaterialInstance* cases...
+		UMaterialInstance* MaterialInstance = Cast<UMaterialInstance>(InMaterialInterface);
+
+		// Walk up the parent chain until we find the first MI with static parameters
+		while (MaterialInstance 
+			&& !MaterialInstance->bHasStaticPermutationResource
+			&& MaterialInstance->Parent 
+			&& MaterialInstance->Parent->IsA<UMaterialInstance>())
+		{
+			MaterialInstance = Cast<UMaterialInstance>(MaterialInstance->Parent);
+		}
+
+		// Special path for a MI with static parameters
+		if (MaterialInstance && MaterialInstance->bHasStaticPermutationResource && MaterialInstance->Parent)
+		{
+			FMaterialResource* MIResource = MaterialInstance->GetMaterialResource(GRHIFeatureLevel);
+
+			// Use the shader map Id from the static permutation
+			// This allows us to create a deterministic yet unique Id for the shader map that will be compiled for this FLightmassMaterialProxy
+			FMaterialShaderMapId ResourceId;
+			//@todo - always use highest quality level for static lighting
+			MaterialInstance->GetMaterialResourceId(GRHIShaderPlatform, EMaterialQualityLevel::Num, ResourceId);
+
+			// Override with a special usage so we won't re-use the shader map used by the MI for rendering
+			ResourceId.Usage = GetShaderMapUsage();
+
+			CacheShaders(ResourceId, GRHIShaderPlatform, true);
+		}
+		else
+		{
+			FMaterialResource* MaterialResource = Material->GetMaterialResource(GRHIFeatureLevel);
+			
+			// Copy the material resource Id
+			// The FLightmassMaterialProxy's GetShaderMapUsage will set it apart from the MI's resource when it comes to finding a shader map
+			CacheShaders(GRHIShaderPlatform, true);
+		}
 	}
 
 	/** This override is required otherwise the shaders aren't ready for use when the surface is rendered resulting in a blank image */
@@ -754,17 +790,17 @@ public:
 		}
 	}
 
-	virtual bool GetVectorValue(const FName ParameterName, FLinearColor* OutValue, const FMaterialRenderContext& Context) const
+	virtual bool GetVectorValue(const FName ParameterName, FLinearColor* OutValue, const FMaterialRenderContext& Context) const override
 	{
 		return MaterialInterface->GetRenderProxy(0)->GetVectorValue(ParameterName, OutValue, Context);
 	}
 
-	virtual bool GetScalarValue(const FName ParameterName, float* OutValue, const FMaterialRenderContext& Context) const
+	virtual bool GetScalarValue(const FName ParameterName, float* OutValue, const FMaterialRenderContext& Context) const override
 	{
 		return MaterialInterface->GetRenderProxy(0)->GetScalarValue(ParameterName, OutValue, Context);
 	}
 
-	virtual bool GetTextureValue(const FName ParameterName,const UTexture** OutValue, const FMaterialRenderContext& Context) const
+	virtual bool GetTextureValue(const FName ParameterName,const UTexture** OutValue, const FMaterialRenderContext& Context) const override
 	{
 		return MaterialInterface->GetRenderProxy(0)->GetTextureValue(ParameterName,OutValue,Context);
 	}
