More robust BP reinstancing. More broken/corrupted/stale blueprint assets can be loaded without any crash. ttp338193

diff --git a/Engine/Source/Editor/UnrealEd/Private/Kismet2/KismetReinstanceUtilities.cpp b/Engine/Source/Editor/UnrealEd/Private/Kismet2/KismetReinstanceUtilities.cpp
--- a/Engine/Source/Editor/UnrealEd/Private/Kismet2/KismetReinstanceUtilities.cpp
+++ b/Engine/Source/Editor/UnrealEd/Private/Kismet2/KismetReinstanceUtilities.cpp
@@ -508,8 +508,15 @@ void FBlueprintCompileReinstancer::ReparentChild(UBlueprint* ChildBP)
 
 void FBlueprintCompileReinstancer::ReparentChild(UClass* ChildClass)
 {
-	check(ChildClass);
-	check(ChildClass->GetSuperClass() == ClassToReinstance || ChildClass->GetSuperClass() == DuplicatedClass);
+	check(ChildClass && ClassToReinstance && DuplicatedClass && ChildClass->GetSuperClass());
+	bool bIsReallyAChild = ChildClass->GetSuperClass() == ClassToReinstance || ChildClass->GetSuperClass() == DuplicatedClass;
+	const auto SuperClassBP = Cast<UBlueprint>(ChildClass->GetSuperClass()->ClassGeneratedBy);
+	if (SuperClassBP && !bIsReallyAChild)
+	{
+		bIsReallyAChild |= (SuperClassBP->SkeletonGeneratedClass == ClassToReinstance) || (SuperClassBP->SkeletonGeneratedClass == DuplicatedClass);
+		bIsReallyAChild |= (SuperClassBP->GeneratedClass == ClassToReinstance) || (SuperClassBP->GeneratedClass == DuplicatedClass);
+	}
+	check(bIsReallyAChild);
 
 	ChildClass->AssembleReferenceTokenStream();
 	ChildClass->SetSuperStruct(DuplicatedClass);
