BP: Name conflict between Local Variable and function parameter.

diff --git a/Engine/Source/Editor/Kismet/Private/BlueprintEditor.cpp b/Engine/Source/Editor/Kismet/Private/BlueprintEditor.cpp
--- a/Engine/Source/Editor/Kismet/Private/BlueprintEditor.cpp
+++ b/Engine/Source/Editor/Kismet/Private/BlueprintEditor.cpp
@@ -3072,30 +3072,6 @@ bool FBlueprintEditor::CanChangePinType() const
 	return false;
 }
 
-struct FFunctionFromNodeHelper
-{
-	UFunction* Function;
-	UK2Node* Node;
-	FFunctionFromNodeHelper(UObject* SelectedObj) : Function(NULL), Node(NULL)
-	{
-		Node = Cast<UK2Node>(SelectedObj);
-		UBlueprint* Blueprint = Node ? Node->GetBlueprint() : NULL;
-		const UClass* SearchScope = Blueprint ? Blueprint->SkeletonGeneratedClass : NULL;
-		if(SearchScope)
-		{
-			if (UK2Node_Event* EventNode = Cast<UK2Node_Event>(Node))
-			{
-				Function = SearchScope->FindFunctionByName(EventNode->EventSignatureName);
-			}
-			else if(UK2Node_FunctionEntry* FunctionNode = Cast<UK2Node_FunctionEntry>(Node))
-			{
-				const FName FunctionName = (FunctionNode->CustomGeneratedFunctionName != NAME_None) ? FunctionNode->CustomGeneratedFunctionName : FunctionNode->GetGraph()->GetFName();
-				Function = SearchScope->FindFunctionByName(FunctionName);
-			}
-		}
-	}
-};
-
 void FBlueprintEditor::OnAddParentNode()
 {
 	const UEdGraphSchema_K2* Schema = GetDefault<UEdGraphSchema_K2>();
diff --git a/Engine/Source/Editor/UnrealEd/Private/Kismet2/BlueprintEditorUtils.cpp b/Engine/Source/Editor/UnrealEd/Private/Kismet2/BlueprintEditorUtils.cpp
--- a/Engine/Source/Editor/UnrealEd/Private/Kismet2/BlueprintEditorUtils.cpp
+++ b/Engine/Source/Editor/UnrealEd/Private/Kismet2/BlueprintEditorUtils.cpp
@@ -3331,13 +3331,46 @@ void FBlueprintEditorUtils::RemoveLocalVariable(UBlueprint* InBlueprint, const F
 	}
 }
 
+FFunctionFromNodeHelper::FFunctionFromNodeHelper(UObject* Obj) : Function(FunctionFromNode(Cast<UK2Node>(Obj))), Node(Cast<UK2Node>(Obj))
+{
+
+}
+
+UFunction* FFunctionFromNodeHelper::FunctionFromNode(UK2Node* Node)
+{
+	UFunction* Function = NULL;
+	UBlueprint* Blueprint = Node ? Node->GetBlueprint() : NULL;
+	const UClass* SearchScope = Blueprint ? Blueprint->SkeletonGeneratedClass : NULL;
+	if (SearchScope)
+	{
+		if (UK2Node_Event* EventNode = Cast<UK2Node_Event>(Node))
+		{
+			Function = SearchScope->FindFunctionByName(EventNode->EventSignatureName);
+		}
+		else if (UK2Node_FunctionEntry* FunctionNode = Cast<UK2Node_FunctionEntry>(Node))
+		{
+			const FName FunctionName = (FunctionNode->CustomGeneratedFunctionName != NAME_None) ? FunctionNode->CustomGeneratedFunctionName : FunctionNode->GetGraph()->GetFName();
+			Function = SearchScope->FindFunctionByName(FunctionName);
+		}
+	}
+
+	return Function;
+}
+
 void FBlueprintEditorUtils::RenameLocalVariable(UBlueprint* InBlueprint, const FName& InOldName, const FName& InNewName)
 {
 	if ((InOldName != InNewName) && (InNewName != NAME_None))
 	{
-		FBPVariableDescription* LocalVariable = FindLocalVariable(InBlueprint, InOldName);
-		
-		if(LocalVariable)
+		UK2Node_FunctionEntry* FoundFunctionEntry = NULL;
+		FBPVariableDescription* LocalVariable = FindLocalVariable(InBlueprint, InOldName, &FoundFunctionEntry);
+		const auto FoundFunction = FFunctionFromNodeHelper::FunctionFromNode(FoundFunctionEntry);
+		const auto ExistingProperty = FindField<const UProperty>(FoundFunction, InNewName);
+		if (ExistingProperty)
+		{
+			UE_LOG(LogBlueprint, Warning, TEXT("Cannot name local variable '%s'. The name is already used."), *InNewName.ToString());
+		}
+
+		if (LocalVariable && !ExistingProperty)
 		{
 			const FScopedTransaction Transaction( LOCTEXT("RenameVariable", "Rename Local Variable") );
 			InBlueprint->Modify();
@@ -3359,7 +3392,7 @@ void FBlueprintEditorUtils::RenameLocalVariable(UBlueprint* InBlueprint, const F
 	}
 }
 
-FBPVariableDescription* FBlueprintEditorUtils::FindLocalVariable(const UBlueprint* InBlueprint, const FName& InVariableName)
+FBPVariableDescription* FBlueprintEditorUtils::FindLocalVariable(const UBlueprint* InBlueprint, const FName& InVariableName, class UK2Node_FunctionEntry** OutFunctionEntry)
 {
 	// Search through all function entry nodes for a local variable with the passed name
 	TArray<UK2Node_FunctionEntry*> FunctionEntryNodes;
@@ -3373,6 +3406,10 @@ FBPVariableDescription* FBlueprintEditorUtils::FindLocalVariable(const UBlueprin
 		{
 			if(InVariableName == Variable.VarName)
 			{
+				if (OutFunctionEntry)
+				{
+					*OutFunctionEntry = FunctionEntry;
+				}
 				ReturnVariable = &Variable;
 				break;
 			}
diff --git a/Engine/Source/Editor/UnrealEd/Public/Kismet2/BlueprintEditorUtils.h b/Engine/Source/Editor/UnrealEd/Public/Kismet2/BlueprintEditorUtils.h
--- a/Engine/Source/Editor/UnrealEd/Public/Kismet2/BlueprintEditorUtils.h
+++ b/Engine/Source/Editor/UnrealEd/Public/Kismet2/BlueprintEditorUtils.h
@@ -27,6 +27,16 @@ namespace EGraphRemoveFlags
 	};
 };
 
+struct UNREALED_API FFunctionFromNodeHelper
+{
+	UFunction* const Function;
+	UK2Node* const Node;
+
+	static UFunction* FunctionFromNode(UK2Node* Node);
+
+	FFunctionFromNodeHelper(UObject* Obj);
+};
+
 class UNREALED_API FBlueprintEditorUtils
 {
 public:
@@ -508,9 +518,10 @@ public:
 	 *
 	 * @param InBlueprint		Blueprint to search for the local variable
 	 * @param InVariableName	Name of the variable to search for
+	 * @param OutFunctionEntry	Optional output parameter. If not null, the found function entry is returned.
 	 * @return					The local variable description
 	 */
-	static FBPVariableDescription* FindLocalVariable(const UBlueprint* InBlueprint, const FName& InVariableName);
+	static FBPVariableDescription* FindLocalVariable(const UBlueprint* InBlueprint, const FName& InVariableName, class UK2Node_FunctionEntry** OutFunctionEntry = NULL);
 
 	/**
 	 * Finds a local variable name using the variable's Guid
