Fixed graph removal to handle animation composite nodes and state machines.
#ttp 339309 - LIVE: PERSONA: ANIM BP: CRASH: Click on the docked tab for a State after deleting its State Machine

diff --git a/Engine/Source/Editor/UnrealEd/Private/Kismet2/BlueprintEditorUtils.cpp b/Engine/Source/Editor/UnrealEd/Private/Kismet2/BlueprintEditorUtils.cpp
--- a/Engine/Source/Editor/UnrealEd/Private/Kismet2/BlueprintEditorUtils.cpp
+++ b/Engine/Source/Editor/UnrealEd/Private/Kismet2/BlueprintEditorUtils.cpp
@@ -18,6 +18,8 @@
 #include "AnimationGraph.h"
 #include "AnimationGraphSchema.h"
 #include "AnimationStateMachineGraph.h"
+#include "AnimStateConduitNode.h"
+#include "AnimGraphNode_StateMachine.h"
 #include "Editor/UnrealEd/Public/Kismet2/KismetEditorUtilities.h"
 #include "Editor/UnrealEd/Public/Kismet2/KismetDebugUtilities.h"
 #include "Editor/UnrealEd/Public/Kismet2/BlueprintEditorUtils.h"
@@ -1799,6 +1801,34 @@ void FBlueprintEditorUtils::RemoveGraph(UBlueprint* Blueprint, class UEdGraph* G
 		}
 	}
 
+	// Animation nodes can contain subgraphs but are not composite nodes, handle their graphs
+	TArray<UAnimStateNodeBase*> AllAnimCompositeNodes;
+	GraphToRemove->GetNodesOfClassEx<UAnimStateNode>(AllAnimCompositeNodes);
+	GraphToRemove->GetNodesOfClassEx<UAnimStateConduitNode>(AllAnimCompositeNodes);
+	GraphToRemove->GetNodesOfClassEx<UAnimStateTransitionNode>(AllAnimCompositeNodes);
+
+	for(UAnimStateNodeBase* Node : AllAnimCompositeNodes)
+	{
+		UEdGraph* BoundGraph = Node->GetBoundGraph();
+		if(BoundGraph && BoundGraph->GetOuter()->IsA(UAnimStateNodeBase::StaticClass()))
+		{
+			FBlueprintEditorUtils::RemoveGraph(Blueprint, BoundGraph, EGraphRemoveFlags::None);
+		}
+	}
+
+	// Handle sub anim state machines
+	TArray<UAnimGraphNode_StateMachineBase*> AllStateMachines;
+	GraphToRemove->GetNodesOfClassEx<UAnimGraphNode_StateMachine>(AllStateMachines);
+
+	for(UAnimGraphNode_StateMachineBase* Node : AllStateMachines)
+	{
+		UEdGraph* BoundGraph = Node->EditorStateMachineGraph;
+		if(BoundGraph && BoundGraph->GetOuter()->IsA(UAnimGraphNode_StateMachineBase::StaticClass()))
+		{
+			FBlueprintEditorUtils::RemoveGraph(Blueprint, BoundGraph, EGraphRemoveFlags::None);
+		}
+	}
+
 	GraphToRemove->GetSchema()->HandleGraphBeingDeleted(*GraphToRemove);
 
 	GraphToRemove->Rename(NULL, Blueprint->GetOuter(), REN_DoNotDirty | REN_DontCreateRedirectors);
