Editor: Removed bogus method PostScriptCompile that is never called

diff --git a/Engine/Source/Editor/UnrealEd/Classes/Editor/EditorEngine.h b/Engine/Source/Editor/UnrealEd/Classes/Editor/EditorEngine.h
--- a/Engine/Source/Editor/UnrealEd/Classes/Editor/EditorEngine.h
+++ b/Engine/Source/Editor/UnrealEd/Classes/Editor/EditorEngine.h
@@ -1782,9 +1782,6 @@ public:
 	 */
 	void DisableDeltaModification(bool bDisable) { bDisableDeltaModification = bDisable; }
 
-	/** called after script compilation to allow for game specific post-compilation steps */
-	virtual void PostScriptCompile() {}
-
 	/**
 	 *	Game-specific function called by Map_Check BEFORE iterating over all actors.
 	 *
