revert back to only building if the project has code and there is a compiler
TTP339487
#ue4

diff --git a/Engine/Source/Editor/MainFrame/Private/Frame/MainFrameActions.cpp b/Engine/Source/Editor/MainFrame/Private/Frame/MainFrameActions.cpp
--- a/Engine/Source/Editor/MainFrame/Private/Frame/MainFrameActions.cpp
+++ b/Engine/Source/Editor/MainFrame/Private/Frame/MainFrameActions.cpp
@@ -583,7 +583,7 @@ void FMainFrameActionCallbacks::PackageProject( const FName InPlatformInfoName )
 	}
 
 	// only build if the project has code that might need to be built
-	if ((bProjectHasCode || !FRocketSupport::IsRocket()) && FSourceCodeNavigation::IsCompilerAvailable())
+	if (bProjectHasCode && FSourceCodeNavigation::IsCompilerAvailable())
 	{
 		OptionalParams += TEXT(" -build");
 	}
diff --git a/Engine/Source/Editor/UnrealEd/Private/PlayLevel.cpp b/Engine/Source/Editor/UnrealEd/Private/PlayLevel.cpp
--- a/Engine/Source/Editor/UnrealEd/Private/PlayLevel.cpp
+++ b/Engine/Source/Editor/UnrealEd/Private/PlayLevel.cpp
@@ -1277,7 +1277,7 @@ void UEditorEngine::PlayUsingLauncher()
 #endif
 
 		ILauncherProfileRef LauncherProfile = LauncherServicesModule.CreateProfile(TEXT("Play On Device"));
-		LauncherProfile->SetBuildGame((bHasCode && FSourceCodeNavigation::IsCompilerAvailable()) || !FRocketSupport::IsRocket());
+		LauncherProfile->SetBuildGame(bHasCode && FSourceCodeNavigation::IsCompilerAvailable());
 		LauncherProfile->SetCookMode(ELauncherProfileCookModes::ByTheBook);
 		LauncherProfile->AddCookedPlatform(PlayUsingLauncherDeviceId.Left(PlayUsingLauncherDeviceId.Find(TEXT("@"))));
 		LauncherProfile->SetForceClose(true);
