UE4: If a World asset gets deleted due to a map transition where the user opted to not save the asset, it will now be properly removed from the asset registry. TTP# 338503

diff --git a/Engine/Source/Editor/UnrealEd/Private/EditorServer.cpp b/Engine/Source/Editor/UnrealEd/Private/EditorServer.cpp
--- a/Engine/Source/Editor/UnrealEd/Private/EditorServer.cpp
+++ b/Engine/Source/Editor/UnrealEd/Private/EditorServer.cpp
@@ -1607,23 +1607,42 @@ void UEditorEngine::BSPIntersectionHelper(UWorld* InWorld, ECsgOper Operation)
 
 void UEditorEngine::EditorDestroyWorld( FWorldContext & Context, const FText& CleanseText, UWorld* NewWorld )
 {
-	if (Context.World() == NULL )
+	UWorld* ContextWorld = Context.World();
+
+	if (ContextWorld == NULL )
 	{
 		return;		// We cannot destroy a world if the pointer is not valid
 	}
 
-	UPackage* WorldPackage = CastChecked<UPackage>(Context.World()->GetOuter());
+	UPackage* WorldPackage = CastChecked<UPackage>(ContextWorld->GetOuter());
 	if (WorldPackage == GetTransientPackage())
 	{
 		// Don't check if the package was properly cleaned up if we were created in the transient package
 		WorldPackage = NULL;
 	}
 
-	if (Context.World()->WorldType != EWorldType::Preview && Context.World()->WorldType != EWorldType::Inactive)
+	if (ContextWorld->WorldType != EWorldType::Preview && ContextWorld->WorldType != EWorldType::Inactive)
 	{
 		// Go away, come again never!
-		Context.World()->ClearFlags(RF_Standalone | RF_RootSet | RF_Transactional);
-		Context.World()->SetFlags(RF_Transient);
+		ContextWorld->ClearFlags(RF_Standalone | RF_RootSet | RF_Transactional);
+
+		// If this was a memory-only world, we should inform the asset registry that this asset is going away forever.
+		if (WorldPackage)
+		{
+			const FString PackageName = WorldPackage->GetName();
+			const bool bIncludeReadOnlyRoots = false;
+			if (FPackageName::IsValidLongPackageName(PackageName, bIncludeReadOnlyRoots))
+			{
+				// Now check if the file exists on disk. If it does, it won't be "lost" when GC'd.
+				if (!FPackageName::DoesPackageExist(PackageName))
+				{
+					// We are preparing the object for GC and there is no file on disk to reload it. Count this as a delete.
+					FAssetRegistryModule::AssetDeleted(ContextWorld);
+				}
+			}
+		}
+
+		ContextWorld->SetFlags(RF_Transient);
 	}
 
 	GUnrealEd->CurrentLODParentActor = NULL;
@@ -1631,10 +1650,10 @@ void UEditorEngine::EditorDestroyWorld( FWorldContext & Context, const FText& Cl
 	EditorClearComponents();
 	ClearPreviewComponents();
 	// Remove all active groups, they belong to a map being unloaded
-	Context.World()->ActiveGroupActors.Empty();
+	ContextWorld->ActiveGroupActors.Empty();
 
 	// Make sure we don't have any apps open on for assets owned by the world we are closing
-	CloseEditedWorldAssets(Context.World());
+	CloseEditedWorldAssets(ContextWorld);
 
 	// Stop all audio and remove references 
 	if ( GetAudioDevice() )
@@ -1653,8 +1672,7 @@ void UEditorEngine::EditorDestroyWorld( FWorldContext & Context, const FText& Cl
 		}
 	}
 
-	UWorld* ContextWorld = Context.World();
-	Context.World()->DestroyWorld( true, NewWorld );
+	ContextWorld->DestroyWorld( true, NewWorld );
 	Context.SetCurrentWorld(NULL);
 
 	// Add the new world to root if it wasn't already and keep track of it so we can remove it from root later if appropriate
