UE4: Added some special-case handling for force-deleting or consolidating the current editor world from the content browser. This will now transition you to a new world, reload the original world as an inactive world, and continue the operation. TTP#338502

diff --git a/Engine/Source/Editor/UnrealEd/Private/ObjectTools.cpp b/Engine/Source/Editor/UnrealEd/Private/ObjectTools.cpp
--- a/Engine/Source/Editor/UnrealEd/Private/ObjectTools.cpp
+++ b/Engine/Source/Editor/UnrealEd/Private/ObjectTools.cpp
@@ -32,6 +32,33 @@
 
 DEFINE_LOG_CATEGORY_STATIC(LogObjectTools, Log, All);
 
+// This function should ONLY be needed by ConsolidateObjects and ForceDeleteObjects
+// Use anywhere else could be dangerous as this involves a map transition and GC
+void ReloadEditorWorldForReferenceReplacementIfNecessary(TArray<UObject*>& InOutObjectsToReplace)
+{
+	// If we are force-deleting or consolidating the editor world, first transition to an empty map to prevent reference problems.
+	// Then, re-load the world from disk to set it up for delete as an inactive world which isn't attached to the editor engine or other systems.
+	UWorld* EditorWorld = GEditor->GetEditorWorldContext().World();
+
+	// Remove the world from ObjectsToDelete since NewMap() will delete the object naturally
+	int32 NumEntriesRemoved = InOutObjectsToReplace.Remove(EditorWorld);
+	if (NumEntriesRemoved > 0)
+	{
+		const FString ObjectPath = EditorWorld->GetPathName();
+
+		// Transition to a new map. This will invoke garbage collection and destroy the EditorWorld
+		GEditor->NewMap();
+
+		// Attempt to reload the editor world so we can make sure the file gets deleted and everything is handled normally.
+		// It is okay for this to fail. If we could not reload the world, it is not on disk and is gone.
+		UObject* ReloadedEditorWorld = LoadObject<UWorld>(nullptr, *ObjectPath, nullptr, LOAD_Quiet | LOAD_NoWarn);
+		if (ReloadedEditorWorld)
+		{
+			InOutObjectsToReplace.Add(ReloadedEditorWorld);
+		}
+	}
+}
+
 namespace ObjectTools
 {
 	/** Returns true if the specified object can be displayed in a content browser */
@@ -722,6 +749,9 @@ namespace ObjectTools
 			// objects they are replacing until the objects are garbage collected
 			TMap<UObjectRedirector*, FName> RedirectorToObjectNameMap;
 
+			// If the current editor world is in this list, transition to a new map and reload the world to finish the consolidate
+			ReloadEditorWorldForReferenceReplacementIfNecessary(ObjectsToConsolidate);
+
 			FForceReplaceInfo ReplaceInfo;
 			// Scope the reregister context below to complete after object deletion and before garbage collection
 			{
@@ -1806,6 +1836,9 @@ namespace ObjectTools
 			}
 		}
 
+		// If the current editor world is in this list, transition to a new map and reload the world to finish the delete
+		ReloadEditorWorldForReferenceReplacementIfNecessary(ObjectsToDelete);
+
 		{
 			// Replacing references inside already loaded objects could cause rendering issues, so globally detach all components from their scenes for now
 			FGlobalComponentReregisterContext ReregisterContext;
