Always kill a PIE session before creating a NewMap in the editor engine.

diff --git a/Engine/Source/Editor/UnrealEd/Private/EditorServer.cpp b/Engine/Source/Editor/UnrealEd/Private/EditorServer.cpp
--- a/Engine/Source/Editor/UnrealEd/Private/EditorServer.cpp
+++ b/Engine/Source/Editor/UnrealEd/Private/EditorServer.cpp
@@ -1844,6 +1844,12 @@ void UEditorEngine::CreateNewMapForEditing()
 
 UWorld* UEditorEngine::NewMap()
 {
+	// If we have a PIE session kill it before creating a new map
+	if (PlayWorld)
+	{
+		EndPlayMap();
+	}
+
 	const FScopedBusyCursor BusyCursor;
 
 	FWorldContext &Context = GetEditorWorldContext();
