Editor: Add UFactory::CanCreateNew and make bCreateNew a protected member

diff --git a/Engine/Source/Editor/BehaviorTreeEditor/Classes/BehaviorTreeFactory.h b/Engine/Source/Editor/BehaviorTreeEditor/Classes/BehaviorTreeFactory.h
--- a/Engine/Source/Editor/BehaviorTreeEditor/Classes/BehaviorTreeFactory.h
+++ b/Engine/Source/Editor/BehaviorTreeEditor/Classes/BehaviorTreeFactory.h
@@ -3,15 +3,15 @@
 #pragma once
 #include "BehaviorTreeFactory.generated.h"
 
-UCLASS(hidecategories=Object)
+UCLASS()
 class UBehaviorTreeFactory : public UFactory
 {
 	GENERATED_UCLASS_BODY()
 
-	// Begin UFactory Interface
+	// UFactory interface
 	virtual UObject* FactoryCreateNew(UClass* Class,UObject* InParent,FName Name,EObjectFlags Flags,UObject* Context,FFeedbackContext* Warn) override;
-	virtual bool ShouldShowInNewMenu() const override;
-	// Begin UFactory Interface	
+	virtual bool CanCreateNew() const override;
+	// End of UFactory interface
 };
 
 
diff --git a/Engine/Source/Editor/BehaviorTreeEditor/Private/BehaviorTreeEditorModule.cpp b/Engine/Source/Editor/BehaviorTreeEditor/Private/BehaviorTreeEditorModule.cpp
--- a/Engine/Source/Editor/BehaviorTreeEditor/Private/BehaviorTreeEditorModule.cpp
+++ b/Engine/Source/Editor/BehaviorTreeEditor/Private/BehaviorTreeEditorModule.cpp
@@ -105,15 +105,6 @@ void FBehaviorTreeEditorModule::ShutdownModule()
 	}
 }
 
-void FBehaviorTreeEditorModule::HandleExperimentalSettingChanged(FName PropertyName)
-{
-	if (PropertyName == TEXT("bBehaviorTreeEditor"))
-	{
-		UBehaviorTreeFactory* EditorUtilityBlueprintFactory = Cast<UBehaviorTreeFactory>(UBehaviorTreeFactory::StaticClass()->GetDefaultObject());
-		EditorUtilityBlueprintFactory->bCreateNew = GetDefault<UEditorExperimentalSettings>()->bBehaviorTreeEditor;;
-	}
-}
-
 TSharedRef<IBehaviorTreeEditor> FBehaviorTreeEditorModule::CreateBehaviorTreeEditor( const EToolkitMode::Type Mode, const TSharedPtr< IToolkitHost >& InitToolkitHost, UObject* Object )
 {
 	if (!ClassCache.IsValid())
diff --git a/Engine/Source/Editor/BehaviorTreeEditor/Private/BehaviorTreeFactory.cpp b/Engine/Source/Editor/BehaviorTreeEditor/Private/BehaviorTreeFactory.cpp
--- a/Engine/Source/Editor/BehaviorTreeEditor/Private/BehaviorTreeFactory.cpp
+++ b/Engine/Source/Editor/BehaviorTreeEditor/Private/BehaviorTreeFactory.cpp
@@ -10,19 +10,11 @@ UBehaviorTreeFactory::UBehaviorTreeFactory(const class FPostConstructInitializeP
 	: Super(PCIP)
 {
 	SupportedClass = UBehaviorTree::StaticClass();
+	bCreateNew = true;
 	bEditAfterNew = true;
-
-	//enable asset creation by Epic Labs
-	bool bCreateNewByEpicLabs = GetDefault<UEditorExperimentalSettings>()->bBehaviorTreeEditor;
-
-	// Check ini to see if we should enable creation
-	bool bBehaviorTreeNewAssetsEnabled = false;
-	GConfig->GetBool(TEXT("BehaviorTreesEd"), TEXT("BehaviorTreeNewAssetsEnabled"), bBehaviorTreeNewAssetsEnabled, GEngineIni);
-
-	bCreateNew = (bCreateNewByEpicLabs || bBehaviorTreeNewAssetsEnabled);
 }
 
-bool UBehaviorTreeFactory::ShouldShowInNewMenu() const
+bool UBehaviorTreeFactory::CanCreateNew() const
 {
 	bool bBehaviorTreeNewAssetsEnabled = false;
 	GConfig->GetBool(TEXT("BehaviorTreesEd"), TEXT("BehaviorTreeNewAssetsEnabled"), bBehaviorTreeNewAssetsEnabled, GEngineIni);
diff --git a/Engine/Source/Editor/BehaviorTreeEditor/Public/BehaviorTreeEditorModule.h b/Engine/Source/Editor/BehaviorTreeEditor/Public/BehaviorTreeEditorModule.h
--- a/Engine/Source/Editor/BehaviorTreeEditor/Public/BehaviorTreeEditorModule.h
+++ b/Engine/Source/Editor/BehaviorTreeEditor/Public/BehaviorTreeEditorModule.h
@@ -32,8 +32,6 @@ public:
 	static const FName BehaviorTreeEditorAppIdentifier;
 
 private:
-	void HandleExperimentalSettingChanged(FName PropertyName);
-
 	TSharedPtr<FExtensibilityManager> MenuExtensibilityManager;
 	TSharedPtr<FExtensibilityManager> ToolBarExtensibilityManager;
 
diff --git a/Engine/Source/Editor/Blutility/Classes/EditorUtilityBlueprintFactory.h b/Engine/Source/Editor/Blutility/Classes/EditorUtilityBlueprintFactory.h
--- a/Engine/Source/Editor/Blutility/Classes/EditorUtilityBlueprintFactory.h
+++ b/Engine/Source/Editor/Blutility/Classes/EditorUtilityBlueprintFactory.h
@@ -20,6 +20,7 @@ class UEditorUtilityBlueprintFactory : public UFactory
 	// UFactory interface
 	virtual bool ConfigureProperties() override;
 	virtual UObject* FactoryCreateNew(UClass* Class, UObject* InParent, FName Name, EObjectFlags Flags, UObject* Context, FFeedbackContext* Warn) override;
+	virtual bool CanCreateNew() const override;
 	// End of UFactory interface
 
 	/** Handler for when a class is picked in the class picker */
diff --git a/Engine/Source/Editor/Blutility/Private/BlutilityModule.cpp b/Engine/Source/Editor/Blutility/Private/BlutilityModule.cpp
--- a/Engine/Source/Editor/Blutility/Private/BlutilityModule.cpp
+++ b/Engine/Source/Editor/Blutility/Private/BlutilityModule.cpp
@@ -45,8 +45,6 @@ public:
 		{
 			RegisterBlutilityShelfTabSpawner();
 		}
-		
-		GetMutableDefault<UEditorExperimentalSettings>()->OnSettingChanged().AddRaw(this, &FBlutilityModule::HandleExperimentalSettingChanged);
 	}
 
 	virtual void ShutdownModule() override
@@ -58,8 +56,6 @@ public:
 
 		UnregisterBlutilityShelfTabSpawner();
 
-		GetMutableDefault<UEditorExperimentalSettings>()->OnSettingChanged().RemoveAll(this);
-
 		// Only unregister if the asset tools module is loaded.  We don't want to forcibly load it during shutdown phase.
 		check( EditorBlueprintAssetTypeActions.IsValid() );
 		if (FModuleManager::Get().IsModuleLoaded("AssetTools"))
@@ -109,12 +105,10 @@ protected:
 			if (GetDefault<UEditorExperimentalSettings>()->bEnableEditorUtilityBlueprints)
 			{
 				RegisterBlutilityShelfTabSpawner();
-				EditorUtilityBlueprintFactory->bCreateNew = true;
 			}
 			else
 			{
 				UnregisterBlutilityShelfTabSpawner();
-				EditorUtilityBlueprintFactory->bCreateNew = false;
 			}
 		}
 	}
diff --git a/Engine/Source/Editor/Blutility/Private/EditorUtilityBlueprintFactory.cpp b/Engine/Source/Editor/Blutility/Private/EditorUtilityBlueprintFactory.cpp
--- a/Engine/Source/Editor/Blutility/Private/EditorUtilityBlueprintFactory.cpp
+++ b/Engine/Source/Editor/Blutility/Private/EditorUtilityBlueprintFactory.cpp
@@ -31,14 +31,7 @@ UEditorUtilityBlueprintFactory::UEditorUtilityBlueprintFactory(const class FPost
 {
 	UClass* DefaultParentClass = APlacedEditorUtilityBase::StaticClass();
 
-	if ( GEditor )
-	{
-		if (GetDefault<UEditorExperimentalSettings>()->bEnableEditorUtilityBlueprints)
-		{
-			bCreateNew = true;
-		}
-	}
-
+	bCreateNew = true;
 	bEditAfterNew = true;
 	SupportedClass = UEditorUtilityBlueprint::StaticClass();
 	ParentClass = DefaultParentClass;
@@ -100,6 +93,11 @@ UObject* UEditorUtilityBlueprintFactory::FactoryCreateNew(UClass* Class, UObject
 	}
 }
 
+bool UEditorUtilityBlueprintFactory::CanCreateNew() const
+{
+	return GetDefault<UEditorExperimentalSettings>()->bEnableEditorUtilityBlueprints;
+}
+
 void UEditorUtilityBlueprintFactory::OnClassPicked(UClass* InChosenClass)
 {
 	ParentClass = InChosenClass;
diff --git a/Engine/Source/Editor/EnvironmentQueryEditor/Classes/EnvironmentQueryFactory.h b/Engine/Source/Editor/EnvironmentQueryEditor/Classes/EnvironmentQueryFactory.h
--- a/Engine/Source/Editor/EnvironmentQueryEditor/Classes/EnvironmentQueryFactory.h
+++ b/Engine/Source/Editor/EnvironmentQueryEditor/Classes/EnvironmentQueryFactory.h
@@ -8,9 +8,10 @@ class UEnvironmentQueryFactory : public UFactory
 {
 	GENERATED_UCLASS_BODY()
 
-	// Begin UFactory Interface
+	// UFactory interface
 	virtual UObject* FactoryCreateNew(UClass* Class,UObject* InParent,FName Name,EObjectFlags Flags,UObject* Context,FFeedbackContext* Warn) override;
-	// Begin UFactory Interface	
+	virtual bool CanCreateNew() const override;
+	// End of UFactory interface	
 };
 
 
diff --git a/Engine/Source/Editor/EnvironmentQueryEditor/Private/EnvironmentQueryFactory.cpp b/Engine/Source/Editor/EnvironmentQueryEditor/Private/EnvironmentQueryFactory.cpp
--- a/Engine/Source/Editor/EnvironmentQueryEditor/Private/EnvironmentQueryFactory.cpp
+++ b/Engine/Source/Editor/EnvironmentQueryEditor/Private/EnvironmentQueryFactory.cpp
@@ -9,15 +9,19 @@ UEnvironmentQueryFactory::UEnvironmentQueryFactory(const class FPostConstructIni
 {
 	SupportedClass = UEnvQuery::StaticClass();
 	bEditAfterNew = true;
-
-	// Check ini to see if we should enable creation
-	bool bEnableEnvironmentQueryEd=false;
-	GConfig->GetBool(TEXT("EnvironmentQueryEd"), TEXT("EnableEnvironmentQueryEd"), bEnableEnvironmentQueryEd, GEngineIni);
-	bCreateNew = bEnableEnvironmentQueryEd;
+	bCreateNew = true;
 }
 
 UObject* UEnvironmentQueryFactory::FactoryCreateNew(UClass* Class,UObject* InParent,FName Name,EObjectFlags Flags,UObject* Context,FFeedbackContext* Warn)
 {
 	check(Class->IsChildOf(UEnvQuery::StaticClass()));
 	return ConstructObject<UEnvQuery>(Class, InParent, Name, Flags);
-}
\ No newline at end of file
+}
+
+bool UEnvironmentQueryFactory::CanCreateNew() const
+{
+	// Check ini to see if we should enable creation
+	bool bEnableEnvironmentQueryEd = false;
+	GConfig->GetBool(TEXT("EnvironmentQueryEd"), TEXT("EnableEnvironmentQueryEd"), bEnableEnvironmentQueryEd, GEngineIni);
+	return bEnableEnvironmentQueryEd;
+}
diff --git a/Engine/Source/Editor/UnrealEd/Classes/Factories/Factory.h b/Engine/Source/Editor/UnrealEd/Classes/Factories/Factory.h
--- a/Engine/Source/Editor/UnrealEd/Classes/Factories/Factory.h
+++ b/Engine/Source/Editor/UnrealEd/Classes/Factories/Factory.h
@@ -26,10 +26,12 @@ class UNREALED_API UFactory : public UObject
 	UPROPERTY()
 	TArray<FString> Formats;
 
-	/** true if the factory creates a new object from scratch. */
+protected:
+	/** The default value to return from CanCreateNew() */
 	UPROPERTY()
 	uint32 bCreateNew:1;
 
+public:
 	/** true if the associated editor should be opened after creating a new object. */
 	UPROPERTY()
 	uint32 bEditAfterNew:1;
@@ -66,7 +68,7 @@ class UNREALED_API UFactory : public UObject
 	static void AddReferencedObjects(UObject* InThis, FReferenceCollector& Collector);
 	// End UObject interface.
 
-	/** Returns true if this factory should be shown in the New Asset menu. */
+	/** Returns true if this factory should be shown in the New Asset menu (by default calls CanCreateNew). */
 	virtual bool ShouldShowInNewMenu() const;
 
 	/** Returns an optional override brush name for the new asset menu. If this is not specified, the thumbnail for the supported class will be used. */
@@ -89,7 +91,7 @@ class UNREALED_API UFactory : public UObject
 	/**
 	 * @return true if it supports this class 
 	 */
-	virtual bool DoesSupportClass(UClass * Class);
+	virtual bool DoesSupportClass(UClass* Class);
 
 	/**
 	 * Resolves SupportedClass for factories which support multiple classes.
@@ -106,6 +108,11 @@ class UNREALED_API UFactory : public UObject
 	/** Opens a dialog to configure the factory properties. Return false if user opted out of configuring properties */
 	virtual bool ConfigureProperties() { return true; }
 
+	/**
+	 * @return true if the factory can currently create a new object from scratch.
+	 */
+	virtual bool CanCreateNew() const { return bCreateNew; }
+
 	// @todo document
 	virtual UObject* FactoryCreateText( UClass* InClass, UObject* InParent, FName InName, EObjectFlags Flags, UObject* Context, const TCHAR* Type, const TCHAR*& Buffer, const TCHAR* BufferEnd, FFeedbackContext* Warn ) {return NULL;}
 
diff --git a/Engine/Source/Editor/UnrealEd/Private/Factories/Factory.cpp b/Engine/Source/Editor/UnrealEd/Private/Factories/Factory.cpp
--- a/Engine/Source/Editor/UnrealEd/Private/Factories/Factory.cpp
+++ b/Engine/Source/Editor/UnrealEd/Private/Factories/Factory.cpp
@@ -73,7 +73,7 @@ bool UFactory::FactoryCanImport( const FString& Filename )
 
 bool UFactory::ShouldShowInNewMenu() const
 {
-	return bCreateNew;
+	return CanCreateNew();
 }
 
 FText UFactory::GetDisplayName() const
@@ -224,7 +224,7 @@ UObject* UFactory::StaticImportObject
 	{
 		UFactory* Factory = Factories[i];
 		UObject* Result = NULL;
-		if( Factory->bCreateNew )
+		if( Factory->CanCreateNew() )
 		{
 			if( FCString::Stricmp(Filename,TEXT(""))==0 )
 			{
