Fixed some issues parsing the COOKDIR command line argument

The commandline can contain quotes that surround either the entire string, or each individual path, however only the outer case was being handled. Additionally, if you had directories separated by +, all but the last directory was being put into the Map array rather than the Dir array.

The following command lines will now parse correctly:
-cookdir="D:\\Unreal Projects\\MyCodeProject\\Content\\RandomDirToCook"+"D:\\Unreal Projects\\MyCodeProject\\Content\\Audio"
-cookdir="D:\\Unreal Projects\\MyCodeProject\\Content\\RandomDirToCook+D:\\Unreal Projects\\MyCodeProject\\Content\\Audio"

diff --git a/Engine/Source/Editor/UnrealEd/Private/Commandlets/CookCommandlet.cpp b/Engine/Source/Editor/UnrealEd/Private/Commandlets/CookCommandlet.cpp
--- a/Engine/Source/Editor/UnrealEd/Private/Commandlets/CookCommandlet.cpp
+++ b/Engine/Source/Editor/UnrealEd/Private/Commandlets/CookCommandlet.cpp
@@ -747,36 +747,43 @@ void UCookCommandlet::CollectFilesToCook(TArray<FString>& FilesInPath)
 	TArray<FString> CmdLineDirEntries;
 	for (int32 SwitchIdx = 0; SwitchIdx < Switches.Num(); SwitchIdx++)
 	{
-		// Check for -MAP=<name of map> entries
 		const FString& Switch = Switches[SwitchIdx];
 
+		// Check for -MAP=<name of map> entries
 		if (Switch.StartsWith(TEXT("MAP=")) == true)
 		{
 			FString MapToCook = Switch.Right(Switch.Len() - 4);
+			
 			// Allow support for -MAP=Map1+Map2+Map3 as well as -MAP=Map1 -MAP=Map2
-			int32 PlusIdx = MapToCook.Find(TEXT("+"));
-			while (PlusIdx != INDEX_NONE)
+			for (int32 PlusIdx = MapToCook.Find(TEXT("+")); PlusIdx != INDEX_NONE; PlusIdx = MapToCook.Find(TEXT("+")))
 			{
 				FString MapName = MapToCook.Left(PlusIdx);
 				CmdLineMapEntries.Add(MapName);
+
 				MapToCook = MapToCook.Right(MapToCook.Len() - (PlusIdx + 1));
-				PlusIdx = MapToCook.Find(TEXT("+"));
 			}
+
 			CmdLineMapEntries.Add(MapToCook);
 		}
 
+		// Check for -COOKDIR=<path to directory> entries
 		if (Switch.StartsWith(TEXT("COOKDIR=")) == true)
 		{
 			FString DirToCook = Switch.Right(Switch.Len() - 8);
+			
 			// Allow support for -COOKDIR=Dir1+Dir2+Dir3 as well as -COOKDIR=Dir1 -COOKDIR=Dir2
-			int32 PlusIdx = DirToCook.Find(TEXT("+"));
-			while (PlusIdx != INDEX_NONE)
+			for (int32 PlusIdx = DirToCook.Find(TEXT("+")); PlusIdx != INDEX_NONE; PlusIdx = DirToCook.Find(TEXT("+")))
 			{
 				FString DirName = DirToCook.Left(PlusIdx);
-				CmdLineMapEntries.Add(DirName);
+				
+				// The dir may be contained within quotes
+				DirName = DirName.TrimQuotes();
+				FPaths::NormalizeDirectoryName(DirName);
+				CmdLineDirEntries.Add(DirName);
+
 				DirToCook = DirToCook.Right(DirToCook.Len() - (PlusIdx + 1));
-				PlusIdx = DirToCook.Find(TEXT("+"));
 			}
+			
 			// The dir may be contained within quotes
 			DirToCook = DirToCook.TrimQuotes();
 			FPaths::NormalizeDirectoryName(DirToCook);
