Editor performance monitor can now disable realtime viewports and automatically apply scalability settings

When the monitor detects that the editor is running under an acceptible framerate it will set a flag which causes a benchmark to run on start-up, choosing the most optimal engine display settings for the system. This can be disabled by manually tweaking the settings.
The feature can also be turned off entirely in the editor preferences under Level Editor/Miscellaneous.

Also fixed a shutdown crash by making the performance monitor a member of UUnrealEdEngine. I've had to move some files around to make this possible.

diff --git a/Engine/Source/Editor/LevelEditor/Private/PerformanceWarningNotification.cpp b/Engine/Source/Editor/LevelEditor/Private/PerformanceWarningNotification.cpp
deleted file mode 100644
diff --git a/Engine/Source/Editor/LevelEditor/Private/SScalabilitySettings.cpp b/Engine/Source/Editor/LevelEditor/Private/SScalabilitySettings.cpp
deleted file mode 100644
diff --git a/Engine/Source/Editor/LevelEditor/Private/SScalabilitySettings.h b/Engine/Source/Editor/LevelEditor/Private/SScalabilitySettings.h
deleted file mode 100644
diff --git a/Engine/Source/Editor/UnrealEd/Classes/Editor/EditorGameAgnosticSettings.h b/Engine/Source/Editor/UnrealEd/Classes/Editor/EditorGameAgnosticSettings.h
--- a/Engine/Source/Editor/UnrealEd/Classes/Editor/EditorGameAgnosticSettings.h
+++ b/Engine/Source/Editor/UnrealEd/Classes/Editor/EditorGameAgnosticSettings.h
@@ -28,8 +28,9 @@ class UEditorGameAgnosticSettings : public UObject
 	UPROPERTY(config)
 	bool bCopyStarterContentPreference;
 
+	/** When set this will prompt a benchmark to be ran on startup to determine the best scalability settings */
 	UPROPERTY(config)
-	bool bShowPerformanceWarningPreference;
+	bool bApplyAutoScalabilityOnStartup;
 
 	/** The id's of the surveys completed */
 	UPROPERTY(config)
diff --git a/Engine/Source/Editor/UnrealEd/Classes/Editor/UnrealEdEngine.h b/Engine/Source/Editor/UnrealEd/Classes/Editor/UnrealEdEngine.h
--- a/Engine/Source/Editor/UnrealEd/Classes/Editor/UnrealEdEngine.h
+++ b/Engine/Source/Editor/UnrealEd/Classes/Editor/UnrealEdEngine.h
@@ -80,6 +80,8 @@ struct FTemplateMapInfo
 	}
 };
 
+class FPerformanceMonitor;
+
 UCLASS(config=Engine, transient)
 class UNREALED_API UUnrealEdEngine : public UEditorEngine, public FNotifyHook
 {
@@ -709,4 +711,7 @@ protected:
 
 	/** The package auto-saver instance used by the editor */
 	TUniquePtr<IPackageAutoSaver> PackageAutoSaver;
+
+	/** Instance responsible for monitoring this editor's performance */
+	FPerformanceMonitor* PerformanceMonitor;
 };
diff --git a/Engine/Source/Editor/UnrealEd/Classes/Settings/LevelEditorMiscSettings.h b/Engine/Source/Editor/UnrealEd/Classes/Settings/LevelEditorMiscSettings.h
--- a/Engine/Source/Editor/UnrealEd/Classes/Settings/LevelEditorMiscSettings.h
+++ b/Engine/Source/Editor/UnrealEd/Classes/Settings/LevelEditorMiscSettings.h
@@ -53,6 +53,12 @@ public:
 	UPROPERTY(EditAnywhere, config, Category=Sound)
 	uint32 bEnableEditorSounds:1;
 
+public:	
+
+	/** When turned on, the editor will constantly monitor performance and adjust scalability settings for you when performance drops */
+	UPROPERTY(EditAnywhere, config, Category=Performance)
+	bool bMonitorEditorPerformance;
+
 protected:
 
 	// Begin UObject overrides
diff --git a/Engine/Source/Editor/UnrealEd/Private/Editor/EditorGameAgnosticSettings.cpp b/Engine/Source/Editor/UnrealEd/Private/Editor/EditorGameAgnosticSettings.cpp
--- a/Engine/Source/Editor/UnrealEd/Private/Editor/EditorGameAgnosticSettings.cpp
+++ b/Engine/Source/Editor/UnrealEd/Private/Editor/EditorGameAgnosticSettings.cpp
@@ -6,7 +6,7 @@ UEditorGameAgnosticSettings::UEditorGameAgnosticSettings(const class FPostConstr
 	: Super(PCIP)
 {
 	bCopyStarterContentPreference = true;
-	bShowPerformanceWarningPreference = true;
+	bApplyAutoScalabilityOnStartup = false;
 }
 
 void UEditorGameAgnosticSettings::PostEditChangeProperty( struct FPropertyChangedEvent& PropertyChangedEvent)
diff --git a/Engine/Source/Editor/UnrealEd/Private/PerformanceMonitor.cpp b/Engine/Source/Editor/UnrealEd/Private/PerformanceMonitor.cpp
new file mode 100644
--- /dev/null
+++ b/Engine/Source/Editor/UnrealEd/Private/PerformanceMonitor.cpp
@@ -0,0 +1,413 @@
+// Copyright 1998-2014 Epic Games, Inc. All Rights Reserved.
+
+
+#include "UnrealEd.h"
+#include "SScalabilitySettings.h"
+#include "PerformanceMonitor.h"
+
+#define LOCTEXT_NAMESPACE "PerformanceMonitor"
+
+/** Scalability dialog widget */
+class SScalabilitySettingsDialog : public SCompoundWidget
+{
+public:
+
+	SLATE_BEGIN_ARGS(SScalabilitySettingsDialog) {}
+		SLATE_ARGUMENT(FOnClicked, OnDoneClicked)
+	SLATE_END_ARGS()
+
+public:
+
+	/** Construct this widget */
+	void Construct(const FArguments& InArgs)
+	{
+		ChildSlot
+		[
+			SNew(SBorder)
+			.HAlign(HAlign_Fill)
+			.BorderImage(FEditorStyle::GetBrush("ChildWindow.Background"))
+			[
+				SNew(SBox)
+				.WidthOverride(500.0f)
+				[
+					SNew(SVerticalBox)
+					+ SVerticalBox::Slot()
+					.AutoHeight()
+					.Padding(5.f)
+					[
+						SNew(STextBlock).Text(LOCTEXT("PerformanceWarningDescription",
+						"The current performance of the editor seems to be low.\nUse the options below to reduce the amount of detail and increase performance.").ToString())
+					]
+					+ SVerticalBox::Slot()
+					.AutoHeight()
+					.Padding(5.f)
+					[
+						SNew(SScalabilitySettings)
+					]
+					+ SVerticalBox::Slot()
+					.AutoHeight()
+					.Padding(5.f)
+					[
+						SNew(STextBlock)
+						.ToolTip(
+							SNew(SToolTip)
+							[
+								SNew(SImage)
+								.Image(FEditorStyle::GetBrush("Scalability.ScalabilitySettings"))
+							]
+						)
+						.AutoWrapText(true)
+						.Text(LOCTEXT("PerformanceWarningChangeLater",
+							"You can modify these settings in future via \"Quick Settings\" button on the level editor toolbar and choosing \"Engine Scalability Settings\".").ToString())
+					]
+					+ SVerticalBox::Slot()
+					.AutoHeight()
+					.Padding(5.f)
+					[
+						SNew(SHorizontalBox)
+						+ SHorizontalBox::Slot()
+						[
+							SNew(SCheckBox)
+							.OnCheckStateChanged(FOnCheckStateChanged::CreateStatic(&OnMonitorPerformanceChanged))
+							.IsChecked(TAttribute<ESlateCheckBoxState::Type>::Create(&IsMonitoringPerformance))
+							.Content()
+							[
+								SNew(STextBlock)
+								.Text(LOCTEXT("PerformanceWarningEnableDisableCheckbox", "Keep monitoring editor performance?"))
+							]
+						]
+						+ SHorizontalBox::Slot()
+						.AutoWidth()
+						[
+							SNew(SButton)
+							.Text(LOCTEXT("ScalabilityDone", "Done"))
+							.OnClicked(InArgs._OnDoneClicked)
+						]
+					]
+				]
+			]
+		];
+	}
+
+	/** Called to get the "Show notification" check box state */
+	static ESlateCheckBoxState::Type IsMonitoringPerformance()
+	{
+		const bool bMonitorEditorPerformance = GetDefault<ULevelEditorMiscSettings>()->bMonitorEditorPerformance;
+		return bMonitorEditorPerformance ? ESlateCheckBoxState::Checked : ESlateCheckBoxState::Unchecked;
+	}
+
+	/** Called when the state of the "Show notification" check box changes */
+	static void OnMonitorPerformanceChanged(ESlateCheckBoxState::Type NewState)
+	{
+		const bool bNewEnabledState = (NewState == ESlateCheckBoxState::Checked);
+
+		auto* Settings = GetMutableDefault<ULevelEditorMiscSettings>();
+		Settings->bMonitorEditorPerformance = bNewEnabledState;
+		Settings->PostEditChange();
+	}
+};
+
+FPerformanceMonitor::FPerformanceMonitor()
+{
+	LastEnableTime = 0;
+	bIsNotificationAllowed = true;
+
+	UEditorGameAgnosticSettings& Settings = GEditor->AccessGameAgnosticSettings();
+	if (Settings.bApplyAutoScalabilityOnStartup)
+	{
+		Settings.bApplyAutoScalabilityOnStartup = false;
+		Settings.PostEditChange();
+
+		AutoApplyScalability();
+	}
+
+	// Register the console variables we need
+	IConsoleVariable* CVar = nullptr;
+
+	CVar = IConsoleManager::Get().RegisterConsoleVariable(TEXT("PerfWarn.FineSampleTime"), 30, TEXT("How many seconds we sample the percentage for the fine-grained minimum FPS."), ECVF_Default);
+	CVar->SetOnChangedCallback(FConsoleVariableDelegate::CreateRaw(this, &FPerformanceMonitor::UpdateMovingAverageSamplers));
+	CVar = IConsoleManager::Get().RegisterConsoleVariable(TEXT("PerfWarn.CoarseSampleTime"), 600, TEXT("How many seconds we sample the percentage for the coarse-grained minimum FPS."), ECVF_Default);
+	CVar->SetOnChangedCallback(FConsoleVariableDelegate::CreateRaw(this, &FPerformanceMonitor::UpdateMovingAverageSamplers));
+
+	IConsoleManager::Get().RegisterConsoleVariable(TEXT("PerfWarn.FineMinFPS"), 10, TEXT("The FPS threshold below which we warn about for fine-grained sampling."), ECVF_Default);
+	IConsoleManager::Get().RegisterConsoleVariable(TEXT("PerfWarn.CoarseMinFPS"), 20, TEXT("The FPS threshold below which we warn about for coarse-grained sampling."), ECVF_Default);
+
+	IConsoleManager::Get().RegisterConsoleVariable(TEXT("PerfWarn.FinePercentThreshold"), 80, TEXT("The percentage of samples that fall below min FPS above which we warn for."), ECVF_Default);
+	IConsoleManager::Get().RegisterConsoleVariable(TEXT("PerfWarn.CoarsePercentThreshold"), 80, TEXT("The percentage of samples that fall below min FPS above which we warn for."), ECVF_Default);
+
+	// Set up the moving average sasmplers
+	UpdateMovingAverageSamplers();
+}
+
+FPerformanceMonitor::~FPerformanceMonitor()
+{
+	IConsoleManager& Console = IConsoleManager::Get();
+	Console.FindConsoleVariable(TEXT("PerfWarn.FineSampleTime"))->SetOnChangedCallback(FConsoleVariableDelegate());
+	Console.FindConsoleVariable(TEXT("PerfWarn.CoarseSampleTime"))->SetOnChangedCallback(FConsoleVariableDelegate());
+}
+
+void FPerformanceMonitor::AutoApplyScalability()
+{
+	GWarn->BeginSlowTask(LOCTEXT("RunningBenchmark", "Running benchmark to determine optimum engine settings"), true);
+	//GWarn->UpdateProgress(-1, -1);
+	
+	Scalability::FQualityLevels BenchmarkLevels = Scalability::BenchmarkQualityLevels();
+	const Scalability::FQualityLevels ExistingLevels = Scalability::GetQualityLevels();
+	
+	// Make sure we don't turn settings up if the user has turned them down for any reason
+	BenchmarkLevels.ResolutionQuality		= FMath::Min(BenchmarkLevels.ResolutionQuality,	ExistingLevels.ResolutionQuality);
+	BenchmarkLevels.ViewDistanceQuality	= FMath::Min(BenchmarkLevels.ViewDistanceQuality,	ExistingLevels.ViewDistanceQuality);
+	BenchmarkLevels.AntiAliasingQuality	= FMath::Min(BenchmarkLevels.AntiAliasingQuality,	ExistingLevels.AntiAliasingQuality);
+	BenchmarkLevels.ShadowQuality			= FMath::Min(BenchmarkLevels.ShadowQuality,		ExistingLevels.ShadowQuality);
+	BenchmarkLevels.PostProcessQuality	= FMath::Min(BenchmarkLevels.PostProcessQuality,	ExistingLevels.PostProcessQuality);
+	BenchmarkLevels.TextureQuality		= FMath::Min(BenchmarkLevels.TextureQuality,		ExistingLevels.TextureQuality);
+	BenchmarkLevels.EffectsQuality		= FMath::Min(BenchmarkLevels.EffectsQuality,		ExistingLevels.EffectsQuality);
+	
+	Scalability::SetQualityLevels(BenchmarkLevels);
+	
+	GWarn->EndSlowTask();
+}
+
+void FPerformanceMonitor::ShowPerformanceWarning(FText MessageText)
+{
+	auto AlreadyOpenItem = PerformanceWarningNotificationPtr.Pin();
+	if (AlreadyOpenItem.IsValid())
+	{
+		AlreadyOpenItem->SetText(MessageText);
+	}
+	else if ((FPlatformTime::Seconds() - LastEnableTime) > 30.0)
+	{
+		// Only show a new one if we've not shown one for a while
+		LastEnableTime = FPlatformTime::Seconds();
+
+		// Create notification item
+		FNotificationInfo Info(MessageText);
+		Info.bFireAndForget = false;
+		Info.FadeOutDuration = 1.0f;
+		Info.ExpireDuration = 0.0f;
+		Info.bUseLargeFont = false;
+
+		Info.ButtonDetails.Add(
+			FNotificationButtonInfo(
+			LOCTEXT("PerformanceWarningScalability", "Tweak Scalability Settings"),
+			LOCTEXT("PerformanceWarningScalabilityToolTip", "Opens the Scalability UI for manual adjustment."),
+			FSimpleDelegate::CreateRaw(this, &FPerformanceMonitor::ShowScalabilityDialog),
+			SNotificationItem::CS_None
+			)
+			);
+
+		Info.ButtonDetails.Add(
+			FNotificationButtonInfo(
+			LOCTEXT("PerformanceWarningDismiss", "Ok"),
+			LOCTEXT("PerformanceWarningDismissToolTip", "Close this notification."),
+			FSimpleDelegate::CreateRaw(this, &FPerformanceMonitor::DismissAndDisableNotification),
+			SNotificationItem::CS_None
+			)
+			);
+
+		PerformanceWarningNotificationPtr = FSlateNotificationManager::Get().AddNotification(Info);
+	}
+}
+
+void FPerformanceMonitor::HidePerformanceWarning()
+{
+	// Finished! Notify the UI.
+	TSharedPtr<SNotificationItem> NotificationItem = PerformanceWarningNotificationPtr.Pin();
+
+	if (NotificationItem.IsValid())
+	{
+		NotificationItem->SetCompletionState(SNotificationItem::CS_Success);
+		NotificationItem->Fadeout();
+
+		PerformanceWarningNotificationPtr.Reset();
+	}
+}
+
+void FPerformanceMonitor::Tick(float DeltaTime)
+{
+	if (GEngine->ShouldThrottleCPUUsage())
+	{
+		return;
+	}
+
+	extern ENGINE_API float GAverageFPS;
+	FineMovingAverage.Tick(FPlatformTime::Seconds(), GAverageFPS);
+	CoarseMovingAverage.Tick(FPlatformTime::Seconds(), GAverageFPS);
+
+	const bool bMonitorEditorPerformance = GetDefault<ULevelEditorMiscSettings>()->bMonitorEditorPerformance;
+	if( !bMonitorEditorPerformance || !bIsNotificationAllowed )
+	{
+		return;
+	}
+
+	FFormatNamedArguments Arguments;
+	bool bLowFramerate = false;
+	int32 SampleTime = 0;
+	float PercentUnderTarget = 0.f;
+
+	if (FineMovingAverage.IsReliable())
+	{
+		static IConsoleVariable* CVarMinFPS = IConsoleManager::Get().FindConsoleVariable(TEXT("PerfWarn.FineMinFPS"));
+		static IConsoleVariable* CVarPercentThreshold = IConsoleManager::Get().FindConsoleVariable(TEXT("PerfWarn.FinePercentThreshold"));
+
+		PercentUnderTarget = FineMovingAverage.PercentageBelowThreshold(CVarMinFPS->GetFloat());
+
+		if (PercentUnderTarget >= CVarPercentThreshold->GetFloat())
+		{
+			Arguments.Add(TEXT("Framerate"), CVarMinFPS->GetInt());
+			Arguments.Add(TEXT("Percentage"), FMath::FloorToFloat(PercentUnderTarget));
+			SampleTime = IConsoleManager::Get().FindConsoleVariable(TEXT("PerfWarn.FineSampleTime"))->GetInt();
+
+			bLowFramerate = true;
+		}
+	}
+
+
+	if (!bLowFramerate && CoarseMovingAverage.IsReliable())
+	{
+		static IConsoleVariable* CVarMinFPS = IConsoleManager::Get().FindConsoleVariable(TEXT("PerfWarn.CoarseMinFPS"));
+		static IConsoleVariable* CVarPercentThreshold = IConsoleManager::Get().FindConsoleVariable(TEXT("PerfWarn.CoarsePercentThreshold"));
+
+		PercentUnderTarget = CoarseMovingAverage.PercentageBelowThreshold(CVarMinFPS->GetFloat());
+
+		if (PercentUnderTarget >= CVarPercentThreshold->GetFloat())
+		{
+			Arguments.Add(TEXT("Framerate"), CVarMinFPS->GetInt());
+			Arguments.Add(TEXT("Percentage"), FMath::FloorToFloat(PercentUnderTarget));
+			SampleTime = IConsoleManager::Get().FindConsoleVariable(TEXT("PerfWarn.CoarseSampleTime"))->GetInt();
+
+			bLowFramerate = true;
+		}
+	}
+
+	if (!bLowFramerate)
+	{
+		HidePerformanceWarning();
+	}
+	else
+	{
+		enum MessagesEnum { Seconds, SecondsPercent, Minute, MinutePercent, Minutes, MinutesPercent };
+		const FText Messages[] = {
+			LOCTEXT("PerformanceWarningInProgress_Seconds",			"Your framerate has been under {Framerate} FPS for the past {SampleTime} seconds.\nYour graphics settings will be reduced next time you restart the editor."),
+			LOCTEXT("PerformanceWarningInProgress_Seconds_Percent", "Your framerate has been under {Framerate} FPS for {Percentage}% of the past {SampleTime} seconds.\nYour graphics settings will be reduced next time you restart the editor."),
+
+			LOCTEXT("PerformanceWarningInProgress_Minute",			"Your framerate has been under {Framerate} FPS for the past minute.\nYour graphics settings will be reduced next time you restart the editor."),
+			LOCTEXT("PerformanceWarningInProgress_Minute_Percent",	"Your framerate has been under {Framerate} FPS for {Percentage}% of the last minute.\nYour graphics settings will be reduced next time you restart the editor."),
+
+			LOCTEXT("PerformanceWarningInProgress_Minutes",			"Your framerate has been below {Framerate} FPS for the past {SampleTime} minutes.\nYour graphics settings will be reduced next time you restart the editor."),
+			LOCTEXT("PerformanceWarningInProgress_Minutes_Percent", "Your framerate has been below {Framerate} FPS for {Percentage}% of the past {SampleTime} minutes.\nYour graphics settings will be reduced next time you restart the editor."),
+		};
+
+		int32 Message;
+		if (SampleTime < 60)
+		{
+			Arguments.Add(TEXT("SampleTime"), SampleTime);
+			Message = Seconds;
+		}
+		else if (SampleTime == 60)
+		{
+			Message = Minute;
+		}
+		else
+		{
+			Arguments.Add(TEXT("SampleTime"), SampleTime / 60);
+			Message = Minutes;
+		}
+
+		// Use the message with the specific percentage on if applicable
+		if (PercentUnderTarget <= 95.f)
+		{
+			++Message;
+		}
+
+		ShowPerformanceWarning(FText::Format(Messages[Message], Arguments));
+	}
+}
+
+void FPerformanceMonitor::DismissAndDisableNotification()
+{
+	bIsNotificationAllowed = false;
+
+	// Set the flag to say we need to run a benchmark on startup
+	GEditor->AccessGameAgnosticSettings().bApplyAutoScalabilityOnStartup = true;
+	GEditor->AccessGameAgnosticSettings().PostEditChange();
+
+	GEditor->DisableRealtimeViewports();
+
+	HidePerformanceWarning();
+}
+
+void FPerformanceMonitor::Reset()
+{
+	FineMovingAverage.Reset();
+	CoarseMovingAverage.Reset();
+
+	HidePerformanceWarning();
+}
+
+void FPerformanceMonitor::UpdateMovingAverageSamplers(IConsoleVariable* Unused)
+{
+	static const int NumberOfSamples = 50;
+	IConsoleManager& Console = IConsoleManager::Get();
+
+	float SampleTime = Console.FindConsoleVariable(TEXT("PerfWarn.FineSampleTime"))->GetFloat();
+	FineMovingAverage = FMovingAverage(NumberOfSamples, SampleTime / NumberOfSamples);
+
+	SampleTime = Console.FindConsoleVariable(TEXT("PerfWarn.CoarseSampleTime"))->GetFloat();
+	CoarseMovingAverage = FMovingAverage(NumberOfSamples, SampleTime / NumberOfSamples);
+}
+
+void FPerformanceMonitor::ShowScalabilityDialog()
+{
+	auto ExistingWindow = ScalabilitySettingsWindowPtr.Pin();
+
+	if (ExistingWindow.IsValid())
+	{
+		ExistingWindow->BringToFront();
+	}
+	else
+	{
+		// Turn off the auto-application of settings on start up if the user is manually tweaking them
+		GEditor->AccessGameAgnosticSettings().bApplyAutoScalabilityOnStartup = false;
+		GEditor->AccessGameAgnosticSettings().PostEditChange();
+
+		// Create the window
+		ScalabilitySettingsWindowPtr = ExistingWindow = SNew(SWindow)
+			.Title(LOCTEXT("PerformanceWarningDialogTitle", "Scalability Options"))
+			.SupportsMaximize(false)
+			.SupportsMinimize(false)
+			.CreateTitleBar(true)
+			.SizingRule(ESizingRule::Autosized);
+
+		ExistingWindow->SetOnWindowClosed(FOnWindowClosed::CreateStatic([](const TSharedRef<SWindow>&, FPerformanceMonitor* PerfWarn){
+			PerfWarn->Reset();
+		}, this));
+
+		ExistingWindow->SetContent(
+			SNew(SScalabilitySettingsDialog)
+			.OnDoneClicked(
+				FOnClicked::CreateStatic([](TWeakPtr<SWindow> Window){
+					auto WindowPin = Window.Pin();
+					if (WindowPin.IsValid())
+					{
+						WindowPin->RequestDestroyWindow();
+					}
+					return FReply::Handled();
+				}, ScalabilitySettingsWindowPtr)
+			)
+		);
+
+
+		TSharedPtr<SWindow> RootWindow = FGlobalTabmanager::Get()->GetRootWindow();
+		if (RootWindow.IsValid())
+		{
+			FSlateApplication::Get().AddModalWindow(ExistingWindow.ToSharedRef(), RootWindow);
+		}
+		else
+		{
+			FSlateApplication::Get().AddWindow(ExistingWindow.ToSharedRef());
+		}
+	}
+}
+
+#undef LOCTEXT_NAMESPACE
\ No newline at end of file
diff --git a/Engine/Source/Editor/UnrealEd/Private/PerformanceMonitor.h b/Engine/Source/Editor/UnrealEd/Private/PerformanceMonitor.h
new file mode 100644
--- /dev/null
+++ b/Engine/Source/Editor/UnrealEd/Private/PerformanceMonitor.h
@@ -0,0 +1,179 @@
+// Copyright 1998-2014 Epic Games, Inc. All Rights Reserved.
+
+/**
+ * Helper class for a calculating a moving average. This works by maintaining an accumulator which is then sampled at regular intervals into the "Samples" array.
+ */
+struct FMovingAverage
+{
+	/** Constructor from a sample size, and sample rate */
+	FMovingAverage(const int32 InSampleSize = 0, const float InSampleRateSeconds = 1.f)
+		: SampleRateSeconds(InSampleRateSeconds)
+		, CurrentSampleAccumulator(0)
+		, CurrentSampleStartTime(0)
+		, SampleSize(InSampleSize)
+		, CurrentSampleCount(0)
+		, NextSampleIndex(0)
+		, SampleAverage(0)
+	{
+		Samples.Reserve(SampleSize);
+	}
+	
+	/** Check if this data is reliable or not. Returns false if the sampler is not populated with enough data */
+	inline float IsReliable() const
+	{
+		return SampleSize != 0 && Samples.Num() == SampleSize;
+	}
+
+	/** Reset this sampler to its default (unpopulated) state */
+	inline void Reset()
+	{
+		*this = FMovingAverage(SampleSize, SampleRateSeconds);
+	}
+
+	/** Get the average of all the samples contained in this sampler */
+	inline float GetAverage() const
+	{
+		return SampleAverage;
+	}
+
+	/** Allow this sampler to tick the frame time, and potentially save a new sample */
+	void Tick(double CurrentTime, float Value)
+	{
+		if (SampleSize == 0)
+		{
+			return;
+		}
+
+		if (CurrentSampleStartTime == 0)
+		{
+			CurrentSampleStartTime = CurrentTime;
+		}
+
+		++CurrentSampleCount;
+		CurrentSampleAccumulator += Value;
+
+		if (CurrentTime - CurrentSampleStartTime > SampleRateSeconds)
+		{
+			// Limit to a minimum of 5 frames per second
+			const float AccumulatorAverage = FMath::Max(CurrentSampleAccumulator / CurrentSampleCount, 5.f);
+
+			if (Samples.Num() == SampleSize)
+			{
+				Samples[NextSampleIndex] = AccumulatorAverage;
+			}
+			else
+			{
+				Samples.Add(AccumulatorAverage);
+			}
+
+			NextSampleIndex = ++NextSampleIndex % SampleSize;
+			ensure(NextSampleIndex >= 0 && NextSampleIndex < SampleSize);
+
+			// calculate the average
+			float Sum = 0.0;
+			for (const auto& Sample : Samples)
+			{
+				Sum += Sample;
+			}
+			SampleAverage = Sum / Samples.Num();
+
+			// reset the accumulator and counter
+			CurrentSampleAccumulator = CurrentSampleCount = 0;
+			CurrentSampleStartTime = CurrentTime;
+		}
+	}
+
+	/** Return the percentage of samples that fall below the specified threshold */
+	float PercentageBelowThreshold(const float Threshold) const
+	{
+		check(IsReliable());
+
+		float NumBelowThreshold = 0;
+		for (const auto& Sample : Samples)
+		{
+			if (Sample < Threshold)
+			{
+				++NumBelowThreshold;
+			}
+		}
+
+		return (NumBelowThreshold / Samples.Num()) * 100.0f;
+	}
+private:
+	/** The number of samples in the current accumulator */
+	int32 CurrentSampleCount;
+	/** The cumulative sum of samples for the current time period */
+	float CurrentSampleAccumulator;
+	/** The time at which we started accumulating the current sample */
+	double CurrentSampleStartTime;
+
+	/** The rate at which to store accumulated samples of FPS in seconds */
+	float SampleRateSeconds;
+	/** The maximum number of accumulated samples to store */
+	int32 SampleSize;
+	/** The actual average stored across all samples */
+	float SampleAverage;
+
+	/** The actual FPS samples */
+	TArray<float> Samples;
+	/** The index of the next sample to write to */
+	int32 NextSampleIndex;
+};
+
+/** Notification class for performance warnings. */
+class FPerformanceMonitor
+	: public FTickableEditorObject
+{
+
+public:
+
+	/** Constructor */
+	FPerformanceMonitor();
+
+	/** Destructor */
+	~FPerformanceMonitor();
+
+protected:
+
+	/** FTickableEditorObject interface */
+	virtual bool IsTickable() const override { return true; }
+	virtual void Tick(float DeltaTime) override;
+	virtual TStatId GetStatId() const override { RETURN_QUICK_DECLARE_CYCLE_STAT(FPerformanceMonitor, STATGROUP_Tickables); }
+
+private:
+
+	/** Starts the notification. */
+	void ShowPerformanceWarning(FText MessageText);
+
+	/** Run a benchmark and auto apply scalability settings */
+	void AutoApplyScalability();
+
+	/** Ends the notification. */
+	void HidePerformanceWarning();
+
+	/** Display the scalability dialog. */
+	void ShowScalabilityDialog();
+
+	/** Dismiss the notification and disable it for this session. */
+	void DismissAndDisableNotification();
+
+	/** Resets the performance warning data and hides the warning */
+	void Reset();
+
+	/** Update the moving average samplers to match the settings specified in the console variables */
+	void UpdateMovingAverageSamplers(IConsoleVariable* Unused = nullptr);
+
+	/** Moving average data for the fine and coarse-grained moving averages */
+	FMovingAverage FineMovingAverage, CoarseMovingAverage;
+
+	/** Tracks the last time the notification was started, used to avoid spamming. */
+	double LastEnableTime;
+
+	/** The notification window ptr */
+	TWeakPtr<SNotificationItem> PerformanceWarningNotificationPtr;
+
+	/** The Scalability Setting window we may be using */
+	TWeakPtr<SWindow> ScalabilitySettingsWindowPtr;
+
+	bool bIsNotificationAllowed;
+};
\ No newline at end of file
diff --git a/Engine/Source/Editor/UnrealEd/Private/SScalabilitySettings.cpp b/Engine/Source/Editor/UnrealEd/Private/SScalabilitySettings.cpp
new file mode 100644
--- /dev/null
+++ b/Engine/Source/Editor/UnrealEd/Private/SScalabilitySettings.cpp
@@ -0,0 +1,206 @@
+// Copyright 1998-2014 Epic Games, Inc. All Rights Reserved.
+
+#include "UnrealEd.h"
+#include "SScalabilitySettings.h"
+
+#define LOCTEXT_NAMESPACE "EngineScalabiltySettings"
+
+ESlateCheckBoxState::Type SScalabilitySettings::IsGroupQualityLevelSelected(const TCHAR* InGroupName, int32 InQualityLevel) const
+{
+	int32 QualityLevel = -1;
+
+	if (FCString::Strcmp(InGroupName, TEXT("ResolutionQuality")) == 0) { QualityLevel = CachedQualityLevels.ResolutionQuality; }
+	else if (FCString::Strcmp(InGroupName, TEXT("ViewDistanceQuality")) == 0) QualityLevel = CachedQualityLevels.ViewDistanceQuality;
+	else if (FCString::Strcmp(InGroupName, TEXT("AntiAliasingQuality")) == 0) QualityLevel = CachedQualityLevels.AntiAliasingQuality;
+	else if (FCString::Strcmp(InGroupName, TEXT("PostProcessQuality")) == 0) QualityLevel = CachedQualityLevels.PostProcessQuality;
+	else if (FCString::Strcmp(InGroupName, TEXT("ShadowQuality")) == 0) QualityLevel = CachedQualityLevels.ShadowQuality;
+	else if (FCString::Strcmp(InGroupName, TEXT("TextureQuality")) == 0) QualityLevel = CachedQualityLevels.TextureQuality;
+	else if (FCString::Strcmp(InGroupName, TEXT("EffectsQuality")) == 0) QualityLevel = CachedQualityLevels.EffectsQuality;
+
+	return (QualityLevel == InQualityLevel) ? ESlateCheckBoxState::Checked : ESlateCheckBoxState::Unchecked;
+}
+
+void SScalabilitySettings::OnGroupQualityLevelChanged(ESlateCheckBoxState::Type NewState, const TCHAR* InGroupName, int32 InQualityLevel)
+{
+	if (FCString::Strcmp(InGroupName, TEXT("ResolutionQuality")) == 0) CachedQualityLevels.ResolutionQuality = InQualityLevel;
+	else if (FCString::Strcmp(InGroupName, TEXT("ViewDistanceQuality")) == 0) CachedQualityLevels.ViewDistanceQuality = InQualityLevel;
+	else if (FCString::Strcmp(InGroupName, TEXT("AntiAliasingQuality")) == 0) CachedQualityLevels.AntiAliasingQuality = InQualityLevel;
+	else if (FCString::Strcmp(InGroupName, TEXT("PostProcessQuality")) == 0) CachedQualityLevels.PostProcessQuality = InQualityLevel;
+	else if (FCString::Strcmp(InGroupName, TEXT("ShadowQuality")) == 0) CachedQualityLevels.ShadowQuality = InQualityLevel;
+	else if (FCString::Strcmp(InGroupName, TEXT("TextureQuality")) == 0) CachedQualityLevels.TextureQuality = InQualityLevel;
+	else if (FCString::Strcmp(InGroupName, TEXT("EffectsQuality")) == 0) CachedQualityLevels.EffectsQuality = InQualityLevel;
+
+	Scalability::SetQualityLevels(CachedQualityLevels);
+	Scalability::SaveState(GEditorUserSettingsIni);
+	GEditor->RedrawAllViewports();
+}
+
+void SScalabilitySettings::OnResolutionScaleChanged(float InValue)
+{
+	CachedQualityLevels.ResolutionQuality = (int32)(FMath::Lerp(Scalability::MinResolutionScale, Scalability::MaxResolutionScale, InValue));
+
+	Scalability::SetQualityLevels(CachedQualityLevels);
+	Scalability::SaveState(GEditorUserSettingsIni);
+	GEditor->RedrawAllViewports();
+}
+
+float SScalabilitySettings::GetResolutionScale() const
+{
+	return (float)(CachedQualityLevels.ResolutionQuality - Scalability::MinResolutionScale) / (float)(Scalability::MaxResolutionScale - Scalability::MinResolutionScale);
+}
+
+FString SScalabilitySettings::GetResolutionScaleString() const
+{
+	return FString::Printf(TEXT("%i%%"), (int32)(100.0f * FMath::Square((float)CachedQualityLevels.ResolutionQuality / 100.0f)));
+}
+
+TSharedRef<SWidget> SScalabilitySettings::MakeButtonWidget(const FText& InName, const TCHAR* InGroupName, int32 InQualityLevel, const FText& InToolTip)
+{
+	return SNew(SCheckBox)
+		.Style(FEditorStyle::Get(), "ToggleButtonCheckbox")
+		.OnCheckStateChanged(this, &SScalabilitySettings::OnGroupQualityLevelChanged, InGroupName, InQualityLevel)
+		.IsChecked(this, &SScalabilitySettings::IsGroupQualityLevelSelected, InGroupName, InQualityLevel)
+		.ToolTipText(InToolTip)
+		.Content()
+		[
+			SNew(STextBlock)
+			.Text(InName)
+		];
+}
+
+TSharedRef<SWidget> SScalabilitySettings::MakeHeaderButtonWidget(const FText& InName, int32 InQualityLevel, const FText& InToolTip)
+{
+	return SNew(SButton)
+		.OnClicked(this, &SScalabilitySettings::OnHeaderClicked, InQualityLevel)
+		.ToolTipText(InToolTip)
+		.Content()
+		[
+			SNew(STextBlock)
+			.Text(InName)
+		];
+}
+
+TSharedRef<SWidget> SScalabilitySettings::MakeAutoButtonWidget()
+{
+	return SNew(SButton)
+		.OnClicked(this, &SScalabilitySettings::OnAutoClicked)
+		.ToolTipText(LOCTEXT("AutoButtonTooltip", "We test your system and try to find the most suitable settings"))
+		.Content()
+		[
+			SNew(STextBlock)
+			.Text(LOCTEXT("AutoLabel", "Auto"))
+		];
+}
+
+FReply SScalabilitySettings::OnHeaderClicked(int32 InQualityLevel)
+{
+	CachedQualityLevels.SetFromSingleQualityLevel(InQualityLevel);
+	Scalability::SetQualityLevels(CachedQualityLevels);
+	Scalability::SaveState(GEditorUserSettingsIni);
+	GEditor->RedrawAllViewports();
+	return FReply::Handled();
+}
+
+FReply SScalabilitySettings::OnAutoClicked()
+{
+	CachedQualityLevels = Scalability::BenchmarkQualityLevels();
+	SetQualityLevels(CachedQualityLevels);
+	Scalability::SaveState(GEditorUserSettingsIni);
+	GEditor->RedrawAllViewports();
+	return FReply::Handled();
+}
+
+SGridPanel::FSlot& SScalabilitySettings::MakeGridSlot(int32 InCol, int32 InRow, int32 InColSpan /*= 1*/, int32 InRowSpan /*= 1*/)
+{
+	float PaddingH = 2.0f;
+	float PaddingV = InRow == 0 ? 8.0f : 2.0f;
+	return SGridPanel::Slot(InCol, InRow)
+		.Padding(PaddingH, PaddingV)
+		.RowSpan(InRowSpan)
+		.ColumnSpan(InColSpan);
+}
+
+void SScalabilitySettings::Construct( const FArguments& InArgs )
+{
+	const FText NamesLow(LOCTEXT("QualityLowLabel", "Low"));
+	const FText NamesMedium(LOCTEXT("QualityMediumLabel", "Medium"));
+	const FText NamesHigh(LOCTEXT("QualityHighLabel", "High"));
+	const FText NamesEpic(LOCTEXT("QualityEpicLabel", "Epic"));
+	const FText NamesAuto(LOCTEXT("QualityAutoLabel", "Auto"));
+
+	auto TitleFont = FEditorStyle::GetFontStyle( FName( "Scalability.TitleFont" ) );
+	auto GroupFont = FEditorStyle::GetFontStyle( FName( "Scalability.GroupFont" ) );
+
+	CachedQualityLevels = Scalability::GetQualityLevels();
+	static float Padding = 1.0f;
+	static float QualityColumnCoeff = 1.0f;
+	static float QualityLevelColumnCoeff = 0.2f;
+
+	this->ChildSlot
+		.HAlign(EHorizontalAlignment::HAlign_Fill)
+		[
+			SNew(SVerticalBox)
+			+SVerticalBox::Slot()
+			[
+				SNew(SGridPanel)
+				.FillColumn(0, QualityColumnCoeff)
+
+				+MakeGridSlot(0,1,5,1) [ SNew (SBorder).BorderImage(FEditorStyle::GetBrush("Scalability.RowBackground")) ]
+				+MakeGridSlot(0,2,5,1) [ SNew (SBorder).BorderImage(FEditorStyle::GetBrush("Scalability.RowBackground")) ]
+				+MakeGridSlot(0,3,5,1) [ SNew (SBorder).BorderImage(FEditorStyle::GetBrush("Scalability.RowBackground")) ]
+				+MakeGridSlot(0,4,5,1) [ SNew (SBorder).BorderImage(FEditorStyle::GetBrush("Scalability.RowBackground")) ]
+				+MakeGridSlot(0,5,5,1) [ SNew (SBorder).BorderImage(FEditorStyle::GetBrush("Scalability.RowBackground")) ]
+				+MakeGridSlot(0,6,5,1) [ SNew (SBorder).BorderImage(FEditorStyle::GetBrush("Scalability.RowBackground")) ]
+				+MakeGridSlot(0,7,5,1) [ SNew (SBorder).BorderImage(FEditorStyle::GetBrush("Scalability.RowBackground")) ]
+
+				+MakeGridSlot(0,0).VAlign(VAlign_Center) [ SNew(STextBlock).Text(LOCTEXT("QualityLabel", "Quality")).Font(TitleFont) ]
+				+MakeGridSlot(1,0) [ MakeHeaderButtonWidget(NamesLow, 0, LOCTEXT("QualityLow", "Set all groups to low quality")) ]
+				+MakeGridSlot(2,0) [ MakeHeaderButtonWidget(NamesMedium, 1, LOCTEXT("QualityMedium", "Set all groups to medium quality")) ]
+				+MakeGridSlot(3,0) [ MakeHeaderButtonWidget(NamesHigh, 2, LOCTEXT("QualityHigh", "Set all groups to high quality")) ]
+				+MakeGridSlot(4,0) [ MakeHeaderButtonWidget(NamesEpic, 3, LOCTEXT("QualityEpic", "Set all groups to epic quality")) ]
+				+MakeGridSlot(5,0) [ MakeAutoButtonWidget() ]
+
+				+MakeGridSlot(0,1) [ SNew(STextBlock).Text(LOCTEXT("ResScaleLabel1", "Resolution Scale")).Font(GroupFont) ]
+				+MakeGridSlot(5,1) [ SNew(STextBlock).Text(this, &SScalabilitySettings::GetResolutionScaleString) ]
+				+MakeGridSlot(1,1,4,1) [ SNew(SSlider).OnValueChanged(this, &SScalabilitySettings::OnResolutionScaleChanged).Value(this, &SScalabilitySettings::GetResolutionScale) ]
+				
+				+MakeGridSlot(0,2) [ SNew(STextBlock).Text(LOCTEXT("ViewDistanceLabel1", "View Distance")).Font(GroupFont) ]
+				+MakeGridSlot(1,2) [ MakeButtonWidget(LOCTEXT("ViewDistanceLabel2", "Near"), TEXT("ViewDistanceQuality"), 0, LOCTEXT("ViewDistanceQualityNear", "Set view distance to near")) ]
+				+MakeGridSlot(2,2) [ MakeButtonWidget(LOCTEXT("ViewDistanceLabel3", "Medium"), TEXT("ViewDistanceQuality"), 1, LOCTEXT("ViewDistanceQualityMed", "Set view distance to medium")) ]
+				+MakeGridSlot(3,2) [ MakeButtonWidget(LOCTEXT("ViewDistanceLabel4", "Far"), TEXT("ViewDistanceQuality"), 2, LOCTEXT("ViewDistanceQualityFar", "Set view distance to far")) ]
+				+MakeGridSlot(4,2) [ MakeButtonWidget(LOCTEXT("ViewDistanceLabel5", "Epic"), TEXT("ViewDistanceQuality"), 3, LOCTEXT("ViewDistanceQualityEpic", "Set view distance to epic")) ]
+
+				+MakeGridSlot(0,3) [ SNew(STextBlock).Text(LOCTEXT("AntiAliasingQualityLabel1", "Anti-Aliasing")).Font(GroupFont) ]
+				+MakeGridSlot(1,3) [ MakeButtonWidget(NamesLow, TEXT("AntiAliasingQuality"), 0, LOCTEXT("AntiAliasingQualityLow", "Set anti-aliasing quality to low")) ]
+				+MakeGridSlot(2,3) [ MakeButtonWidget(NamesMedium, TEXT("AntiAliasingQuality"), 1, LOCTEXT("AntiAliasingQualityMed", "Set anti-aliasing quality to medium")) ]
+				+MakeGridSlot(3,3) [ MakeButtonWidget(NamesHigh, TEXT("AntiAliasingQuality"), 2, LOCTEXT("AntiAliasingQualityHigh", "Set anti-aliasing quality to high")) ]
+				+MakeGridSlot(4,3) [ MakeButtonWidget(NamesEpic, TEXT("AntiAliasingQuality"), 3, LOCTEXT("AntiAliasingQualityEpic", "Set anti-aliasing quality to epic")) ]
+
+				+MakeGridSlot(0,4) [ SNew(STextBlock).Text(LOCTEXT("PostProcessQualityLabel1", "Post Processing")).Font(GroupFont) ]
+				+MakeGridSlot(1,4) [ MakeButtonWidget(NamesLow, TEXT("PostProcessQuality"), 0, LOCTEXT("PostProcessQualityLow", "Set post processing quality to low")) ]
+				+MakeGridSlot(2,4) [ MakeButtonWidget(NamesMedium, TEXT("PostProcessQuality"), 1, LOCTEXT("PostProcessQualityMed", "Set post processing quality to medium")) ]
+				+MakeGridSlot(3,4) [ MakeButtonWidget(NamesHigh, TEXT("PostProcessQuality"), 2, LOCTEXT("PostProcessQualityHigh", "Set post processing quality to high")) ]
+				+MakeGridSlot(4,4) [ MakeButtonWidget(NamesEpic, TEXT("PostProcessQuality"), 3, LOCTEXT("PostProcessQualityEpic", "Set post processing quality to epic")) ]
+
+				+MakeGridSlot(0,5) [ SNew(STextBlock).Text(LOCTEXT("ShadowLabel1", "Shadows")).Font(GroupFont) ]
+				+MakeGridSlot(1,5) [ MakeButtonWidget(NamesLow, TEXT("ShadowQuality"), 0, LOCTEXT("ShadowQualityLow", "Set shadow quality to low")) ]
+				+MakeGridSlot(2,5) [ MakeButtonWidget(NamesMedium, TEXT("ShadowQuality"), 1, LOCTEXT("ShadowQualityMed", "Set shadow quality to medium")) ]
+				+MakeGridSlot(3,5) [ MakeButtonWidget(NamesHigh, TEXT("ShadowQuality"), 2, LOCTEXT("ShadowQualityHigh", "Set shadow quality to high")) ]
+				+MakeGridSlot(4,5) [ MakeButtonWidget(NamesEpic, TEXT("ShadowQuality"), 3, LOCTEXT("ShadowQualityEpic", "Set shadow quality to epic")) ]
+
+				+MakeGridSlot(0,6) [ SNew(STextBlock).Text(LOCTEXT("TextureQualityLabel1", "Textures")).Font(GroupFont) ]
+				+MakeGridSlot(1,6) [ MakeButtonWidget(NamesLow, TEXT("TextureQuality"), 0, LOCTEXT("TextureQualityLow", "Set texture quality to low")) ]
+				+MakeGridSlot(2,6) [ MakeButtonWidget(NamesMedium, TEXT("TextureQuality"), 1, LOCTEXT("TextureQualityMed", "Set texture quality to medium")) ]
+				+MakeGridSlot(3,6) [ MakeButtonWidget(NamesHigh, TEXT("TextureQuality"), 2, LOCTEXT("TextureQualityHigh", "Set texture quality to high")) ]
+				+MakeGridSlot(4,6) [ MakeButtonWidget(NamesEpic, TEXT("TextureQuality"), 3, LOCTEXT("TextureQualityEpic", "Set texture quality to epic")) ]
+
+				+MakeGridSlot(0,7) [ SNew(STextBlock).Text(LOCTEXT("EffectsQualityLabel1", "Effects")).Font(GroupFont) ]
+				+MakeGridSlot(1,7) [ MakeButtonWidget(NamesLow, TEXT("EffectsQuality"), 0, LOCTEXT("EffectsQualityLow", "Set effects quality to low")) ]
+				+MakeGridSlot(2,7) [ MakeButtonWidget(NamesMedium, TEXT("EffectsQuality"), 1, LOCTEXT("EffectsQualityMed", "Set effects quality to medium")) ]
+				+MakeGridSlot(3,7) [ MakeButtonWidget(NamesHigh, TEXT("EffectsQuality"), 2, LOCTEXT("EffectsQualityHigh", "Set effects quality to high")) ]
+				+MakeGridSlot(4,7) [ MakeButtonWidget(NamesEpic, TEXT("EffectsQuality"), 3, LOCTEXT("EffectsQualityEpic", "Set effects quality to epic")) ]
+			]
+		];
+}
+
+#undef LOCTEXT_NAMESPACE
\ No newline at end of file
diff --git a/Engine/Source/Editor/UnrealEd/Private/Settings/SettingsClasses.cpp b/Engine/Source/Editor/UnrealEd/Private/Settings/SettingsClasses.cpp
--- a/Engine/Source/Editor/UnrealEd/Private/Settings/SettingsClasses.cpp
+++ b/Engine/Source/Editor/UnrealEd/Private/Settings/SettingsClasses.cpp
@@ -134,6 +134,7 @@ ULevelEditorMiscSettings::ULevelEditorMiscSettings( const class FPostConstructIn
 	: Super(PCIP)
 {
 	bAutoApplyLightingEnable = true;
+	bMonitorEditorPerformance = true;
 }
 
 
diff --git a/Engine/Source/Editor/UnrealEd/Private/UnrealEdEngine.cpp b/Engine/Source/Editor/UnrealEd/Private/UnrealEdEngine.cpp
--- a/Engine/Source/Editor/UnrealEd/Private/UnrealEdEngine.cpp
+++ b/Engine/Source/Editor/UnrealEd/Private/UnrealEdEngine.cpp
@@ -24,6 +24,7 @@
 #include "Editor/StatsViewer/Public/StatsViewerModule.h"
 #include "SnappingUtils.h"
 #include "PackageAutoSaver.h"
+#include "PerformanceMonitor.h"
 #include "BSPOps.h"
 #include "ComponentVisualizer.h"
 #include "Editor/EditorLiveStreaming/Public/IEditorLiveStreaming.h"
@@ -38,6 +39,8 @@ void UUnrealEdEngine::Init(IEngineLoop* InEngineLoop)
 	PackageAutoSaver.Reset(new FPackageAutoSaver);
 	PackageAutoSaver->LoadRestoreFile();
 
+	PerformanceMonitor = new FPerformanceMonitor;
+
 	// Register for the package dirty state updated callback to catch packages that have been modified and need to be checked out.
 	UPackage::PackageDirtyStateChangedEvent.AddUObject(this, &UUnrealEdEngine::OnPackageDirtyStateUpdated);
 	
@@ -223,6 +226,8 @@ void UUnrealEdEngine::FinishDestroy()
 		PackageAutoSaver.Reset();
 	}
 
+	delete PerformanceMonitor;
+
 	UPackage::PackageDirtyStateChangedEvent.RemoveAll(this);
 	FCoreDelegates::PostGarbageCollect.RemoveAll(this);
 	FCoreDelegates::ColorPickerChanged.RemoveAll(this);
diff --git a/Engine/Source/Editor/UnrealEd/Public/SScalabilitySettings.h b/Engine/Source/Editor/UnrealEd/Public/SScalabilitySettings.h
new file mode 100644
--- /dev/null
+++ b/Engine/Source/Editor/UnrealEd/Public/SScalabilitySettings.h
@@ -0,0 +1,60 @@
+// Copyright 1998-2014 Epic Games, Inc. All Rights Reserved.
+#pragma once
+
+#include "Scalability.h"
+#include "Slate.h"
+
+/**
+ * Scalability settings configuration widget                                                                  
+ **/
+
+class SScalabilitySettings : public SCompoundWidget
+{
+public:
+
+	SLATE_BEGIN_ARGS( SScalabilitySettings )
+	{}
+
+	SLATE_END_ARGS()
+
+	// Widget construction
+	UNREALED_API void Construct( const FArguments& InArgs );
+
+private:
+
+	// Stores state of scalability settings at the point of opening the menu
+	Scalability::FQualityLevels CachedQualityLevels;
+
+	// Checks cached quality levels to see if the specified group is at the specified quality level
+	ESlateCheckBoxState::Type IsGroupQualityLevelSelected(const TCHAR* InGroupName, int32 InQualityLevel) const;
+
+	// Callback for when a particular scalability group has its quality level changed
+	void OnGroupQualityLevelChanged(ESlateCheckBoxState::Type NewState, const TCHAR* InGroupName, int32 InQualityLevel);
+
+	// Callback for when the resolution scale slider changes
+	void OnResolutionScaleChanged(float InValue);
+
+	// Callback to retrieve current resolution scale
+	float GetResolutionScale() const;
+
+	// Callback to retrieve current resolution scale as a display string
+	FString GetResolutionScaleString() const;
+
+	// Makes a button widget for the group quality levels
+	TSharedRef<SWidget> MakeButtonWidget(const FText& InName, const TCHAR* InGroupName, int32 InQualityLevel, const FText& InToolTip);
+
+	// Makes a general quality level header button widget
+	TSharedRef<SWidget> MakeHeaderButtonWidget(const FText& InName, int32 InQualityLevel, const FText& InToolTip);
+
+	// Makes the auto benchmark button
+	TSharedRef<SWidget> MakeAutoButtonWidget();
+
+	// Callback for when a quality level header button is pressed
+	FReply OnHeaderClicked(int32 InQualityLevel);
+
+	// Callback for auto benchmark button
+	FReply OnAutoClicked();
+
+	// Create a gridslot for the group quality level with all the required formatting
+	SGridPanel::FSlot& MakeGridSlot(int32 InCol, int32 InRow, int32 InColSpan = 1, int32 InRowSpan = 1);
+};
\ No newline at end of file
