added message at the end of launch on for iOS on PC when users uses the dummy device All_iOS_on_PC
#ue4

diff --git a/Engine/Source/Editor/UnrealEd/Private/PlayLevel.cpp b/Engine/Source/Editor/UnrealEd/Private/PlayLevel.cpp
--- a/Engine/Source/Editor/UnrealEd/Private/PlayLevel.cpp
+++ b/Engine/Source/Editor/UnrealEd/Private/PlayLevel.cpp
@@ -1198,10 +1198,23 @@ void UEditorEngine::HandleLaunchCompleted(bool Succeeded, double TotalTime, bool
 {
 	if (Succeeded)
 	{
+		FText CompletionMsg;
+		const FString DummyDeviceName(FString::Printf(TEXT("All_iOS_On_%s"), FPlatformProcess::ComputerName()));
+		if (PlayUsingLauncherDeviceId.Left(PlayUsingLauncherDeviceId.Find(TEXT("@"))) == TEXT("IOS") && PlayUsingLauncherDeviceName.Contains(DummyDeviceName))
+		{
+			CompletionMsg = LOCTEXT("LauncherTaskCompleted", "Deployment complete!! Open the app on your device to launch.");
+			TSharedPtr<SNotificationItem> NotificationItem = NotificationItemPtr.Pin();
+			NotificationItem->SetExpireDuration(30.0f);
+		}
+		else
+		{
+			CompletionMsg = LOCTEXT("LauncherTaskCompleted", "Launch complete!!");
+		}
+
 		TGraphTask<FLauncherNotificationTask>::CreateTask().ConstructAndDispatchWhenReady(
 			NotificationItemPtr,
 			SNotificationItem::CS_Success,
-			LOCTEXT("LauncherTaskCompleted", "Launch complete!!")
+			CompletionMsg
 		);
 
 		// analytics for launch on
@@ -1219,10 +1232,20 @@ void UEditorEngine::HandleLaunchCompleted(bool Succeeded, double TotalTime, bool
 	}
 	else
 	{
+		FText CompletionMsg;
+		const FString DummyDeviceName(FString::Printf(TEXT("All_iOS_On_%s"), FPlatformProcess::ComputerName()));
+		if (PlayUsingLauncherDeviceId.Left(PlayUsingLauncherDeviceId.Find(TEXT("@"))) == TEXT("IOS") && PlayUsingLauncherDeviceName.Contains(DummyDeviceName))
+		{
+			CompletionMsg = LOCTEXT("LauncherTaskFailed", "Deployment failed!!");
+		}
+		else
+		{
+			CompletionMsg = LOCTEXT("LauncherTaskFailed", "Launch failed!!");
+		}
 		TGraphTask<FLauncherNotificationTask>::CreateTask().ConstructAndDispatchWhenReady(
 			NotificationItemPtr,
 			SNotificationItem::CS_Fail,
-			LOCTEXT("LauncherTaskFailed", "Launch failed!!")
+			CompletionMsg
 			);
 
 		// analytics for launch on
@@ -1300,7 +1323,11 @@ void UEditorEngine::PlayUsingLauncher()
 		LauncherProfile->SetDeploymentMode(ELauncherProfileDeploymentModes::CopyToDevice);
 		LauncherProfile->SetDeployedDeviceGroup(DeviceGroup);
 		LauncherProfile->SetHideFileServerWindow(false);
-		LauncherProfile->SetLaunchMode(ELauncherProfileLaunchModes::DefaultRole);
+		const FString DummyDeviceName(FString::Printf(TEXT("All_iOS_On_%s"), FPlatformProcess::ComputerName()));
+		if (PlayUsingLauncherDeviceId.Left(PlayUsingLauncherDeviceId.Find(TEXT("@"))) != TEXT("IOS") || !PlayUsingLauncherDeviceName.Contains(DummyDeviceName))
+		{
+			LauncherProfile->SetLaunchMode(ELauncherProfileLaunchModes::DefaultRole);
+		}
 
 		TArray<FString> MapNames;
 		FWorldContext & EditorContext = GetEditorWorldContext();
@@ -1343,7 +1370,7 @@ void UEditorEngine::PlayUsingLauncher()
 
 		Info.Image = FEditorStyle::GetBrush(TEXT("MainFrame.CookContent"));
 		Info.bFireAndForget = false;
-		Info.ExpireDuration = 3.0f;
+		Info.ExpireDuration = 10.0f;
 		Info.Hyperlink = FSimpleDelegate::CreateStatic(HandleHyperlinkNavigate);
 		Info.HyperlinkText = LOCTEXT("ShowOutputLogHyperlink", "Show Output Log");
 		Info.ButtonDetails.Add(
diff --git a/Engine/Source/Runtime/Slate/Private/Widgets/Notifications/SNotificationList.cpp b/Engine/Source/Runtime/Slate/Private/Widgets/Notifications/SNotificationList.cpp
--- a/Engine/Source/Runtime/Slate/Private/Widgets/Notifications/SNotificationList.cpp
+++ b/Engine/Source/Runtime/Slate/Private/Widgets/Notifications/SNotificationList.cpp
@@ -501,6 +501,12 @@ public:
 		}
 	}
 
+	/** Sets the ExpireDuration */
+	virtual void SetExpireDuration(float Duration)
+	{
+		ExpireDuration = Duration;
+	}
+
 protected:
 
 	/* Used to determine whether the button is visible */
diff --git a/Engine/Source/Runtime/Slate/Public/Widgets/Notifications/SNotificationList.h b/Engine/Source/Runtime/Slate/Public/Widgets/Notifications/SNotificationList.h
--- a/Engine/Source/Runtime/Slate/Public/Widgets/Notifications/SNotificationList.h
+++ b/Engine/Source/Runtime/Slate/Public/Widgets/Notifications/SNotificationList.h
@@ -23,6 +23,9 @@ public:
 	/** Sets the text and delegate for the hyperlink */
 	virtual void SetHyperlink( const FSimpleDelegate& InHyperlink, const TAttribute< FText >& InHyperlinkText = TAttribute< FText >() ) = 0;
 
+	/** Sets the ExpireDuration */
+	virtual void SetExpireDuration(float ExpireDuration) = 0;
+
 	/** Gets the visibility state of the throbber, success, and fail images */
 	virtual ECompletionState GetCompletionState() const = 0;
 
