Engine: Fix undo issues with any class derived from UDataAsset due to missing RF_Transactional flag

diff --git a/Engine/Source/Editor/UnrealEd/Classes/Factories/DataAssetFactory.h b/Engine/Source/Editor/UnrealEd/Classes/Factories/DataAssetFactory.h
--- a/Engine/Source/Editor/UnrealEd/Classes/Factories/DataAssetFactory.h
+++ b/Engine/Source/Editor/UnrealEd/Classes/Factories/DataAssetFactory.h
@@ -11,10 +11,10 @@ class UDataAssetFactory : public UFactory
 	UPROPERTY(EditAnywhere, Category=DataAsset)
 	TSubclassOf<UDataAsset> DataAssetClass;
 
-	// Begin UFactory Interface
+	// UFactory interface
 	virtual bool ConfigureProperties() override;
-	virtual UObject* FactoryCreateNew(UClass* Class,UObject* InParent,FName Name,EObjectFlags Flags,UObject* Context,FFeedbackContext* Warn) override;
-	// Begin UFactory Interface
+	virtual UObject* FactoryCreateNew(UClass* Class, UObject* InParent, FName Name, EObjectFlags Flags, UObject* Context, FFeedbackContext* Warn) override;
+	// End of UFactory interface
 };
 
 
diff --git a/Engine/Source/Editor/UnrealEd/Private/Factories/EditorFactories.cpp b/Engine/Source/Editor/UnrealEd/Private/Factories/EditorFactories.cpp
--- a/Engine/Source/Editor/UnrealEd/Private/Factories/EditorFactories.cpp
+++ b/Engine/Source/Editor/UnrealEd/Private/Factories/EditorFactories.cpp
@@ -6600,7 +6600,6 @@ UObject* UObjectLibraryFactory::FactoryCreateNew(UClass* Class, UObject* InParen
 UDataAssetFactory::UDataAssetFactory(const class FPostConstructInitializeProperties& PCIP)
 	: Super(PCIP)
 {
-	
 	bCreateNew = true;
 	bEditAfterNew = true;
 	SupportedClass = UDataAsset::StaticClass();
@@ -6634,13 +6633,13 @@ bool UDataAssetFactory::ConfigureProperties()
 	}
 
 	return bPressedOk;
-};
+}
 
 UObject* UDataAssetFactory::FactoryCreateNew(UClass* Class, UObject* InParent, FName Name, EObjectFlags Flags, UObject* Context, FFeedbackContext* Warn)
 {
-	if(DataAssetClass != NULL)
+	if (DataAssetClass != NULL)
 	{
-		return CastChecked<UDataAsset>(StaticConstructObject(DataAssetClass,InParent,Name,Flags));
+		return CastChecked<UDataAsset>(StaticConstructObject(DataAssetClass, InParent, Name, Flags | RF_Transactional));
 	}
 	else
 	{
diff --git a/Engine/Source/Runtime/Core/Public/UObject/ObjectVersion.h b/Engine/Source/Runtime/Core/Public/UObject/ObjectVersion.h
--- a/Engine/Source/Runtime/Core/Public/UObject/ObjectVersion.h
+++ b/Engine/Source/Runtime/Core/Public/UObject/ObjectVersion.h
@@ -570,6 +570,8 @@ enum EUnrealEngineObjectUE4Version
 	VER_UE4_FOLIAGE_SETTINGS_TYPE,
 	// Lights serialize static shadow depth maps
 	VER_UE4_STATIC_SHADOW_DEPTH_MAPS,
+	// Add RF_Transactional to data assets, fixing undo problems when editing them
+	VER_UE4_ADD_TRANSACTIONAL_TO_DATA_ASSETS,
 
 	// -----<new versions can be added before this line>-------------------------------------------------
 	// - this needs to be the last line (see note below)
diff --git a/Engine/Source/Runtime/Engine/Classes/Engine/DataAsset.h b/Engine/Source/Runtime/Engine/Classes/Engine/DataAsset.h
--- a/Engine/Source/Runtime/Engine/Classes/Engine/DataAsset.h
+++ b/Engine/Source/Runtime/Engine/Classes/Engine/DataAsset.h
@@ -9,4 +9,10 @@ UCLASS(abstract,MinimalAPI)
 class UDataAsset : public UObject
 {
 	GENERATED_UCLASS_BODY()
+
+	// UObject interface
+#if WITH_EDITORONLY_DATA
+	ENGINE_API virtual void Serialize(FArchive& Ar) override;
+#endif
+	// End of UObject interface
 };
\ No newline at end of file
diff --git a/Engine/Source/Runtime/Engine/Private/DataAsset.cpp b/Engine/Source/Runtime/Engine/Private/DataAsset.cpp
--- a/Engine/Source/Runtime/Engine/Private/DataAsset.cpp
+++ b/Engine/Source/Runtime/Engine/Private/DataAsset.cpp
@@ -6,3 +6,15 @@ UDataAsset::UDataAsset(const class FPostConstructInitializeProperties& PCIP)
 	: Super(PCIP)
 {
 }
+
+#if WITH_EDITORONLY_DATA
+void UDataAsset::Serialize(FArchive& Ar)
+{
+	Super::Serialize(Ar);
+
+	if (Ar.IsLoading() && (Ar.UE4Ver() < VER_UE4_ADD_TRANSACTIONAL_TO_DATA_ASSETS))
+	{
+		SetFlags(RF_Transactional);
+	}
+}
+#endif
