Fixing Linux Editor build issues, mostly order of initialization.

diff --git a/Engine/Source/Editor/Kismet/Public/FindInBlueprintManager.h b/Engine/Source/Editor/Kismet/Public/FindInBlueprintManager.h
--- a/Engine/Source/Editor/Kismet/Public/FindInBlueprintManager.h
+++ b/Engine/Source/Editor/Kismet/Public/FindInBlueprintManager.h
@@ -80,9 +80,9 @@ struct FSearchData
 	int32 DDCRetrievalID;
 
 	FSearchData()
-		: bMarkedForDeletion(false)
+		: Blueprint(nullptr)
+		, bMarkedForDeletion(false)
 		, DDCRetrievalID(INDEX_NONE)
-		, Blueprint(nullptr)
 	{
 
 	}
diff --git a/Engine/Source/Editor/KismetCompiler/Private/KismetCompiler.cpp b/Engine/Source/Editor/KismetCompiler/Private/KismetCompiler.cpp
--- a/Engine/Source/Editor/KismetCompiler/Private/KismetCompiler.cpp
+++ b/Engine/Source/Editor/KismetCompiler/Private/KismetCompiler.cpp
@@ -46,12 +46,12 @@ DEFINE_STAT(EKismetCompilerStats_UpdateBlueprintGeneratedClass);
 FKismetCompilerContext::FKismetCompilerContext(UBlueprint* SourceSketch, FCompilerResultsLog& InMessageLog, const FKismetCompilerOptions& InCompilerOptions, TArray<UObject*>* InObjLoaded)
 	: FGraphCompilerContext(InMessageLog)
 	, Schema(NULL)
+	, ObjLoaded(InObjLoaded)
 	, Blueprint(SourceSketch)
 	, NewClass(NULL)
+	, CompileOptions(InCompilerOptions)
 	, ConsolidatedEventGraph(NULL)
 	, UbergraphContext(NULL)
-	, CompileOptions(InCompilerOptions)
-	, ObjLoaded(InObjLoaded)
 {
 	MacroRowMaxHeight = 0;
 
diff --git a/Engine/Source/Editor/UnrealEd/Private/AssetThumbnail.cpp b/Engine/Source/Editor/UnrealEd/Private/AssetThumbnail.cpp
--- a/Engine/Source/Editor/UnrealEd/Private/AssetThumbnail.cpp
+++ b/Engine/Source/Editor/UnrealEd/Private/AssetThumbnail.cpp
@@ -563,10 +563,10 @@ FAssetThumbnail::FAssetThumbnail( UObject* InAsset, uint32 InWidth, uint32 InHei
 }
 
 FAssetThumbnail::FAssetThumbnail( const FAssetData& InAssetData , uint32 InWidth, uint32 InHeight, const TSharedPtr<class FAssetThumbnailPool>& InThumbnailPool )
-	: AssetData ( InAssetData )
+	: ThumbnailPool( InThumbnailPool )
+	, AssetData ( InAssetData )
 	, Width( InWidth )
 	, Height( InHeight )
-	, ThumbnailPool( InThumbnailPool )
 {
 	if ( InThumbnailPool.IsValid() )
 	{
@@ -680,10 +680,10 @@ void FAssetThumbnail::RefreshThumbnail()
 }
 
 FAssetThumbnailPool::FAssetThumbnailPool( uint32 InNumInPool, const TAttribute<bool>& InAreRealTimeThumbnailsAllowed, double InMaxFrameTimeAllowance, uint32 InMaxRealTimeThumbnailsPerFrame )
-	: NumInPool( InNumInPool )
-	, AreRealTimeThumbnailsAllowed( InAreRealTimeThumbnailsAllowed )
-	, MaxFrameTimeAllowance( InMaxFrameTimeAllowance )
+	: AreRealTimeThumbnailsAllowed( InAreRealTimeThumbnailsAllowed )
+	, NumInPool( InNumInPool )
 	, MaxRealTimeThumbnailsPerFrame( InMaxRealTimeThumbnailsPerFrame )
+	, MaxFrameTimeAllowance( InMaxFrameTimeAllowance )
 {
 	FCoreDelegates::OnObjectPropertyChanged.AddRaw(this, &FAssetThumbnailPool::OnObjectPropertyChanged);
 	FCoreDelegates::OnAssetLoaded.AddRaw(this, &FAssetThumbnailPool::OnAssetLoaded);
diff --git a/Engine/Source/Editor/UnrealEd/Private/DragTool_FrustumSelect.cpp b/Engine/Source/Editor/UnrealEd/Private/DragTool_FrustumSelect.cpp
--- a/Engine/Source/Editor/UnrealEd/Private/DragTool_FrustumSelect.cpp
+++ b/Engine/Source/Editor/UnrealEd/Private/DragTool_FrustumSelect.cpp
@@ -485,7 +485,7 @@ bool FDragTool_ActorFrustumSelect::IntersectsFrustum( const UModel& InModel, int
 	// Does the box intersect the frustum
 	bool bIntersects = InFrustum.IntersectBox( NodeBB.GetCenter(), NodeBB.GetExtent(), bFullyContained );
 
-	return bIntersects && (!bUseStrictSelection || bUseStrictSelection && bFullyContained );
+	return bIntersects && (!bUseStrictSelection || (bUseStrictSelection && bFullyContained));
 }
 
 /** 
diff --git a/Engine/Source/Editor/UnrealEd/Private/EditorCsg.cpp b/Engine/Source/Editor/UnrealEd/Private/EditorCsg.cpp
--- a/Engine/Source/Editor/UnrealEd/Private/EditorCsg.cpp
+++ b/Engine/Source/Editor/UnrealEd/Private/EditorCsg.cpp
@@ -24,9 +24,9 @@ public:
 	 * Default constructor, initializing all member variables and iterating to first.
 	 */
 	FStaticBrushIterator( UWorld* InWorld )
-	:	World( InWorld ),
-		ActorIndex( -1 ),
-		ReachedEnd( false )
+	:	ActorIndex( -1 ),
+		ReachedEnd( false ),
+		World( InWorld )
 	{
 		// Iterate to first.
 		++(*this);
diff --git a/Engine/Source/Editor/UnrealEd/Private/Kismet2/KismetReinstanceUtilities.cpp b/Engine/Source/Editor/UnrealEd/Private/Kismet2/KismetReinstanceUtilities.cpp
--- a/Engine/Source/Editor/UnrealEd/Private/Kismet2/KismetReinstanceUtilities.cpp
+++ b/Engine/Source/Editor/UnrealEd/Private/Kismet2/KismetReinstanceUtilities.cpp
@@ -11,8 +11,8 @@
 // FBlueprintCompileReinstancer
 
 FBlueprintCompileReinstancer::FBlueprintCompileReinstancer(UClass* InClassToReinstance, bool bIsBytecodeOnly, bool bSkipGC)
-	: DuplicatedClass(NULL)
-	, ClassToReinstance(InClassToReinstance)
+	: ClassToReinstance(InClassToReinstance)
+	, DuplicatedClass(NULL)
 	, OriginalCDO(NULL)
 	, bHasReinstanced(false)
 	, bSkipGarbageCollection(bSkipGC)
diff --git a/Engine/Source/Editor/UnrealEd/Private/Tests/EditorAutomationTests.cpp b/Engine/Source/Editor/UnrealEd/Private/Tests/EditorAutomationTests.cpp
--- a/Engine/Source/Editor/UnrealEd/Private/Tests/EditorAutomationTests.cpp
+++ b/Engine/Source/Editor/UnrealEd/Private/Tests/EditorAutomationTests.cpp
@@ -2024,10 +2024,10 @@ namespace StaticMeshUVTest
 
 		/** Constructor */
 		FUVTestHelper(FAutomationTestBase* BaseTest) :
-			CurrentTest(BaseTest),
+			CurrentState(EStaticMeshUVTestState::Ready),
 			CurrentPackage(NULL),
 			LoadedPackageCount(0),
-			CurrentState(EStaticMeshUVTestState::Ready)
+			CurrentTest(BaseTest)
 		{
 			FPackageName::FindPackagesInDirectory(ContentPackages, *FPaths::EngineContentDir());
 			FPackageName::FindPackagesInDirectory(ContentPackages, *FPaths::GameContentDir());
diff --git a/Engine/Source/Editor/UnrealEd/Public/CubemapUnwrapUtils.h b/Engine/Source/Editor/UnrealEd/Public/CubemapUnwrapUtils.h
--- a/Engine/Source/Editor/UnrealEd/Public/CubemapUnwrapUtils.h
+++ b/Engine/Source/Editor/UnrealEd/Public/CubemapUnwrapUtils.h
@@ -114,7 +114,8 @@ class UNREALED_API FMipLevelBatchedElementParameters : public FBatchedElementPar
 {
 public:
 	FMipLevelBatchedElementParameters(float InMipLevel, bool bInHDROutput = false)
-		: MipLevel(InMipLevel), bHDROutput(bInHDROutput)
+		: bHDROutput(bInHDROutput)
+		, MipLevel(InMipLevel)
 	{
 	}
 
diff --git a/Engine/Source/Editor/UnrealEd/Public/EditorModeRegistry.h b/Engine/Source/Editor/UnrealEd/Public/EditorModeRegistry.h
--- a/Engine/Source/Editor/UnrealEd/Public/EditorModeRegistry.h
+++ b/Engine/Source/Editor/UnrealEd/Public/EditorModeRegistry.h
@@ -60,6 +60,9 @@ struct FEditorModeInfo
 
 struct IEditorModeFactory : public TSharedFromThis<IEditorModeFactory>
 {
+	/** Virtual destructor */
+	virtual ~IEditorModeFactory() {}
+
 	/**
 	 * Allows mode factories to handle selection change events, and potentially activate/deactivate modes
 
@@ -82,6 +85,7 @@ struct IEditorModeFactory : public TSharedFromThis<IEditorModeFactory>
 struct FEditorModeFactory : IEditorModeFactory
 {
 	FEditorModeFactory(FEditorModeInfo InModeInfo) : ModeInfo(InModeInfo) {}
+	virtual ~FEditorModeFactory() {}
 
 	/** Information pertaining to this factory's mode */
 	FEditorModeInfo ModeInfo;
