Added deprecated API methods for editor modes.

This change adds backwards compatibility for Editor Modes that were created prior to version 4.3 (most specifically, CL#2109245, sha$ba3528c9d4aa3d8cf432f110a871ec3b4e0379ee)

diff --git a/Engine/Source/Editor/UnrealEd/Private/EditorModes.cpp b/Engine/Source/Editor/UnrealEd/Private/EditorModes.cpp
--- a/Engine/Source/Editor/UnrealEd/Private/EditorModes.cpp
+++ b/Engine/Source/Editor/UnrealEd/Private/EditorModes.cpp
@@ -190,6 +190,8 @@ FEdMode::FEdMode()
 	, EditedPropertyIndex( INDEX_NONE )
 	, bEditedPropertyIsTransform( false )
 	, bPendingDeletion( false )
+	/** Deprecated members */
+	, Name(Info.Name), ID(Info.ID), IconBrush(Info.IconBrush), bVisible(Info.bVisible), PriorityOrder(Info.PriorityOrder)
 {
 	bDrawKillZ = true;
 }
@@ -1182,6 +1184,39 @@ FEditorModeTools::~FEditorModeTools()
 	USelection::SelectObjectEvent.RemoveAll(this);
 }
 
+/*********** DEPRECATED ***********/
+void FEditorModeTools::RegisterMode(TSharedRef<FEdMode> ModeToRegister)
+{
+	struct LegacyModeFactory : IEditorModeFactory
+	{
+		LegacyModeFactory(TSharedRef<FEdMode> InMode) : EditorMode(InMode){}
+		TSharedRef<FEdMode> EditorMode;
+
+		virtual void OnSelectionChanged(FEditorModeTools& Tools, UObject* ItemUndergoingChange) const override { }
+		virtual TSharedRef<FEdMode> CreateMode() const override { return EditorMode; }
+		virtual FEditorModeInfo GetModeInfo() const override { return EditorMode->GetModeInfo(); }
+	};
+
+	FEditorModeRegistry::Get().RegisterMode(ModeToRegister->GetID(), MakeShareable(new LegacyModeFactory(ModeToRegister)));
+}
+void FEditorModeTools::UnregisterMode(TSharedRef<FEdMode> ModeToUnregister)
+{
+	FEditorModeRegistry::Get().UnregisterMode(ModeToUnregister->GetID());
+}
+void FEditorModeTools::GetModes(TArray<FEdMode*>& OutModes)
+{
+	OutModes.Empty();
+	for (TSharedPtr<FEdMode> Mode : Modes)
+	{
+		OutModes.Add(Mode.Get());
+	}
+	for (const auto& Pair : RecycledModes)
+	{
+		OutModes.Add(Pair.Value.Get());
+	}
+}
+/*********** END DEPRECATED ***********/
+
 /**
  * Loads the state that was saved in the INI file
  */
@@ -1543,6 +1578,11 @@ bool FEditorModeTools::StartTracking(FEditorViewportClient* InViewportClient, FV
 
 	for ( const auto& Mode : Modes)
 	{
+		// Call the deprecated method if applicable
+		if (InViewportClient->IsLevelEditorClient())
+		{
+			Mode->StartTracking_Deprecated(static_cast<FLevelEditorViewportClient*>(InViewportClient), InViewport);
+		}
 		bTransactionHandled |= Mode->StartTracking(InViewportClient, InViewport);
 	}
 
@@ -1557,6 +1597,11 @@ bool FEditorModeTools::EndTracking(FEditorViewportClient* InViewportClient, FVie
 
 	for ( const auto& Mode : Modes)
 	{
+		// Call the deprecated method if applicable
+		if (InViewportClient->IsLevelEditorClient())
+		{
+			Mode->EndTracking_Deprecated(static_cast<FLevelEditorViewportClient*>(InViewportClient), InViewport);
+		}
 		bTransactionHandled |= Mode->EndTracking(InViewportClient, InViewportClient->Viewport);
 	}
 
@@ -1646,6 +1691,11 @@ bool FEditorModeTools::HandleClick(FEditorViewportClient* InViewportClient,  HHi
 	bool bHandled = false;
 	for( const auto& Mode : Modes)
 	{
+		// Call the deprecated method if applicable
+		if (InViewportClient->IsLevelEditorClient())
+		{
+			bHandled |= Mode->HandleClick_Deprecated(static_cast<FLevelEditorViewportClient*>(InViewportClient), HitProxy, Click);
+		}
 		bHandled |= Mode->HandleClick(InViewportClient, HitProxy, Click);
 	}
 
@@ -1697,6 +1747,12 @@ void FEditorModeTools::Tick( FEditorViewportClient* ViewportClient, float DeltaT
 
 	for( const auto& Mode : Modes)
 	{
+		// Call the deprecated method if applicable
+		if (ViewportClient->IsLevelEditorClient())
+		{
+			Mode->Tick_Deprecated( static_cast<FLevelEditorViewportClient*>(ViewportClient), DeltaTime );
+		}
+
 		Mode->Tick( ViewportClient, DeltaTime );
 	}
 }
@@ -1707,6 +1763,11 @@ bool FEditorModeTools::InputDelta( FEditorViewportClient* InViewportClient,FView
 	bool bHandled = false;
 	for( const auto& Mode : Modes)
 	{
+		// Call the deprecated method if applicable
+		if (InViewportClient->IsLevelEditorClient())
+		{
+			bHandled |= Mode->InputDelta_Deprecated(static_cast<FLevelEditorViewportClient*>(InViewportClient), InViewport, InDrag, InRot, InScale);
+		}
 		bHandled |= Mode->InputDelta( InViewportClient, InViewport, InDrag, InRot, InScale );
 	}
 	return bHandled;
@@ -1718,6 +1779,11 @@ bool FEditorModeTools::CapturedMouseMove( FEditorViewportClient* InViewportClien
 	bool bHandled = false;
 	for( const auto& Mode : Modes)
 	{
+		// Call the deprecated method if applicable
+		if (InViewportClient->IsLevelEditorClient())
+		{
+			bHandled |= Mode->CapturedMouseMove_Deprecated(static_cast<FLevelEditorViewportClient*>(InViewportClient), InViewport, InMouseX, InMouseY);
+		}
 		bHandled |= Mode->CapturedMouseMove( InViewportClient, InViewport, InMouseX, InMouseY );
 	}
 	return bHandled;
@@ -1729,6 +1795,11 @@ bool FEditorModeTools::InputKey(FEditorViewportClient* InViewportClient, FViewpo
 	bool bHandled = false;
 	for (const auto& Mode : Modes)
 	{
+		// Call the deprecated method if applicable
+		if (InViewportClient->IsLevelEditorClient())
+		{
+			bHandled |= Mode->InputKey_Deprecated(static_cast<FLevelEditorViewportClient*>(InViewportClient), Viewport, Key, Event);
+		}
 		bHandled |= Mode->InputKey( InViewportClient, Viewport, Key, Event );
 	}
 	return bHandled;
@@ -1740,6 +1811,11 @@ bool FEditorModeTools::InputAxis(FEditorViewportClient* InViewportClient, FViewp
 	bool bHandled = false;
 	for( const auto& Mode : Modes)
 	{
+		// Call the deprecated method if applicable
+		if (InViewportClient->IsLevelEditorClient())
+		{
+			bHandled |= Mode->InputAxis_Deprecated(static_cast<FLevelEditorViewportClient*>(InViewportClient), Viewport, ControllerId, Key, Delta, DeltaTime);
+		}
 		bHandled |= Mode->InputAxis( InViewportClient, Viewport, ControllerId, Key, Delta, DeltaTime );
 	}
 	return bHandled;
@@ -1750,6 +1826,11 @@ bool FEditorModeTools::MouseEnter( FEditorViewportClient* InViewportClient, FVie
 	bool bHandled = false;
 	for( const auto& Mode : Modes)
 	{
+		// Call the deprecated method if applicable
+		if (InViewportClient->IsLevelEditorClient())
+		{
+			bHandled |= Mode->MouseEnter_Deprecated(static_cast<FLevelEditorViewportClient*>(InViewportClient), Viewport, X, Y);
+		}
 		bHandled |= Mode->MouseEnter( InViewportClient, Viewport, X, Y );
 	}
 	return bHandled;
@@ -1760,6 +1841,11 @@ bool FEditorModeTools::MouseLeave( FEditorViewportClient* InViewportClient, FVie
 	bool bHandled = false;
 	for( const auto& Mode : Modes)
 	{
+		// Call the deprecated method if applicable
+		if (InViewportClient->IsLevelEditorClient())
+		{
+			bHandled |= Mode->MouseLeave_Deprecated(static_cast<FLevelEditorViewportClient*>(InViewportClient), Viewport);
+		}
 		bHandled |= Mode->MouseLeave( InViewportClient, Viewport );
 	}
 	return bHandled;
@@ -1771,6 +1857,11 @@ bool FEditorModeTools::MouseMove( FEditorViewportClient* InViewportClient, FView
 	bool bHandled = false;
 	for( const auto& Mode : Modes)
 	{
+		// Call the deprecated method if applicable
+		if (InViewportClient->IsLevelEditorClient())
+		{
+			bHandled |= Mode->MouseMove_Deprecated(static_cast<FLevelEditorViewportClient*>(InViewportClient), Viewport, X, Y);
+		}
 		bHandled |= Mode->MouseMove( InViewportClient, Viewport, X, Y );
 	}
 	return bHandled;
@@ -1781,6 +1872,11 @@ bool FEditorModeTools::ReceivedFocus( FEditorViewportClient* InViewportClient, F
 	bool bHandled = false;
 	for( const auto& Mode : Modes)
 	{
+		// Call the deprecated method if applicable
+		if (InViewportClient->IsLevelEditorClient())
+		{
+			bHandled |= Mode->ReceivedFocus_Deprecated(static_cast<FLevelEditorViewportClient*>(InViewportClient), Viewport);
+		}
 		bHandled |= Mode->ReceivedFocus( InViewportClient, Viewport );
 	}
 	return bHandled;
@@ -1791,6 +1887,11 @@ bool FEditorModeTools::LostFocus( FEditorViewportClient* InViewportClient, FView
 	bool bHandled = false;
 	for( const auto& Mode : Modes)
 	{
+		// Call the deprecated method if applicable
+		if (InViewportClient->IsLevelEditorClient())
+		{
+			bHandled |= Mode->LostFocus_Deprecated(static_cast<FLevelEditorViewportClient*>(InViewportClient), Viewport);
+		}
 		bHandled |= Mode->LostFocus( InViewportClient, Viewport );
 	}
 	return bHandled;
@@ -1819,6 +1920,11 @@ void FEditorModeTools::DrawHUD( FEditorViewportClient* InViewportClient,FViewpor
 {
 	for( const auto& Mode : Modes)
 	{
+		// Call the deprecated method if applicable
+		if (InViewportClient->IsLevelEditorClient())
+		{
+			Mode->DrawHUD_Deprecated(static_cast<FLevelEditorViewportClient*>(InViewportClient), Viewport, View, Canvas);
+		}
 		Mode->DrawHUD( InViewportClient, Viewport, View, Canvas );
 	}
 }
diff --git a/Engine/Source/Editor/UnrealEd/Private/UnrealEd.cpp b/Engine/Source/Editor/UnrealEd/Private/UnrealEd.cpp
--- a/Engine/Source/Editor/UnrealEd/Private/UnrealEd.cpp
+++ b/Engine/Source/Editor/UnrealEd/Private/UnrealEd.cpp
@@ -36,6 +36,11 @@ class FEditorModeTools& GLevelEditorModeTools()
 	return *EditorModeToolsSingleton;
 }
 
+FEditorModeTools& GEditorModeTools()
+{
+	return GLevelEditorModeTools();
+}
+
 FLevelEditorViewportClient* GCurrentLevelEditingViewportClient = NULL;
 /** Tracks the last level editing viewport client that received a key press. */
 FLevelEditorViewportClient* GLastKeyLevelEditingViewportClient = NULL;
diff --git a/Engine/Source/Editor/UnrealEd/Public/Editor.h b/Engine/Source/Editor/UnrealEd/Public/Editor.h
--- a/Engine/Source/Editor/UnrealEd/Public/Editor.h
+++ b/Engine/Source/Editor/UnrealEd/Public/Editor.h
@@ -556,7 +556,8 @@ UNREALED_API bool IsUniqueObjectName( const FName& InName, UObject* Outer, FText
  * Provides access to the FEditorModeTools for the level editor
  */
 UNREALED_API class FEditorModeTools& GLevelEditorModeTools();
-
+DEPRECATED(4.3, "GEditorModeTools is now deprecated. Please use GLevelEditorModeTools instead.")
+UNREALED_API class FEditorModeTools& GEditorModeTools();
 
 namespace EditorUtilities
 {
diff --git a/Engine/Source/Editor/UnrealEd/Public/EditorModes.h b/Engine/Source/Editor/UnrealEd/Public/EditorModes.h
--- a/Engine/Source/Editor/UnrealEd/Public/EditorModes.h
+++ b/Engine/Source/Editor/UnrealEd/Public/EditorModes.h
@@ -104,6 +104,85 @@ public:
 	FEdMode();
 	virtual ~FEdMode();
 
+protected:
+#ifdef __clang__
+	#pragma clang diagnostic push
+	#pragma clang diagnostic ignored "-Wdeprecated-declarations"
+#else
+	#pragma warning(push)
+	#pragma warning(disable:4996)
+#endif
+
+	friend FEditorModeTools;
+	FEdMode(const FEdMode &);	// Need this to stop the compiler generating one for us and causing a compile warning with the deprecated API
+
+	/** Deprecated API. We use a non-virtual method of a different name to ensure that we remove all deprecated usage from the engine when the deprecated methods are removed. */
+	DEPRECATED(4.3, "FLevelEditorViewportClient* overrides are now deprecated, please use FEditorViewportClient* overrides instead.")
+	virtual bool MouseEnter(FLevelEditorViewportClient* ViewportClient, FViewport* Viewport, int32 x, int32 y) { return false; }
+	bool MouseEnter_Deprecated(FLevelEditorViewportClient* ViewportClient, FViewport* Viewport, int32 x, int32 y) { return MouseEnter(ViewportClient, Viewport, x, y); }
+	DEPRECATED(4.3, "FLevelEditorViewportClient* overrides are now deprecated, please use FEditorViewportClient* overrides instead.")
+	virtual bool MouseLeave(FLevelEditorViewportClient* ViewportClient, FViewport* Viewport) { return false; }
+	bool MouseLeave_Deprecated(FLevelEditorViewportClient* ViewportClient, FViewport* Viewport) { return MouseLeave(ViewportClient, Viewport); }
+	DEPRECATED(4.3, "FLevelEditorViewportClient* overrides are now deprecated, please use FEditorViewportClient* overrides instead.")
+	virtual bool MouseMove(FLevelEditorViewportClient* ViewportClient, FViewport* Viewport, int32 x, int32 y) { return false; }
+	bool MouseMove_Deprecated(FLevelEditorViewportClient* ViewportClient, FViewport* Viewport, int32 x, int32 y) { return MouseMove(ViewportClient, Viewport, x, y); }
+	DEPRECATED(4.3, "FLevelEditorViewportClient* overrides are now deprecated, please use FEditorViewportClient* overrides instead.")
+	virtual bool ReceivedFocus(FLevelEditorViewportClient* ViewportClient, FViewport* Viewport) { return false; }
+	bool ReceivedFocus_Deprecated(FLevelEditorViewportClient* ViewportClient, FViewport* Viewport) { return ReceivedFocus(ViewportClient, Viewport); }
+	DEPRECATED(4.3, "FLevelEditorViewportClient* overrides are now deprecated, please use FEditorViewportClient* overrides instead.")
+	virtual bool LostFocus(FLevelEditorViewportClient* ViewportClient, FViewport* Viewport) { return false; }
+	bool LostFocus_Deprecated(FLevelEditorViewportClient* ViewportClient, FViewport* Viewport) { return LostFocus(ViewportClient, Viewport); }
+	DEPRECATED(4.3, "FLevelEditorViewportClient* overrides are now deprecated, please use FEditorViewportClient* overrides instead.")
+	virtual bool CapturedMouseMove(FLevelEditorViewportClient* InViewportClient, FViewport* InViewport, int32 InMouseX, int32 InMouseY) { return false; }
+	bool CapturedMouseMove_Deprecated(FLevelEditorViewportClient* InViewportClient, FViewport* InViewport, int32 InMouseX, int32 InMouseY) { return CapturedMouseMove(InViewportClient	, InViewport, InMouseX, InMouseY); }
+	DEPRECATED(4.3, "FLevelEditorViewportClient* overrides are now deprecated, please use FEditorViewportClient* overrides instead.")
+	virtual bool InputKey(FLevelEditorViewportClient* ViewportClient, FViewport* Viewport, FKey Key, EInputEvent Event) { return false; }
+	bool InputKey_Deprecated(FLevelEditorViewportClient* ViewportClient, FViewport* Viewport, FKey Key, EInputEvent Event) { return InputKey(ViewportClient, Viewport, Key, Event); }
+	DEPRECATED(4.3, "FLevelEditorViewportClient* overrides are now deprecated, please use FEditorViewportClient* overrides instead.")
+	virtual bool InputAxis(FLevelEditorViewportClient* InViewportClient, FViewport* Viewport, int32 ControllerId, FKey Key, float Delta, float DeltaTime) { return false; }
+	bool InputAxis_Deprecated(FLevelEditorViewportClient* InViewportClient, FViewport* Viewport, int32 ControllerId, FKey Key, float Delta, float DeltaTime) { return InputAxis(InViewportClient, Viewport, ControllerId, Key, Delta, DeltaTime); }
+	DEPRECATED(4.3, "FLevelEditorViewportClient* overrides are now deprecated, please use FEditorViewportClient* overrides instead.")
+	virtual bool InputDelta(FLevelEditorViewportClient* InViewportClient, FViewport* InViewport, FVector& InDrag, FRotator& InRot, FVector& InScale) { return false; }
+	bool InputDelta_Deprecated(FLevelEditorViewportClient* InViewportClient, FViewport* InViewport, FVector& InDrag, FRotator& InRot, FVector& InScale) { return InputDelta(InViewportClient, InViewport, InDrag, InRot, InScale); }
+	DEPRECATED(4.3, "FLevelEditorViewportClient* overrides are now deprecated, please use FEditorViewportClient* overrides instead.")
+	virtual bool StartTracking(FLevelEditorViewportClient* InViewportClient, FViewport* InViewport) { return false; }
+	bool StartTracking_Deprecated(FLevelEditorViewportClient* InViewportClient, FViewport* InViewport) { return StartTracking(InViewportClient, InViewport); }
+	DEPRECATED(4.3, "FLevelEditorViewportClient* overrides are now deprecated, please use FEditorViewportClient* overrides instead.")
+	virtual bool EndTracking(FLevelEditorViewportClient* InViewportClient, FViewport* InViewport) { return false; }
+	bool EndTracking_Deprecated(FLevelEditorViewportClient* InViewportClient, FViewport* InViewport) { return EndTracking(InViewportClient, InViewport); }
+	DEPRECATED(4.3, "FLevelEditorViewportClient* overrides are now deprecated, please use FEditorViewportClient* overrides instead.")
+	virtual void Tick(FLevelEditorViewportClient* ViewportClient, float DeltaTime) {}
+	void Tick_Deprecated(FLevelEditorViewportClient* ViewportClient, float DeltaTime) { Tick(ViewportClient, DeltaTime); }
+	DEPRECATED(4.3, "FLevelEditorViewportClient* overrides are now deprecated, please use FEditorViewportClient* overrides instead.")
+	virtual void DrawBrackets(FLevelEditorViewportClient* ViewportClient, FViewport* Viewport, const FSceneView* View, FCanvas* Canvas) {}
+	void DrawBrackets_Deprecated(FLevelEditorViewportClient* ViewportClient, FViewport* Viewport, const FSceneView* View, FCanvas* Canvas) { DrawBrackets(ViewportClient, Viewport, View, Canvas); }
+	DEPRECATED(4.3, "FLevelEditorViewportClient* overrides are now deprecated, please use FEditorViewportClient* overrides instead.")
+	virtual void DrawHUD(FLevelEditorViewportClient* ViewportClient, FViewport* Viewport, const FSceneView* View, FCanvas* Canvas) {}
+	void DrawHUD_Deprecated(FLevelEditorViewportClient* ViewportClient, FViewport* Viewport, const FSceneView* View, FCanvas* Canvas) { DrawHUD(ViewportClient, Viewport, View, Canvas); }
+	DEPRECATED(4.3, "FLevelEditorViewportClient* overrides are now deprecated, please use FEditorViewportClient* overrides instead.")
+	virtual bool HandleClick(FLevelEditorViewportClient* InViewportClient, HHitProxy *HitProxy, const FViewportClick &Click) { return false; }
+	bool HandleClick_Deprecated(FLevelEditorViewportClient* InViewportClient, HHitProxy *HitProxy, const FViewportClick &Click) { return HandleClick(InViewportClient, HitProxy, Click); }
+
+	/** All the following members have been moved into the "Info" member, assigned on construction from FEditorModeRegistry::CreateMode */
+	DEPRECATED(4.3, "Editor mode information is now provided by FEditorModeRegistry.")
+	FText& Name;
+	DEPRECATED(4.3, "Editor mode information is now provided by FEditorModeRegistry.")
+	FEditorModeID& ID;
+	DEPRECATED(4.3, "Editor mode information is now provided by FEditorModeRegistry.")
+	FSlateIcon& IconBrush;
+	DEPRECATED(4.3, "Editor mode information is now provided by FEditorModeRegistry.")
+	bool& bVisible;
+	DEPRECATED(4.3, "Editor mode information is now provided by FEditorModeRegistry.")
+	int32& PriorityOrder;
+
+#ifdef __clang__
+	#pragma clang diagnostic pop
+#else
+	#pragma warning(pop)
+#endif
+	/** End Deprecated API */
+
+public:
 	virtual void Initialize() {}
 
 	virtual bool MouseEnter( FEditorViewportClient* ViewportClient,FViewport* Viewport,int32 x, int32 y );
@@ -429,6 +508,19 @@ public:
 	FEditorModeTools();
 	virtual ~FEditorModeTools();
 
+	/** Deprecated API */
+	DEPRECATED(4.3, "Modes are no longer registered through FEditorModeTools. Please use FEditorModeRegistry instead.")
+	void RegisterMode(TSharedRef<FEdMode> ModeToRegister);
+	DEPRECATED(4.3, "Modes are no longer registered through FEditorModeTools. Please use FEditorModeRegistry instead.")
+	void UnregisterMode(TSharedRef<FEdMode> ModeToUnregister);
+	DEPRECATED(4.3, "Mode compatability is now determined through FEdMode::IsModeCompatible.")
+	void RegisterCompatibleModes( const FEditorModeID& Mode, const FEditorModeID& CompatibleMode ){}
+	DEPRECATED(4.3, "Mode compatability is now determined through FEdMode::IsModeCompatible.")
+	void UnregisterCompatibleModes(const FEditorModeID& Mode, const FEditorModeID& CompatibleMode){}
+	DEPRECATED(4.3, "GetModes is now deprecated. Please use GetActiveModes or FEditorModeRegistry::GetSortedModeInfo.")
+	void GetModes(TArray<FEdMode*>& OutModes);
+	/** End Deprecated API */
+
 	/**
 	 * Set the default editor mode for these tools
 	 */
