General - Don't need the credits screen text field to have auto wrap text enabled.

diff --git a/Engine/Source/Editor/UnrealEd/Private/CreditsScreen.cpp b/Engine/Source/Editor/UnrealEd/Private/CreditsScreen.cpp
--- a/Engine/Source/Editor/UnrealEd/Private/CreditsScreen.cpp
+++ b/Engine/Source/Editor/UnrealEd/Private/CreditsScreen.cpp
@@ -41,7 +41,6 @@ void SCreditsScreen::Construct(const FArguments& InArgs)
 					.Text(FText::FromString(CreditsText))
 					.TextStyle(FEditorStyle::Get(), "Credits.Normal")
 					.DecoratorStyleSet(&FEditorStyle::Get())
-					.AutoWrapText(true)
 					.Justification(ETextJustify::Center)
 					+ SRichTextBlock::HyperlinkDecorator(TEXT("browser"), this, &SCreditsScreen::OnBrowserLinkClicked)
 				]
