Grid lines now zoom correctly when switching to power of two in preferences.

#ttp 338594 - EDITOR: LEVEL EDITOR: VIEWPORT: Grid lines switch to decimal when zooming a power of 2 grid

#branch UE4

#change Knock on from CL#2039190, variable remained static.

diff --git a/Engine/Source/Editor/UnrealEd/Private/EditorComponents.cpp b/Engine/Source/Editor/UnrealEd/Private/EditorComponents.cpp
--- a/Engine/Source/Editor/UnrealEd/Private/EditorComponents.cpp
+++ b/Engine/Source/Editor/UnrealEd/Private/EditorComponents.cpp
@@ -502,7 +502,7 @@ void FEditorCommonDrawHelper::DrawGridSection(float ViewportGridY,FVector* A,FVe
 	}
 
 	// todo
-	static int32 Exponent = GEditor->IsGridSizePowerOfTwo() ? 8 : 10;
+	int32 Exponent = GEditor->IsGridSizePowerOfTwo() ? 8 : 10;
 
 	const float SizeX = View->ViewRect.Width();
 	const float Zoom = (1.0f / View->ViewMatrices.ProjMatrix.M[0][0]) * 2.0f / SizeX;
