Brushes now have unique labels

#ttp 337884 - EDITOR: BSP: Geometry brushes can get the same default names if the user closes and reopens the project

#branch UE4

diff --git a/Engine/Source/Editor/UnrealEd/Private/BSPOps.cpp b/Engine/Source/Editor/UnrealEd/Private/BSPOps.cpp
--- a/Engine/Source/Editor/UnrealEd/Private/BSPOps.cpp
+++ b/Engine/Source/Editor/UnrealEd/Private/BSPOps.cpp
@@ -612,7 +612,7 @@ ABrush*	FBSPOps::csgAddOperation( ABrush* Actor, uint32 PolyFlags, EBrushType Br
 
 	if( Result->GetBrushBuilder() )
 	{
-		Result->SetActorLabel( FText::Format( NSLOCTEXT("BSPBrushOps", "BrushName", "{0} Brush"), FText::FromString( Result->GetBrushBuilder()->GetClass()->GetDescription() ) ).ToString() );
+		GEditor->SetActorLabelUnique(Result, FText::Format(NSLOCTEXT("BSPBrushOps", "BrushName", "{0} Brush"), FText::FromString(Result->GetBrushBuilder()->GetClass()->GetDescription())).ToString());
 	}
 	// Assign the default material to the brush's polys.
 	for( int32 i=0; i<Result->Brush->Polys->Element.Num(); i++ )
