Fixed a divide by 0 crash from the slow task dialog.
#TTP 338368 - EDITOR: CRASH: Div by zero in FFeedbackContextEditor::UpdateProgress() called from FSimplygonMeshReduction

#branch UE4

#proj Editor.UnrealEd

#change Defaulted the Saved and New Ratios to 0, then tested the divisors for > 0. Added an ensure onto the passed in devisors check, if the user is passing in 0 they should be notified as they are using the interface wrong.

diff --git a/Engine/Source/Editor/UnrealEd/Private/FeedbackContextEditor.cpp b/Engine/Source/Editor/UnrealEd/Private/FeedbackContextEditor.cpp
--- a/Engine/Source/Editor/UnrealEd/Private/FeedbackContextEditor.cpp
+++ b/Engine/Source/Editor/UnrealEd/Private/FeedbackContextEditor.cpp
@@ -439,8 +439,20 @@ void FFeedbackContextEditor::UpdateProgress( int32 Numerator, int32 Denominator
 		if( StatusMessageStack.Num() <= 1 || bCanCancelTask)
 		{
 			// calculate our previous percentage and our new one
-			float SavedRatio = (float)StatusMessage.SavedNumerator / (float)StatusMessage.SavedDenominator;
-			float NewRatio = (float)StatusMessage.ProgressNumerator / (float)StatusMessage.ProgressDenominator;
+			float SavedRatio = 0.0f;
+			float NewRatio = 0.0f;
+
+			// Check divisors to guard against divide by zeros
+			if ( StatusMessage.SavedDenominator > 0 )
+			{
+				SavedRatio = (float)StatusMessage.SavedNumerator / (float)StatusMessage.SavedDenominator;
+			}
+
+			if ( ensure(StatusMessage.ProgressDenominator > 0) )
+			{
+				NewRatio = (float)StatusMessage.ProgressNumerator / (float)StatusMessage.ProgressDenominator;
+			}
+
 			double Now = FPlatformTime::Seconds();
 
 			// update the progress bar if we've moved enough since last time, or we are going to start or end of bar,
