Fixed bad glyphs on the credits screen.

The SRichTextBlock doesn't correctly handle \\r\\n newlines, leading to Slate trying to render a \\r character. I've done a replace from \\r\\n -> \\n before passing the text into the SRichTextBlock.

diff --git a/Engine/Source/Editor/UnrealEd/Private/CreditsScreen.cpp b/Engine/Source/Editor/UnrealEd/Private/CreditsScreen.cpp
--- a/Engine/Source/Editor/UnrealEd/Private/CreditsScreen.cpp
+++ b/Engine/Source/Editor/UnrealEd/Private/CreditsScreen.cpp
@@ -19,6 +19,7 @@ void SCreditsScreen::Construct(const FArguments& InArgs)
 	FString CreditsText;
 	FFileHelper::LoadFileToString(CreditsText, *( FPaths::EngineContentDir() + TEXT("Editor/Credits.rt") ));
 	CreditsText.ReplaceInline(TEXT("%VERSION%"), *Version);
+	CreditsText.ReplaceInline(TEXT("\r\n"), TEXT("\n"));
 
 	ChildSlot
 	[
