Fixed problems with zooming ortho views with Magic Mouse on Mac by disabling scroll gesture processing for ortho views when a mouse button is pressed

diff --git a/Engine/Source/Editor/UnrealEd/Private/EditorViewportClient.cpp b/Engine/Source/Editor/UnrealEd/Private/EditorViewportClient.cpp
--- a/Engine/Source/Editor/UnrealEd/Private/EditorViewportClient.cpp
+++ b/Engine/Source/Editor/UnrealEd/Private/EditorViewportClient.cpp
@@ -2831,7 +2831,9 @@ bool FEditorViewportClient::InputGesture(FViewport* InViewport, EGestureEvent::T
 	case LVT_OrthoXZ:
 	case LVT_OrthoYZ:
 		{
-			if (GestureType == EGestureEvent::Scroll)
+			const bool LeftMouseButtonDown = Viewport->KeyState(EKeys::LeftMouseButton);
+			const bool RightMouseButtonDown = Viewport->KeyState(EKeys::RightMouseButton);
+			if (GestureType == EGestureEvent::Scroll && !LeftMouseButtonDown && !RightMouseButtonDown)
 			{
 				const float UnitsPerPixel = GetOrthoUnitsPerPixel(Viewport);
 				
