Editor: Created a base class with common Options, View, Show, TransformBar, etc... sections for other in-viewport editor toolbars to extend from (WIP, especially Show)

diff --git a/Engine/Source/Editor/LevelEditor/Public/SLevelViewport.h b/Engine/Source/Editor/LevelEditor/Public/SLevelViewport.h
--- a/Engine/Source/Editor/LevelEditor/Public/SLevelViewport.h
+++ b/Engine/Source/Editor/LevelEditor/Public/SLevelViewport.h
@@ -93,6 +93,7 @@ public:
 
 	/** SEditorViewport Interface */
 	virtual void OnFocusViewportToSelection() override;
+	virtual EVisibility GetTransformToolbarVisibility() const override;
 
 	/**
 	 * Called when the maximize command is executed                   
@@ -127,11 +128,6 @@ public:
 	 */
 	EVisibility GetMaximizeToggleVisibility() const;
 
-	/**
-	 * Called to get the visibility of the viewport's transform toolbar.
-	 */
-	EVisibility GetTransformToolbarVisibility() const;
-
 	/**
 	 * @return true if the active viewport is currently being used for play in editor
 	 */
diff --git a/Engine/Source/Editor/UnrealEd/Private/SCommonEditorViewportToolbarBase.cpp b/Engine/Source/Editor/UnrealEd/Private/SCommonEditorViewportToolbarBase.cpp
new file mode 100644
--- /dev/null
+++ b/Engine/Source/Editor/UnrealEd/Private/SCommonEditorViewportToolbarBase.cpp
@@ -0,0 +1,427 @@
+// Copyright 1998-2014 Epic Games, Inc. All Rights Reserved.
+
+#include "UnrealEd.h"
+
+#include "Editor/UnrealEd/Public/SCommonEditorViewportToolbarBase.h"
+
+#include "Editor/UnrealEd/Public/STransformViewportToolbar.h"
+#include "Editor/UnrealEd/Public/EditorShowFlags.h"
+#include "Editor/UnrealEd/Public/SEditorViewport.h"
+#include "Editor/UnrealEd/Public/EditorViewportCommands.h"
+#include "Editor/UnrealEd/Public/SEditorViewportToolBarMenu.h"
+#include "Editor/UnrealEd/Public/SEditorViewportToolBarButton.h"
+#include "Editor/UnrealEd/Public/SEditorViewportViewMenu.h"
+
+#define LOCTEXT_NAMESPACE "LevelViewportToolBar"
+
+//////////////////////////////////////////////////////////////////////////
+// SCommonEditorViewportToolbarBase
+
+void SCommonEditorViewportToolbarBase::Construct(const FArguments& InArgs, TSharedPtr<class ICommonEditorViewportToolbarInfoProvider> InInfoProvider)
+{
+	InfoProviderPtr = InInfoProvider;
+ 	TSharedRef<SEditorViewport> ViewportRef = GetInfoProvider().GetViewportWidget();
+
+	const FMargin ToolbarSlotPadding( 2.0f, 2.0f );
+	const FMargin ToolbarButtonPadding( 2.0f, 0.0f );
+
+	ChildSlot
+	[
+		SNew( SBorder )
+		.BorderImage( FEditorStyle::GetBrush("NoBorder") )
+		// Color and opacity is changed based on whether or not the mouse cursor is hovering over the toolbar area
+		.ColorAndOpacity( this, &SViewportToolBar::OnGetColorAndOpacity )
+		.ForegroundColor( FEditorStyle::GetSlateColor("DefaultForeground") )
+		[
+			SNew( SVerticalBox )
+			+ SVerticalBox::Slot()
+			.AutoHeight()
+			[
+				SNew( SHorizontalBox )
+
+				// Options menu
+				+ SHorizontalBox::Slot()
+				.AutoWidth()
+				.Padding( ToolbarSlotPadding )
+				[
+					SNew( SEditorViewportToolbarMenu )
+					.ParentToolBar( SharedThis( this ) )
+					.Cursor( EMouseCursor::Default )
+					.Image( "EditorViewportToolBar.MenuDropdown" )
+					.OnGetMenuContent(this, &SCommonEditorViewportToolbarBase::GenerateOptionsMenu)
+				]
+
+				// Camera mode menu
+				+ SHorizontalBox::Slot()
+				.AutoWidth()
+				.Padding( ToolbarSlotPadding )
+				[
+					SNew( SEditorViewportToolbarMenu )
+					.ParentToolBar( SharedThis( this ) )
+					.Cursor( EMouseCursor::Default )
+					.Label(this, &SCommonEditorViewportToolbarBase::GetCameraMenuLabel)
+					.LabelIcon(this, &SCommonEditorViewportToolbarBase::GetCameraMenuLabelIcon)
+					.OnGetMenuContent(this, &SCommonEditorViewportToolbarBase::GenerateCameraMenu)
+				]
+
+				// View menu
+				+ SHorizontalBox::Slot()
+				.AutoWidth()
+				.Padding( ToolbarSlotPadding )
+				[
+					MakeViewMenu()
+				]
+
+				// Show menu
+				+ SHorizontalBox::Slot()
+				.AutoWidth()
+				.Padding( ToolbarSlotPadding )
+				[
+					SNew( SEditorViewportToolbarMenu )
+					.Label( LOCTEXT("ShowMenuTitle", "Show") )
+					.Cursor( EMouseCursor::Default )
+					.ParentToolBar( SharedThis( this ) )
+					.OnGetMenuContent(this, &SCommonEditorViewportToolbarBase::GenerateShowMenu)
+				]
+
+				// Transform toolbar
+				+ SHorizontalBox::Slot()
+				.Padding( ToolbarSlotPadding )
+				.HAlign( HAlign_Right )
+				[
+					SNew(STransformViewportToolBar)
+					.Viewport(ViewportRef)
+					.CommandList(ViewportRef->GetCommandList())
+					.Extenders(GetInfoProvider().GetExtenders())
+					.Visibility(ViewportRef, &SEditorViewport::GetTransformToolbarVisibility)
+				]
+			]
+		]
+	];
+
+	SViewportToolBar::Construct(SViewportToolBar::FArguments());
+}
+
+FText SCommonEditorViewportToolbarBase::GetCameraMenuLabel() const
+{
+	FText Label = LOCTEXT("CameraMenuTitle_Default", "Camera");
+
+	switch (GetViewportClient().GetViewportType())
+	{
+		case LVT_Perspective:
+			Label = LOCTEXT("CameraMenuTitle_Perspective", "Perspective");
+			break;
+
+		case LVT_OrthoXY:
+			Label = LOCTEXT("CameraMenuTitle_Top", "Top");
+			break;
+
+		case LVT_OrthoYZ:
+			Label = LOCTEXT("CameraMenuTitle_Side", "Side");
+			break;
+
+		case LVT_OrthoXZ:
+			Label = LOCTEXT("CameraMenuTitle_Front", "Front");
+			break;
+
+		case LVT_OrthoFreelook:
+			Label = LOCTEXT("CameraMenuTitle_OrthoFreelook", "Ortho");
+			break;
+	}
+
+	return Label;
+}
+
+const FSlateBrush* SCommonEditorViewportToolbarBase::GetCameraMenuLabelIcon() const
+{
+	FName Icon = NAME_None;
+
+	switch (GetViewportClient().GetViewportType())
+	{
+		case LVT_Perspective:
+			Icon = FName( "EditorViewport.Perspective" );
+			break;
+
+		case LVT_OrthoXY:
+			Icon = FName( "EditorViewport.Top" );
+			break;
+
+		case LVT_OrthoYZ:
+			Icon = FName( "EditorViewport.Side" );
+			break;
+
+		case LVT_OrthoXZ:
+			Icon = FName( "EditorViewport.Front" );
+			break;
+	}
+
+	return FEditorStyle::GetBrush( Icon );
+}
+
+TSharedRef<SWidget> SCommonEditorViewportToolbarBase::GenerateOptionsMenu() const
+{
+	GetInfoProvider().OnFloatingButtonClicked();
+	TSharedRef<SEditorViewport> ViewportRef = GetInfoProvider().GetViewportWidget();
+
+	const bool bIsPerspective = GetViewportClient().GetViewportType() == LVT_Perspective;
+	
+	const bool bInShouldCloseWindowAfterMenuSelection = true;
+	FMenuBuilder OptionsMenuBuilder(bInShouldCloseWindowAfterMenuSelection, ViewportRef->GetCommandList());
+	{
+		OptionsMenuBuilder.BeginSection("LevelViewportViewportOptions", LOCTEXT("OptionsMenuHeader", "Viewport Options") );
+		{
+			OptionsMenuBuilder.AddMenuEntry( FEditorViewportCommands::Get().ToggleRealTime );
+			OptionsMenuBuilder.AddMenuEntry( FEditorViewportCommands::Get().ToggleStats );
+			OptionsMenuBuilder.AddMenuEntry( FEditorViewportCommands::Get().ToggleFPS );
+
+			if (bIsPerspective)
+			{
+				OptionsMenuBuilder.AddWidget( GenerateFOVMenu(), LOCTEXT("FOVAngle", "Field of View") );
+				OptionsMenuBuilder.AddWidget( GenerateFarViewPlaneMenu(), LOCTEXT("FarViewPlane", "Far View Plane") );
+			}
+		}
+		OptionsMenuBuilder.EndSection();
+	}
+
+	return OptionsMenuBuilder.MakeWidget();
+}
+
+TSharedRef<SWidget> SCommonEditorViewportToolbarBase::GenerateCameraMenu() const
+{
+	GetInfoProvider().OnFloatingButtonClicked();
+	TSharedRef<SEditorViewport> ViewportRef = GetInfoProvider().GetViewportWidget();
+
+	const bool bInShouldCloseWindowAfterMenuSelection = true;
+	FMenuBuilder CameraMenuBuilder( bInShouldCloseWindowAfterMenuSelection, ViewportRef->GetCommandList() );
+
+	// Camera types
+	CameraMenuBuilder.AddMenuEntry( FEditorViewportCommands::Get().Perspective );
+
+	CameraMenuBuilder.BeginSection("LevelViewportCameraType_Ortho", LOCTEXT("CameraTypeHeader_Ortho", "Othographic") );
+		CameraMenuBuilder.AddMenuEntry(FEditorViewportCommands::Get().Top);
+		CameraMenuBuilder.AddMenuEntry(FEditorViewportCommands::Get().Side);
+		CameraMenuBuilder.AddMenuEntry(FEditorViewportCommands::Get().Front);
+	CameraMenuBuilder.EndSection();
+
+	return CameraMenuBuilder.MakeWidget();
+}
+
+TSharedRef<SWidget> SCommonEditorViewportToolbarBase::GenerateShowMenu() const
+{
+	GetInfoProvider().OnFloatingButtonClicked();
+	TSharedRef<SEditorViewport> ViewportRef = GetInfoProvider().GetViewportWidget();
+
+#if 0
+	const FLevelViewportCommands& Actions = FLevelViewportCommands::Get();
+
+	const TArray<FShowFlagData>& ShowFlagData = GetShowFlagMenuItems();
+
+	TArray<CommonEditorViewportUtils::FShowMenuCommand> ShowMenu[SFG_Max];
+
+	// Get each show flag command and put them in their corresponding groups
+	for( int32 ShowFlag = 0; ShowFlag < ShowFlagData.Num(); ++ShowFlag )
+	{
+		const FShowFlagData& SFData = ShowFlagData[ShowFlag];
+		
+		ShowMenu[SFData.Group].Add( Actions.ShowFlagCommands[ ShowFlag ] );
+	}
+#endif
+
+	const bool bInShouldCloseWindowAfterMenuSelection = true;
+	FMenuBuilder ShowMenuBuilder(bInShouldCloseWindowAfterMenuSelection, ViewportRef->GetCommandList());
+	{
+
+#if 0
+		ShowMenuBuilder.AddMenuEntry( Actions.UseDefaultShowFlags );
+		
+		if( ShowMenu[SFG_Normal].Num() > 0 )
+		{
+			// Generate entries for the standard show flags
+			ShowMenuBuilder.BeginSection("LevelViewportShowFlagsCommon", LOCTEXT("CommonShowFlagHeader", "Common") );
+			{
+				for( int32 EntryIndex = 0; EntryIndex < ShowMenu[SFG_Normal].Num(); ++EntryIndex )
+				{
+					ShowMenuBuilder.AddMenuEntry( ShowMenu[SFG_Normal][ EntryIndex ].ShowMenuItem, NAME_None, ShowMenu[SFG_Normal][ EntryIndex ].LabelOverride );
+				}
+			}
+			ShowMenuBuilder.EndSection();
+		}
+
+		// Generate entries for the different show flags groups
+		ShowMenuBuilder.BeginSection("LevelViewportShowFlags");
+		{
+			ShowMenuBuilder.AddSubMenu( LOCTEXT("PostProcessShowFlagsMenu", "Post Processing"), LOCTEXT("PostProcessShowFlagsMenu_ToolTip", "Post process show flags"),
+				FNewMenuDelegate::CreateStatic(&CommonEditorViewportUtils::FillShowMenu, ShowMenu[SFG_PostProcess], 0));
+
+			ShowMenuBuilder.AddSubMenu( LOCTEXT("LightingComponentsShowFlagsMenu", "Lighting Components"), LOCTEXT("LightingComponentsShowFlagsMenu_ToolTip", "Lighting Components show flags"),
+				FNewMenuDelegate::CreateStatic(&CommonEditorViewportUtils::FillShowMenu, ShowMenu[SFG_LightingComponents], 0));
+
+			ShowMenuBuilder.AddSubMenu( LOCTEXT("LightingFeaturesShowFlagsMenu", "Lighting Features"), LOCTEXT("LightingFeaturesShowFlagsMenu_ToolTip", "Lighting Features show flags"),
+				FNewMenuDelegate::CreateStatic(&CommonEditorViewportUtils::FillShowMenu, ShowMenu[SFG_LightingFeatures], 0));
+
+			ShowMenuBuilder.AddSubMenu( LOCTEXT("DeveloperShowFlagsMenu", "Developer"), LOCTEXT("DeveloperShowFlagsMenu_ToolTip", "Developer show flags"),
+				FNewMenuDelegate::CreateStatic(&CommonEditorViewportUtils::FillShowMenu, ShowMenu[SFG_Developer], 0));
+
+			ShowMenuBuilder.AddSubMenu( LOCTEXT("VisualizeShowFlagsMenu", "Visualize"), LOCTEXT("VisualizeShowFlagsMenu_ToolTip", "Visualize show flags"),
+				FNewMenuDelegate::CreateStatic(&CommonEditorViewportUtils::FillShowMenu, ShowMenu[SFG_Visualize], 0));
+
+			ShowMenuBuilder.AddSubMenu( LOCTEXT("AdvancedShowFlagsMenu", "Advanced"), LOCTEXT("AdvancedShowFlagsMenu_ToolTip", "Advanced show flags"),
+				FNewMenuDelegate::CreateStatic(&CommonEditorViewportUtils::FillShowMenu, ShowMenu[SFG_Advanced], 0));
+		}
+		ShowMenuBuilder.EndSection();
+
+
+		FText ShowAllLabel = LOCTEXT("ShowAllLabel", "Show All");
+		FText HideAllLabel = LOCTEXT("HideAllLabel", "Hide All");
+
+		// Show Volumes sub-menu
+		{
+			TArray< FLevelViewportCommands::FShowMenuCommand > ShowVolumesMenu;
+
+			// 'Show All' and 'Hide All' buttons
+			ShowVolumesMenu.Add( FLevelViewportCommands::FShowMenuCommand( Actions.ShowAllVolumes, ShowAllLabel ) );
+			ShowVolumesMenu.Add( FLevelViewportCommands::FShowMenuCommand( Actions.HideAllVolumes, HideAllLabel ) );
+
+			// Get each show flag command and put them in their corresponding groups
+			ShowVolumesMenu += Actions.ShowVolumeCommands;
+
+			ShowMenuBuilder.AddSubMenu( LOCTEXT("ShowVolumesMenu", "Volumes"), LOCTEXT("ShowVolumesMenu_ToolTip", "Show volumes flags"),
+				FNewMenuDelegate::CreateStatic( &FillShowMenu, ShowVolumesMenu, 2 ) );
+		}
+#endif
+	}
+
+	return ShowMenuBuilder.MakeWidget();
+}
+
+TSharedRef<SWidget> SCommonEditorViewportToolbarBase::GenerateFOVMenu() const
+{
+	const float FOVMin = 5.f;
+	const float FOVMax = 170.f;
+
+	return
+		SNew( SBox )
+		.HAlign( HAlign_Right )
+		[
+			SNew( SBox )
+			.Padding( FMargin(4.0f, 0.0f, 0.0f, 0.0f) )
+			.WidthOverride( 100.0f )
+			[
+				SNew(SSpinBox<float>)
+				.Font( FEditorStyle::GetFontStyle( TEXT( "MenuItem.Font" ) ) )
+				.MinValue(FOVMin)
+				.MaxValue(FOVMax)
+				.Value(this, &SCommonEditorViewportToolbarBase::OnGetFOVValue)
+				.OnValueChanged(this, &SCommonEditorViewportToolbarBase::OnFOVValueChanged)
+			]
+		];
+}
+
+float SCommonEditorViewportToolbarBase::OnGetFOVValue() const
+{
+	return GetViewportClient().ViewFOV;
+}
+
+void SCommonEditorViewportToolbarBase::OnFOVValueChanged(float NewValue) const
+{
+	FEditorViewportClient& ViewportClient = GetViewportClient();
+	ViewportClient.FOVAngle = NewValue;
+	ViewportClient.ViewFOV = NewValue;
+	ViewportClient.Invalidate();
+}
+
+TSharedRef<SWidget> SCommonEditorViewportToolbarBase::GenerateFarViewPlaneMenu() const
+{
+	return
+		SNew(SBox)
+		.HAlign(HAlign_Right)
+		[
+			SNew(SBox)
+			.Padding(FMargin(4.0f, 0.0f, 0.0f, 0.0f))
+			.WidthOverride(100.0f)
+			[
+				SNew(SSpinBox<float>)
+				.ToolTipText(LOCTEXT("FarViewPlaneTooltip", "Distance to use as the far view plane, or zero to enable an infinite far view plane"))
+				.MinValue(0.0f)
+				.MaxValue(100000.0f)
+				.Font(FEditorStyle::GetFontStyle(TEXT("MenuItem.Font")))
+				.Value(this, &SCommonEditorViewportToolbarBase::OnGetFarViewPlaneValue)
+				.OnValueChanged(this, &SCommonEditorViewportToolbarBase::OnFarViewPlaneValueChanged)
+			]
+		];
+}
+
+float SCommonEditorViewportToolbarBase::OnGetFarViewPlaneValue() const
+{
+	return GetViewportClient().GetFarClipPlaneOverride();
+}
+
+void SCommonEditorViewportToolbarBase::OnFarViewPlaneValueChanged(float NewValue)
+{
+	GetViewportClient().OverrideFarClipPlane(NewValue);
+}
+
+TSharedPtr<FExtender> SCommonEditorViewportToolbarBase::GetCombinedExtenderList(TSharedRef<FExtender> MenuExtender) const
+{
+	TSharedPtr<FExtender> HostEditorExtenders = GetInfoProvider().GetExtenders();
+
+	TArray<TSharedPtr<FExtender>> Extenders;
+	Extenders.Reserve(2);
+	Extenders.Add(HostEditorExtenders);
+	Extenders.Add(MenuExtender);
+
+	return FExtender::Combine(Extenders);
+}
+
+TSharedPtr<FExtender> SCommonEditorViewportToolbarBase::GetViewMenuExtender() const
+{
+	TSharedRef<FExtender> ViewModeExtender(new FExtender());
+	ViewModeExtender->AddMenuExtension(
+		TEXT("ViewMode"),
+		EExtensionHook::After,
+		GetInfoProvider().GetViewportWidget()->GetCommandList(),
+		FMenuExtensionDelegate::CreateSP(this, &SCommonEditorViewportToolbarBase::CreateViewMenuExtensions));
+
+	return GetCombinedExtenderList(ViewModeExtender);
+}
+
+void SCommonEditorViewportToolbarBase::CreateViewMenuExtensions(FMenuBuilder& MenuBuilder)
+{
+	MenuBuilder.BeginSection("LevelViewportDeferredRendering", LOCTEXT("DeferredRenderingHeader", "Deferred Rendering") );
+	MenuBuilder.EndSection();
+
+	MenuBuilder.BeginSection("LevelViewportCollision", LOCTEXT("CollisionViewModeHeader", "Collision") );
+	{
+		MenuBuilder.AddMenuEntry( FEditorViewportCommands::Get().CollisionPawn, NAME_None, LOCTEXT("CollisionPawnViewModeDisplayName", "Player Collision") );
+		MenuBuilder.AddMenuEntry( FEditorViewportCommands::Get().CollisionVisibility, NAME_None, LOCTEXT("CollisionVisibilityViewModeDisplayName", "Visibility Collision") );
+	}
+	MenuBuilder.EndSection();
+
+//FINDME
+// 	MenuBuilder.BeginSection("LevelViewportLandscape", LOCTEXT("LandscapeHeader", "Landscape") );
+// 	{
+// 		MenuBuilder.AddSubMenu(LOCTEXT("LandscapeLODDisplayName", "LOD"), LOCTEXT("LandscapeLODMenu_ToolTip", "Override Landscape LOD in this viewport"), FNewMenuDelegate::CreateStatic(&Local::BuildLandscapeLODMenu, this), /*Default*/false, FSlateIcon());
+// 	}
+// 	MenuBuilder.EndSection();
+}
+
+ICommonEditorViewportToolbarInfoProvider& SCommonEditorViewportToolbarBase::GetInfoProvider() const
+{
+	return *InfoProviderPtr.Pin().Get();
+}
+
+FEditorViewportClient& SCommonEditorViewportToolbarBase::GetViewportClient() const
+{
+	return *GetInfoProvider().GetViewportWidget()->GetViewportClient().Get();
+}
+
+TSharedRef<SEditorViewportViewMenu> SCommonEditorViewportToolbarBase::MakeViewMenu()
+{
+	TSharedRef<SEditorViewport> ViewportRef = GetInfoProvider().GetViewportWidget();
+
+	return SNew(SEditorViewportViewMenu, ViewportRef, SharedThis(this))
+		.Cursor(EMouseCursor::Default)
+		.MenuExtenders(GetViewMenuExtender());
+}
+
+#undef LOCTEXT_NAMESPACE
diff --git a/Engine/Source/Editor/UnrealEd/Private/SEditorViewport.cpp b/Engine/Source/Editor/UnrealEd/Private/SEditorViewport.cpp
--- a/Engine/Source/Editor/UnrealEd/Private/SEditorViewport.cpp
+++ b/Engine/Source/Editor/UnrealEd/Private/SEditorViewport.cpp
@@ -379,6 +379,11 @@ void SEditorViewport::OnScreenCaptureForProjectThumbnail()
 	}
 }
 
+EVisibility SEditorViewport::GetTransformToolbarVisibility() const
+{
+	return (Client->GetWidgetMode() != FWidget::WM_None) ? EVisibility::Visible : EVisibility::Hidden;
+}
+
 bool SEditorViewport::IsWidgetModeActive( FWidget::EWidgetMode Mode ) const
 {
 	return Client->GetWidgetMode() == Mode;
diff --git a/Engine/Source/Editor/UnrealEd/Public/SCommonEditorViewportToolbarBase.h b/Engine/Source/Editor/UnrealEd/Public/SCommonEditorViewportToolbarBase.h
new file mode 100644
--- /dev/null
+++ b/Engine/Source/Editor/UnrealEd/Public/SCommonEditorViewportToolbarBase.h
@@ -0,0 +1,170 @@
+// Copyright 1998-2014 Epic Games, Inc. All Rights Reserved.
+
+#pragma once
+
+#include "Editor/UnrealEd/Public/SViewportToolBar.h"
+
+// This is the interface that the host of a SCommonEditorViewportToolbarBase must implement
+class ICommonEditorViewportToolbarInfoProvider
+{
+public:
+	// Get the viewport widget
+	virtual TSharedRef<class SEditorViewport> GetViewportWidget() = 0;
+
+// 	FLevelEditorModule& LevelEditorModule = FModuleManager::GetModuleChecked<FLevelEditorModule>(TEXT("LevelEditor"));
+// 	TSharedPtr<FExtender> LevelEditorExtenders = LevelEditorModule.GetMenuExtensibilityManager()->GetAllExtenders();
+	virtual TSharedPtr<FExtender> GetExtenders() const = 0;
+
+	// Called to inform the host that a button was clicked (typically used to focus on a particular viewport in a multi-viewport setup)
+	virtual void OnFloatingButtonClicked() = 0;
+};
+
+namespace CommonEditorViewportUtils
+{
+	struct FShowMenuCommand
+	{
+		TSharedPtr<FUICommandInfo> ShowMenuItem;
+		FText LabelOverride;
+
+		FShowMenuCommand(TSharedPtr<FUICommandInfo> InShowMenuItem, const FText& InLabelOverride)
+			: ShowMenuItem(InShowMenuItem)
+			, LabelOverride(InLabelOverride)
+		{
+		}
+
+		FShowMenuCommand(TSharedPtr<FUICommandInfo> InShowMenuItem)
+			: ShowMenuItem(InShowMenuItem)
+		{
+		}
+	};
+
+	static inline void FillShowMenu(class FMenuBuilder& MenuBuilder, TArray<FShowMenuCommand> MenuCommands, int32 EntryOffset)
+	{
+		// Generate entries for the standard show flags
+		// Assumption: the first 'n' entries types like 'Show All' and 'Hide All' buttons, so insert a separator after them
+		for (int32 EntryIndex = 0; EntryIndex < MenuCommands.Num(); ++EntryIndex)
+		{
+			MenuBuilder.AddMenuEntry(MenuCommands[EntryIndex].ShowMenuItem, NAME_None, MenuCommands[EntryIndex].LabelOverride);
+			if (EntryIndex == EntryOffset - 1)
+			{
+				MenuBuilder.AddMenuSeparator();
+			}
+		}
+	}
+}
+
+/**
+ * A viewport toolbar widget for an asset or level editor that is placed in a viewport
+ */
+class UNREALED_API SCommonEditorViewportToolbarBase : public SViewportToolBar
+{
+public:
+	SLATE_BEGIN_ARGS(SCommonEditorViewportToolbarBase){}
+	SLATE_END_ARGS()
+
+	void Construct(const FArguments& InArgs, TSharedPtr<class ICommonEditorViewportToolbarInfoProvider> InInfoProvider);
+
+private:
+	/**
+	 * Returns the label for the "Camera" tool bar menu, which changes depending on the viewport type
+	 *
+	 * @return	Label to use for this menu label
+	 */
+	FText GetCameraMenuLabel() const;
+
+	/**
+	 * Returns the label icon for the "Camera" tool bar menu, which changes depending on the viewport type
+	 *
+	 * @return	Label icon to use for this menu label
+	 */
+	const FSlateBrush* GetCameraMenuLabelIcon() const;
+
+	/**
+	 * Returns the label for the "View" tool bar menu, which changes depending on viewport show flags
+	 *
+	 * @return	Label to use for this menu label
+	 */
+	FText GetViewMenuLabel() const;
+
+	/**
+	 * Returns the label icon for the "View" tool bar menu, which changes depending on viewport show flags
+	 *
+	 * @return	Label icon to use for this menu label
+	 */
+	const FSlateBrush* GetViewMenuLabelIcon() const;
+
+	/**
+	 * Generates the toolbar options menu content 
+	 *
+	 * @return The widget containing the options menu content
+	 */
+	TSharedRef<SWidget> GenerateOptionsMenu() const;
+
+	/**
+	 * Generates the toolbar camera menu content 
+	 *
+	 * @return The widget containing the view menu content
+	 */
+	TSharedRef<SWidget> GenerateCameraMenu() const;
+
+	/**
+	 * Generates the toolbar view menu content 
+	 *
+	 * @return The widget containing the view menu content
+	 */
+	TSharedRef<SWidget> GenerateViewMenu() const;
+
+	/**
+	 * Generates the toolbar show menu content 
+	 *
+	 * @return The widget containing the show menu content
+	 */
+	TSharedRef<SWidget> GenerateShowMenu() const;
+
+	/**
+	 * @return The widget containing the perspective only FOV window.
+	 */
+	TSharedRef<SWidget> GenerateFOVMenu() const;
+
+	/** Called by the FOV slider in the perspective viewport to get the FOV value */
+	float OnGetFOVValue() const;
+
+	/**
+	 * @return The widget containing the far view plane slider.
+	 */
+	TSharedRef<SWidget> GenerateFarViewPlaneMenu() const;
+
+	/** Called by the far view plane slider in the perspective viewport to get the far view plane value */
+	float OnGetFarViewPlaneValue() const;
+
+	/** Called when the far view plane slider is adjusted in the perspective viewport */
+	void OnFarViewPlaneValueChanged( float NewValue );
+
+protected:
+	// Merges the extender list from the host with the specified extender and returns the results
+	TSharedPtr<FExtender> GetCombinedExtenderList(TSharedRef<FExtender> MenuExtender) const;
+
+	/** Gets the extender for the view menu */
+	virtual TSharedPtr<FExtender> GetViewMenuExtender() const;
+
+	void CreateViewMenuExtensions(FMenuBuilder& MenuBuilder);
+
+protected:
+	// Returns the info provider for this viewport
+	ICommonEditorViewportToolbarInfoProvider& GetInfoProvider() const;
+
+	// Get the viewport client
+	class FEditorViewportClient& GetViewportClient() const;
+
+protected:
+	// Creates the view menu widget (override point for children)
+	virtual TSharedRef<class SEditorViewportViewMenu> MakeViewMenu();
+
+	/** Called when the FOV slider is adjusted in the perspective viewport */
+	virtual void OnFOVValueChanged(float NewValue) const;
+
+private:
+	/** The viewport that we are in */
+	TWeakPtr<class ICommonEditorViewportToolbarInfoProvider> InfoProviderPtr;
+};
+
diff --git a/Engine/Source/Editor/UnrealEd/Public/SEditorViewport.h b/Engine/Source/Editor/UnrealEd/Public/SEditorViewport.h
--- a/Engine/Source/Editor/UnrealEd/Public/SEditorViewport.h
+++ b/Engine/Source/Editor/UnrealEd/Public/SEditorViewport.h
@@ -33,6 +33,11 @@ public:
 
 	TSharedPtr<FEditorViewportClient> GetViewportClient() const { return Client; }
 
+	/*
+	 * Controls the visibility of the widget transform toolbar, if there is an associated toolbar
+	 */
+	virtual EVisibility GetTransformToolbarVisibility() const;
+
 protected:
 	virtual TSharedRef<FEditorViewportClient> MakeEditorViewportClient() = 0;
 	virtual TSharedPtr<SWidget> MakeViewportToolbar() = 0;
