Editor: Added drop shadow and scaling based on thumbnail tile size to IES lumens display to make it readable

diff --git a/Engine/Source/Editor/UnrealEd/Private/ThumbnailRendering/TextureThumbnailRenderer.cpp b/Engine/Source/Editor/UnrealEd/Private/ThumbnailRendering/TextureThumbnailRenderer.cpp
--- a/Engine/Source/Editor/UnrealEd/Private/ThumbnailRendering/TextureThumbnailRenderer.cpp
+++ b/Engine/Source/Editor/UnrealEd/Private/ThumbnailRendering/TextureThumbnailRenderer.cpp
@@ -98,7 +98,9 @@ void UTextureThumbnailRenderer::Draw(UObject* Object, int32 X, int32 Y, uint32 W
 			// Brightness in Lumens
 			FText BrightnessText = FText::AsNumber( Brightness );
 			FCanvasTextItem TextItem( FVector2D( 5.0f, 5.0f ), BrightnessText, GEngine->GetLargeFont(), FLinearColor::White );
-			TextItem.Draw( Canvas );
+			TextItem.EnableShadow(FLinearColor::Black);
+			TextItem.Scale = FVector2D(Width / 128.0f, Height / 128.0f);
+			TextItem.Draw(Canvas);
 		}
 	}
 }
