Editor: Enable box selecting other kinds of actors in orthographic views (based on the bounds of any visible primitive components)

diff --git a/Engine/Source/Editor/UnrealEd/Private/LevelEditorViewport.cpp b/Engine/Source/Editor/UnrealEd/Private/LevelEditorViewport.cpp
--- a/Engine/Source/Editor/UnrealEd/Private/LevelEditorViewport.cpp
+++ b/Engine/Source/Editor/UnrealEd/Private/LevelEditorViewport.cpp
@@ -3190,7 +3190,7 @@ bool FLevelEditorViewportClient::ComponentIsTouchingSelectionBox(AActor* InActor
 	}
 	else if ( !bConsiderOnlyBSP )
 	{
-		bool bSelectByBoundingBox = false;
+		bool bSelectByBoundingBox = true;
 
 		if( bSelectByBoundingBox )
 		{
