Editor performance monitoring is now disabled in debug

diff --git a/Engine/Source/Editor/UnrealEd/Classes/Settings/LevelEditorMiscSettings.h b/Engine/Source/Editor/UnrealEd/Classes/Settings/LevelEditorMiscSettings.h
--- a/Engine/Source/Editor/UnrealEd/Classes/Settings/LevelEditorMiscSettings.h
+++ b/Engine/Source/Editor/UnrealEd/Classes/Settings/LevelEditorMiscSettings.h
@@ -55,7 +55,7 @@ public:
 
 public:	
 
-	/** When turned on, the editor will constantly monitor performance and adjust scalability settings for you when performance drops */
+	/** When turned on, the editor will constantly monitor performance and adjust scalability settings for you when performance drops (disabled in debug) */
 	UPROPERTY(EditAnywhere, config, Category=Performance)
 	bool bMonitorEditorPerformance;
 
diff --git a/Engine/Source/Editor/UnrealEd/Private/PerformanceMonitor.cpp b/Engine/Source/Editor/UnrealEd/Private/PerformanceMonitor.cpp
--- a/Engine/Source/Editor/UnrealEd/Private/PerformanceMonitor.cpp
+++ b/Engine/Source/Editor/UnrealEd/Private/PerformanceMonitor.cpp
@@ -223,6 +223,17 @@ void FPerformanceMonitor::HidePerformanceWarning()
 	}
 }
 
+
+bool FPerformanceMonitor::IsTickable() const
+{
+#if UE_BUILD_DEBUG
+	// We don't do this in debug
+	return false;
+#else
+	return true;
+#endif
+}
+
 void FPerformanceMonitor::Tick(float DeltaTime)
 {
 	if (GEngine->ShouldThrottleCPUUsage())
diff --git a/Engine/Source/Editor/UnrealEd/Private/PerformanceMonitor.h b/Engine/Source/Editor/UnrealEd/Private/PerformanceMonitor.h
--- a/Engine/Source/Editor/UnrealEd/Private/PerformanceMonitor.h
+++ b/Engine/Source/Editor/UnrealEd/Private/PerformanceMonitor.h
@@ -136,7 +136,7 @@ public:
 protected:
 
 	/** FTickableEditorObject interface */
-	virtual bool IsTickable() const override { return true; }
+	virtual bool IsTickable() const override;
 	virtual void Tick(float DeltaTime) override;
 	virtual TStatId GetStatId() const override { RETURN_QUICK_DECLARE_CYCLE_STAT(FPerformanceMonitor, STATGROUP_Tickables); }
 
