A persistent guid is stored as a part of property name for UDS variables

problem:
There is a user-defined structure S with a variable V. The default value of V is 1.
In blueprint B there is a variable M of type S. In M the V has value 2.
When V is renamed, the V value in M is reset to 1.

diff --git a/Engine/Source/Editor/UnrealEd/Private/Kismet2/StructureEditorUtils.cpp b/Engine/Source/Editor/UnrealEd/Private/Kismet2/StructureEditorUtils.cpp
--- a/Engine/Source/Editor/UnrealEd/Private/Kismet2/StructureEditorUtils.cpp
+++ b/Engine/Source/Editor/UnrealEd/Private/Kismet2/StructureEditorUtils.cpp
@@ -188,8 +188,10 @@ bool FStructureEditorUtils::CanHaveAMemberVariableOfType(const UUserDefinedStruc
 
 struct FMemberVariableNameHelper
 {
-	static FName Generate(UUserDefinedStruct* Struct, const FString& NameBase)
+	static FName Generate(UUserDefinedStruct* Struct, const FString& NameBase, const FGuid Guid, FString* OutFriendlyName = NULL)
 	{
+		check(Struct);
+
 		FString Result;
 		if (!NameBase.IsEmpty())
 		{
@@ -205,15 +207,36 @@ struct FMemberVariableNameHelper
 			Result = TEXT("MemberVar");
 		}
 
-		if(Struct)
+		const uint32 UniqueNameId = CastChecked<UUserDefinedStructEditorData>(Struct->EditorData)->GenerateUniqueNameIdForMemberVariable();
+		const FString FriendlyName = FString::Printf(TEXT("%s_%u"), *Result, UniqueNameId);
+		if (OutFriendlyName)
 		{
-			const uint32 UniqueNameId = CastChecked<UUserDefinedStructEditorData>(Struct->EditorData)->GenerateUniqueNameIdForMemberVariable();
-			Result = FString::Printf(TEXT("%s_%u"), *Result, UniqueNameId);
+			*OutFriendlyName = FriendlyName;
 		}
-		FName NameResult = FName(*Result);
+		const FName NameResult = *FString::Printf(TEXT("%s_%s"), *FriendlyName, *Guid.ToString(EGuidFormats::Digits));
 		check(NameResult.IsValidXName(INVALID_OBJECTNAME_CHARACTERS));
 		return NameResult;
 	}
+
+	static FGuid GetGuidFromName(const FName Name)
+	{
+		const FString NameStr = Name.ToString();
+		const int32 GuidStrLen = 32;
+		if (NameStr.Len() > (GuidStrLen + 1))
+		{
+			const int32 UnderscoreIndex = NameStr.Len() - GuidStrLen - 1;
+			if (TCHAR('_') == NameStr[UnderscoreIndex])
+			{
+				const FString GuidStr = NameStr.Right(GuidStrLen);
+				FGuid Guid;
+				if (FGuid::ParseExact(GuidStr, EGuidFormats::Digits, Guid))
+				{
+					return Guid;
+				}
+			}
+		}
+		return FGuid();
+	}
 };
 
 bool FStructureEditorUtils::AddVariable(UUserDefinedStruct* Struct, const FEdGraphPinType& VarType)
@@ -230,16 +253,17 @@ bool FStructureEditorUtils::AddVariable(UUserDefinedStruct* Struct, const FEdGra
 			return false;
 		}
 
-		const FName VarName = FMemberVariableNameHelper::Generate(Struct, FString());
+		const FGuid Guid = FGuid::NewGuid();
+		FString DisplayName;
+		const FName VarName = FMemberVariableNameHelper::Generate(Struct, FString(), Guid, &DisplayName);
 		check(NULL == GetVarDesc(Struct).FindByPredicate(FStructureEditorUtils::FFindByNameHelper<FStructVariableDescription>(VarName)));
-		const FString DisplayName = VarName.ToString();
 		check(IsUniqueVariableDisplayName(Struct, DisplayName));
 
 		FStructVariableDescription NewVar;
 		NewVar.VarName = VarName;
 		NewVar.FriendlyName = DisplayName;
 		NewVar.SetPinType(VarType);
-		NewVar.VarGuid = FGuid::NewGuid();
+		NewVar.VarGuid = Guid;
 		NewVar.bDontEditoOnInstance = false;
 		NewVar.bInvalidMember = false;
 		GetVarDesc(Struct).Add(NewVar);
@@ -281,14 +305,21 @@ bool FStructureEditorUtils::RenameVariable(UUserDefinedStruct* Struct, FGuid Var
 	if (Struct)
 	{
 		auto VarDesc = GetVarDescByGuid(Struct, VarGuid);
-		if (VarDesc &&
-			FName(*NewDisplayNameStr).IsValidXName(INVALID_OBJECTNAME_CHARACTERS) &&
+		if (VarDesc && FName(*NewDisplayNameStr).IsValidXName(INVALID_OBJECTNAME_CHARACTERS) &&
 			IsUniqueVariableDisplayName(Struct, NewDisplayNameStr))
 		{
 			const FScopedTransaction Transaction(LOCTEXT("RenameVariable", "Rename Variable"));
 			ModifyStructData(Struct);
 
 			VarDesc->FriendlyName = NewDisplayNameStr;
+			//>>> TEMPORARY it's more important to prevent changes in structs instances, than to have consistent names
+			if (GetGuidFromPropertyName(VarDesc->VarName).IsValid())
+			//<<< TEMPORARY
+			{
+				const FName NewName = FMemberVariableNameHelper::Generate(Struct, NewDisplayNameStr, VarGuid);
+				check(NULL == GetVarDesc(Struct).FindByPredicate(FFindByNameHelper<FStructVariableDescription>(NewName)))
+				VarDesc->VarName = NewName;
+			}
 			OnStructureChanged(Struct);
 			return true;
 		}
@@ -296,6 +327,7 @@ bool FStructureEditorUtils::RenameVariable(UUserDefinedStruct* Struct, FGuid Var
 	return false;
 }
 
+
 bool FStructureEditorUtils::ChangeVariableType(UUserDefinedStruct* Struct, FGuid VarGuid, const FEdGraphPinType& NewType)
 {
 	if (Struct)
@@ -316,7 +348,7 @@ bool FStructureEditorUtils::ChangeVariableType(UUserDefinedStruct* Struct, FGuid
 				const FScopedTransaction Transaction(LOCTEXT("ChangeVariableType", "Change Variable Type"));
 				ModifyStructData(Struct);
 
-				VarDesc->VarName = FMemberVariableNameHelper::Generate(Struct, VarDesc->FriendlyName);
+				VarDesc->VarName = FMemberVariableNameHelper::Generate(Struct, VarDesc->FriendlyName, VarDesc->VarGuid);
 				VarDesc->DefaultValue = FString();
 				VarDesc->SetPinType(NewType);
 
@@ -652,4 +684,15 @@ FGuid FStructureEditorUtils::GetGuidForProperty(const UProperty* Property)
 	return VarDesc ? VarDesc->VarGuid : FGuid();
 }
 
+UProperty* FStructureEditorUtils::GetPropertyByGuid(const UUserDefinedStruct* Struct, const FGuid VarGuid)
+{
+	const auto VarDesc = GetVarDescByGuid(Struct, VarGuid);
+	return VarDesc ? FindField<UProperty>(Struct, VarDesc->VarName) : NULL;
+}
+
+FGuid FStructureEditorUtils::GetGuidFromPropertyName(const FName Name)
+{
+	return FMemberVariableNameHelper::GetGuidFromName(Name);
+}
+
 #undef LOCTEXT_NAMESPACE
diff --git a/Engine/Source/Editor/UnrealEd/Public/Kismet2/StructureEditorUtils.h b/Engine/Source/Editor/UnrealEd/Public/Kismet2/StructureEditorUtils.h
--- a/Engine/Source/Editor/UnrealEd/Public/Kismet2/StructureEditorUtils.h
+++ b/Engine/Source/Editor/UnrealEd/Public/Kismet2/StructureEditorUtils.h
@@ -80,7 +80,7 @@ public:
 
 	static bool Is3dWidgetEnabled(const UUserDefinedStruct* Struct, FGuid VarGuid);
 
-	//MISC
+	//GUID AND VAR DESC
 	static TArray<FStructVariableDescription>& GetVarDesc(UUserDefinedStruct* Struct);
 
 	static const TArray<FStructVariableDescription>& GetVarDesc(const UUserDefinedStruct* Struct);
@@ -95,14 +95,20 @@ public:
 		return Struct ? GetVarDesc(Struct).FindByPredicate(FFindByGuidHelper<FStructVariableDescription>(VarGuid)) : NULL;
 	}
 
+	static FGuid GetGuidForProperty(const UProperty* Property);
+
+	static UProperty* GetPropertyByGuid(const UUserDefinedStruct* Struct, FGuid VarGuid);
+
+	static FGuid GetGuidFromPropertyName(FName Name);
+
+	//MISC
 	static void ModifyStructData(UUserDefinedStruct* Struct);
 
 	static bool UserDefinedStructEnabled();
 
 	static void RemoveInvalidStructureMemberVariableFromBlueprint(UBlueprint* Blueprint);
 
-	static FGuid GetGuidForProperty(const UProperty* Property);
-
+	//DEFAULT VALUE
 	/*
 	 Default values for member variables in User Defined Structure are stored in meta data "MakeStructureDefaultValue"
 	 The following functions are used to fill an instance of user defined struct with those default values.
diff --git a/Engine/Source/Runtime/CoreUObject/Private/UObject/Class.cpp b/Engine/Source/Runtime/CoreUObject/Private/UObject/Class.cpp
--- a/Engine/Source/Runtime/CoreUObject/Private/UObject/Class.cpp
+++ b/Engine/Source/Runtime/CoreUObject/Private/UObject/Class.cpp
@@ -718,13 +718,6 @@ void UStruct::SerializeTaggedProperties(FArchive& Ar, uint8* Data, UStruct* Defa
 			}
 			PropertyName = Tag.Name;
 
-			static FName DEBUG_Name(TEXT("NewVar"));
-			if (Tag.Name == DEBUG_Name)
-			{
-				static volatile int32 xx = 0;
-				xx++;
-			}
-
 			// Move to the next property to be serialized
 			if( AdvanceProperty && --RemainingArrayDim <= 0 )
 			{
@@ -813,7 +806,12 @@ void UStruct::SerializeTaggedProperties(FArchive& Ar, uint8* Data, UStruct* Defa
 
 				RemainingArrayDim = Property ? Property->ArrayDim : 0;
 			}
-
+#if WITH_EDITOR
+			if (!Property)
+			{
+				Property = CustomFindProperty(Tag.Name);
+			}
+#endif // WITH_EDITOR
 			// Check if this is a struct property and we have a redirector
 			if (Tag.Type==NAME_StructProperty && Property != NULL && Tag.Type == Property->GetID())
 			{
@@ -1033,8 +1031,8 @@ void UStruct::SerializeTaggedProperties(FArchive& Ar, uint8* Data, UStruct* Defa
 				}
 				else
 				{
-					UE_LOG(LogClass, Warning, TEXT("Array Inner Type mismatch in %s of %s - Previous (%s) Current(%s) for package:  %s"), *Tag.Name.ToString(), *GetName(), *Tag.InnerType.ToString(), *CastChecked<UArrayProperty>(Property)->Inner->GetID().ToString(), *Ar.GetArchiveName() );
-				}
+				UE_LOG(LogClass, Warning, TEXT("Array Inner Type mismatch in %s of %s - Previous (%s) Current(%s) for package:  %s"), *Tag.Name.ToString(), *GetName(), *Tag.InnerType.ToString(), *CastChecked<UArrayProperty>(Property)->Inner->GetID().ToString(), *Ar.GetArchiveName() );
+			}
 			}
 			else if( Tag.Type==NAME_StructProperty && Tag.StructName!=CastChecked<UStructProperty>(Property)->Struct->GetFName() 
 				&& CastChecked<UStructProperty>(Property)->UseBinaryOrNativeSerialization(Ar) )
@@ -1090,14 +1088,14 @@ void UStruct::SerializeTaggedProperties(FArchive& Ar, uint8* Data, UStruct* Defa
 			}
 			else
 			{
-				uint8* DestAddress = Property->ContainerPtrToValuePtr<uint8>(Data, Tag.ArrayIndex);  
+					uint8* DestAddress = Property->ContainerPtrToValuePtr<uint8>(Data, Tag.ArrayIndex);  
 
-				// This property is ok.			
-				Tag.SerializeTaggedProperty( Ar, Property, DestAddress, Tag.Size, NULL );
+					// This property is ok.			
+					Tag.SerializeTaggedProperty( Ar, Property, DestAddress, Tag.Size, NULL );
 
-				AdvanceProperty = true;
-				continue;
-			}
+					AdvanceProperty = true;
+					continue;
+				}
 
 			AdvanceProperty = false;
 
diff --git a/Engine/Source/Runtime/CoreUObject/Public/UObject/Class.h b/Engine/Source/Runtime/CoreUObject/Public/UObject/Class.h
--- a/Engine/Source/Runtime/CoreUObject/Public/UObject/Class.h
+++ b/Engine/Source/Runtime/CoreUObject/Public/UObject/Class.h
@@ -301,6 +301,11 @@ public:
 
 	virtual void SerializeTaggedProperties( FArchive& Ar, uint8* Data, UStruct* DefaultsStruct, uint8* Defaults ) const;
 
+#if WITH_EDITOR
+private:
+	virtual UProperty* CustomFindProperty(const FName Name) const { return NULL; };
+#endif // WITH_EDITOR
+public:
 	virtual EExprToken SerializeExpr(int32& iCode, FArchive& Ar);
 	virtual void TagSubobjects(EObjectFlags NewFlags) override;
 
diff --git a/Engine/Source/Runtime/Engine/Classes/Engine/UserDefinedStruct.h b/Engine/Source/Runtime/Engine/Classes/Engine/UserDefinedStruct.h
--- a/Engine/Source/Runtime/Engine/Classes/Engine/UserDefinedStruct.h
+++ b/Engine/Source/Runtime/Engine/Classes/Engine/UserDefinedStruct.h
@@ -53,4 +53,10 @@ public:
 	// UScriptStruct interface.
 	virtual void RecursivelyPreload() override;
 	// End of  UScriptStruct interface.
+
+#if WITH_EDITOR
+private:
+	// UStruct
+	virtual UProperty* CustomFindProperty(const FName Name) const override;
+#endif	// WITH_EDITOR
 };
diff --git a/Engine/Source/Runtime/Engine/Private/UserDefinedStruct.cpp b/Engine/Source/Runtime/Engine/Private/UserDefinedStruct.cpp
--- a/Engine/Source/Runtime/Engine/Private/UserDefinedStruct.cpp
+++ b/Engine/Source/Runtime/Engine/Private/UserDefinedStruct.cpp
@@ -45,6 +45,14 @@ void UUserDefinedStruct::GetAssetRegistryTags(TArray<FAssetRegistryTag>& OutTags
 	OutTags.Add(FAssetRegistryTag(TEXT("Tooltip"), FStructureEditorUtils::GetTooltip(this), FAssetRegistryTag::TT_Hidden));
 }
 
+UProperty* UUserDefinedStruct::CustomFindProperty(const FName Name) const
+{
+	const FGuid Guid = FStructureEditorUtils::GetGuidFromPropertyName(Name);
+	UProperty* Property = Guid.IsValid() ? FStructureEditorUtils::GetPropertyByGuid(this, Guid) : NULL;
+	ensure(!Property || Guid == FStructureEditorUtils::GetGuidForProperty(Property));
+	return Property;
+}
+
 #endif	// WITH_EDITOR
 
 void UUserDefinedStruct::SerializeTaggedProperties(FArchive& Ar, uint8* Data, UStruct* DefaultsStruct, uint8* Defaults) const
