CIS build fix (enum forward declaration).

diff --git a/Engine/Source/Editor/UnrealEd/Classes/Factories/AnimBlueprintFactory.h b/Engine/Source/Editor/UnrealEd/Classes/Factories/AnimBlueprintFactory.h
--- a/Engine/Source/Editor/UnrealEd/Classes/Factories/AnimBlueprintFactory.h
+++ b/Engine/Source/Editor/UnrealEd/Classes/Factories/AnimBlueprintFactory.h
@@ -5,6 +5,7 @@
  */
 
 #pragma once
+#include "Engine/Blueprint.h"
 #include "AnimBlueprintFactory.generated.h"
 
 UCLASS(HideCategories=Object,MinimalAPI)
@@ -14,7 +15,7 @@ class UAnimBlueprintFactory : public UFactory
 
 	// The type of blueprint that will be created
 	UPROPERTY(EditAnywhere, Category=AnimBlueprintFactory)
-	TEnumAsByte<enum EBlueprintType> BlueprintType;
+	TEnumAsByte<EBlueprintType> BlueprintType;
 
 	// The parent class of the created blueprint
 	UPROPERTY(EditAnywhere, Category=AnimBlueprintFactory, meta=(AllowAbstract = ""))
