Fixing mismatched macro definitions between Engine and UnrealEd modules

diff --git a/Engine/Source/Editor/UnrealEd/Public/UnrealEd.h b/Engine/Source/Editor/UnrealEd/Public/UnrealEd.h
--- a/Engine/Source/Editor/UnrealEd/Public/UnrealEd.h
+++ b/Engine/Source/Editor/UnrealEd/Public/UnrealEd.h
@@ -7,6 +7,8 @@
 
 #if WITH_EDITOR
 
+#include "EngineDefines.h"
+
 #include "Slate.h"
 #include "EditorStyle.h"
 
