GetDependentBlueprints optimization. The change was pending after Pretest merge.

diff --git a/Engine/Source/Editor/UnrealEd/Private/Kismet2/BlueprintEditorUtils.cpp b/Engine/Source/Editor/UnrealEd/Private/Kismet2/BlueprintEditorUtils.cpp
--- a/Engine/Source/Editor/UnrealEd/Private/Kismet2/BlueprintEditorUtils.cpp
+++ b/Engine/Source/Editor/UnrealEd/Private/Kismet2/BlueprintEditorUtils.cpp
@@ -2091,17 +2091,11 @@ void FBlueprintEditorUtils::GetDependentBlueprints(UBlueprint* Blueprint, TArray
 {
 	DependentBlueprints.Empty();
 
-	//@TODO PRETEST MERGE:  Swap this when main is merged to pretest!
-#if 0
 	TArray<UObject*> AllBlueprints;
 	bool const bIncludeDerivedClasses = true;
 	GetObjectsOfClass(UBlueprint::StaticClass(), AllBlueprints, bIncludeDerivedClasses );
 
 	for ( auto ObjIt = AllBlueprints.CreateConstIterator(); ObjIt; ++ObjIt )
-#else
-	// Use old school slower iterator until fast object collection makes it to pretest
-	for( TObjectIterator<UBlueprint> ObjIt; ObjIt; ++ObjIt )
-#endif
 	{
 		// we know the class is correct so a fast cast is ok here
 		UBlueprint* const TestBP = (UBlueprint*) *ObjIt;
