UE4: All thumbnail renderers now have a chance to conditionally not render a thumbnail for a given asset. This was previously only possible for blueprints

diff --git a/Engine/Source/Editor/UnrealEd/Classes/ThumbnailRendering/BlueprintThumbnailRenderer.h b/Engine/Source/Editor/UnrealEd/Classes/ThumbnailRendering/BlueprintThumbnailRenderer.h
--- a/Engine/Source/Editor/UnrealEd/Classes/ThumbnailRendering/BlueprintThumbnailRenderer.h
+++ b/Engine/Source/Editor/UnrealEd/Classes/ThumbnailRendering/BlueprintThumbnailRenderer.h
@@ -15,6 +15,7 @@ class UBlueprintThumbnailRenderer : public UDefaultSizedThumbnailRenderer
 
 
 	// Begin UThumbnailRenderer Object
+	virtual bool CanVisualizeAsset(UObject* Object) override;
 	UNREALED_API virtual void Draw(UObject* Object, int32 X, int32 Y, uint32 Width, uint32 Height, FRenderTarget*, FCanvas* Canvas) override;
 	// End UThumbnailRenderer Object
 
@@ -22,9 +23,6 @@ class UBlueprintThumbnailRenderer : public UDefaultSizedThumbnailRenderer
 	UNREALED_API virtual void BeginDestroy() override;
 	// End UObject implementation
 
-	/** Returns true if this scene is capable of creating a thumbnail with the specified blueprint */
-	bool CanVisualizeBlueprint(class UBlueprint* Blueprint);
-
 	/** Notifies the thumbnail scene to refresh components for the specified blueprint */
 	void BlueprintChanged(class UBlueprint* Blueprint);
 
diff --git a/Engine/Source/Editor/UnrealEd/Classes/ThumbnailRendering/ThumbnailRenderer.h b/Engine/Source/Editor/UnrealEd/Classes/ThumbnailRendering/ThumbnailRenderer.h
--- a/Engine/Source/Editor/UnrealEd/Classes/ThumbnailRendering/ThumbnailRenderer.h
+++ b/Engine/Source/Editor/UnrealEd/Classes/ThumbnailRendering/ThumbnailRenderer.h
@@ -17,6 +17,13 @@ class UThumbnailRenderer : public UObject
 
 
 public:
+	/**
+	 * Returns true if the renderer is capable of producing a thumbnail for the specified asset.
+	 *
+	 * @param Object the asset to attempt to render
+	 */
+	virtual bool CanVisualizeAsset(UObject* Object) { return true; }
+
 	/**
 	 * Calculates the size the thumbnail would be at the specified zoom level
 	 *
diff --git a/Engine/Source/Editor/UnrealEd/Private/ThumbnailManager.cpp b/Engine/Source/Editor/UnrealEd/Private/ThumbnailManager.cpp
--- a/Engine/Source/Editor/UnrealEd/Private/ThumbnailManager.cpp
+++ b/Engine/Source/Editor/UnrealEd/Private/ThumbnailManager.cpp
@@ -134,13 +134,11 @@ FThumbnailRenderingInfo* UThumbnailManager::GetRenderingInfo(UObject* Object)
 		RenderInfoMap.Add(ClassToCheck, (RenderInfo != nullptr) ? RenderInfo : &NotSupported);
 	}
 
-	if ( RenderInfo && RenderInfo->Renderer && Object->IsA(UBlueprint::StaticClass()) && RenderInfo->Renderer->IsA(UBlueprintThumbnailRenderer::StaticClass()) )
+	if ( RenderInfo && RenderInfo->Renderer )
 	{
-		UBlueprint* BP = Cast<UBlueprint>(Object);
-		UBlueprintThumbnailRenderer* BPRenderer = Cast<UBlueprintThumbnailRenderer>(RenderInfo->Renderer);
-		if ( !BPRenderer->CanVisualizeBlueprint(BP) )
+		if ( !RenderInfo->Renderer->CanVisualizeAsset(Object) )
 		{
-			// This is a blueprint, but it can't visualized (i.e it doesnt contain any visible primitive components)
+			// This is an asset with a thumbnail renderer, but it can't visualized (i.e it is something like a blueprint that doesn't contain any visible primitive components)
 			RenderInfo = nullptr;
 		}
 	}
diff --git a/Engine/Source/Editor/UnrealEd/Private/ThumbnailRendering/BlueprintThumbnailRenderer.cpp b/Engine/Source/Editor/UnrealEd/Private/ThumbnailRendering/BlueprintThumbnailRenderer.cpp
--- a/Engine/Source/Editor/UnrealEd/Private/ThumbnailRendering/BlueprintThumbnailRenderer.cpp
+++ b/Engine/Source/Editor/UnrealEd/Private/ThumbnailRendering/BlueprintThumbnailRenderer.cpp
@@ -12,49 +12,17 @@ UBlueprintThumbnailRenderer::UBlueprintThumbnailRenderer(const class FPostConstr
 	ThumbnailScene = nullptr;
 }
 
-void UBlueprintThumbnailRenderer::Draw(UObject* Object, int32 X, int32 Y, uint32 Width, uint32 Height, FRenderTarget* RenderTarget, FCanvas* Canvas)
+bool UBlueprintThumbnailRenderer::CanVisualizeAsset(UObject* Object)
 {
-	UBlueprint* Blueprint = Cast<UBlueprint>(Object);
-	if (Blueprint != nullptr)
-	{
-		if ( ThumbnailScene == nullptr )
-		{
-			ThumbnailScene = new FBlueprintThumbnailScene();
-		}
-
-		ThumbnailScene->SetBlueprint(Blueprint);
-		FSceneViewFamilyContext ViewFamily( FSceneViewFamily::ConstructionValues( RenderTarget, ThumbnailScene->GetScene(), FEngineShowFlags(ESFIM_Game) )
-			.SetWorldTimes(FApp::GetCurrentTime() - GStartTime, FApp::GetDeltaTime(), FApp::GetCurrentTime() - GStartTime));
-
-		ViewFamily.EngineShowFlags.DisableAdvancedFeatures();
-		ViewFamily.EngineShowFlags.MotionBlur = 0;
-
-		ThumbnailScene->GetView(&ViewFamily, X, Y, Width, Height);
-		GetRendererModule().BeginRenderingViewFamily(Canvas,&ViewFamily);
-		ThumbnailScene->SetBlueprint(nullptr);
-	}
-}
-
-void UBlueprintThumbnailRenderer::BeginDestroy()
-{
-	if ( ThumbnailScene != nullptr )
-	{
-		delete ThumbnailScene;
-		ThumbnailScene = nullptr;
-	}
-
-	Super::BeginDestroy();
-}
-
-bool UBlueprintThumbnailRenderer::CanVisualizeBlueprint(class UBlueprint* Blueprint)
-{
-	if ( ThumbnailScene == nullptr )
+	if (ThumbnailScene == nullptr)
 	{
 		ThumbnailScene = new FBlueprintThumbnailScene();
 	}
 
+	UBlueprint* Blueprint = Cast<UBlueprint>(Object);
+
 	// Only visualize actor based blueprints
-	if ( Blueprint->GeneratedClass && Blueprint->GeneratedClass->IsChildOf(AActor::StaticClass()) )
+	if (Blueprint && Blueprint->GeneratedClass && Blueprint->GeneratedClass->IsChildOf(AActor::StaticClass()))
 	{
 		// Try to find any visible primitive components in the native class' CDO
 		AActor* CDO = Blueprint->GeneratedClass->GetDefaultObject<AActor>();
@@ -62,9 +30,9 @@ bool UBlueprintThumbnailRenderer::CanVisualizeBlueprint(class UBlueprint* Bluepr
 		TArray<UActorComponent*> Components;
 		CDO->GetComponents(Components);
 
-		for ( auto CompIt = Components.CreateConstIterator(); CompIt; ++CompIt )
+		for (auto CompIt = Components.CreateConstIterator(); CompIt; ++CompIt)
 		{
-			if ( ThumbnailScene->IsValidComponentForVisualization(*CompIt) )
+			if (ThumbnailScene->IsValidComponentForVisualization(*CompIt))
 			{
 				return true;
 			}
@@ -74,17 +42,17 @@ bool UBlueprintThumbnailRenderer::CanVisualizeBlueprint(class UBlueprint* Bluepr
 		// Do this for all parent blueprint generated classes as well
 		UBlueprint* BlueprintToHarvestComponents = Blueprint;
 		TSet<UBlueprint*> AllVisitedBlueprints;
-		while ( BlueprintToHarvestComponents )
+		while (BlueprintToHarvestComponents)
 		{
 			AllVisitedBlueprints.Add(BlueprintToHarvestComponents);
 
-			if ( BlueprintToHarvestComponents->SimpleConstructionScript )
+			if (BlueprintToHarvestComponents->SimpleConstructionScript)
 			{
 				TArray<USCS_Node*> AllNodes = BlueprintToHarvestComponents->SimpleConstructionScript->GetAllNodes();
 
-				for ( auto NodeIt = AllNodes.CreateConstIterator(); NodeIt; ++NodeIt )
+				for (auto NodeIt = AllNodes.CreateConstIterator(); NodeIt; ++NodeIt)
 				{
-					if ( ThumbnailScene->IsValidComponentForVisualization((*NodeIt)->ComponentTemplate) )
+					if (ThumbnailScene->IsValidComponentForVisualization((*NodeIt)->ComponentTemplate))
 					{
 						return true;
 					}
@@ -95,12 +63,12 @@ bool UBlueprintThumbnailRenderer::CanVisualizeBlueprint(class UBlueprint* Bluepr
 			BlueprintToHarvestComponents = nullptr;
 
 			// If the parent class was a blueprint generated class, check it's simple construction script components as well
-			if ( ParentClass )
+			if (ParentClass)
 			{
 				UBlueprint* ParentBlueprint = Cast<UBlueprint>(ParentClass->ClassGeneratedBy);
 
 				// Also make sure we haven't visited the blueprint already. This would only happen if there was a loop of parent classes.
-				if ( ParentBlueprint && !AllVisitedBlueprints.Contains(ParentBlueprint) )
+				if (ParentBlueprint && !AllVisitedBlueprints.Contains(ParentBlueprint))
 				{
 					BlueprintToHarvestComponents = ParentBlueprint;
 				}
@@ -111,6 +79,40 @@ bool UBlueprintThumbnailRenderer::CanVisualizeBlueprint(class UBlueprint* Bluepr
 	return false;
 }
 
+void UBlueprintThumbnailRenderer::Draw(UObject* Object, int32 X, int32 Y, uint32 Width, uint32 Height, FRenderTarget* RenderTarget, FCanvas* Canvas)
+{
+	UBlueprint* Blueprint = Cast<UBlueprint>(Object);
+	if (Blueprint != nullptr)
+	{
+		if ( ThumbnailScene == nullptr )
+		{
+			ThumbnailScene = new FBlueprintThumbnailScene();
+		}
+
+		ThumbnailScene->SetBlueprint(Blueprint);
+		FSceneViewFamilyContext ViewFamily( FSceneViewFamily::ConstructionValues( RenderTarget, ThumbnailScene->GetScene(), FEngineShowFlags(ESFIM_Game) )
+			.SetWorldTimes(FApp::GetCurrentTime() - GStartTime, FApp::GetDeltaTime(), FApp::GetCurrentTime() - GStartTime));
+
+		ViewFamily.EngineShowFlags.DisableAdvancedFeatures();
+		ViewFamily.EngineShowFlags.MotionBlur = 0;
+
+		ThumbnailScene->GetView(&ViewFamily, X, Y, Width, Height);
+		GetRendererModule().BeginRenderingViewFamily(Canvas,&ViewFamily);
+		ThumbnailScene->SetBlueprint(nullptr);
+	}
+}
+
+void UBlueprintThumbnailRenderer::BeginDestroy()
+{
+	if ( ThumbnailScene != nullptr )
+	{
+		delete ThumbnailScene;
+		ThumbnailScene = nullptr;
+	}
+
+	Super::BeginDestroy();
+}
+
 void UBlueprintThumbnailRenderer::BlueprintChanged(class UBlueprint* Blueprint)
 {
 	if ( ThumbnailScene != nullptr )
