Landscape : support hole export for FBX

diff --git a/Engine/Source/Editor/UnrealEd/Private/Fbx/FbxMainExport.cpp b/Engine/Source/Editor/UnrealEd/Private/Fbx/FbxMainExport.cpp
--- a/Engine/Source/Editor/UnrealEd/Private/Fbx/FbxMainExport.cpp
+++ b/Engine/Source/Editor/UnrealEd/Private/Fbx/FbxMainExport.cpp
@@ -2373,6 +2373,10 @@ FbxNode* FFbxExporter::ExportLandscapeToFbx(ALandscapeProxy* Landscape, const TC
 	FbxLayerElementArrayTemplate<FbxVector2>& WeightmapUVs = LayerElementWeightmapUVs->GetDirectArray();
 	WeightmapUVs.Resize(VertexCount);
 
+	TArray<uint8> VisibiltyData;
+	VisibiltyData.Empty(VertexCount);
+	VisibiltyData.AddZeroed(VertexCount);
+
 	for (int32 ComponentIndex = 0, SelectedComponentIndex = 0; ComponentIndex < Landscape->LandscapeComponents.Num(); ComponentIndex++)
 	{
 		ULandscapeComponent* Component = Landscape->LandscapeComponents[ComponentIndex];
@@ -2385,6 +2389,21 @@ FbxNode* FFbxExporter::ExportLandscapeToFbx(ALandscapeProxy* Landscape, const TC
 		FLandscapeComponentDataInterface CDI(Component, Landscape->ExportLOD);
 		const int32 BaseVertIndex = SelectedComponentIndex++ * VertexCountPerComponent;
 
+		TArray<uint8> CompVisData;
+		for (int32 AllocIdx = 0; AllocIdx < Component->WeightmapLayerAllocations.Num(); AllocIdx++)
+		{
+			FWeightmapLayerAllocationInfo& AllocInfo = Component->WeightmapLayerAllocations[AllocIdx];
+			if (AllocInfo.LayerInfo == ALandscapeProxy::DataLayer)
+			{
+				CDI.GetWeightmapTextureData(AllocInfo.LayerInfo, CompVisData);
+			}
+		}
+
+		for (int32 i = 0; i < CompVisData.Num(); ++i)
+		{
+			VisibiltyData[BaseVertIndex + i] = CompVisData[i];
+		}
+		
 		for (int32 VertIndex = 0; VertIndex < VertexCountPerComponent; VertIndex++)
 		{
 			int32 VertX, VertY;
@@ -2451,6 +2470,7 @@ FbxNode* FFbxExporter::ExportLandscapeToFbx(ALandscapeProxy* Landscape, const TC
 	const int32 MaterialIndex = FbxActor->AddMaterial(FbxMaterial);
 	LayerElementMaterials->GetIndexArray().Add(MaterialIndex);
 
+	const int32 VisThreshold = 170;
 	// Copy over the index buffer into the FBX polygons set.
 	for (int32 ComponentIndex = 0; ComponentIndex < NumComponents; ComponentIndex++)
 	{
@@ -2460,17 +2480,20 @@ FbxNode* FFbxExporter::ExportLandscapeToFbx(ALandscapeProxy* Landscape, const TC
 		{
 			for (int32 X = 0; X < ComponentSizeQuads; X++)
 			{
-				Mesh->BeginPolygon();
-				Mesh->AddPolygon(BaseVertIndex + (X+0) + (Y+0)*(ComponentSizeQuads+1));
-				Mesh->AddPolygon(BaseVertIndex + (X+1) + (Y+1)*(ComponentSizeQuads+1));
-				Mesh->AddPolygon(BaseVertIndex + (X+1) + (Y+0)*(ComponentSizeQuads+1));
-				Mesh->EndPolygon();
-
-				Mesh->BeginPolygon();
-				Mesh->AddPolygon(BaseVertIndex + (X+0) + (Y+0)*(ComponentSizeQuads+1));
-				Mesh->AddPolygon(BaseVertIndex + (X+0) + (Y+1)*(ComponentSizeQuads+1));
-				Mesh->AddPolygon(BaseVertIndex + (X+1) + (Y+1)*(ComponentSizeQuads+1));
-				Mesh->EndPolygon();
+				if (VisibiltyData[BaseVertIndex + Y * (ComponentSizeQuads + 1) + X] < VisThreshold)
+				{
+					Mesh->BeginPolygon();
+					Mesh->AddPolygon(BaseVertIndex + (X + 0) + (Y + 0)*(ComponentSizeQuads + 1));
+					Mesh->AddPolygon(BaseVertIndex + (X + 1) + (Y + 1)*(ComponentSizeQuads + 1));
+					Mesh->AddPolygon(BaseVertIndex + (X + 1) + (Y + 0)*(ComponentSizeQuads + 1));
+					Mesh->EndPolygon();
+
+					Mesh->BeginPolygon();
+					Mesh->AddPolygon(BaseVertIndex + (X + 0) + (Y + 0)*(ComponentSizeQuads + 1));
+					Mesh->AddPolygon(BaseVertIndex + (X + 0) + (Y + 1)*(ComponentSizeQuads + 1));
+					Mesh->AddPolygon(BaseVertIndex + (X + 1) + (Y + 1)*(ComponentSizeQuads + 1));
+					Mesh->EndPolygon();
+				}
 			}
 		}
 	}
diff --git a/Engine/Source/Runtime/Engine/Private/Landscape/LandscapeEdit.cpp b/Engine/Source/Runtime/Engine/Private/Landscape/LandscapeEdit.cpp
--- a/Engine/Source/Runtime/Engine/Private/Landscape/LandscapeEdit.cpp
+++ b/Engine/Source/Runtime/Engine/Private/Landscape/LandscapeEdit.cpp
@@ -2273,23 +2273,14 @@ bool ALandscapeProxy::ExportToRawMesh(int32 InExportLOD, FRawMesh& OutRawMesh) c
 		
 				
 		// Check if there are any holes
-		TArray<uint8> RawVisData;
-		uint8* VisDataMap = NULL;
-		const int32 ChannelOffsets[4] = {
-			(int32)STRUCT_OFFSET(FColor,R),
-			(int32)STRUCT_OFFSET(FColor,G),
-			(int32)STRUCT_OFFSET(FColor,B),
-			(int32)STRUCT_OFFSET(FColor,A)
-		};
+		TArray<uint8> VisDataMap;
 		
 		for (int32 AllocIdx = 0; AllocIdx < Component->WeightmapLayerAllocations.Num(); AllocIdx++)
 		{
 			FWeightmapLayerAllocationInfo& AllocInfo = Component->WeightmapLayerAllocations[AllocIdx];
 			if (AllocInfo.LayerInfo == ALandscapeProxy::DataLayer)
 			{
-				int32 TexIndex = AllocInfo.WeightmapTextureIndex;
-				Component->WeightmapTextures[TexIndex]->Source.GetMipData(RawVisData, 0);
-				VisDataMap = RawVisData.GetTypedData() + ChannelOffsets[AllocInfo.WeightmapTextureChannel];
+				CDI.GetWeightmapTextureData(AllocInfo.LayerInfo, VisDataMap);
 			}
 		}
 
@@ -2319,8 +2310,7 @@ bool ALandscapeProxy::ExportToRawMesh(int32 InExportLOD, FRawMesh& OutRawMesh) c
 				{
 					int32 SubNumX, SubNumY, SubX, SubY;
 					CDI.ComponentXYToSubsectionXY(x, y, SubNumX, SubNumY, SubX, SubY);
-					int32 WeightIndex = SubX + SubNumX*(SubsectionSizeQuadsLOD + 1) + (SubY+SubNumY*(SubsectionSizeQuadsLOD + 1))*WeightMapSize;
-					bool bInvisible = VisDataMap && VisDataMap[WeightIndex * sizeof(FColor)] >= VisThreshold;
+					bool bInvisible = VisDataMap[x + y * (ComponentSizeQuadsLOD + 1)] >= VisThreshold;
 
 					// triangulation matches FLandscapeIndexBuffer constructor
 					Faces[0] = VertexIdx;
diff --git a/Engine/Source/Runtime/Engine/Public/Landscape/LandscapeDataAccess.h b/Engine/Source/Runtime/Engine/Public/Landscape/LandscapeDataAccess.h
--- a/Engine/Source/Runtime/Engine/Public/Landscape/LandscapeDataAccess.h
+++ b/Engine/Source/Runtime/Engine/Public/Landscape/LandscapeDataAccess.h
@@ -295,11 +295,11 @@ struct FLandscapeComponentDataInterface
 			return false;
 		}
 
-		int32 WeightmapSize = (Component->SubsectionSizeQuads+1) * Component->NumSubsections;
+		int32 WeightmapSize = ((Component->SubsectionSizeQuads + 1) * Component->NumSubsections) >> MipLevel;
 		OutData.Empty(FMath::Square(WeightmapSize));
 		OutData.AddUninitialized(FMath::Square(WeightmapSize));
 
-		FColor* WeightMipData = (FColor*)DataInterface->LockMip(Component->WeightmapTextures[Component->WeightmapLayerAllocations[LayerIdx].WeightmapTextureIndex], 0);
+		FColor* WeightMipData = (FColor*)DataInterface->LockMip(Component->WeightmapTextures[Component->WeightmapLayerAllocations[LayerIdx].WeightmapTextureIndex], MipLevel);
 
 		// Channel remapping
 		int32 ChannelOffsets[4] = {(int32)STRUCT_OFFSET(FColor,R),(int32)STRUCT_OFFSET(FColor,G),(int32)STRUCT_OFFSET(FColor,B),(int32)STRUCT_OFFSET(FColor,A)};
@@ -311,7 +311,7 @@ struct FLandscapeComponentDataInterface
 			OutData[i] = SrcTextureData[i*4];
 		}
 
-		DataInterface->UnlockMip(Component->WeightmapTextures[Component->WeightmapLayerAllocations[LayerIdx].WeightmapTextureIndex], 0);
+		DataInterface->UnlockMip(Component->WeightmapTextures[Component->WeightmapLayerAllocations[LayerIdx].WeightmapTextureIndex], MipLevel);
 		return true;
 	}
 
