Disable performance and auto-scalablity settings in debug builds.

diff --git a/Engine/Source/Editor/UnrealEd/Private/UnrealEdEngine.cpp b/Engine/Source/Editor/UnrealEd/Private/UnrealEdEngine.cpp
--- a/Engine/Source/Editor/UnrealEd/Private/UnrealEdEngine.cpp
+++ b/Engine/Source/Editor/UnrealEd/Private/UnrealEdEngine.cpp
@@ -39,7 +39,9 @@ void UUnrealEdEngine::Init(IEngineLoop* InEngineLoop)
 	PackageAutoSaver.Reset(new FPackageAutoSaver);
 	PackageAutoSaver->LoadRestoreFile();
 
+#if !UE_BUILD_DEBUG
 	PerformanceMonitor = new FPerformanceMonitor;
+#endif
 
 	// Register for the package dirty state updated callback to catch packages that have been modified and need to be checked out.
 	UPackage::PackageDirtyStateChangedEvent.AddUObject(this, &UUnrealEdEngine::OnPackageDirtyStateUpdated);
@@ -226,7 +228,10 @@ void UUnrealEdEngine::FinishDestroy()
 		PackageAutoSaver.Reset();
 	}
 
-	delete PerformanceMonitor;
+	if( PerformanceMonitor )
+	{
+		delete PerformanceMonitor;
+	}
 
 	UPackage::PackageDirtyStateChangedEvent.RemoveAll(this);
 	FCoreDelegates::PostGarbageCollect.RemoveAll(this);
