Fixed engine scalability settings not being applied correctly when running a project from the project browser

Also added an extra button to the performance monitor notification to allow users to auto apply settings immediately if they so choose.
This addresses TTP 340255 (CRITICAL: UE4: EDITOR: Low FPS detection.  Does not lower Scalability Settings after editor restart)

diff --git a/Engine/Source/Editor/UnrealEd/Private/PerformanceMonitor.cpp b/Engine/Source/Editor/UnrealEd/Private/PerformanceMonitor.cpp
--- a/Engine/Source/Editor/UnrealEd/Private/PerformanceMonitor.cpp
+++ b/Engine/Source/Editor/UnrealEd/Private/PerformanceMonitor.cpp
@@ -113,7 +113,7 @@ FPerformanceMonitor::FPerformanceMonitor()
 	bIsNotificationAllowed = true;
 
 	UEditorGameAgnosticSettings& Settings = GEditor->AccessGameAgnosticSettings();
-	if (Settings.bApplyAutoScalabilityOnStartup)
+	if (Settings.bApplyAutoScalabilityOnStartup && FApp::HasGameName())
 	{
 		Settings.bApplyAutoScalabilityOnStartup = false;
 		Settings.PostEditChange();
@@ -164,7 +164,11 @@ void FPerformanceMonitor::AutoApplyScalability()
 	BenchmarkLevels.EffectsQuality		= FMath::Min(BenchmarkLevels.EffectsQuality,		ExistingLevels.EffectsQuality);
 	
 	Scalability::SetQualityLevels(BenchmarkLevels);
-	
+	Scalability::SaveState(GEditorUserSettingsIni);
+
+	GEditor->DisableRealtimeViewports();
+
+	HidePerformanceWarning();
 	GWarn->EndSlowTask();
 }
 
@@ -189,17 +193,26 @@ void FPerformanceMonitor::ShowPerformanceWarning(FText MessageText)
 
 		Info.ButtonDetails.Add(
 			FNotificationButtonInfo(
-			LOCTEXT("PerformanceWarningScalability", "Tweak Scalability Settings"),
+			LOCTEXT("PerformanceWarningScalability", "Tweak Manually"),
 			LOCTEXT("PerformanceWarningScalabilityToolTip", "Opens the Scalability UI for manual adjustment."),
 			FSimpleDelegate::CreateRaw(this, &FPerformanceMonitor::ShowScalabilityDialog),
 			SNotificationItem::CS_None
 			)
 			);
 
+		Info.ButtonDetails.Add(
+			FNotificationButtonInfo(
+			LOCTEXT("PerformanceWarningApplyNow", "Apply Now"),
+			LOCTEXT("PerformanceWarningApplyNowToolTip", "Run a benchmark now and apply settings immediately."),
+			FSimpleDelegate::CreateRaw(this, &FPerformanceMonitor::AutoApplyScalability),
+			SNotificationItem::CS_None
+			)
+			);
+
 		Info.ButtonDetails.Add(
 			FNotificationButtonInfo(
 			LOCTEXT("PerformanceWarningDismiss", "Ok"),
-			LOCTEXT("PerformanceWarningDismissToolTip", "Close this notification."),
+			LOCTEXT("PerformanceWarningDismissToolTip", "Close this notification and apply settings on restart."),
 			FSimpleDelegate::CreateRaw(this, &FPerformanceMonitor::DismissAndDisableNotification),
 			SNotificationItem::CS_None
 			)
