Fix crash in PhAT when trying to create multi convex hull with soft vertices

diff --git a/Engine/Source/Editor/PhAT/Private/PhATSharedData.cpp b/Engine/Source/Editor/PhAT/Private/PhATSharedData.cpp
--- a/Engine/Source/Editor/PhAT/Private/PhATSharedData.cpp
+++ b/Engine/Source/Editor/PhAT/Private/PhATSharedData.cpp
@@ -800,14 +800,21 @@ void FPhATSharedData::MakeNewBody(int32 NewBoneIndex)
 	
 	BodySetup->Modify();
 
+	bool bCreatedBody = false;
 	// Create a new physics body for this bone.
 	if (NewBodyData.VertWeight == EVW_DominantWeight)
 	{
-		FPhysicsAssetUtils::CreateCollisionFromBone(BodySetup, EditorSkelMesh, NewBoneIndex, NewBodyData, DominantWeightBoneInfos);
+		bCreatedBody = FPhysicsAssetUtils::CreateCollisionFromBone(BodySetup, EditorSkelMesh, NewBoneIndex, NewBodyData, DominantWeightBoneInfos);
 	}
 	else
 	{
-		FPhysicsAssetUtils::CreateCollisionFromBone(BodySetup, EditorSkelMesh, NewBoneIndex, NewBodyData, AnyWeightBoneInfos);
+		bCreatedBody = FPhysicsAssetUtils::CreateCollisionFromBone(BodySetup, EditorSkelMesh, NewBoneIndex, NewBodyData, AnyWeightBoneInfos);
+	}
+
+	if (bCreatedBody == false)
+	{
+		FPhysicsAssetUtils::DestroyBody(PhysicsAsset, NewBodyIndex);
+		return;
 	}
 
 	// Check if the bone of the new body has any physical children bones
diff --git a/Engine/Source/Editor/UnrealEd/Private/PhysicsAssetUtils.cpp b/Engine/Source/Editor/UnrealEd/Private/PhysicsAssetUtils.cpp
--- a/Engine/Source/Editor/UnrealEd/Private/PhysicsAssetUtils.cpp
+++ b/Engine/Source/Editor/UnrealEd/Private/PhysicsAssetUtils.cpp
@@ -165,10 +165,10 @@ bool CreateFromSkeletalMeshInternal(UPhysicsAsset* PhysicsAsset, USkeletalMesh*
 			check(bs->BoneName == BoneName);
 
 			// Fill in collision info for this bone.
-			CreateCollisionFromBone(bs, SkelMesh, i, Params, Infos);
+			bool bSuccess = CreateCollisionFromBone(bs, SkelMesh, i, Params, Infos);
 
 			// If not root - create joint to parent body.
-			if(bHitRoot && Params.bCreateJoints)
+			if(bHitRoot && Params.bCreateJoints && bSuccess)
 			{
 				int32 NewConstraintIndex = CreateNewConstraint(PhysicsAsset, BoneName);
 				UPhysicsConstraintTemplate* CS = PhysicsAsset->ConstraintSetup[NewConstraintIndex];
@@ -197,6 +197,11 @@ bool CreateFromSkeletalMeshInternal(UPhysicsAsset* PhysicsAsset, USkeletalMesh*
 			}
 
 			bHitRoot = true;
+
+			if (bSuccess == false)
+			{
+				DestroyBody(PhysicsAsset, NewBodyIndex);
+			}
 		}
 	}
 
@@ -301,11 +306,14 @@ FVector ComputeEigenVector(const FMatrix & A)
 }
 
 
-void CreateCollisionFromBone( UBodySetup* bs, USkeletalMesh* skelMesh, int32 BoneIndex, FPhysAssetCreateParams& Params, const TArray<FBoneVertInfo>& Infos )
+bool CreateCollisionFromBone( UBodySetup* bs, USkeletalMesh* skelMesh, int32 BoneIndex, FPhysAssetCreateParams& Params, const TArray<FBoneVertInfo>& Infos )
 {
 #if WITH_EDITOR
-	// Empty any existing collision.
-	bs->RemoveSimpleCollision();
+	if (Params.GeomType != EFG_MultiConvexHull)	//multi convex hull can fail so wait to clear it
+	{
+		// Empty any existing collision.
+		bs->RemoveSimpleCollision();
+	}
 #endif // WITH_EDITOR
 
 	// Calculate orientation of to use for collision primitive.
@@ -454,8 +462,10 @@ void CreateCollisionFromBone( UBodySetup* bs, USkeletalMesh* skelMesh, int32 Bon
 			if ( bSoftVertex )
 			{
 				// We dont want to support soft verts, only rigid
-				UE_LOG(LogPhysics, Log, TEXT("Unable to create physics asset with a multi convex hull due to the presence of soft vertices!"));
-				return;
+				FMessageLog EditorErrors("EditorErrors");
+				EditorErrors.Warning(NSLOCTEXT("PhysicsAssetUtils", "MultiConvexSoft", "Unable to create physics asset with a multi convex hull due to the presence of soft vertices."));
+				EditorErrors.Open();
+				return false;
 			}
 
 			// Using the same code in GetSkinnedVertexPosition
@@ -480,12 +490,16 @@ void CreateCollisionFromBone( UBodySetup* bs, USkeletalMesh* skelMesh, int32 Bon
 
 		if ( Params.GeomType == EFG_MultiConvexHull )
 		{
+#if WITH_EDITOR
+			bs->RemoveSimpleCollision();
+#endif
 			// Create the convex hull from the data we got from the skeletal mesh
 			DecomposeMeshToHulls( bs, Verts, Indices, Params.MaxHullCount, Params.MaxHullVerts );
 		}
 		else
 		{
 			//Support triangle mesh soon
+			return false;
 		}		
 	}
 	else
@@ -500,6 +514,8 @@ void CreateCollisionFromBone( UBodySetup* bs, USkeletalMesh* skelMesh, int32 Bon
 
 		bs->AggGeom.SphylElems.Add(SphylElem);
 	}
+
+	return true;
 }
 
 
diff --git a/Engine/Source/Editor/UnrealEd/Public/PhysicsAssetUtils.h b/Engine/Source/Editor/UnrealEd/Public/PhysicsAssetUtils.h
--- a/Engine/Source/Editor/UnrealEd/Public/PhysicsAssetUtils.h
+++ b/Engine/Source/Editor/UnrealEd/Public/PhysicsAssetUtils.h
@@ -55,8 +55,9 @@ namespace FPhysicsAssetUtils
 	 * @param	BoneIndex			Index of the bone the collision is created for
 	 * @param	Params				Additional parameters to control the creation 
 	 * @param	Infos				The vertices to create the collirion for
+	 * @return  Returns true if successfully created collision from bone
 	 */
-	UNREALED_API void CreateCollisionFromBone( UBodySetup* bs, USkeletalMesh* skelMesh, int32 BoneIndex, FPhysAssetCreateParams& Params, const TArray<FBoneVertInfo>& Infos );
+	UNREALED_API bool CreateCollisionFromBone( UBodySetup* bs, USkeletalMesh* skelMesh, int32 BoneIndex, FPhysAssetCreateParams& Params, const TArray<FBoneVertInfo>& Infos );
 
 	/**
 	 * Does a few things:
