Fix non unity build

diff --git a/Engine/Source/Editor/UnrealEd/Private/PhysicsAssetUtils.cpp b/Engine/Source/Editor/UnrealEd/Private/PhysicsAssetUtils.cpp
--- a/Engine/Source/Editor/UnrealEd/Private/PhysicsAssetUtils.cpp
+++ b/Engine/Source/Editor/UnrealEd/Private/PhysicsAssetUtils.cpp
@@ -4,6 +4,7 @@
 #include "PhysicsAssetUtils.h"
 #include "Developer/MeshUtilities/Public/MeshUtilities.h"
 #include "Editor/UnrealEd/Private/ConvexDecompTool.h"
+#include "MessageLog.h"
 
 
 void FPhysAssetCreateParams::Initialize()
