UE4: CloseEditedWorldAssets now also closes asset editors that are editing world assets. Also, it closes ALL asset editors for that asset, not just the first primary one.

diff --git a/Engine/Source/Editor/UnrealEd/Private/Editor.cpp b/Engine/Source/Editor/UnrealEd/Private/Editor.cpp
--- a/Engine/Source/Editor/UnrealEd/Private/Editor.cpp
+++ b/Engine/Source/Editor/UnrealEd/Private/Editor.cpp
@@ -1858,14 +1858,24 @@ void UEditorEngine::CloseEditedWorldAssets(UWorld* InWorld)
 
 	for(int32 i=0; i<AllAssets.Num(); i++)
 	{
-		UWorld* AssetWorld = AllAssets[i]->GetTypedOuter<UWorld>();
+		UObject* Asset = AllAssets[i];
+		UWorld* AssetWorld = Asset->GetTypedOuter<UWorld>();
+
+		if ( !AssetWorld )
+		{
+			// This might be a world, itself
+			AssetWorld = Cast<UWorld>(Asset);
+		}
 
 		if (AssetWorld && ClosingWorlds.Contains(AssetWorld))
-		{		
-			IAssetEditorInstance* EditorInstance = EditorManager.FindEditorForAsset(AllAssets[i], false);
-			if(EditorInstance != NULL)
+		{
+			const TArray<IAssetEditorInstance*> AssetEditors = EditorManager.FindEditorsForAsset(Asset);
+			for (IAssetEditorInstance* EditorInstance : AssetEditors )
 			{
-				EditorInstance->CloseWindow();
+				if (EditorInstance != NULL)
+				{
+					EditorInstance->CloseWindow();
+				}
 			}
 		}
 	}
