UE4: Added a thumbnail renderer for world assets. This is slightly different than the level thumbnail renderer in that it needs more features for arbitrary levels and renders streaming levels in the same world as well.

diff --git a/Engine/Config/BaseEditor.ini b/Engine/Config/BaseEditor.ini
--- a/Engine/Config/BaseEditor.ini
+++ b/Engine/Config/BaseEditor.ini
@@ -29,6 +29,7 @@ MaxNumberOfLogLines=2000
 +RenderableThumbnailTypes=(ClassNeedingThumbnailName="/Script/Engine.StaticMesh",RendererClassName="/Script/UnrealEd.StaticMeshThumbnailRenderer")
 +RenderableThumbnailTypes=(ClassNeedingThumbnailName="/Script/Engine.SkeletalMesh",RendererClassName="/Script/UnrealEd.SkeletalMeshThumbnailRenderer")
 +RenderableThumbnailTypes=(ClassNeedingThumbnailName="/Script/Engine.Level",RendererClassName="/Script/UnrealEd.LevelThumbnailRenderer")
++RenderableThumbnailTypes=(ClassNeedingThumbnailName="/Script/Engine.World",RendererClassName="/Script/UnrealEd.WorldThumbnailRenderer")
 +RenderableThumbnailTypes=(ClassNeedingThumbnailName="/Script/Engine.SoundWave",RendererClassName="/Script/UnrealEd.SoundWaveThumbnailRenderer")
 +RenderableThumbnailTypes=(ClassNeedingThumbnailName="/Script/Engine.SlateBrushAsset",RendererClassName="/Script/UnrealEd.SlateBrushThumbnailRenderer")
 
@@ -36,6 +37,10 @@ MaxNumberOfLogLines=2000
 DefaultSizeX=512
 DefaultSizeY=512
 
+[/Script/UnrealEd.WorldThumbnailRenderer]
+GlobalOrbitYawOffset=0
+bUseUnlitScene=false
+
 [/Script/UnrealEd.ActorFactory]
 MenuPriority=10
 
diff --git a/Engine/Source/Developer/AssetTools/Private/AssetTypeActions/AssetTypeActions_World.cpp b/Engine/Source/Developer/AssetTools/Private/AssetTypeActions/AssetTypeActions_World.cpp
--- a/Engine/Source/Developer/AssetTools/Private/AssetTypeActions/AssetTypeActions_World.cpp
+++ b/Engine/Source/Developer/AssetTools/Private/AssetTypeActions/AssetTypeActions_World.cpp
@@ -29,4 +29,17 @@ void FAssetTypeActions_World::OpenAssetEditor( const TArray<UObject*>& InObjects
 	}
 }
 
+UThumbnailInfo* FAssetTypeActions_World::GetThumbnailInfo(UObject* Asset) const
+{
+	UWorld* World = CastChecked<UWorld>(Asset);
+	UThumbnailInfo* ThumbnailInfo = World->ThumbnailInfo;
+	if (ThumbnailInfo == NULL)
+	{
+		ThumbnailInfo = ConstructObject<UWorldThumbnailInfo>(UWorldThumbnailInfo::StaticClass(), World);
+		World->ThumbnailInfo = ThumbnailInfo;
+	}
+
+	return ThumbnailInfo;
+}
+
 #undef LOCTEXT_NAMESPACE
diff --git a/Engine/Source/Developer/AssetTools/Private/AssetTypeActions/AssetTypeActions_World.h b/Engine/Source/Developer/AssetTools/Private/AssetTypeActions/AssetTypeActions_World.h
--- a/Engine/Source/Developer/AssetTools/Private/AssetTypeActions/AssetTypeActions_World.h
+++ b/Engine/Source/Developer/AssetTools/Private/AssetTypeActions/AssetTypeActions_World.h
@@ -12,4 +12,5 @@ public:
 	virtual bool HasActions ( const TArray<UObject*>& InObjects ) const override { return false; }
 	virtual void OpenAssetEditor( const TArray<UObject*>& InObjects, TSharedPtr<class IToolkitHost> EditWithinLevelEditor = TSharedPtr<IToolkitHost>() ) override;
 	virtual uint32 GetCategories() override { return EAssetTypeCategories::Basic; }
+	virtual class UThumbnailInfo* GetThumbnailInfo(UObject* Asset) const override;
 };
\ No newline at end of file
diff --git a/Engine/Source/Editor/UnrealEd/Classes/ThumbnailRendering/SceneThumbnailInfo.h b/Engine/Source/Editor/UnrealEd/Classes/ThumbnailRendering/SceneThumbnailInfo.h
--- a/Engine/Source/Editor/UnrealEd/Classes/ThumbnailRendering/SceneThumbnailInfo.h
+++ b/Engine/Source/Editor/UnrealEd/Classes/ThumbnailRendering/SceneThumbnailInfo.h
@@ -5,6 +5,8 @@
  */
 
 #pragma once
+
+#include "Runtime/Engine/Classes/EditorFramework/ThumbnailInfo.h"
 #include "SceneThumbnailInfo.generated.h"
 
 
diff --git a/Engine/Source/Editor/UnrealEd/Classes/ThumbnailRendering/WorldThumbnailInfo.h b/Engine/Source/Editor/UnrealEd/Classes/ThumbnailRendering/WorldThumbnailInfo.h
new file mode 100644
--- /dev/null
+++ b/Engine/Source/Editor/UnrealEd/Classes/ThumbnailRendering/WorldThumbnailInfo.h
@@ -0,0 +1,38 @@
+// Copyright 1998-2014 Epic Games, Inc. All Rights Reserved.
+
+/**
+ * Thumbnail information for world assets
+ */
+
+#pragma once
+
+#include "Editor/UnrealEd/Classes/ThumbnailRendering/SceneThumbnailInfo.h"
+#include "WorldThumbnailInfo.generated.h"
+
+UENUM()
+namespace EOrthoThumbnailDirection
+{
+	enum Type
+	{
+		Top,
+		Bottom,
+		Left,
+		Right,
+		Front,
+		Back
+	};
+}
+
+UCLASS(MinimalAPI)
+class UWorldThumbnailInfo : public USceneThumbnailInfo
+{
+	GENERATED_UCLASS_BODY()
+
+	/** The type of projection to use */
+	UPROPERTY(EditAnywhere, Category=Thumbnail)
+	TEnumAsByte<ECameraProjectionMode::Type> CameraMode;
+
+	/** The direction of the camera when in Orthographic CameraMode */
+	UPROPERTY(EditAnywhere, Category=Thumbnail)
+	TEnumAsByte<EOrthoThumbnailDirection::Type> OrthoDirection;
+};
diff --git a/Engine/Source/Editor/UnrealEd/Classes/ThumbnailRendering/WorldThumbnailRenderer.h b/Engine/Source/Editor/UnrealEd/Classes/ThumbnailRendering/WorldThumbnailRenderer.h
new file mode 100644
--- /dev/null
+++ b/Engine/Source/Editor/UnrealEd/Classes/ThumbnailRendering/WorldThumbnailRenderer.h
@@ -0,0 +1,33 @@
+// Copyright 1998-2014 Epic Games, Inc. All Rights Reserved.
+
+/**
+ *
+ * This thumbnail renderer displays a given world
+ */
+
+#pragma once
+#include "WorldThumbnailRenderer.generated.h"
+
+UCLASS(config=Editor)
+class UWorldThumbnailRenderer : public UDefaultSizedThumbnailRenderer
+{
+	GENERATED_UCLASS_BODY()
+
+	// Begin UThumbnailRenderer Object
+	virtual bool CanVisualizeAsset(UObject* Object) override;
+	virtual void Draw(UObject* Object, int32 X, int32 Y, uint32 Width, uint32 Height, FRenderTarget*, FCanvas* Canvas) override;
+	// End UThumbnailRenderer Object
+
+private:
+	void GetView(UWorld* World, FSceneViewFamily* ViewFamily, int32 X, int32 Y, uint32 SizeX, uint32 SizeY) const; 
+
+private:
+	/** Offset used to orient all worlds to face the camera in degrees when using a perspective camera. Individual thumbnail infos can provide additional offset. */
+	UPROPERTY(config)
+	float GlobalOrbitYawOffset;
+
+	/** If true, all world thumbnails will be rendered unlit. This is useful in games that have shared lighting in a common map */
+	UPROPERTY(config)
+	bool bUseUnlitScene;
+};
+
diff --git a/Engine/Source/Editor/UnrealEd/Private/AssetThumbnail.cpp b/Engine/Source/Editor/UnrealEd/Private/AssetThumbnail.cpp
--- a/Engine/Source/Editor/UnrealEd/Private/AssetThumbnail.cpp
+++ b/Engine/Source/Editor/UnrealEd/Private/AssetThumbnail.cpp
@@ -1210,11 +1210,22 @@ void FAssetThumbnailPool::OnAssetLoaded( UObject* Asset )
 
 void FAssetThumbnailPool::OnObjectPropertyChanged( UObject* ObjectBeingModified, FPropertyChangedEvent& PropertyChangedEvent )
 {
-	if ( ObjectBeingModified->HasAnyFlags(RF_ClassDefaultObject) && ObjectBeingModified->GetClass()->ClassGeneratedBy != NULL )
+	if ( ObjectBeingModified->HasAnyFlags(RF_ClassDefaultObject) )
 	{
-		// This is a blueprint modification. Check to see if this thumbnail is the blueprint of the modified CDO
-		ObjectBeingModified = ObjectBeingModified->GetClass()->ClassGeneratedBy;
+		if ( ObjectBeingModified->GetClass()->ClassGeneratedBy != NULL )
+		{
+			// This is a blueprint modification. Check to see if this thumbnail is the blueprint of the modified CDO
+			ObjectBeingModified = ObjectBeingModified->GetClass()->ClassGeneratedBy;
+		}
+	}
+	else if ( AActor* ActorBeingModified = Cast<AActor>(ObjectBeingModified) )
+	{
+		// This is a non CDO actor getting modified. Update the actor's world's thumbnail.
+		ObjectBeingModified = ActorBeingModified->GetWorld();
 	}
 
-	RefreshThumbnailsFor( FName(*ObjectBeingModified->GetPathName()) );
+	if ( ObjectBeingModified )
+	{
+		RefreshThumbnailsFor( FName(*ObjectBeingModified->GetPathName()) );
+	}
 }
diff --git a/Engine/Source/Editor/UnrealEd/Private/Factories/WorldFactory.cpp b/Engine/Source/Editor/UnrealEd/Private/Factories/WorldFactory.cpp
--- a/Engine/Source/Editor/UnrealEd/Private/Factories/WorldFactory.cpp
+++ b/Engine/Source/Editor/UnrealEd/Private/Factories/WorldFactory.cpp
@@ -24,6 +24,7 @@ UObject* UWorldFactory::FactoryCreateNew(UClass* Class, UObject* InParent, FName
 	const bool bAddToRoot = false;
 	UWorld* NewWorld = UWorld::CreateWorld(WorldType, bInformEngineOfWorld, Name, Cast<UPackage>(InParent), bAddToRoot);
 	NewWorld->SetFlags(Flags);
+	NewWorld->ThumbnailInfo = ConstructObject<UWorldThumbnailInfo>(UWorldThumbnailInfo::StaticClass(), NewWorld);
 
 	return NewWorld;
 }
diff --git a/Engine/Source/Editor/UnrealEd/Private/ThumbnailRendering/WorldThumbnailInfo.cpp b/Engine/Source/Editor/UnrealEd/Private/ThumbnailRendering/WorldThumbnailInfo.cpp
new file mode 100644
--- /dev/null
+++ b/Engine/Source/Editor/UnrealEd/Private/ThumbnailRendering/WorldThumbnailInfo.cpp
@@ -0,0 +1,10 @@
+// Copyright 1998-2014 Epic Games, Inc. All Rights Reserved.
+
+#include "UnrealEd.h"
+
+UWorldThumbnailInfo::UWorldThumbnailInfo(const class FPostConstructInitializeProperties& PCIP)
+	: Super(PCIP)
+{
+	CameraMode = ECameraProjectionMode::Orthographic;
+	OrthoDirection = EOrthoThumbnailDirection::Top;
+}
diff --git a/Engine/Source/Editor/UnrealEd/Private/ThumbnailRendering/WorldThumbnailRenderer.cpp b/Engine/Source/Editor/UnrealEd/Private/ThumbnailRendering/WorldThumbnailRenderer.cpp
new file mode 100644
--- /dev/null
+++ b/Engine/Source/Editor/UnrealEd/Private/ThumbnailRendering/WorldThumbnailRenderer.cpp
@@ -0,0 +1,252 @@
+// Copyright 1998-2014 Epic Games, Inc. All Rights Reserved.
+
+#include "UnrealEd.h"
+#include "Runtime/Engine/Public/EngineModule.h"
+#include "Runtime/Engine/Public/ContentStreaming.h"
+
+UWorldThumbnailRenderer::UWorldThumbnailRenderer(const class FPostConstructInitializeProperties& PCIP)
+	: Super(PCIP)
+{
+	bUseUnlitScene = false;
+	GlobalOrbitYawOffset = 0.f;
+}
+
+bool UWorldThumbnailRenderer::CanVisualizeAsset(UObject* Object)
+{
+	UWorld* World = Cast<UWorld>(Object);
+	if (World && World->PersistentLevel)
+	{
+		// If this is a world, only render the current persistent editor world. Other worlds don't have an initialized scene to render.
+		return World->WorldType == EWorldType::Editor && World->PersistentLevel->OwningWorld == World;
+	}
+
+	return false;
+}
+
+void UWorldThumbnailRenderer::Draw(UObject* Object, int32 X, int32 Y, uint32 Width, uint32 Height, FRenderTarget* RenderTarget, FCanvas* Canvas)
+{
+	UWorld* World = Cast<UWorld>(Object);
+	if (World != nullptr)
+	{
+		FSceneViewFamilyContext ViewFamily( FSceneViewFamily::ConstructionValues( RenderTarget, World->Scene, FEngineShowFlags(ESFIM_Game) )
+			.SetWorldTimes(FApp::GetCurrentTime() - GStartTime, FApp::GetDeltaTime(), FApp::GetCurrentTime() - GStartTime));
+
+		ViewFamily.EngineShowFlags.DisableAdvancedFeatures();
+		ViewFamily.EngineShowFlags.SetMotionBlur(false);
+		ViewFamily.EngineShowFlags.SetLighting(!bUseUnlitScene);
+		ViewFamily.EngineShowFlags.SetPostProcessing(false);
+		ViewFamily.EngineShowFlags.SetDistanceCulledPrimitives(true);
+		ViewFamily.EngineShowFlags.SetFog(false);
+
+		GetView(World, &ViewFamily, X, Y, Width, Height);
+
+		if (ViewFamily.Views.Num() > 0)
+		{
+			GetRendererModule().BeginRenderingViewFamily(Canvas, &ViewFamily);
+		}
+	}
+}
+
+void UWorldThumbnailRenderer::GetView(UWorld* World, FSceneViewFamily* ViewFamily, int32 X, int32 Y, uint32 SizeX, uint32 SizeY) const 
+{
+	check(ViewFamily);
+	check(World);
+	check(World->PersistentLevel);
+
+	FIntRect ViewRect(
+		FMath::Max<int32>(X, 0),
+		FMath::Max<int32>(Y, 0),
+		FMath::Max<int32>(X + SizeX, 0),
+		FMath::Max<int32>(Y + SizeY, 0));
+
+	if (ViewRect.Width() > 0 && ViewRect.Height() > 0)
+	{
+		FBox WorldBox(0);
+		TArray<ULevel*> LevelsToRender = World->GetLevels();
+		for ( auto* Level : LevelsToRender )
+		{
+			if (Level)
+			{
+				if (Level->LevelBoundsActor.IsValid())
+				{
+					WorldBox += Level->LevelBoundsActor.Get()->GetComponentsBoundingBox();
+				}
+				else
+				{
+					// Ensure a Level Bounds Actor exists for future renders
+					FActorSpawnParameters SpawnParameters;
+					SpawnParameters.OverrideLevel = Level;
+					World->SpawnActor<ALevelBounds>(SpawnParameters);
+
+					// Calculate the bounds manually for this frame. Real data should be available shortly
+					WorldBox += ALevelBounds::CalculateLevelBounds(Level);
+				}
+			}
+		}
+
+		UWorldThumbnailInfo* ThumbnailInfo = Cast<UWorldThumbnailInfo>(World->ThumbnailInfo);
+		if (!ThumbnailInfo)
+		{
+			ThumbnailInfo = UWorldThumbnailInfo::StaticClass()->GetDefaultObject<UWorldThumbnailInfo>();
+		}
+
+		const FVector Origin = WorldBox.GetCenter();
+		FMatrix ViewMatrix = FTranslationMatrix(-Origin);
+		FMatrix ProjectionMatrix;
+		float FOVScreenSize = 0; // Screen size taking FOV into account
+		if (ThumbnailInfo->CameraMode == ECameraProjectionMode::Perspective)
+		{
+			const float FOVDegrees = 30.f;
+			const float HalfFOVRadians = FMath::DegreesToRadians<float>(FOVDegrees) * 0.5f;
+			const float WorldRadius = WorldBox.GetSize().Size() / 2.f;
+			float TargetDistance = WorldRadius / FMath::Tan(HalfFOVRadians);
+
+			if (ensure(ThumbnailInfo))
+			{
+				if (TargetDistance + ThumbnailInfo->OrbitZoom < 0)
+				{
+					ThumbnailInfo->OrbitZoom = -TargetDistance;
+				}
+			}
+
+			float OrbitPitch = ThumbnailInfo->OrbitPitch;
+			float OrbitYaw = GlobalOrbitYawOffset + ThumbnailInfo->OrbitYaw;
+			float OrbitZoom = TargetDistance + ThumbnailInfo->OrbitZoom;
+
+			// Ensure a minimum camera distance to prevent problems with really small objects
+			const float MinCameraDistance = 48;
+			OrbitZoom = FMath::Max<float>(MinCameraDistance, OrbitZoom);
+
+			const FRotator RotationOffsetToViewCenter(0.f, 90.f, 0.f);
+			ViewMatrix = ViewMatrix *
+				FRotationMatrix(FRotator(0, OrbitYaw, 0)) *
+				FRotationMatrix(FRotator(0, 0, OrbitPitch)) *
+				FTranslationMatrix(FVector(0, OrbitZoom, 0)) *
+				FInverseRotationMatrix(RotationOffsetToViewCenter);
+
+			ViewMatrix = ViewMatrix * FMatrix(
+				FPlane(0, 0, 1, 0),
+				FPlane(1, 0, 0, 0),
+				FPlane(0, 1, 0, 0),
+				FPlane(0, 0, 0, 1));
+
+			const float NearPlane = 1.0f;
+			ProjectionMatrix = FReversedZPerspectiveMatrix(
+				HalfFOVRadians,
+				1.0f,
+				1.0f,
+				NearPlane
+				);
+
+			FOVScreenSize = SizeX / FMath::Tan(FOVDegrees);
+		}
+		else if (ThumbnailInfo->CameraMode == ECameraProjectionMode::Orthographic)
+		{
+			FVector2D WorldSizeMin2D;
+			FVector2D WorldSizeMax2D;
+			switch (ThumbnailInfo->OrthoDirection)
+			{
+				case EOrthoThumbnailDirection::Top:
+					ViewMatrix = ViewMatrix * FMatrix(
+						FPlane(1, 0, 0, 0),
+						FPlane(0, -1, 0, 0),
+						FPlane(0, 0, -1, 0),
+						FPlane(0, 0, Origin.Z, 1));
+					WorldSizeMin2D = FVector2D(WorldBox.Min.X,WorldBox.Min.Y);
+					WorldSizeMax2D = FVector2D(WorldBox.Max.X,WorldBox.Max.Y);
+					break;
+				case EOrthoThumbnailDirection::Bottom:
+					ViewMatrix = ViewMatrix * FMatrix(
+						FPlane(1, 0, 0, 0),
+						FPlane(0, -1, 0, 0),
+						FPlane(0, 0, 1, 0),
+						FPlane(0, 0, Origin.Z, 1));
+					WorldSizeMin2D = FVector2D(WorldBox.Min.X, WorldBox.Min.Y);
+					WorldSizeMax2D = FVector2D(WorldBox.Max.X, WorldBox.Max.Y);
+					break;
+				case EOrthoThumbnailDirection::Front:
+					ViewMatrix = ViewMatrix * FMatrix(
+						FPlane(1, 0, 0, 0),
+						FPlane(0, 0, -1, 0),
+						FPlane(0, 1, 0, 0),
+						FPlane(0, 0, Origin.Y, 1));
+					WorldSizeMin2D = FVector2D(WorldBox.Min.X, WorldBox.Min.Z);
+					WorldSizeMax2D = FVector2D(WorldBox.Max.X, WorldBox.Max.Z);
+					break;
+				case EOrthoThumbnailDirection::Back:
+					ViewMatrix = ViewMatrix * FMatrix(
+						FPlane(-1, 0, 0, 0),
+						FPlane(0, 0, 1, 0),
+						FPlane(0, 1, 0, 0),
+						FPlane(0, 0, Origin.Y, 1));
+					WorldSizeMin2D = FVector2D(WorldBox.Min.X, WorldBox.Min.Z);
+					WorldSizeMax2D = FVector2D(WorldBox.Max.X, WorldBox.Max.Z);
+					break;
+				case EOrthoThumbnailDirection::Left:
+					ViewMatrix = ViewMatrix * FMatrix(
+						FPlane(0, 0, -1, 0),
+						FPlane(-1, 0, 0, 0),
+						FPlane(0, 1, 0, 0),
+						FPlane(0, 0, Origin.X, 1));
+					WorldSizeMin2D = FVector2D(WorldBox.Min.Y, WorldBox.Min.Z);
+					WorldSizeMax2D = FVector2D(WorldBox.Max.Y, WorldBox.Max.Z);
+					break;
+				case EOrthoThumbnailDirection::Right:
+					ViewMatrix = ViewMatrix * FMatrix(
+						FPlane(0, 0, 1, 0),
+						FPlane(1, 0, 0, 0),
+						FPlane(0, 1, 0, 0),
+						FPlane(0, 0, Origin.X, 1));
+					WorldSizeMin2D = FVector2D(WorldBox.Min.Y, WorldBox.Min.Z);
+					WorldSizeMax2D = FVector2D(WorldBox.Max.Y, WorldBox.Max.Z);
+					break;
+				default:
+					// Unknown OrthoDirection
+					ensureMsgf(false, TEXT("Unknown thumbnail OrthoDirection"));
+					break;
+			}
+
+			FVector2D WorldSize2D = (WorldSizeMax2D - WorldSizeMin2D);
+			WorldSize2D.X = FMath::Abs(WorldSize2D.X);
+			WorldSize2D.Y = FMath::Abs(WorldSize2D.Y);
+			const bool bUseXAxis = (WorldSize2D.X / WorldSize2D.Y) > 1.f;
+			const float WorldAxisSize = bUseXAxis ? WorldSize2D.X : WorldSize2D.Y;
+			const uint32 ViewportAxisSize = bUseXAxis ? SizeX : SizeY;
+			const float OrthoZoom = WorldAxisSize / ViewportAxisSize / 2.f;
+			const float OrthoWidth = FMath::Max(1.f, SizeX * OrthoZoom);
+			const float OrthoHeight = FMath::Max(1.f, SizeY * OrthoZoom);
+
+			const float ZOffset = HALF_WORLD_MAX;
+			ProjectionMatrix = FReversedZOrthoMatrix(
+				OrthoWidth,
+				OrthoHeight,
+				0.5f / ZOffset,
+				ZOffset
+				);
+
+			FOVScreenSize = SizeX;
+		}
+		else
+		{
+			// Unknown CameraMode
+			ensureMsgf(false, TEXT("Unknown thumbnail CameraMode"));
+		}
+
+		FSceneViewInitOptions ViewInitOptions;
+		ViewInitOptions.ViewFamily = ViewFamily;
+		ViewInitOptions.SetViewRectangle(ViewRect);
+		ViewInitOptions.BackgroundColor = FLinearColor::Black;
+		ViewInitOptions.ViewMatrix = ViewMatrix;
+		ViewInitOptions.ProjectionMatrix = ProjectionMatrix;
+
+		FSceneView* NewView = new FSceneView(ViewInitOptions);
+
+		ViewFamily->Views.Add(NewView);
+
+		// Tell the texture streaming system about this thumbnail view, so the textures will stream in as needed
+		// NOTE: Sizes may not actually be in screen space depending on how the thumbnail ends up stretched by the UI.  Not a big deal though.
+		// NOTE: Textures still take a little time to stream if the view has not been re-rendered recently, so they may briefly appear blurry while mips are prepared
+		// NOTE: Content Browser only renders thumbnails for loaded assets, and only when the mouse is over the panel. They'll be frozen in their last state while the mouse cursor is not over the panel.  This is for performance reasons
+		IStreamingManager::Get().AddViewInformation(Origin, SizeX, FOVScreenSize);
+	}
+}
diff --git a/Engine/Source/Runtime/Engine/Classes/Engine/World.h b/Engine/Source/Runtime/Engine/Classes/Engine/World.h
--- a/Engine/Source/Runtime/Engine/Classes/Engine/World.h
+++ b/Engine/Source/Runtime/Engine/Classes/Engine/World.h
@@ -477,6 +477,10 @@ class ENGINE_API UWorld : public UObject, public FNetworkNotify
 	// Group actors currently "active"
 	UPROPERTY(transient)
 	TArray<class AActor*> ActiveGroupActors;
+
+	/** Information for thumbnail rendering */
+	UPROPERTY(VisibleAnywhere, EditInline, Category=Thumbnail)
+	class UThumbnailInfo* ThumbnailInfo;
 #endif // WITH_EDITORONLY_DATA
 
 	/** Persistent level containing the world info, default brush and actors spawned during gameplay among other things			*/
diff --git a/Engine/Source/Runtime/Engine/Private/World.cpp b/Engine/Source/Runtime/Engine/Private/World.cpp
--- a/Engine/Source/Runtime/Engine/Private/World.cpp
+++ b/Engine/Source/Runtime/Engine/Private/World.cpp
@@ -35,6 +35,7 @@
 #if WITH_EDITOR
 	#include "Editor/UnrealEd/Public/Kismet2/KismetEditorUtilities.h"
 	#include "Editor/UnrealEd/Public/Kismet2/BlueprintEditorUtils.h"
+	#include "Editor/UnrealEd/Classes/ThumbnailRendering/WorldThumbnailInfo.h"
 	#include "Slate.h"
 #endif
 
@@ -569,6 +570,12 @@ void UWorld::PostLoad()
 				Model->Rename(TEXT("Brush"), DefaultBrush->GetOuter(), REN_DoNotDirty | REN_DontCreateRedirectors | REN_ForceNoResetLoaders | REN_NonTransactional);
 			}
 		}
+
+		// Make sure thumbnail info exists
+		if ( !ThumbnailInfo )
+		{
+			ThumbnailInfo = ConstructObject<UWorldThumbnailInfo>(UWorldThumbnailInfo::StaticClass(), this);
+		}
 	}
 #endif
 }
