UE4: World asset thumbnails are now updated when moving actors.

diff --git a/Engine/Source/Editor/UnrealEd/Private/AssetThumbnail.cpp b/Engine/Source/Editor/UnrealEd/Private/AssetThumbnail.cpp
--- a/Engine/Source/Editor/UnrealEd/Private/AssetThumbnail.cpp
+++ b/Engine/Source/Editor/UnrealEd/Private/AssetThumbnail.cpp
@@ -687,12 +687,20 @@ FAssetThumbnailPool::FAssetThumbnailPool( uint32 InNumInPool, const TAttribute<b
 {
 	FCoreDelegates::OnObjectPropertyChanged.AddRaw(this, &FAssetThumbnailPool::OnObjectPropertyChanged);
 	FCoreDelegates::OnAssetLoaded.AddRaw(this, &FAssetThumbnailPool::OnAssetLoaded);
+	if ( GEditor )
+	{
+		GEditor->OnActorMoved().AddRaw( this, &FAssetThumbnailPool::OnActorPostEditMove );
+	}
 }
 
 FAssetThumbnailPool::~FAssetThumbnailPool()
 {
 	FCoreDelegates::OnObjectPropertyChanged.RemoveAll(this);
 	FCoreDelegates::OnAssetLoaded.RemoveAll(this);
+	if ( GEditor )
+	{
+		GEditor->OnActorMoved().RemoveAll(this);
+	}
 
 	// Release all the texture resources
 	ReleaseResources();
@@ -1208,6 +1216,18 @@ void FAssetThumbnailPool::OnAssetLoaded( UObject* Asset )
 	}
 }
 
+void FAssetThumbnailPool::OnActorPostEditMove( AActor* Actor )
+{
+	if (Actor != NULL)
+	{
+		UWorld* World = Actor->GetWorld();
+		if (World)
+		{
+			RefreshThumbnailsFor(FName(*World->GetPathName()));
+		}
+	}
+}
+
 void FAssetThumbnailPool::OnObjectPropertyChanged( UObject* ObjectBeingModified, FPropertyChangedEvent& PropertyChangedEvent )
 {
 	if ( ObjectBeingModified->HasAnyFlags(RF_ClassDefaultObject) )
diff --git a/Engine/Source/Editor/UnrealEd/Public/AssetThumbnail.h b/Engine/Source/Editor/UnrealEd/Public/AssetThumbnail.h
--- a/Engine/Source/Editor/UnrealEd/Public/AssetThumbnail.h
+++ b/Engine/Source/Editor/UnrealEd/Public/AssetThumbnail.h
@@ -188,6 +188,9 @@ private:
 	/** Handler for when an asset is loaded */
 	void OnAssetLoaded( UObject* Asset );
 
+	/** Handler for when an actor is moved in a level. Used to update world asset thumbnails. */
+	void OnActorPostEditMove( AActor* Actor );
+
 	/** Handler for when an asset is loaded */
 	void OnObjectPropertyChanged( UObject* Asset, FPropertyChangedEvent& PropertyChangedEvent );
 
