UE4: Disabling world asset thumbnails by default for now. Some performance and bounds issues need to be figured out for the common case. Turned off many rendering flags for the thumbnail to boost perf. Fixed a crash when attempting to render without a scene.

diff --git a/Engine/Config/BaseEditor.ini b/Engine/Config/BaseEditor.ini
--- a/Engine/Config/BaseEditor.ini
+++ b/Engine/Config/BaseEditor.ini
@@ -40,6 +40,7 @@ DefaultSizeY=512
 [/Script/UnrealEd.WorldThumbnailRenderer]
 GlobalOrbitYawOffset=0
 bUseUnlitScene=false
+bAllowWorldThumbnails=false
 
 [/Script/UnrealEd.ActorFactory]
 MenuPriority=10
diff --git a/Engine/Source/Editor/UnrealEd/Classes/ThumbnailRendering/WorldThumbnailRenderer.h b/Engine/Source/Editor/UnrealEd/Classes/ThumbnailRendering/WorldThumbnailRenderer.h
--- a/Engine/Source/Editor/UnrealEd/Classes/ThumbnailRendering/WorldThumbnailRenderer.h
+++ b/Engine/Source/Editor/UnrealEd/Classes/ThumbnailRendering/WorldThumbnailRenderer.h
@@ -29,5 +29,9 @@ private:
 	/** If true, all world thumbnails will be rendered unlit. This is useful in games that have shared lighting in a common map */
 	UPROPERTY(config)
 	bool bUseUnlitScene;
+
+	/** If false, all world thumbnails rendering will be disabled */
+	UPROPERTY(config)
+	bool bAllowWorldThumbnails;
 };
 
diff --git a/Engine/Source/Editor/UnrealEd/Private/ThumbnailRendering/WorldThumbnailRenderer.cpp b/Engine/Source/Editor/UnrealEd/Private/ThumbnailRendering/WorldThumbnailRenderer.cpp
--- a/Engine/Source/Editor/UnrealEd/Private/ThumbnailRendering/WorldThumbnailRenderer.cpp
+++ b/Engine/Source/Editor/UnrealEd/Private/ThumbnailRendering/WorldThumbnailRenderer.cpp
@@ -7,17 +7,21 @@
 UWorldThumbnailRenderer::UWorldThumbnailRenderer(const class FPostConstructInitializeProperties& PCIP)
 	: Super(PCIP)
 {
-	bUseUnlitScene = false;
 	GlobalOrbitYawOffset = 0.f;
+	bUseUnlitScene = false;
+	bAllowWorldThumbnails = false;
 }
 
 bool UWorldThumbnailRenderer::CanVisualizeAsset(UObject* Object)
 {
-	UWorld* World = Cast<UWorld>(Object);
-	if (World && World->PersistentLevel)
+	if ( bAllowWorldThumbnails )
 	{
-		// If this is a world, only render the current persistent editor world. Other worlds don't have an initialized scene to render.
-		return World->WorldType == EWorldType::Editor && World->PersistentLevel->OwningWorld == World;
+		UWorld* World = Cast<UWorld>(Object);
+		if (World && World->PersistentLevel)
+		{
+			// If this is a world, only render the current persistent editor world. Other worlds don't have an initialized scene to render.
+			return World->WorldType == EWorldType::Editor && World->PersistentLevel->OwningWorld == World;
+		}
 	}
 
 	return false;
@@ -26,17 +30,38 @@ bool UWorldThumbnailRenderer::CanVisualizeAsset(UObject* Object)
 void UWorldThumbnailRenderer::Draw(UObject* Object, int32 X, int32 Y, uint32 Width, uint32 Height, FRenderTarget* RenderTarget, FCanvas* Canvas)
 {
 	UWorld* World = Cast<UWorld>(Object);
-	if (World != nullptr)
+	if (World != nullptr && World->Scene)
 	{
-		FSceneViewFamilyContext ViewFamily( FSceneViewFamily::ConstructionValues( RenderTarget, World->Scene, FEngineShowFlags(ESFIM_Game) )
+		FSceneViewFamilyContext ViewFamily( FSceneViewFamily::ConstructionValues( RenderTarget, World->Scene, FEngineShowFlags(ESFIM_All0) )
 			.SetWorldTimes(FApp::GetCurrentTime() - GStartTime, FApp::GetDeltaTime(), FApp::GetCurrentTime() - GStartTime));
 
-		ViewFamily.EngineShowFlags.DisableAdvancedFeatures();
-		ViewFamily.EngineShowFlags.SetMotionBlur(false);
-		ViewFamily.EngineShowFlags.SetLighting(!bUseUnlitScene);
-		ViewFamily.EngineShowFlags.SetPostProcessing(false);
-		ViewFamily.EngineShowFlags.SetDistanceCulledPrimitives(true);
-		ViewFamily.EngineShowFlags.SetFog(false);
+		ViewFamily.EngineShowFlags.SetDiffuse(true);
+		ViewFamily.EngineShowFlags.SetSkeletalMeshes(true);
+		ViewFamily.EngineShowFlags.SetTranslucency(true);
+		ViewFamily.EngineShowFlags.SetBillboardSprites(true);
+		ViewFamily.EngineShowFlags.SetLOD(true);
+		ViewFamily.EngineShowFlags.SetMaterials(true);
+		ViewFamily.EngineShowFlags.SetStaticMeshes(true);
+		ViewFamily.EngineShowFlags.SetLandscape(true);
+		ViewFamily.EngineShowFlags.SetGame(true);
+		ViewFamily.EngineShowFlags.SetBSP(true);
+		ViewFamily.EngineShowFlags.SetRendering(true);
+		ViewFamily.EngineShowFlags.SetPaper2DSprites(true);
+
+		if ( !bUseUnlitScene )
+		{
+			ViewFamily.EngineShowFlags.SetSpecular(true);
+			ViewFamily.EngineShowFlags.SetLighting(true);
+			ViewFamily.EngineShowFlags.SetDirectLighting(true);
+			ViewFamily.EngineShowFlags.SetIndirectLightingCache(true);
+			ViewFamily.EngineShowFlags.SetDeferredLighting(true);
+			ViewFamily.EngineShowFlags.SetDirectionalLights(true);
+			ViewFamily.EngineShowFlags.SetGlobalIllumination(true);
+			ViewFamily.EngineShowFlags.SetPointLights(true);
+			ViewFamily.EngineShowFlags.SetSpotLights(true);
+			ViewFamily.EngineShowFlags.SetSkyLighting(true);
+			ViewFamily.EngineShowFlags.SetReflectionEnvironment(true);
+		}
 
 		GetView(World, &ViewFamily, X, Y, Width, Height);
 
