MeshMerging: added option for specifing max atlased lightmap resolution

diff --git a/Engine/Source/Developer/MeshUtilities/Private/MeshUtilities.cpp b/Engine/Source/Developer/MeshUtilities/Private/MeshUtilities.cpp
--- a/Engine/Source/Developer/MeshUtilities/Private/MeshUtilities.cpp
+++ b/Engine/Source/Developer/MeshUtilities/Private/MeshUtilities.cpp
@@ -3333,7 +3333,7 @@ public:
 	 */ 
 	bool Pack(const TArray<uint32>& LigthmapsList)
 	{
-		// Calculate total ligtmaps area and sort lightmaps list by resolution
+		// Calculate total lightmaps area and sort lightmaps list by resolution
 		TArray<TPair<int32, uint32>> SortedLightmaps;
 		float TotalArea = 0;
 
@@ -3498,34 +3498,36 @@ void FMeshUtilities::MergeActors(
 
 	FRawMeshExt MergedMesh;
 
-	// Pack lightmaps
-	float MergedLightmapScale = 1.f;
-	FLightmapPacker LightmapPacker;	
-					
-	if (InSettings.bGnerateAtlasedLightmapUV)
+	// Attempt to pack lightmaps
+	float MergedLightMapScale = 1.f;
+	bool bCreateLightMapChannel = false;
+	FLightmapPacker LightMapPacker;	
+						
+	if (InSettings.bGenerateAtlasedLightMapUV)
 	{
-		static const uint32 MaxLightmapRes = 2048;
-				
 		// Set target channel for lightmap UV
-		MergedMesh.LightMapCoordinateIndex = InSettings.TargetLightmapUVChannel;
+		MergedMesh.LightMapCoordinateIndex = InSettings.TargetLightMapUVChannel;
 		
 		// Collect lightmap sizes from all meshes
-		TArray<uint32> LightmapResList;
+		TArray<uint32> LightMapResList;
 		for (const FRawMeshExt& SourceMesh : SourceMeshes)
 		{
-			LightmapResList.Add(SourceMesh.LightMapRes);
+			LightMapResList.Add(SourceMesh.LightMapRes);
 		}
 		
 		// Pack them into one atlas
-		LightmapPacker.Pack(LightmapResList);
-		MergedMesh.LightMapRes = LightmapPacker.GetAtlasResolution();
+		LightMapPacker.Pack(LightMapResList);
+		MergedMesh.LightMapRes = LightMapPacker.GetAtlasResolution();
 
-		// Whether we need to scale down UV coordiantes
-		if (MergedMesh.LightMapRes > MaxLightmapRes)
+		// Whether we need to scale down lightUV coordiantes
+		if (MergedMesh.LightMapRes > InSettings.MaxAltlasedLightMapResolution)
 		{
-			MergedLightmapScale = MaxLightmapRes/(float)MergedMesh.LightMapRes;
-			MergedMesh.LightMapRes = MaxLightmapRes;
+			MergedLightMapScale = InSettings.MaxAltlasedLightMapResolution/(float)MergedMesh.LightMapRes;
+			MergedMesh.LightMapRes = InSettings.MaxAltlasedLightMapResolution;
 		}
+
+		// Create lightmap channel in a merged mesh only if we succeed to generate atlas for it
+		bCreateLightMapChannel = (MergedMesh.LightMapRes > 0);
 	}
 
 	// Use first mesh for naming and pivot
@@ -3577,21 +3579,21 @@ void FMeshUtilities::MergeActors(
 		}
 		
 		// Write atlased UVs into user specified TargetLightmapChannel 
-		if (InSettings.bGnerateAtlasedLightmapUV && MergedMesh.LightMapRes)
+		if (bCreateLightMapChannel)
 		{
 			FVector2D	UVOffset = FVector2D::ZeroVector;
-			FIntRect	PackedLightmapRect = LightmapPacker.GetPackedLightmapRect(SourceMeshIdx);
-			if (PackedLightmapRect.Area() != 0)
+			FIntRect	PackedLightMapRect = LightMapPacker.GetPackedLightmapRect(SourceMeshIdx);
+			if (PackedLightMapRect.Area() != 0)
 			{
-				UVOffset = FVector2D(PackedLightmapRect.Min) * MergedLightmapScale / MergedMesh.LightMapRes;
+				UVOffset = FVector2D(PackedLightMapRect.Min) * MergedLightMapScale / MergedMesh.LightMapRes;
 			}
 			
 			const TArray<FVector2D>& SourceWedgeTexCoords = SourceRawMesh.WedgeTexCoords[SourceMeshes[SourceMeshIdx].LightMapCoordinateIndex];
 			
 			for (FVector2D LightMapUV : SourceWedgeTexCoords)
 			{
-				const float SourceMeshLightmapRes = SourceMeshes[SourceMeshIdx].LightMapRes;
-				const float UVScale = (SourceMeshLightmapRes * MergedLightmapScale) / MergedMesh.LightMapRes;
+				const float SourceMeshLightMapRes = SourceMeshes[SourceMeshIdx].LightMapRes;
+				const float UVScale = (SourceMeshLightMapRes * MergedLightMapScale) / MergedMesh.LightMapRes;
 				TargetRawMesh.WedgeTexCoords[MergedMesh.LightMapCoordinateIndex].Add(LightMapUV * UVScale + UVOffset);
 			}
 		}
@@ -3600,7 +3602,7 @@ void FMeshUtilities::MergeActors(
 		for (int32 ChannelIdx = 0; ChannelIdx < MAX_MESH_TEXTURE_COORDS; ++ChannelIdx)
 		{
 			// Skip Lightmap channel if any
-			if (InSettings.bGnerateAtlasedLightmapUV && ChannelIdx == MergedMesh.LightMapCoordinateIndex)
+			if (bCreateLightMapChannel && ChannelIdx == MergedMesh.LightMapCoordinateIndex)
 			{
 				continue;
 			}
@@ -3660,7 +3662,7 @@ void FMeshUtilities::MergeActors(
 		StaticMesh->LightingGuid = FGuid::NewGuid();
 
 		// Set it to use textured lightmaps. Note that Build Lighting will do the error-checking (texcoordindex exists for all LODs, etc).
-		if (InSettings.bGnerateAtlasedLightmapUV)
+		if (bCreateLightMapChannel)
 		{
 			StaticMesh->LightMapResolution = MergedMesh.LightMapRes;
 			StaticMesh->LightMapCoordinateIndex = MergedMesh.LightMapCoordinateIndex;	
diff --git a/Engine/Source/Developer/MeshUtilities/Public/MeshUtilities.h b/Engine/Source/Developer/MeshUtilities/Public/MeshUtilities.h
--- a/Engine/Source/Developer/MeshUtilities/Public/MeshUtilities.h
+++ b/Engine/Source/Developer/MeshUtilities/Public/MeshUtilities.h
@@ -73,28 +73,27 @@ public:
  */
 struct FMeshMergingSettings
 {
-	/** */
-	bool bGnerateAtlasedLightmapUV;
+	/** Whether to generate atlased lightmap UVs for a merged mesh*/
+	bool bGenerateAtlasedLightMapUV;
 	
-	/** */
-	bool bImportVertexColors;
-
-	/** */
-	bool bImportAllUVChannels;
-
-	/** */
-	int32 TargetLightmapUVChannel;
-
-	/** */
-	bool bPivotPointAtZero;
+	/** Target UV channel in a merged mesh for an atlased lightmap */
+	int32 TargetLightMapUVChannel;
 
+	/** Upper bounds for an atlased lightmap resolution */
+	int32 MaxAltlasedLightMapResolution;
+		
+	/** Whether we should import vertex colors into merged mesh */
+	bool bImportVertexColors;
 	
+	/** Whether merged mesh should have pivot at world origin, or at first merged component otherwise */
+	bool bPivotPointAtZero;
+		
 	/** Default settings. */
 	FMeshMergingSettings()
-		: bGnerateAtlasedLightmapUV(true)
+		: bGenerateAtlasedLightMapUV(true)
+		, TargetLightMapUVChannel(1)
+		, MaxAltlasedLightMapResolution(1024)
 		, bImportVertexColors(false)
-		, bImportAllUVChannels(false)
-		, TargetLightmapUVChannel(1)
 		, bPivotPointAtZero(false)
 	{
 	}
diff --git a/Engine/Source/Editor/UnrealEd/Private/Dialogs/SMeshMergingDialog.cpp b/Engine/Source/Editor/UnrealEd/Private/Dialogs/SMeshMergingDialog.cpp
--- a/Engine/Source/Editor/UnrealEd/Private/Dialogs/SMeshMergingDialog.cpp
+++ b/Engine/Source/Editor/UnrealEd/Private/Dialogs/SMeshMergingDialog.cpp
@@ -29,9 +29,22 @@ SMeshMergingDialog::~SMeshMergingDialog()
 void SMeshMergingDialog::Construct(const FArguments& InArgs)
 {
 	ParentWindow = InArgs._ParentWindow;
-	
+
 	FLevelEditorModule& LevelEditor = FModuleManager::GetModuleChecked<FLevelEditorModule>( "LevelEditor" );
 	LevelEditor.OnActorSelectionChanged().AddSP(this, &SMeshMergingDialog::OnActorSelectionChanged);
+	
+	// Setup available resolutions for an atlased lightmap
+	const auto& LightMapGroup = GSystemSettings.TextureLODSettings.GetTextureLODGroup(TEXTUREGROUP_Lightmap);
+	for (int32 Index = LightMapGroup.MinLODMipCount; Index <= LightMapGroup.MaxLODMipCount; Index++)
+	{
+		LightMapResolutionOptions.Add(MakeShareable(new FString(FString::FormatAsNumber(1 << Index))));
+	}
+	
+	// Setup available UV channels for an atlased lightmap
+	for (int32 Index = 0; Index < MAX_MESH_TEXTURE_COORDS; Index++)
+	{
+		LightMapChannelOptions.Add(MakeShareable(new FString(FString::FormatAsNumber(Index))));
+	}
 
 	GenerateNewPackageName();
 
@@ -41,161 +54,180 @@ void SMeshMergingDialog::Construct(const FArguments& InArgs)
 		SNew(SVerticalBox)
 
 		+SVerticalBox::Slot()
-		.Padding(2,2,2,4)
+		.AutoHeight()
+		.Padding(0,2,0,0)
 		[
+			// Lightmap settings
 			SNew(SBorder)
 			.BorderImage(FEditorStyle::GetBrush("ToolPanel.GroupBorder"))
 			[
 				SNew(SVerticalBox)
-
-				//
-				// Merging options
-				//
+									
+				// Enable atlasing
 				+SVerticalBox::Slot()
+				.AutoHeight()
+				.Padding(FEditorStyle::GetMargin("StandardDialog.ContentPadding"))
 				[
-					SNew(SVerticalBox)
-				
-					// Lightmaps atlasing
-					+SVerticalBox::Slot()
-					.AutoHeight()
-					.Padding(FEditorStyle::GetMargin("StandardDialog.ContentPadding"))
+					SNew(SCheckBox)
+					.Type(ESlateCheckBoxType::CheckBox)
+					.IsChecked(this, &SMeshMergingDialog::GetAtlasLightmapUV)
+					.OnCheckStateChanged(this, &SMeshMergingDialog::SetAtlasLightmapUV)
+					.Content()
 					[
-						SNew(SCheckBox)
-						.Type(ESlateCheckBoxType::CheckBox)
-						.IsChecked(this, &SMeshMergingDialog::GetAtlasLightmapUV)
-						.OnCheckStateChanged(this, &SMeshMergingDialog::SetAtlasLightmapUV)
-						.Content()
-						[
-							SNew(STextBlock).Text(LOCTEXT("AtlasLightmapUVLabel", "Generate Atlased Lightmap UVs"))
-						]
+						SNew(STextBlock).Text(LOCTEXT("AtlasLightmapUVLabel", "Generate Atlased Lightmap UVs"))
 					]
+				]
 					
-					// Target lightmap channel
-					+SVerticalBox::Slot()
-					.AutoHeight()
-					.Padding(FEditorStyle::GetMargin("StandardDialog.ContentPadding"))
+				// Target lightmap channel / Max lightmap resolution
+				+SVerticalBox::Slot()
+				.AutoHeight()
+				.Padding(FEditorStyle::GetMargin("StandardDialog.ContentPadding"))
+				[
+					SNew(SHorizontalBox)
+
+					+SHorizontalBox::Slot()
+					.AutoWidth()
+					.VAlign(VAlign_Center)
 					[
-						SNew(SBorder)
-						.Padding(FMargin(8,0,0,0))
-						.BorderImage( FEditorStyle::GetBrush("ToolPanel.GroupBorder"))
-						[
-							SNew(SNumericEntryBox<int32>)
-							.IsEnabled( this, &SMeshMergingDialog::IsLightmapChannelEnabled )
-							.MinValue(0)
-							.MaxValue(MAX_MESH_TEXTURE_COORDS-1)
-							.Value(this, &SMeshMergingDialog::GetTargetLightmapChannelValue)
-							.OnValueCommitted(this, &SMeshMergingDialog::OnTargetLightmapChannelValueCommited)
-							.Label()
-							[
-								SNumericEntryBox<int32>::BuildLabel(LOCTEXT("TargetLightMapChannelLabel", "Target LightMap Channel: "), FLinearColor::Black, FLinearColor(.33f,.33f,.33f))
-							]
-						]
+						SNew(STextBlock)
+						.Text(LOCTEXT("TargetLightMapChannelLabel", "Target Channel:"))
 					]
-					
-					// Vertex colors
-					+SVerticalBox::Slot()
-					.AutoHeight()
-					.Padding(FEditorStyle::GetMargin("StandardDialog.ContentPadding"))
+
+					+SHorizontalBox::Slot()
+					.AutoWidth()
+					.VAlign(VAlign_Center)
+					.Padding(4,0,4,0)
 					[
-						SNew(SCheckBox)
-						.Type(ESlateCheckBoxType::CheckBox)
-						.IsChecked(this, &SMeshMergingDialog::GetImportVertexColors)
-						.OnCheckStateChanged(this, &SMeshMergingDialog::SetImportVertexColors)
-						.Content()
-						[
-							SNew(STextBlock).Text(LOCTEXT("ImportVertexColorsLabel", "Import Vertex Colors"))
-						]
+						SNew(STextComboBox)
+						.IsEnabled( this, &SMeshMergingDialog::IsLightmapChannelEnabled )
+						.OptionsSource(&LightMapChannelOptions)
+						.InitiallySelectedItem(LightMapChannelOptions[MergingSettings.TargetLightMapUVChannel])
+						.OnSelectionChanged(this, &SMeshMergingDialog::SetTargetLightMapChannel)
 					]
 
-					// Make one section per material
-					+SVerticalBox::Slot()
-					.AutoHeight()
-					.Padding(FEditorStyle::GetMargin("StandardDialog.ContentPadding"))
+					+SHorizontalBox::Slot()
+					.AutoWidth()
+					.VAlign(VAlign_Center)
 					[
-						SNew(SCheckBox)
-						.Type(ESlateCheckBoxType::CheckBox)
-						.IsChecked(this, &SMeshMergingDialog::GetOneSectionPerMaterial)
-						.OnCheckStateChanged(this, &SMeshMergingDialog::SetOneSectionPerMaterial)
-						.Content()
-						[
-							SNew(STextBlock).Text(LOCTEXT("MakeOneSectionPerMaterialLabel", "Make One Section Per Material"))
-						]
+						SNew(STextBlock)
+						.Text(LOCTEXT("MaxLightMapResolutionLabel", "Max Resolution:"))
 					]
-
-					// Pivot at zero
-					+SVerticalBox::Slot()
-					.AutoHeight()
-					.Padding(FEditorStyle::GetMargin("StandardDialog.ContentPadding"))
+												
+					+SHorizontalBox::Slot()
+					.AutoWidth()
+					.VAlign(VAlign_Center)
+					.Padding(4,0,4,0)
 					[
-						SNew(SCheckBox)
-						.Type(ESlateCheckBoxType::CheckBox)
-						.IsChecked(this, &SMeshMergingDialog::GetPivotPointAtZero)
-						.OnCheckStateChanged(this, &SMeshMergingDialog::SetPivotPointAtZero)
-						.Content()
-						[
-							SNew(STextBlock).Text(LOCTEXT("PivotPointAtZeroLabel", "Pivot Point At (0,0,0)"))
-						]
+						SNew(STextComboBox)
+						.IsEnabled( this, &SMeshMergingDialog::IsLightmapChannelEnabled )
+						.OptionsSource(&LightMapResolutionOptions)
+						.InitiallySelectedItem(LightMapResolutionOptions[FMath::FloorLog2(MergingSettings.MaxAltlasedLightMapResolution)])
+						.OnSelectionChanged(this, &SMeshMergingDialog::SetMaxLightMapResolution)
 					]
+				]
+			]
+		]
 
-					// Place in world
-					+SVerticalBox::Slot()
-					.AutoHeight()
-					.Padding(FEditorStyle::GetMargin("StandardDialog.ContentPadding"))
+		// Other merging settings
+		+SVerticalBox::Slot()
+		.AutoHeight()
+		.Padding(0,2,0,0)
+		[
+			SNew(SBorder)
+			.BorderImage(FEditorStyle::GetBrush("ToolPanel.GroupBorder"))
+			[
+				SNew(SVerticalBox)
+					
+				// Vertex colors
+				+SVerticalBox::Slot()
+				.AutoHeight()
+				.Padding(FEditorStyle::GetMargin("StandardDialog.ContentPadding"))
+				[
+					SNew(SCheckBox)
+					.Type(ESlateCheckBoxType::CheckBox)
+					.IsChecked(this, &SMeshMergingDialog::GetImportVertexColors)
+					.OnCheckStateChanged(this, &SMeshMergingDialog::SetImportVertexColors)
+					.Content()
 					[
-						SNew(SCheckBox)
-						.Type(ESlateCheckBoxType::CheckBox)
-						.IsChecked(this, &SMeshMergingDialog::GetPlaceInWorld)
-						.OnCheckStateChanged(this, &SMeshMergingDialog::SetPlaceInWorld)
-						.Content()
-						[
-							SNew(STextBlock).Text(LOCTEXT("PlaceInWorldLabel", "Place In World"))
-						]
+						SNew(STextBlock).Text(LOCTEXT("ImportVertexColorsLabel", "Import Vertex Colors"))
 					]
-						
-					// Asset name and picker
-					+SVerticalBox::Slot()
-					.AutoHeight()
-					.Padding(FEditorStyle::GetMargin("StandardDialog.ContentPadding"))
+				]
+
+				// Pivot at zero
+				+SVerticalBox::Slot()
+				.AutoHeight()
+				.Padding(FEditorStyle::GetMargin("StandardDialog.ContentPadding"))
+				[
+					SNew(SCheckBox)
+					.Type(ESlateCheckBoxType::CheckBox)
+					.IsChecked(this, &SMeshMergingDialog::GetPivotPointAtZero)
+					.OnCheckStateChanged(this, &SMeshMergingDialog::SetPivotPointAtZero)
+					.Content()
 					[
-						SNew(STextBlock).Text(LOCTEXT("AssetNameLabel", "Asset Name:"))
+						SNew(STextBlock).Text(LOCTEXT("PivotPointAtZeroLabel", "Pivot Point At (0,0,0)"))
 					]
+				]
 
-					+SVerticalBox::Slot()
-					.AutoHeight()
-					.Padding(FEditorStyle::GetMargin("StandardDialog.ContentPadding"))
+				// Place in world
+				+SVerticalBox::Slot()
+				.AutoHeight()
+				.Padding(FEditorStyle::GetMargin("StandardDialog.ContentPadding"))
+				[
+					SNew(SCheckBox)
+					.Type(ESlateCheckBoxType::CheckBox)
+					.IsChecked(this, &SMeshMergingDialog::GetPlaceInWorld)
+					.OnCheckStateChanged(this, &SMeshMergingDialog::SetPlaceInWorld)
+					.Content()
 					[
-						SNew(SHorizontalBox)
+						SNew(STextBlock).Text(LOCTEXT("PlaceInWorldLabel", "Place In World"))
+					]
+				]
+						
+				// Asset name and picker
+				+SVerticalBox::Slot()
+				.AutoHeight()
+				.Padding(FEditorStyle::GetMargin("StandardDialog.ContentPadding"))
+				[
+					SNew(STextBlock).Text(LOCTEXT("AssetNameLabel", "Asset Name:"))
+				]
+
+				+SVerticalBox::Slot()
+				.AutoHeight()
+				.Padding(FEditorStyle::GetMargin("StandardDialog.ContentPadding"))
+				[
+					SNew(SHorizontalBox)
 					
-						+SHorizontalBox::Slot()
-						.FillWidth(1.0)
-						.Padding(0,0,2,0)
-						.HAlign(HAlign_Fill)
-						.VAlign(VAlign_Center)
-						[
-							SNew(SEditableTextBox)
-							.Text(this, &SMeshMergingDialog::GetMergedMeshPackageName)
-							.OnTextCommitted(this, &SMeshMergingDialog::OnMergedMeshPackageNameTextCommited)
-						]
+					+SHorizontalBox::Slot()
+					.FillWidth(1.0)
+					.Padding(0,0,2,0)
+					.HAlign(HAlign_Fill)
+					.VAlign(VAlign_Center)
+					[
+						SNew(SEditableTextBox)
+						.Text(this, &SMeshMergingDialog::GetMergedMeshPackageName)
+						.OnTextCommitted(this, &SMeshMergingDialog::OnMergedMeshPackageNameTextCommited)
+					]
 					
-						+SHorizontalBox::Slot()
-						.AutoWidth()
-						.HAlign(HAlign_Right)
-						.VAlign(VAlign_Center)
-						[
-							SNew(SButton)
-							.OnClicked(this, &SMeshMergingDialog::OnSelectPackageNameClicked)
-							.Text(LOCTEXT("SelectPackageButton", "..."))
-						]
+					+SHorizontalBox::Slot()
+					.AutoWidth()
+					.HAlign(HAlign_Right)
+					.VAlign(VAlign_Center)
+					[
+						SNew(SButton)
+						.OnClicked(this, &SMeshMergingDialog::OnSelectPackageNameClicked)
+						.Text(LOCTEXT("SelectPackageButton", "..."))
 					]
 				]
 			]
 		]
 
+		// Merge, Cancel buttons
 		+SVerticalBox::Slot()
 		.AutoHeight()
 		.HAlign(HAlign_Right)
 		.VAlign(VAlign_Bottom)
 		.AutoHeight()
+		.Padding(0,4,0,0)
 		[
 			SNew(SUniformGridPanel)
 			.SlotPadding(FEditorStyle::GetMargin("StandardDialog.SlotPadding"))
@@ -236,27 +268,27 @@ FReply SMeshMergingDialog::OnMergeClicked()
 
 ESlateCheckBoxState::Type SMeshMergingDialog::GetAtlasLightmapUV() const
 {
-	return (MergingSettings.bGnerateAtlasedLightmapUV ? ESlateCheckBoxState::Checked : ESlateCheckBoxState::Unchecked);
+	return (MergingSettings.bGenerateAtlasedLightMapUV ? ESlateCheckBoxState::Checked : ESlateCheckBoxState::Unchecked);
 }
 
 void SMeshMergingDialog::SetAtlasLightmapUV(ESlateCheckBoxState::Type NewValue)
 {
-	MergingSettings.bGnerateAtlasedLightmapUV = (ESlateCheckBoxState::Checked == NewValue);
+	MergingSettings.bGenerateAtlasedLightMapUV = (ESlateCheckBoxState::Checked == NewValue);
 }
 
 bool SMeshMergingDialog::IsLightmapChannelEnabled() const
 {
-	return MergingSettings.bGnerateAtlasedLightmapUV;
+	return MergingSettings.bGenerateAtlasedLightMapUV;
 }
 
-TOptional<int32> SMeshMergingDialog::GetTargetLightmapChannelValue() const
+void SMeshMergingDialog::SetTargetLightMapChannel(TSharedPtr<FString> NewSelection, ESelectInfo::Type SelectInfo)
 {
-	return MergingSettings.TargetLightmapUVChannel;
+	TTypeFromString<int32>::FromString(MergingSettings.TargetLightMapUVChannel, **NewSelection);
 }
 
-void SMeshMergingDialog::OnTargetLightmapChannelValueCommited(int32 NewValue, ETextCommit::Type CommitInfo)
+void SMeshMergingDialog::SetMaxLightMapResolution(TSharedPtr<FString> NewSelection, ESelectInfo::Type SelectInfo)
 {
-	MergingSettings.TargetLightmapUVChannel = FMath::Clamp(NewValue, 0, MAX_MESH_TEXTURE_COORDS-1);
+	TTypeFromString<int32>::FromString(MergingSettings.MaxAltlasedLightMapResolution, **NewSelection);
 }
 
 ESlateCheckBoxState::Type SMeshMergingDialog::GetImportVertexColors() const
@@ -274,17 +306,6 @@ void SMeshMergingDialog::OnActorSelectionChanged(const TArray<UObject*>& NewSele
 	GenerateNewPackageName();
 }
 
-ESlateCheckBoxState::Type SMeshMergingDialog::GetOneSectionPerMaterial() const
-{
-	return ESlateCheckBoxState::Checked;
-}
-
-void SMeshMergingDialog::SetOneSectionPerMaterial(ESlateCheckBoxState::Type NewValue)
-{
-	// TODO: Need to add support for one mesh per material
-	// Default is one mesh with all materials (one section per material)
-}
-
 ESlateCheckBoxState::Type SMeshMergingDialog::GetPivotPointAtZero() const
 {
 	return (MergingSettings.bPivotPointAtZero ? ESlateCheckBoxState::Checked : ESlateCheckBoxState::Unchecked);
diff --git a/Engine/Source/Editor/UnrealEd/Public/Dialogs/SMeshMergingDialog.h b/Engine/Source/Editor/UnrealEd/Public/Dialogs/SMeshMergingDialog.h
--- a/Engine/Source/Editor/UnrealEd/Public/Dialogs/SMeshMergingDialog.h
+++ b/Engine/Source/Editor/UnrealEd/Public/Dialogs/SMeshMergingDialog.h
@@ -39,17 +39,13 @@ private:
 
 	/** Target lightmap channel */
 	bool IsLightmapChannelEnabled() const;
-	TOptional<int32> GetTargetLightmapChannelValue() const;
-	void OnTargetLightmapChannelValueCommited(int32 NewValue, ETextCommit::Type CommitInfo);
-	
+	void SetTargetLightMapChannel(TSharedPtr<FString> NewSelection, ESelectInfo::Type SelectInfo);
+	void SetMaxLightMapResolution(TSharedPtr<FString> NewSelection, ESelectInfo::Type SelectInfo);
+		
 	/**  */
 	ESlateCheckBoxState::Type GetImportVertexColors() const;
 	void SetImportVertexColors(ESlateCheckBoxState::Type NewValue);
 
-	/**  */
-	ESlateCheckBoxState::Type GetOneSectionPerMaterial() const;
-	void SetOneSectionPerMaterial(ESlateCheckBoxState::Type NewValue);
-
 	/**  */
 	ESlateCheckBoxState::Type GetPivotPointAtZero() const;
 	void SetPivotPointAtZero(ESlateCheckBoxState::Type NewValue);
@@ -85,4 +81,8 @@ private:
 
 	/** Whether to spawn merged actor in the world */
 	bool bPlaceInWorld;
+
+	/**  */
+	TArray<TSharedPtr<FString>>	LightMapResolutionOptions;
+	TArray<TSharedPtr<FString>>	LightMapChannelOptions;
 };
