Asset editors docked into the MainFrame area will now correctly restore to that area when next re-opened

TTP# 338363 - When I relaunch the editor it reopens the items in the wrong place.

IÆve added a "DockedToolkit" tab which resides next to the "LevelEditor" tab. This is used for two purposes; firstly, it can be used to identify asset editors that have been dragged into the MainFrame, and secondly, it's used as a marker (akin to the "StandaloneToolkit" tab) to restore the asset editor to the correct place.

When an asset editor is closed, an entry is added to the GEditorUserSettingsIni file so that FAssetEditorToolkit::InitAssetEditor can restore the asset editor to either the "DockedToolkit" or "StandaloneToolkit" area.

diff --git a/Engine/Source/Editor/MainFrame/Private/MainFrameModule.cpp b/Engine/Source/Editor/MainFrame/Private/MainFrameModule.cpp
--- a/Engine/Source/Editor/MainFrame/Private/MainFrameModule.cpp
+++ b/Engine/Source/Editor/MainFrame/Private/MainFrameModule.cpp
@@ -114,7 +114,7 @@ void FMainFrameModule::CreateDefaultMainFrame( const bool bStartImmersivePIE )
 			TSharedRef<FTabManager::FLayout> LoadedLayout = FLayoutSaveRestore::LoadFromConfig(GEditorLayoutIni,
 				// We persist the positioning of the level editor and the content browser.
 				// The asset editors currently do not get saved.
-				FTabManager::NewLayout( "UnrealEd_Layout_v1.1" )
+				FTabManager::NewLayout( "UnrealEd_Layout_v1.3" )
 				->AddArea
 				(
 					FTabManager::NewPrimaryArea()
@@ -123,6 +123,7 @@ void FMainFrameModule::CreateDefaultMainFrame( const bool bStartImmersivePIE )
 						FTabManager::NewStack()
 						->SetSizeCoefficient(2.0f)
 						->AddTab("LevelEditor", ETabState::OpenedTab)
+						->AddTab("DockedToolkit", ETabState::ClosedTab)
 					)
 				)
 				->AddArea
diff --git a/Engine/Source/Editor/UnrealEd/Private/Toolkits/AssetEditorToolkit.cpp b/Engine/Source/Editor/UnrealEd/Private/Toolkits/AssetEditorToolkit.cpp
--- a/Engine/Source/Editor/UnrealEd/Private/Toolkits/AssetEditorToolkit.cpp
+++ b/Engine/Source/Editor/UnrealEd/Private/Toolkits/AssetEditorToolkit.cpp
@@ -84,7 +84,19 @@ void FAssetEditorToolkit::InitAssetEditor( const EToolkitMode::Type Mode, const
 						.Icon( this, &FAssetEditorToolkit::GetDefaultTabIcon )
 						.Label( Label );
 		{
-			FGlobalTabmanager::Get()->InsertNewDocumentTab( "StandaloneToolkit", FTabManager::ESearchPreference::PreferLiveTab, NewMajorTab.ToSharedRef() );
+			static_assert(sizeof(EAssetEditorToolkitTabLocation) == sizeof(int32), "EAssetEditorToolkitTabLocation is the incorrect size");
+
+			// Work out where we should create this asset editor
+			EAssetEditorToolkitTabLocation SavedAssetEditorToolkitTabLocation = EAssetEditorToolkitTabLocation::Standalone;
+			GConfig->GetInt(
+				TEXT("AssetEditorToolkitTabLocation"), 
+				*ObjectsToEdit[0]->GetPathName(), 
+				reinterpret_cast<int32&>(SavedAssetEditorToolkitTabLocation), 
+				GEditorUserSettingsIni
+				);
+
+			const FName AssetEditorToolkitTab = (SavedAssetEditorToolkitTabLocation == EAssetEditorToolkitTabLocation::Docked) ? "DockedToolkit" : "StandaloneToolkit";
+			FGlobalTabmanager::Get()->InsertNewDocumentTab( AssetEditorToolkitTab, FTabManager::ESearchPreference::PreferLiveTab, NewMajorTab.ToSharedRef() );
 		}
 
 		IUserFeedbackModule& UserFeedback = FModuleManager::LoadModuleChecked<IUserFeedbackModule>(TEXT("UserFeedback"));
diff --git a/Engine/Source/Editor/UnrealEd/Private/Toolkits/SStandaloneAssetEditorToolkitHost.cpp b/Engine/Source/Editor/UnrealEd/Private/Toolkits/SStandaloneAssetEditorToolkitHost.cpp
--- a/Engine/Source/Editor/UnrealEd/Private/Toolkits/SStandaloneAssetEditorToolkitHost.cpp
+++ b/Engine/Source/Editor/UnrealEd/Private/Toolkits/SStandaloneAssetEditorToolkitHost.cpp
@@ -114,6 +114,7 @@ void SStandaloneAssetEditorToolkitHost::RestoreFromLayout( const TSharedRef<FTab
 {
 	const TSharedRef<SDockTab> HostTab = HostTabPtr.Pin().ToSharedRef();
 	HostTab->SetCanCloseTab(EditorCloseRequest);
+	HostTab->SetOnTabClosed(SDockTab::FOnTabClosedCallback::CreateSP(this, &SStandaloneAssetEditorToolkitHost::OnTabClosed));
 
 	this->ChildSlot[SNullWidget::NullWidget];
 	MyTabManager->CloseAllAreas();
@@ -254,4 +255,35 @@ FReply SStandaloneAssetEditorToolkitHost::OnKeyDown( const FGeometry& MyGeometry
 }
 
 
+void SStandaloneAssetEditorToolkitHost::OnTabClosed(TSharedRef<SDockTab> TabClosed) const
+{
+	check(TabClosed == HostTabPtr.Pin());
+
+	const TSharedPtr<FAssetEditorToolkit> AssetEditorToolkit = StaticCastSharedPtr<FAssetEditorToolkit>(HostedToolkit);
+	if(AssetEditorToolkit.IsValid())
+	{
+		const TArray<UObject*>* const ObjectsBeingEdited = AssetEditorToolkit->GetObjectsCurrentlyBeingEdited();
+		if(ObjectsBeingEdited)
+		{
+			const bool IsDockedAssetEditor = TabClosed->HasSiblingTab(FName("DockedToolkit"), false/*TreatIndexNoneAsWildcard*/);
+			const EAssetEditorToolkitTabLocation AssetEditorToolkitTabLocation = (IsDockedAssetEditor) ? EAssetEditorToolkitTabLocation::Docked : EAssetEditorToolkitTabLocation::Standalone;
+			for(const UObject* ObjectBeingEdited : *ObjectsBeingEdited)
+			{
+				// Only record assets that have a valid saved package
+				UPackage* const Package = ObjectBeingEdited->GetOutermost();
+				if(Package && Package->GetFileSize())
+				{
+					GConfig->SetInt(
+						TEXT("AssetEditorToolkitTabLocation"), 
+						*ObjectBeingEdited->GetPathName(), 
+						static_cast<int32>(AssetEditorToolkitTabLocation), 
+						GEditorUserSettingsIni
+						);
+				}
+			}
+		}
+	}
+}
+
+
 #undef LOCTEXT_NAMESPACE
\ No newline at end of file
diff --git a/Engine/Source/Editor/UnrealEd/Private/Toolkits/SStandaloneAssetEditorToolkitHost.h b/Engine/Source/Editor/UnrealEd/Private/Toolkits/SStandaloneAssetEditorToolkitHost.h
--- a/Engine/Source/Editor/UnrealEd/Private/Toolkits/SStandaloneAssetEditorToolkitHost.h
+++ b/Engine/Source/Editor/UnrealEd/Private/Toolkits/SStandaloneAssetEditorToolkitHost.h
@@ -76,6 +76,8 @@ public:
 	void SetMenuOverlay( TSharedRef<SWidget> NewOverlay );
 
 private:
+	void OnTabClosed(TSharedRef<SDockTab> TabClosed) const;
+
 	/** Manages internal tab layout */
 	TSharedPtr<FTabManager> MyTabManager;
 
diff --git a/Engine/Source/Editor/UnrealEd/Public/Toolkits/AssetEditorToolkit.h b/Engine/Source/Editor/UnrealEd/Public/Toolkits/AssetEditorToolkit.h
--- a/Engine/Source/Editor/UnrealEd/Public/Toolkits/AssetEditorToolkit.h
+++ b/Engine/Source/Editor/UnrealEd/Public/Toolkits/AssetEditorToolkit.h
@@ -9,6 +9,18 @@
 
 DECLARE_DELEGATE_RetVal( bool, FRequestAssetEditorClose );
 
+/**
+ * The location of the asset editor toolkit tab
+ * Note: These values are serialized into an ini file as an int32
+ */
+enum class EAssetEditorToolkitTabLocation : int32
+{
+	/** The tab is within the "DockedToolkit" tab stack */
+	Docked,
+	/** The tab is within the "StandaloneToolkit" tab stack */
+	Standalone,
+};
+
 /**
  * Base class for toolkits that are used for asset editing (abstract)
  */
diff --git a/Engine/Source/Runtime/Slate/Private/Framework/Docking/SDockingTabStack.h b/Engine/Source/Runtime/Slate/Private/Framework/Docking/SDockingTabStack.h
--- a/Engine/Source/Runtime/Slate/Private/Framework/Docking/SDockingTabStack.h
+++ b/Engine/Source/Runtime/Slate/Private/Framework/Docking/SDockingTabStack.h
@@ -215,9 +215,10 @@ private:
 
 struct FTabMatcher
 {
-	FTabMatcher( const FTabId& InTabId, ETabState::Type InTabState = static_cast<ETabState::Type>(ETabState::ClosedTab | ETabState::OpenedTab) )
+	FTabMatcher( const FTabId& InTabId, ETabState::Type InTabState = static_cast<ETabState::Type>(ETabState::ClosedTab | ETabState::OpenedTab), const bool InTreatIndexNoneAsWildcard = true )
 		: TabIdToMatch( InTabId )
 		, RequiredTabState( InTabState )
+		, TreatIndexNoneAsWildcard( InTreatIndexNoneAsWildcard )
 	{
 	}
 
@@ -226,10 +227,11 @@ struct FTabMatcher
 		return
 			( (Candidate.TabState & RequiredTabState) != 0 ) &&
 			( Candidate.TabId.TabType == TabIdToMatch.TabType ) &&
-			// NAME_None is treated as a wildcard
-			( TabIdToMatch.InstanceId == INDEX_NONE || TabIdToMatch.InstanceId == Candidate.TabId.InstanceId );;
+			// INDEX_NONE is treated as a wildcard
+			( (TreatIndexNoneAsWildcard && TabIdToMatch.InstanceId == INDEX_NONE) || TabIdToMatch.InstanceId == Candidate.TabId.InstanceId );;
 	}
 
 	FTabId TabIdToMatch;
 	ETabState::Type RequiredTabState;
+	bool TreatIndexNoneAsWildcard;
 };
diff --git a/Engine/Source/Runtime/Slate/Private/Widgets/Docking/SDockTab.cpp b/Engine/Source/Runtime/Slate/Private/Widgets/Docking/SDockTab.cpp
--- a/Engine/Source/Runtime/Slate/Private/Widgets/Docking/SDockTab.cpp
+++ b/Engine/Source/Runtime/Slate/Private/Widgets/Docking/SDockTab.cpp
@@ -400,6 +400,12 @@ void SDockTab::SetDraggedOverDockArea( const TSharedPtr<SDockingArea>& Area )
 	DraggedOverDockingArea = Area;
 }
 
+bool SDockTab::HasSiblingTab(const FTabId& SiblingTabId, const bool TreatIndexNoneAsWildcard) const
+{
+	TSharedPtr<SDockingTabStack> ParentTabStack = GetParentDockTabStack();
+	return (ParentTabStack.IsValid()) ? ParentTabStack->HasTab(FTabMatcher(SiblingTabId, static_cast<ETabState::Type>(ETabState::ClosedTab | ETabState::OpenedTab), TreatIndexNoneAsWildcard)) : false;
+}
+
 void SDockTab::Construct( const FArguments& InArgs )
 {
 	SpawnAnimCurve = FCurveSequence(0, 0.15f);
diff --git a/Engine/Source/Runtime/Slate/Public/Framework/Docking/TabManager.h b/Engine/Source/Runtime/Slate/Public/Framework/Docking/TabManager.h
--- a/Engine/Source/Runtime/Slate/Public/Framework/Docking/TabManager.h
+++ b/Engine/Source/Runtime/Slate/Public/Framework/Docking/TabManager.h
@@ -782,13 +782,6 @@ private:
 		TWeakPtr<FTabManager> TabManager;
 	};
 
-	struct FFoobar {
-			bool operator()(const FSubTabManager& InItem) const
-			{
-				return !InItem.MajorTab.IsValid();
-			}
-		};
-
 	struct FindByTab
 	{
 		FindByTab(const TSharedRef<SDockTab>& InTabToFind)
diff --git a/Engine/Source/Runtime/Slate/Public/Widgets/Docking/SDockTab.h b/Engine/Source/Runtime/Slate/Public/Widgets/Docking/SDockTab.h
--- a/Engine/Source/Runtime/Slate/Public/Widgets/Docking/SDockTab.h
+++ b/Engine/Source/Runtime/Slate/Public/Widgets/Docking/SDockTab.h
@@ -243,6 +243,16 @@ public:
 	/** Used by the drag/drop operation to signal to this tab what it is dragging over. */
 	void SetDraggedOverDockArea( const TSharedPtr<SDockingArea>& Area );
 
+	/** 
+	 * Check to see whether this tab has a sibling tab with the given tab ID
+	 * 
+	 * @param	SiblingTabId				The ID of the tab we want to find
+	 * @param	TreatIndexNoneAsWildcard	Note that this variable only takes effect if SiblingTabId has an InstanceId of INDEX_NONE.
+	 *										If true, we will consider this a "wildcard" search (matching any tab with the correct TabType, regardless 
+	 *										of its InstanceId). If false, we will explicitly look for a tab with an InstanceId of INDEX_NONE
+	 */
+	bool HasSiblingTab(const FTabId& SiblingTabId, const bool TreatIndexNoneAsWildcard = true) const;
+
 protected:
 
 	/** Gets the dock tab stack this dockable tab resides within, if any */
