[GitHub] Don't load Unreal Engine Launcher while debugging in Visual Studio or Xcode
- Added "-ForceLauncher" command-line argument that can override this

https://github.com/EpicGames/UnrealEngine/pull/216

diff --git a/Engine/Source/Editor/UnrealEd/Private/Editor.cpp b/Engine/Source/Editor/UnrealEd/Private/Editor.cpp
--- a/Engine/Source/Editor/UnrealEd/Private/Editor.cpp
+++ b/Engine/Source/Editor/UnrealEd/Private/Editor.cpp
@@ -359,13 +359,18 @@ void UEditorEngine::InitEditor(IEngineLoop* InEngineLoop)
 	// Call base.
 	UEngine::Init(InEngineLoop);
 
-	if( !FEngineBuildSettings::IsInternalBuild() &&
-		!FEngineBuildSettings::IsPerforceBuild() && 
-		!FApp::IsBenchmarking() &&
-		!GIsDemoMode && 
-		!IsRunningCommandlet() &&
-		!FPlatformProcess::IsApplicationRunning(TEXT("UnrealEngineLauncher") ) &&
-		!FPlatformProcess::IsApplicationRunning(TEXT("Unreal Engine Launcher") ) )
+	// Specify "-ForceLauncher" on the command-line to always open the launcher, even in unusual cases.  This is useful for debugging the Launcher startup.
+	const bool bForceLauncherToOpen = FParse::Param( FCommandLine::Get(),TEXT( "ForceLauncher" ) );
+
+	if( bForceLauncherToOpen ||
+		( !FEngineBuildSettings::IsInternalBuild() &&
+		  !FEngineBuildSettings::IsPerforceBuild() && 
+		  !FPlatformMisc::IsDebuggerPresent() &&	// Don't spawn launcher while running in the Visual Studio debugger by default
+		  !FApp::IsBenchmarking() &&
+		  !GIsDemoMode && 
+		  !IsRunningCommandlet() &&
+		  !FPlatformProcess::IsApplicationRunning(TEXT("UnrealEngineLauncher") ) &&
+		  !FPlatformProcess::IsApplicationRunning(TEXT("Unreal Engine Launcher") ) ) )
 	{
 		IDesktopPlatform* DesktopPlatform = FDesktopPlatformModule::Get();
 		if( DesktopPlatform != NULL )
