Keeping the dump blueprint info commandlet from using BP classes that it creates as base classes for later blueprints (+ various minor fixes to it).
(TTP #339501)

diff --git a/Engine/Source/Editor/UnrealEd/Private/Commandlets/DumpBlueprintsInfoCommandlet.cpp b/Engine/Source/Editor/UnrealEd/Private/Commandlets/DumpBlueprintsInfoCommandlet.cpp
--- a/Engine/Source/Editor/UnrealEd/Private/Commandlets/DumpBlueprintsInfoCommandlet.cpp
+++ b/Engine/Source/Editor/UnrealEd/Private/Commandlets/DumpBlueprintsInfoCommandlet.cpp
@@ -14,6 +14,7 @@
 #include "Engine/LevelScriptBlueprint.h"
 #include "BehaviorTree/BehaviorTree.h"
 #include "AssetSelection.h"
+#include "K2Node_CustomEvent.h"
 
 DEFINE_LOG_CATEGORY_STATIC(LogBlueprintInfoDump, Log, All);
 
@@ -145,30 +146,63 @@ DumpBlueprintsInfo commandlet params: \n\
 
 	/** Tracks spawned level actors (so we don't have to create more the we have to). */
 	static TMap<UClass*, AActor*> LevelActors;
+	
+	/** Tracks instantiated blueprints (so we don't have to create more the we have to). */
+	static TMap<UClass*, UBlueprint*> ClassBlueprints;
 
 	/**
-	 * Utility class for tracking the duration of a scoped action (tracks time 
-	 * in seconds and adds it to the specified variable on destruction).
+	 * Utility stopwatch class for tracking the duration of some action (tracks 
+	 * time in seconds and adds it to the specified variable on destruction).
 	 */
-	class FScopedDurationTimer
+	class FDurationTimer
 	{
 	public:
-		FScopedDurationTimer(double& AccumulatorIn)
+		FDurationTimer(double& AccumulatorIn)
 			: StartTime(FPlatformTime::Seconds())
 			, Accumulator(AccumulatorIn)
 		{}
-
-		/** Dtor, updating seconds with time delta. */
-		~FScopedDurationTimer()
+		
+		double Start()
 		{
-			Accumulator += FPlatformTime::Seconds() - StartTime;
+			StartTime = FPlatformTime::Seconds();
+			return StartTime;
 		}
+		
+		double Stop()
+		{
+			double StopTime = FPlatformTime::Seconds();
+			Accumulator += (StopTime - StartTime);
+			StartTime = StopTime;
+			
+			return StopTime;
+		}
+
 	private:
 		/** Start time, captured in ctor. */
 		double StartTime;
 		/** Time variable to update. */
 		double& Accumulator;
 	};
+	
+	/**
+	 * Utility class for tracking the duration of a scoped action (the user 
+	 * doesn't have to call Start() and Stop() manually).
+	 */
+	class FScopedDurationTimer : private FDurationTimer
+	{
+	public:
+		FScopedDurationTimer(double& AccumulatorIn)
+			: FDurationTimer(AccumulatorIn)
+		{
+			Start();
+		}
+		
+		/** Dtor, updating seconds with time delta. */
+		~FScopedDurationTimer()
+		{
+			Stop();
+		}
+	};
 
 	/**
 	 * Certain blueprints (like level blueprints) require a level outer, and 
@@ -495,6 +529,8 @@ DumpBlueprintInfoUtils::CommandletOptions::CommandletOptions(TArray<FString> con
 		{
 			FString DiffSwitch;
 			Switch.Split(TEXT("="), &DiffSwitch, &DiffCommand);
+			
+			DiffCommand = DiffCommand.Replace(TEXT("\""), TEXT(""));
 		}
 		else if (!Switch.Compare("palette", ESearchCase::IgnoreCase))
 		{
@@ -627,102 +663,127 @@ static AActor* DumpBlueprintInfoUtils::SpawnLevelActor(UClass* ActorClass, bool
 //------------------------------------------------------------------------------
 static UBlueprint* DumpBlueprintInfoUtils::MakeTempBlueprint(UClass* ParentClass)
 {
-	UObject* BlueprintOuter = GetTransientPackage();
-
-	bool bIsAnimBlueprint = ParentClass->IsChildOf(UAnimInstance::StaticClass());
-	bool bIsLevelBlueprint = ParentClass->IsChildOf(ALevelScriptActor::StaticClass());
-
-	UClass* BlueprintClass = UBlueprint::StaticClass();
-	UClass* GeneratedClass = UBlueprintGeneratedClass::StaticClass();
-	EBlueprintType BlueprintType = BPTYPE_Normal;
-
-	if (bIsAnimBlueprint)
+	UBlueprint* MadeBlueprint = nullptr;
+	if (UBlueprint** FoundBlueprint = ClassBlueprints.Find(ParentClass))
 	{
-		BlueprintClass = UAnimBlueprint::StaticClass();
-		GeneratedClass = UAnimBlueprintGeneratedClass::StaticClass();
+		MadeBlueprint = *FoundBlueprint;
 	}
-	else if (bIsLevelBlueprint)
+	else
 	{
-		UWorld* World = GetWorld();
-		if (World == nullptr)
+		UObject* BlueprintOuter = GetTransientPackage();
+
+		bool bIsAnimBlueprint = ParentClass->IsChildOf(UAnimInstance::StaticClass());
+		bool bIsLevelBlueprint = ParentClass->IsChildOf(ALevelScriptActor::StaticClass());
+
+		UClass* BlueprintClass = UBlueprint::StaticClass();
+		UClass* GeneratedClass = UBlueprintGeneratedClass::StaticClass();
+		EBlueprintType BlueprintType = BPTYPE_Normal;
+
+		if (bIsAnimBlueprint)
 		{
-			UE_LOG(LogBlueprintInfoDump, Error, TEXT("Cannot make a proper level blueprint without a valid editor level for its outer."));
+			BlueprintClass = UAnimBlueprint::StaticClass();
+			GeneratedClass = UAnimBlueprintGeneratedClass::StaticClass();
 		}
-		else
+		else if (bIsLevelBlueprint)
 		{
-			BlueprintClass = ULevelScriptBlueprint::StaticClass();
-			BlueprintType = BPTYPE_LevelScript;
-			BlueprintOuter = World->GetCurrentLevel();
+			UWorld* World = GetWorld();
+			if (World == nullptr)
+			{
+				UE_LOG(LogBlueprintInfoDump, Error, TEXT("Cannot make a proper level blueprint without a valid editor level for its outer."));
+			}
+			else
+			{
+				BlueprintClass = ULevelScriptBlueprint::StaticClass();
+				BlueprintType = BPTYPE_LevelScript;
+				BlueprintOuter = World->GetCurrentLevel();
+			}
 		}
-	}
-	// @TODO: UEditorUtilityBlueprint
+		// @TODO: UEditorUtilityBlueprint
 
-	FString const ClassName = ParentClass->GetName();
-	FString const DesiredName = FString::Printf(TEXT("COMMANDLET_TEMP_Blueprint_%s"), *ClassName);
-	FName   const TempBpName = MakeUniqueObjectName(BlueprintOuter, BlueprintClass, FName(*DesiredName));
+		FString const ClassName = ParentClass->GetName();
+		FString const DesiredName = FString::Printf(TEXT("COMMANDLET_TEMP_Blueprint_%s"), *ClassName);
+		FName   const TempBpName = MakeUniqueObjectName(BlueprintOuter, BlueprintClass, FName(*DesiredName));
 
-	check(FKismetEditorUtilities::CanCreateBlueprintOfClass(ParentClass));
-	UBlueprint* TempBlueprint = FKismetEditorUtilities::CreateBlueprint(ParentClass, BlueprintOuter, TempBpName, BlueprintType, BlueprintClass, GeneratedClass);
+		check(FKismetEditorUtilities::CanCreateBlueprintOfClass(ParentClass));
+		MadeBlueprint = FKismetEditorUtilities::CreateBlueprint(ParentClass, BlueprintOuter, TempBpName, BlueprintType, BlueprintClass, GeneratedClass);
+		
+		// if this is an animation blueprint, then we want anim specific graphs to test as well (if it has an anim graph)...
+		if (bIsAnimBlueprint && (MadeBlueprint->FunctionGraphs.Num() > 0))
+		{
+			UAnimationGraph* AnimGraph = CastChecked<UAnimationGraph>(MadeBlueprint->FunctionGraphs[0]);
+			check(AnimGraph != nullptr);
 
-	// if this is an animation blueprint, then we want anim specific graphs to test as well...
-	if (bIsAnimBlueprint)
-	{
-		check(TempBlueprint->FunctionGraphs.Num() > 0);
-		UAnimationGraph* AnimGraph = CastChecked<UAnimationGraph>(TempBlueprint->FunctionGraphs[0]);
+			// should add a state-machine graph
+			UAnimGraphNode_StateMachine* StateMachineNode = AddNodeToGraph<UAnimGraphNode_StateMachine>(AnimGraph);
 
-		// should add a state-machine graph
-		UAnimGraphNode_StateMachine* StateMachineNode = AddNodeToGraph<UAnimGraphNode_StateMachine>(AnimGraph);
+			UAnimationStateMachineGraph* StateMachineGraph = StateMachineNode->EditorStateMachineGraph;
+			// should add an conduit graph 
+			UAnimStateConduitNode* ConduitNode = AddNodeToGraph<UAnimStateConduitNode>((UEdGraph*)StateMachineGraph);
 
-		UAnimationStateMachineGraph* StateMachineGraph = StateMachineNode->EditorStateMachineGraph;
-		// should add an conduit graph 
-		UAnimStateConduitNode* ConduitNode = AddNodeToGraph<UAnimStateConduitNode>((UEdGraph*)StateMachineGraph);
+			UAnimStateNode* StateNode = AddNodeToGraph<UAnimStateNode>((UEdGraph*)StateMachineGraph);
+			// should create a transition graph
+			StateNode->AutowireNewNode(ConduitNode->GetOutputPin());
+		}
+		else if (bIsLevelBlueprint)
+		{
+			ULevel* Level = CastChecked<ULevel>(BlueprintOuter);
+			Level->LevelScriptBlueprint = Cast<ULevelScriptBlueprint>(MadeBlueprint);
+		}
 
-		UAnimStateNode* StateNode = AddNodeToGraph<UAnimStateNode>((UEdGraph*)StateMachineGraph);
-		// should create a transition graph
-		StateNode->AutowireNewNode(ConduitNode->GetOutputPin());
-	}
+		// may have been altered in CreateBlueprint()
+		BlueprintType = MadeBlueprint->BlueprintType;
 
-	// may have been altered in CreateBlueprint()
-	BlueprintType = TempBlueprint->BlueprintType;
+		bool bCanAddFunctions = (BlueprintType != BPTYPE_MacroLibrary); // taken from FBlueprintEditor::NewDocument_IsVisibleForType()
+		if (bCanAddFunctions)
+		{
+			// add a functions graph that isn't the construction script (or an animation graph)
+			FName FuncGraphName = MakeUniqueObjectName(MadeBlueprint, UEdGraph::StaticClass(), FName(TEXT("NewFunction")));
+			UEdGraph* FuncGraph = FBlueprintEditorUtils::CreateNewGraph(MadeBlueprint, FuncGraphName, UEdGraph::StaticClass(), UEdGraphSchema_K2::StaticClass());
+			FBlueprintEditorUtils::AddFunctionGraph<UClass>(MadeBlueprint, FuncGraph, /*bIsUserCreated =*/true, nullptr);
+		}
 
-	bool bCanAddFunctions = (BlueprintType != BPTYPE_MacroLibrary); // taken from FBlueprintEditor::NewDocument_IsVisibleForType()
-	if (bCanAddFunctions)
-	{
-		// add a functions graph that isn't the construction script (or an animation graph)
-		FName FuncGraphName = MakeUniqueObjectName(TempBlueprint, UEdGraph::StaticClass(), FName(TEXT("NewFunction")));
-		UEdGraph* FuncGraph = FBlueprintEditorUtils::CreateNewGraph(TempBlueprint, FuncGraphName, UEdGraph::StaticClass(), UEdGraphSchema_K2::StaticClass());
-		FBlueprintEditorUtils::AddFunctionGraph<UClass>(TempBlueprint, FuncGraph, /*bIsUserCreated =*/true, nullptr);
-	}
+		bool bCanAddMacros = ((BlueprintType == BPTYPE_MacroLibrary) || (BlueprintType == BPTYPE_Normal) || (BlueprintType == BPTYPE_LevelScript));
+		if (bCanAddMacros)
+		{
+			FName MacroGraphName = MakeUniqueObjectName(MadeBlueprint, UEdGraph::StaticClass(), FName(TEXT("NewMacro")));
+			UEdGraph* MacroGraph = FBlueprintEditorUtils::CreateNewGraph(MadeBlueprint, MacroGraphName, UEdGraph::StaticClass(), UEdGraphSchema_K2::StaticClass());
+			FBlueprintEditorUtils::AddMacroGraph(MadeBlueprint, MacroGraph, /*bIsUserCreated =*/true, nullptr);
+		}
 
-	bool bCanAddMacros = ((BlueprintType == BPTYPE_MacroLibrary) || (BlueprintType == BPTYPE_Normal) || (BlueprintType == BPTYPE_LevelScript));
-	if (bCanAddMacros)
-	{
-		FName MacroGraphName = MakeUniqueObjectName(TempBlueprint, UEdGraph::StaticClass(), FName(TEXT("NewMacro")));
-		UEdGraph* MacroGraph = FBlueprintEditorUtils::CreateNewGraph(TempBlueprint, MacroGraphName, UEdGraph::StaticClass(), UEdGraphSchema_K2::StaticClass());
-		FBlueprintEditorUtils::AddMacroGraph(TempBlueprint, MacroGraph, /*bIsUserCreated =*/true, nullptr);
-	}
+		// if you can add custom events to this blueprint, do so (so show that we 
+		// can call an event on ourselves)
+		if (MadeBlueprint->UbergraphPages.Num() > 0)
+		{
+			UK2Node_CustomEvent* CustomEventNode = AddNodeToGraph<UK2Node_CustomEvent>(MadeBlueprint->UbergraphPages[0]);
+			CustomEventNode->CustomFunctionName = FBlueprintEditorUtils::FindUniqueCustomEventName(MadeBlueprint);
+		}
 
-	UClass* ObjTypeToAdd = CommandOptions.SelectedObjectType;
-	if (ObjTypeToAdd && ObjTypeToAdd->IsChildOf(UActorComponent::StaticClass()))
-	{
-		if (!AddComponentToBlueprint(TempBlueprint, ObjTypeToAdd))
+		UClass* ObjTypeToAdd = CommandOptions.SelectedObjectType;
+		if (ObjTypeToAdd && ObjTypeToAdd->IsChildOf(UActorComponent::StaticClass()))
 		{
-			UE_LOG(LogBlueprintInfoDump, Error, TEXT("Cannot add a '%s' to a '%s' blueprint."), *ObjTypeToAdd->GetName(), *ClassName);
+			if (!AddComponentToBlueprint(MadeBlueprint, ObjTypeToAdd))
+			{
+				UE_LOG(LogBlueprintInfoDump, Error, TEXT("Cannot add a '%s' to a '%s' blueprint."), *ObjTypeToAdd->GetName(), *ClassName);
+			}
 		}
-	}
-	else if (CommandOptions.DumpFlags & BPDUMP_SelectAllObjTypes)
-	{
-		for (TObjectIterator<UClass> ClassIt; ClassIt; ++ClassIt)
+		else if (CommandOptions.DumpFlags & BPDUMP_SelectAllObjTypes)
 		{
-			UClass* Class = *ClassIt;
-			if (Class->IsChildOf(UActorComponent::StaticClass()))
+			for (TObjectIterator<UClass> ClassIt; ClassIt; ++ClassIt)
 			{
-				AddComponentToBlueprint(TempBlueprint, Class);
+				UClass* Class = *ClassIt;
+				if (Class->IsChildOf(UActorComponent::StaticClass()))
+				{
+					AddComponentToBlueprint(MadeBlueprint, Class);
+				}
 			}
 		}
+
+		FBlueprintEditorUtils::MarkBlueprintAsStructurallyModified(MadeBlueprint);		
+		ClassBlueprints.Add(ParentClass, MadeBlueprint);
 	}
 
-	return TempBlueprint;
+	check(MadeBlueprint != nullptr);
+	return MadeBlueprint;
 }
 
 //------------------------------------------------------------------------------
@@ -770,6 +831,7 @@ static NodeType* DumpBlueprintInfoUtils::AddNodeToGraph(UEdGraph* Graph)
 	NodeType* NewNode = NewObject<NodeType>(Graph);
 	Graph->AddNode(NewNode, /*bFromUI =*/true, /*bSelectNewNode =*/false);
 
+	NewNode->CreateNewGuid();
 	NewNode->PostPlacedNewNode();
 	NewNode->AllocateDefaultPins();
 	return NewNode;
@@ -778,21 +840,31 @@ static NodeType* DumpBlueprintInfoUtils::AddNodeToGraph(UEdGraph* Graph)
 //------------------------------------------------------------------------------
 static FString DumpBlueprintInfoUtils::BuildDumpFilePath(UClass* BlueprintClass)
 {
-	FString CommandletSaveDir = FPaths::GameSavedDir() + TEXT("Commandlets/");
-	CommandletSaveDir = FPaths::ConvertRelativePathToFull(CommandletSaveDir);
-	IFileManager::Get().MakeDirectory(*CommandletSaveDir);
-
+	
+	
 	FString Pathname = FString::Printf(TEXT("BlueprintsInfoDump_%s"), FPlatformTime::StrTimestamp());
 	Pathname = Pathname.Replace(TEXT(" "), TEXT("_"));
 	Pathname = Pathname.Replace(TEXT("/"), TEXT("-"));
 	Pathname = Pathname.Replace(TEXT(":"), TEXT("."));
-
+	
 	bool const bSplitBlueprintsByFile = (CommandOptions.DumpFlags & BPDUMP_FilePerBlueprint) != 0;
-	if (bSplitBlueprintsByFile && (BlueprintClass != nullptr))
+	
+	static FString CommandletSaveDir;
+	if (CommandletSaveDir.IsEmpty())
 	{
-		CommandletSaveDir += Pathname + "/";
+		CommandletSaveDir = FPaths::GameSavedDir() + TEXT("Commandlets/");
+		CommandletSaveDir = FPaths::ConvertRelativePathToFull(CommandletSaveDir);
 		IFileManager::Get().MakeDirectory(*CommandletSaveDir);
-
+		
+		if (bSplitBlueprintsByFile)
+		{
+			CommandletSaveDir += Pathname + "/";
+			IFileManager::Get().MakeDirectory(*CommandletSaveDir);
+		}
+	}
+	
+	if (bSplitBlueprintsByFile && (BlueprintClass != nullptr))
+	{
 		Pathname = ("BlueprintInfo_" + BlueprintClass->GetName() + ".json");
 	}
 	else
@@ -1106,6 +1178,7 @@ static void DumpBlueprintInfoUtils::DumpActionMenuItem(uint32 Indent, FGraphActi
 		FString const IndentedNewline = "\n" + BuildIndentString(Indent);
 
 		ActionEntry += " : {";
+		ActionEntry += IndentedNewline + "\"ActionType\"  : \"" + PrimeAction->GetTypeId() + "\",";
 		ActionEntry += IndentedNewline + "\"Name\"        : \"" + ActionName + "\",";
 		ActionEntry += IndentedNewline + "\"Category\"    : \"";
 		if (bHasCategory)
@@ -1121,7 +1194,7 @@ static void DumpBlueprintInfoUtils::DumpActionMenuItem(uint32 Indent, FGraphActi
 
 		ActionEntry += IndentedNewline + TooltipFieldLabel + TooltipStr + "\",";
 		ActionEntry += IndentedNewline + "\"Keywords\"    : \"" + PrimeAction->Keywords + "\",";
-		ActionEntry += IndentedNewline + "\"SearchTitle\" : \"" + PrimeAction->SearchTitle.ToString() + "\",";
+		ActionEntry += IndentedNewline + "\"SearchTitle\" : \"" + PrimeAction->GetSearchTitle().ToString() + "\",";
 		ActionEntry += IndentedNewline + FString::Printf(TEXT("\"Grouping\"    : %d"), PrimeAction->Grouping);
 		
 		// Get action node type info
@@ -1196,6 +1269,7 @@ static void DumpBlueprintInfoUtils::DumpGraphContextActions(uint32 Indent, UEdGr
 	FileOutWriter->Serialize(TCHAR_TO_ANSI(*BeginGraphEntry), BeginGraphEntry.Len());
 
 	FBlueprintGraphActionListBuilder ActionBuilder(Graph);
+	UE_LOG(LogBlueprintInfoDump, Display, TEXT("%sDumping graph context actions..."), *BuildIndentString(Indent, true));
 	DumpContextActionList(Indent, ActionBuilder, FileOutWriter);
 
 	if (CommandOptions.DumpFlags & BPDUMP_SelectAllObjTypes)
@@ -1203,11 +1277,19 @@ static void DumpBlueprintInfoUtils::DumpGraphContextActions(uint32 Indent, UEdGr
 		for (TObjectIterator<UClass> ClassIt; ClassIt; ++ClassIt)
 		{
 			UClass* Class = *ClassIt;
-			if (Class->IsChildOf(AActor::StaticClass()))
+			
+			bool const bIsTempBpClass = (Class->GetOuterUPackage() == GetTransientPackage());
+			// have to exclude classes that were created explicitly in this commandlet
+			if (bIsTempBpClass)
 			{
+				continue;
+			}
+			
+			if (Class->IsChildOf<AActor>() && !Class->IsChildOf<ALevelScriptActor>())
+			{				
 				if (AActor* LevelActor = SpawnLevelActor(Class, true))
 				{
-					UE_LOG(LogBlueprintInfoDump, Display, TEXT("%sDumping level actor actions: '%s'..."), *BuildIndentString(Indent, true), *Class->GetName());
+					UE_LOG(LogBlueprintInfoDump, Display, TEXT("%sDumping actions with actor selection: '%s'..."), *BuildIndentString(Indent, true), *Class->GetName());
 
 					FString ActorSelectionEntry = "," + IndentedNewline + "\"LevelActorMenu-" + Class->GetName() + "\" : \n";
 					FileOutWriter->Serialize(TCHAR_TO_ANSI(*ActorSelectionEntry), ActorSelectionEntry.Len());
@@ -1218,13 +1300,21 @@ static void DumpBlueprintInfoUtils::DumpGraphContextActions(uint32 Indent, UEdGr
 		}
 	}
 
-	UBlueprint* Blueprint = CastChecked<UBlueprint>(Graph->GetOuter());
+	UObject* GraphOuter = Graph->GetOuter();
+	UBlueprint* Blueprint = Cast<UBlueprint>(GraphOuter);
+	while ((Blueprint == nullptr) && (GraphOuter != nullptr))
+	{
+		GraphOuter = GraphOuter->GetOuter();
+		Blueprint = Cast<UBlueprint>(GraphOuter);
+	}
+	check(Blueprint != nullptr);
+
 	TArray<UObjectProperty*> ComponentProperties;
 	GetComponentProperties(Blueprint, ComponentProperties);
 
 	for (UObjectProperty* Component : ComponentProperties)
 	{
-		UE_LOG(LogBlueprintInfoDump, Display, TEXT("%sDumping component actions: '%s'..."), *BuildIndentString(Indent, true), *Component->GetName());
+		UE_LOG(LogBlueprintInfoDump, Display, TEXT("%sDumping actions with component selection: '%s'..."), *BuildIndentString(Indent, true), *Component->GetName());
 
 		FString SelectionContextEntry = "," + IndentedNewline + "\"ComponentContextMenu-" + Component->GetName() + "\" : \n";
 		FileOutWriter->Serialize(TCHAR_TO_ANSI(*SelectionContextEntry), SelectionContextEntry.Len());
@@ -1393,9 +1483,9 @@ static void DumpBlueprintInfoUtils::DumpContextActionList(uint32 Indent, FBluepr
 {
 	ActionBuilder.Empty();
 
-	FString PinTypeLog("<UNKNOWN>");
 	if (ActionBuilder.FromPin != nullptr)
 	{
+		FString PinTypeLog("<NONE>");
 		if (ActionBuilder.FromPin->Direction == EGPD_Input)
 		{
 			PinTypeLog = "INPUT ";
@@ -1405,9 +1495,10 @@ static void DumpBlueprintInfoUtils::DumpContextActionList(uint32 Indent, FBluepr
 			PinTypeLog = "OUTPUT";
 		}
 		PinTypeLog += UEdGraphSchema_K2::TypeToString(ActionBuilder.FromPin->PinType);
-	}
-	UE_LOG(LogBlueprintInfoDump, Display, TEXT("%sDumping pin actions: %s"), *BuildIndentString(Indent, true), *PinTypeLog);
 
+		UE_LOG(LogBlueprintInfoDump, Display, TEXT("%sDumping pin actions: %s"), *BuildIndentString(Indent, true), *PinTypeLog);
+	}
+	
 	double MenuBuildDuration = GetContextMenuActions(ActionBuilder);
 
 	FString const ContextEntryIndent = BuildIndentString(Indent);
@@ -1436,10 +1527,9 @@ static double DumpBlueprintInfoUtils::GetContextMenuActions(FBlueprintGraphActio
 	UEdGraphSchema const* GraphSchema = GetDefault<UEdGraphSchema>(ActionBuilder.CurrentGraph->Schema);
 
 	double MenuBuildDuration = 0.0;
-	FBlueprintPaletteListBuilder PaletteBuilder(ActionBuilder.Blueprint);
 	{
-		FScopedDurationTimer DurationTimer(MenuBuildDuration);
-		GraphSchema->GetGraphContextActions(ActionBuilder);
+ 		FScopedDurationTimer DurationTimer(MenuBuildDuration);
+ 		GraphSchema->GetGraphContextActions(ActionBuilder);
 	}
 
 	return MenuBuildDuration;
@@ -1576,14 +1666,16 @@ static void DumpBlueprintInfoUtils::DiffDumpFiles(FString const& NewFilePath, FS
 			// -dw : Ignore line ending and all whitespace differences
 			DiffCommand = TEXT("p4merge.exe -dw \"{2}\" \"{1}\"");
 		}
-		DiffCommand = DiffCommand.Replace(TEXT("{1}"), *NewFilePath).Replace(TEXT("{2}"), *QualifiedOldFilePath);
-
+		
 		FString DiffArgs;
 		if (int32 ArgsIndex = DiffCommand.Find(TEXT(" ")))
 		{
 			DiffArgs = *DiffCommand + ArgsIndex + 1;
 			DiffCommand = DiffCommand.Left(ArgsIndex);
 		}
+		
+		DiffArgs = DiffArgs.Replace(TEXT("{1}"), *NewFilePath);
+		DiffArgs = DiffArgs.Replace(TEXT("{2}"), *QualifiedOldFilePath);
 
 		FProcHandle DiffProc = FPlatformProcess::CreateProc(*DiffCommand,
 			/*Params =*/              *DiffArgs,
@@ -1642,7 +1734,7 @@ int32 UDumpBlueprintsInfoCommandlet::Main(FString const& Params)
 	// closing out the writer (and diffing the resultant file if the user deigns us to do so)
 	auto CloseFileStream = [bDiffGeneratedFile, &ActiveFilePath, &CommandOptions](FArchive** FileOutPtr)
 	{
-		FArchive* FileOut = (*FileOutPtr);
+		FArchive*& FileOut = (*FileOutPtr);
 		if (FileOut != nullptr)
 		{
 			FileOut->Serialize(TCHAR_TO_ANSI(TEXT("\n}")), 2);
@@ -1702,7 +1794,11 @@ int32 UDumpBlueprintsInfoCommandlet::Main(FString const& Params)
 	// class is a blueprintable type... broken into it's own lambda to save on reuse
 	auto IsInvalidBlueprintClass = [](UClass* Class)->bool
 	{
-		return Class->HasAnyClassFlags(CLASS_NewerVersionExists) || FKismetEditorUtilities::IsClassABlueprintSkeleton(Class) || !FKismetEditorUtilities::CanCreateBlueprintOfClass(Class);
+		return !IsValid(Class) ||
+			Class->HasAnyClassFlags(CLASS_NewerVersionExists) ||
+			FKismetEditorUtilities::IsClassABlueprintSkeleton(Class)  ||
+			!FKismetEditorUtilities::CanCreateBlueprintOfClass(Class) ||
+			(Class->GetOuterUPackage() == GetTransientPackage());
 	};
 
 	UClass* const BlueprintClass = CommandOptions.BlueprintClass;
