Fixing circular dependency issue with blueprint interfaces.

diff --git a/Engine/Source/Editor/UnrealEd/Private/Kismet2/BlueprintEditorUtils.cpp b/Engine/Source/Editor/UnrealEd/Private/Kismet2/BlueprintEditorUtils.cpp
--- a/Engine/Source/Editor/UnrealEd/Private/Kismet2/BlueprintEditorUtils.cpp
+++ b/Engine/Source/Editor/UnrealEd/Private/Kismet2/BlueprintEditorUtils.cpp
@@ -553,7 +553,7 @@ private:
 
 struct FRegenerationHelper
 {
-	static bool FrocedLoad(UObject* Obj)
+	static bool ForcedLoad(UObject* Obj)
 	{
 		auto Linker = Obj->GetLinker();
 		if (Linker && !Obj->HasAnyFlags(RF_LoadCompleted))
@@ -565,7 +565,7 @@ struct FRegenerationHelper
 		return false;
 	}
 
-	static void FrocedLoadMembers(UObject* InObject)
+	static void ForcedLoadMembers(UObject* InObject)
 	{
 		// Collect a list of all things this element owns
 		TArray<UObject*> MemberReferences;
@@ -576,9 +576,9 @@ struct FRegenerationHelper
 		for (TArray<UObject*>::TIterator it(MemberReferences); it; ++it)
 		{
 			UObject* CurrentObject = *it;
-			if (FrocedLoad(CurrentObject))
+			if (ForcedLoad(CurrentObject))
 			{
-				FrocedLoadMembers(CurrentObject);
+				ForcedLoadMembers(CurrentObject);
 			}
 		}
 	}
@@ -600,13 +600,13 @@ struct FRegenerationHelper
 			check(Package);
 			UMetaData* MetaData = Package->GetMetaData();
 			check(MetaData);
-			FrocedLoad(MetaData);
+			ForcedLoad(MetaData);
 		}
 
 		const int32 OldPropertiesSize = Struct->GetPropertiesSize();
 		for (UField* Field = Struct->Children; Field; Field = Field->Next)
 		{
-			bChanged |= FrocedLoad(Field);
+			bChanged |= ForcedLoad(Field);
 		}
 
 		if (bChanged)
@@ -662,11 +662,38 @@ struct FRegenerationHelper
 						Linker->Preload(BP);
 					}
 				}
-				FrocedLoadMembers(BP);
+				ForcedLoadMembers(BP);
 			}
 		}
 	}
 
+	/**
+	 * A helper function that loads (and regenerates) interface dependencies.
+	 * Accounts for circular dependencies by following how we handle parent 
+	 * classes in ULinkerLoad::RegenerateBlueprintClass() (that is, to complete 
+	 * the interface's compilation/regeneration before we utilize it for the
+	 * specified blueprint).
+	 * 
+	 * @param  Blueprint	The blueprint whose implemented interfaces you want loaded.
+	 */
+	static void PreloadInterfaces(UBlueprint* Blueprint)
+	{
+#if WITH_EDITORONLY_DATA // ImplementedInterfaces is wrapped WITH_EDITORONLY_DATA 
+		for (FBPInterfaceDescription const& InterfaceDesc : Blueprint->ImplementedInterfaces)
+		{
+			UClass* InterfaceClass = InterfaceDesc.Interface;
+			if (UBlueprint* InterfaceBlueprint = Cast<UBlueprint>(InterfaceClass->ClassGeneratedBy))
+			{
+				ForcedLoadMembers(InterfaceBlueprint);
+				if (InterfaceBlueprint->HasAnyFlags(RF_BeingRegenerated))
+				{
+					InterfaceBlueprint->RegenerateClass(InterfaceClass, nullptr, GObjLoaded);
+				}
+			}
+		}
+#endif // #if WITH_EDITORONLY_DATA
+	}
+
 	static void LinkExternalDependencies(UBlueprint* Blueprint)
 	{
 		check(Blueprint);
@@ -730,6 +757,8 @@ struct FRegenerationHelper
 			}
 		}
 		PreloadMacroSources(MacroSources);
+
+		PreloadInterfaces(Blueprint);
 	}
 };
 
@@ -949,6 +978,29 @@ UClass* FBlueprintEditorUtils::RegenerateBlueprintClass(UBlueprint* Blueprint, U
 
 	bool bNeedsSkelRefRemoval = false;
 
+	// Preload the blueprint and all its parts before refreshing nodes. 
+	// Otherwise, the nodes might not maintain their proper linkages... 
+	//
+	// This all should also happen here, first thing, before 
+	// bIsRegeneratingOnLoad is set, so that we can re-enter this function for 
+	// the same class further down the callstack (presumably from 
+	// PreloadInterfaces() or some other dependency load). This is here to 
+	// handle circular dependencies, where pre-loading a member here sets off a  
+	// subsequent load that in turn, relies on this class and requires this  
+	// class to be fully generated... A second call to this function with the 
+	// same class will continue to preload all it's members (from where it left
+	// off, since they're gated by a RF_NeedLoad check) and then fall through to
+	// finish compiling the class (while it's still technically pre-loading a
+	// member further up the stack).
+	if (!Blueprint->bHasBeenRegenerated)
+	{
+		if (PreviousCDO)
+		{
+			FBlueprintEditorUtils::PreloadMembers(PreviousCDO);
+		}
+		FBlueprintEditorUtils::PreloadMembers(Blueprint);
+	}
+	
 	if( ShouldRegenerateBlueprint(Blueprint) && !Blueprint->bHasBeenRegenerated )
 	{
 		Blueprint->bIsRegeneratingOnLoad = true;
@@ -987,30 +1039,9 @@ UClass* FBlueprintEditorUtils::RegenerateBlueprintClass(UBlueprint* Blueprint, U
 			}
 		}
 
-
 		// Make sure the simple construction script is loaded, since the outer hierarchy isn't compatible with PreloadMembers past the root node
 		FBlueprintEditorUtils::PreloadConstructionScript(Blueprint);
 
-#if WITH_EDITORONLY_DATA
-		// Make sure all interface dependencies are loaded at this point
-		for( TArray<FBPInterfaceDescription>::TIterator InterfaceIt(Blueprint->ImplementedInterfaces); InterfaceIt; ++InterfaceIt )
-		{
-			UClass* InterfaceClass = (*InterfaceIt).Interface;
-			if( InterfaceClass && InterfaceClass->HasAnyFlags(RF_NeedLoad) )
-			{
-				InterfaceClass->GetLinker()->Preload(InterfaceClass);
-			}
-		}
-#endif // WITH_EDITORONLY_DATA
-
-		// Preload the blueprint and all its parts before refreshing nodes. Otherwise, the nodes might not maintain their proper linkages
-		if( PreviousCDO )
-		{
-			FBlueprintEditorUtils::PreloadMembers(PreviousCDO);
-		}
-
-		FBlueprintEditorUtils::PreloadMembers(Blueprint);
-
 		// Purge any NULL graphs
 		FBlueprintEditorUtils::PurgeNullGraphs(Blueprint);
 
@@ -1047,7 +1078,7 @@ UClass* FBlueprintEditorUtils::RegenerateBlueprintClass(UBlueprint* Blueprint, U
 
 			// Make sure events are up to date
 			FBlueprintEditorUtils::ConformImplementedEvents(Blueprint);
-
+			
 			// Make sure interfaces are up to date
 			FBlueprintEditorUtils::ConformImplementedInterfaces(Blueprint);
 
diff --git a/Engine/Source/Runtime/CoreUObject/Public/UObject/UObjectGlobals.h b/Engine/Source/Runtime/CoreUObject/Public/UObject/UObjectGlobals.h
--- a/Engine/Source/Runtime/CoreUObject/Public/UObject/UObjectGlobals.h
+++ b/Engine/Source/Runtime/CoreUObject/Public/UObject/UObjectGlobals.h
@@ -49,7 +49,7 @@ extern int32						GImportCount;
 /** Forced exports for EndLoad optimization.							*/
 extern int32						GForcedExportCount;
 /** Objects that might need preloading.									*/
-extern TArray<UObject*>			GObjLoaded;
+extern COREUOBJECT_API TArray<UObject*>			GObjLoaded;
 
 #if !IS_MONOLITHIC
 	/** Adds and entry for the UFunction native pointer remap table */
