Unexposing GObjLoaded (using indirectly through function params instead).

diff --git a/Engine/Source/Editor/UnrealEd/Private/Kismet2/BlueprintEditorUtils.cpp b/Engine/Source/Editor/UnrealEd/Private/Kismet2/BlueprintEditorUtils.cpp
--- a/Engine/Source/Editor/UnrealEd/Private/Kismet2/BlueprintEditorUtils.cpp
+++ b/Engine/Source/Editor/UnrealEd/Private/Kismet2/BlueprintEditorUtils.cpp
@@ -676,7 +676,7 @@ struct FRegenerationHelper
 	 * 
 	 * @param  Blueprint	The blueprint whose implemented interfaces you want loaded.
 	 */
-	static void PreloadInterfaces(UBlueprint* Blueprint)
+	static void PreloadInterfaces(UBlueprint* Blueprint, TArray<UObject*>& ObjLoaded)
 	{
 #if WITH_EDITORONLY_DATA // ImplementedInterfaces is wrapped WITH_EDITORONLY_DATA 
 		for (FBPInterfaceDescription const& InterfaceDesc : Blueprint->ImplementedInterfaces)
@@ -687,14 +687,14 @@ struct FRegenerationHelper
 				ForcedLoadMembers(InterfaceBlueprint);
 				if (InterfaceBlueprint->HasAnyFlags(RF_BeingRegenerated))
 				{
-					InterfaceBlueprint->RegenerateClass(InterfaceClass, nullptr, GObjLoaded);
+					InterfaceBlueprint->RegenerateClass(InterfaceClass, nullptr, ObjLoaded);
 				}
 			}
 		}
 #endif // #if WITH_EDITORONLY_DATA
 	}
 
-	static void LinkExternalDependencies(UBlueprint* Blueprint)
+	static void LinkExternalDependencies(UBlueprint* Blueprint, TArray<UObject*>& ObjLoaded)
 	{
 		check(Blueprint);
 		const UEdGraphSchema_K2* Schema = GetDefault<UEdGraphSchema_K2>();
@@ -758,7 +758,7 @@ struct FRegenerationHelper
 		}
 		PreloadMacroSources(MacroSources);
 
-		PreloadInterfaces(Blueprint);
+		PreloadInterfaces(Blueprint, ObjLoaded);
 	}
 };
 
@@ -1050,7 +1050,7 @@ UClass* FBlueprintEditorUtils::RegenerateBlueprintClass(UBlueprint* Blueprint, U
 		const bool bHasCode = !FBlueprintEditorUtils::IsDataOnlyBlueprint(Blueprint) && !bIsMacro;
 
 		// Make sure all used external classes/functions/structures/macros/etc are loaded and linked
-		FRegenerationHelper::LinkExternalDependencies(Blueprint);
+		FRegenerationHelper::LinkExternalDependencies(Blueprint, ObjLoaded);
 
 		FKismetEditorUtilities::GenerateBlueprintSkeleton(Blueprint);
 
diff --git a/Engine/Source/Runtime/CoreUObject/Public/UObject/UObjectGlobals.h b/Engine/Source/Runtime/CoreUObject/Public/UObject/UObjectGlobals.h
--- a/Engine/Source/Runtime/CoreUObject/Public/UObject/UObjectGlobals.h
+++ b/Engine/Source/Runtime/CoreUObject/Public/UObject/UObjectGlobals.h
@@ -49,7 +49,7 @@ extern int32						GImportCount;
 /** Forced exports for EndLoad optimization.							*/
 extern int32						GForcedExportCount;
 /** Objects that might need preloading.									*/
-extern COREUOBJECT_API TArray<UObject*>			GObjLoaded;
+extern TArray<UObject*>			GObjLoaded;
 
 #if !IS_MONOLITHIC
 	/** Adds and entry for the UFunction native pointer remap table */
