More throttling of minimized windows and moved performance monitoring settings to the same location as the rest of the editor performance settings

diff --git a/Engine/Source/Editor/UnrealEd/Classes/Editor/EditorUserSettings.h b/Engine/Source/Editor/UnrealEd/Classes/Editor/EditorUserSettings.h
--- a/Engine/Source/Editor/UnrealEd/Classes/Editor/EditorUserSettings.h
+++ b/Engine/Source/Editor/UnrealEd/Classes/Editor/EditorUserSettings.h
@@ -35,6 +35,10 @@ class UEditorUserSettings : public UObject
 	UPROPERTY(EditAnywhere, config, Category=Performance, meta=(DisplayName="Use Less CPU when in Background") )
 	uint32 bThrottleWhenNotForeground:1;
 
+	/** When turned on, the editor will constantly monitor performance and adjust scalability settings for you when performance drops (disabled in debug) */
+	UPROPERTY(EditAnywhere, config, Category=Performance)
+	uint32 bMonitorEditorPerformance:1;
+
 	/** Select to make Distributions use the curves, not the baked lookup tables. */
 	UPROPERTY(config)
 	uint32 bUseCurvesForDistributions:1; //(GDistributionType == 0)
diff --git a/Engine/Source/Editor/UnrealEd/Classes/Settings/LevelEditorMiscSettings.h b/Engine/Source/Editor/UnrealEd/Classes/Settings/LevelEditorMiscSettings.h
--- a/Engine/Source/Editor/UnrealEd/Classes/Settings/LevelEditorMiscSettings.h
+++ b/Engine/Source/Editor/UnrealEd/Classes/Settings/LevelEditorMiscSettings.h
@@ -53,12 +53,6 @@ public:
 	UPROPERTY(EditAnywhere, config, Category=Sound)
 	uint32 bEnableEditorSounds:1;
 
-public:	
-
-	/** When turned on, the editor will constantly monitor performance and adjust scalability settings for you when performance drops (disabled in debug) */
-	UPROPERTY(EditAnywhere, config, Category=Performance)
-	bool bMonitorEditorPerformance;
-
 protected:
 
 	// Begin UObject overrides
diff --git a/Engine/Source/Editor/UnrealEd/Private/Editor.cpp b/Engine/Source/Editor/UnrealEd/Private/Editor.cpp
--- a/Engine/Source/Editor/UnrealEd/Private/Editor.cpp
+++ b/Engine/Source/Editor/UnrealEd/Private/Editor.cpp
@@ -1459,34 +1459,40 @@ void UEditorEngine::Tick( float DeltaSeconds, bool bIdleMode )
 
 float UEditorEngine::GetMaxTickRate( float DeltaTime, bool bAllowFrameRateSmoothing )
 {
-	const float SuperMaxTickRate = Super::GetMaxTickRate( DeltaTime, bAllowFrameRateSmoothing );
-	if( SuperMaxTickRate != 0.0f )
-	{
-		return SuperMaxTickRate;
-	}
-
 	float MaxTickRate = 0.0f;
-
-	// Clamp editor frame rate, even if smoothing is disabled
-	if( !bSmoothFrameRate && GIsEditor && !GIsPlayInEditorWorld )
+	if( !ShouldThrottleCPUUsage() )
 	{
-		MaxTickRate = 1.0f / DeltaTime;
-		if (SmoothedFrameRateRange.HasLowerBound())
+		const float SuperMaxTickRate = Super::GetMaxTickRate( DeltaTime, bAllowFrameRateSmoothing );
+		if( SuperMaxTickRate != 0.0f )
+		{
+			return SuperMaxTickRate;
+		}
+
+		// Clamp editor frame rate, even if smoothing is disabled
+		if( !bSmoothFrameRate && GIsEditor && !GIsPlayInEditorWorld )
 		{
-			MaxTickRate = FMath::Max(MaxTickRate, SmoothedFrameRateRange.GetLowerBoundValue());
+			MaxTickRate = 1.0f / DeltaTime;
+			if (SmoothedFrameRateRange.HasLowerBound())
+			{
+				MaxTickRate = FMath::Max(MaxTickRate, SmoothedFrameRateRange.GetLowerBoundValue());
+			}
+			if (SmoothedFrameRateRange.HasUpperBound())
+			{
+				MaxTickRate = FMath::Min(MaxTickRate, SmoothedFrameRateRange.GetUpperBoundValue());
+			}
 		}
-		if (SmoothedFrameRateRange.HasUpperBound())
+
+		// Laptops should throttle to 60 hz in editor to reduce battery drain
+		static const auto CVarDontLimitOnBattery = IConsoleManager::Get().FindTConsoleVariableDataInt(TEXT("r.DontLimitOnBattery"));
+		const bool bLimitOnBattery = (FPlatformMisc::IsRunningOnBattery() && CVarDontLimitOnBattery->GetValueOnGameThread() == 0);
+		if( bLimitOnBattery )
 		{
-			MaxTickRate = FMath::Min(MaxTickRate, SmoothedFrameRateRange.GetUpperBoundValue());
+			MaxTickRate = 60.0f;
 		}
 	}
-
-	// TTP 337189 - Laptops should throttle to 60 hz in editor to reduce battery drain
-	static const auto CVarDontLimitOnBattery = IConsoleManager::Get().FindTConsoleVariableDataInt(TEXT("r.DontLimitOnBattery"));
-	const bool bLimitOnBattery = (FPlatformMisc::IsRunningOnBattery() && CVarDontLimitOnBattery->GetValueOnGameThread() == 0);
-	if( bLimitOnBattery )
+	else
 	{
-		MaxTickRate = 60.0f;
+		MaxTickRate = 3.0f;
 	}
 
 	return MaxTickRate;
diff --git a/Engine/Source/Editor/UnrealEd/Private/Editor/EditorUserSettings.cpp b/Engine/Source/Editor/UnrealEd/Private/Editor/EditorUserSettings.cpp
--- a/Engine/Source/Editor/UnrealEd/Private/Editor/EditorUserSettings.cpp
+++ b/Engine/Source/Editor/UnrealEd/Private/Editor/EditorUserSettings.cpp
@@ -11,7 +11,7 @@ UEditorUserSettings::UEditorUserSettings(const class FPostConstructInitializePro
 {
 	//Default to high quality
 	MaterialQualityLevel = 1;
-
+	bMonitorEditorPerformance = true;
 	BlueprintFavorites = ConstructObject<UBlueprintPaletteFavorites>(UBlueprintPaletteFavorites::StaticClass(), this);
 }
 
diff --git a/Engine/Source/Editor/UnrealEd/Private/PerformanceMonitor.cpp b/Engine/Source/Editor/UnrealEd/Private/PerformanceMonitor.cpp
--- a/Engine/Source/Editor/UnrealEd/Private/PerformanceMonitor.cpp
+++ b/Engine/Source/Editor/UnrealEd/Private/PerformanceMonitor.cpp
@@ -92,7 +92,7 @@ public:
 	/** Called to get the "Show notification" check box state */
 	static ESlateCheckBoxState::Type IsMonitoringPerformance()
 	{
-		const bool bMonitorEditorPerformance = GetDefault<ULevelEditorMiscSettings>()->bMonitorEditorPerformance;
+		const bool bMonitorEditorPerformance = GetDefault<UEditorUserSettings>()->bMonitorEditorPerformance;
 		return bMonitorEditorPerformance ? ESlateCheckBoxState::Checked : ESlateCheckBoxState::Unchecked;
 	}
 
@@ -101,7 +101,7 @@ public:
 	{
 		const bool bNewEnabledState = (NewState == ESlateCheckBoxState::Checked);
 
-		auto* Settings = GetMutableDefault<ULevelEditorMiscSettings>();
+		auto* Settings = GetMutableDefault<UEditorUserSettings>();
 		Settings->bMonitorEditorPerformance = bNewEnabledState;
 		Settings->PostEditChange();
 	}
@@ -258,7 +258,7 @@ void FPerformanceMonitor::Tick(float DeltaTime)
 	FineMovingAverage.Tick(FPlatformTime::Seconds(), GAverageFPS);
 	CoarseMovingAverage.Tick(FPlatformTime::Seconds(), GAverageFPS);
 
-	const bool bMonitorEditorPerformance = GetDefault<ULevelEditorMiscSettings>()->bMonitorEditorPerformance;
+	const bool bMonitorEditorPerformance = GetDefault<UEditorUserSettings>()->bMonitorEditorPerformance;
 	if( !bMonitorEditorPerformance || !bIsNotificationAllowed )
 	{
 		return;
diff --git a/Engine/Source/Editor/UnrealEd/Private/Settings/SettingsClasses.cpp b/Engine/Source/Editor/UnrealEd/Private/Settings/SettingsClasses.cpp
--- a/Engine/Source/Editor/UnrealEd/Private/Settings/SettingsClasses.cpp
+++ b/Engine/Source/Editor/UnrealEd/Private/Settings/SettingsClasses.cpp
@@ -134,7 +134,6 @@ ULevelEditorMiscSettings::ULevelEditorMiscSettings( const class FPostConstructIn
 	: Super(PCIP)
 {
 	bAutoApplyLightingEnable = true;
-	bMonitorEditorPerformance = true;
 }
 
 
diff --git a/Engine/Source/Runtime/Launch/Private/LaunchEngineLoop.cpp b/Engine/Source/Runtime/Launch/Private/LaunchEngineLoop.cpp
--- a/Engine/Source/Runtime/Launch/Private/LaunchEngineLoop.cpp
+++ b/Engine/Source/Runtime/Launch/Private/LaunchEngineLoop.cpp
@@ -2064,10 +2064,7 @@ void FEngineLoop::Tick()
 
 		// Idle mode prevents ticking and rendering completely
 		const bool bIdleMode = ShouldUseIdleMode();
-
-		// Should we throttle CPU usage based on current state (usually in editor if we are minimized or not in foreground)
-		const bool bShouldThrottleCPU = GEngine->ShouldThrottleCPUUsage(); 
-		if (bIdleMode || bShouldThrottleCPU )
+		if (bIdleMode)
 		{
 			// Yield CPU time
 			FPlatformProcess::Sleep(.1f);
diff --git a/Engine/Source/Runtime/Slate/Private/Framework/Application/SlateApplication.cpp b/Engine/Source/Runtime/Slate/Private/Framework/Application/SlateApplication.cpp
--- a/Engine/Source/Runtime/Slate/Private/Framework/Application/SlateApplication.cpp
+++ b/Engine/Source/Runtime/Slate/Private/Framework/Application/SlateApplication.cpp
@@ -483,7 +483,7 @@ FWidgetPath FSlateApplication::LocateWindowUnderMouse( FVector2D ScreenspaceMous
  */
 void FSlateApplication::TickWindowAndChildren( TSharedRef<SWindow> WindowToTick )
 {
-	if ( WindowToTick->IsVisible() )
+	if ( WindowToTick->IsVisible() && !WindowToTick->IsWindowMinimized() )
 	{
 		// Switch to the appropriate world for ticking
 		FScopedSwitchWorldHack SwitchWorld( WindowToTick );
@@ -633,7 +633,7 @@ void FSlateApplication::DrawPrepass( TSharedPtr<SWindow> DrawOnlyThisWindow )
 		for (TArray< TSharedRef<SWindow> >::TConstIterator CurrentWindowIt(SlateWindows); CurrentWindowIt; ++CurrentWindowIt)
 		{
 			TSharedRef<SWindow> CurrentWindow = *CurrentWindowIt;
-			if (CurrentWindow->IsVisible())
+			if (CurrentWindow->IsVisible() && !CurrentWindow->IsWindowMinimized())
 			{
 				DrawWindowPrepass(CurrentWindow );
 			}
