Fix landscape editor mode being disabled after undoing the creation of a landscape
Also fixes the "actor picker" tool only working once.

diff --git a/Engine/Source/Editor/UnrealEd/Private/EditorModes.cpp b/Engine/Source/Editor/UnrealEd/Private/EditorModes.cpp
--- a/Engine/Source/Editor/UnrealEd/Private/EditorModes.cpp
+++ b/Engine/Source/Editor/UnrealEd/Private/EditorModes.cpp
@@ -560,6 +560,8 @@ void FEdMode::Enter()
 		SelectedActor->MarkComponentsRenderStateDirty();
 	}
 
+	bPendingDeletion = false;
+
 	FEditorDelegates::EditorModeEnter.Broadcast( this );
 	const bool bIsEnteringMode = true;
 	Owner->BroadcastEditorModeChanged( this, bIsEnteringMode );
