Engine scalability settings are now calculated for the first run of an editor

  * The benchmark results are cached in Editor Game Agnostic Settings for later use
    - This saves us having to run a benchmark every time we want to see the automatical scalability settings. Also moved scalability settings to game agnostic config as they are not game specific.
  * Refined performance warning notification to automatically apply the settings when FPS drops too low. The user has the option of manually tweaking if they prefer.

diff --git a/Engine/Source/Editor/UnrealEd/Classes/Editor/EditorGameAgnosticSettings.h b/Engine/Source/Editor/UnrealEd/Classes/Editor/EditorGameAgnosticSettings.h
--- a/Engine/Source/Editor/UnrealEd/Classes/Editor/EditorGameAgnosticSettings.h
+++ b/Engine/Source/Editor/UnrealEd/Classes/Editor/EditorGameAgnosticSettings.h
@@ -1,6 +1,7 @@
 // Copyright 1998-2014 Epic Games, Inc. All Rights Reserved.
 
 #pragma once
+#include "Scalability.h"
 #include "EditorGameAgnosticSettings.generated.h"
 
 UCLASS(config=EditorGameAgnostic)
@@ -28,10 +29,6 @@ class UEditorGameAgnosticSettings : public UObject
 	UPROPERTY(config)
 	bool bCopyStarterContentPreference;
 
-	/** When set this will prompt a benchmark to be ran on startup to determine the best scalability settings */
-	UPROPERTY(config)
-	bool bApplyAutoScalabilityOnStartup;
-
 	/** The id's of the surveys completed */
 	UPROPERTY(config)
 	TArray<FGuid> CompletedSurveys;
@@ -40,6 +37,12 @@ class UEditorGameAgnosticSettings : public UObject
 	UPROPERTY(config)
 	TArray<FGuid> InProgressSurveys;
 
+	/** Engine scalability benchmark results */
+	Scalability::FQualityLevels EngineBenchmarkResult;
+
+	/** Load the engine scalability benchmark results. Performs a benchmark if not yet valid. */
+	void LoadScalabilityBenchmark();
+
 	// Begin UObject Interface
 	virtual void PostEditChangeProperty( struct FPropertyChangedEvent& PropertyChangedEvent) override;
 	// End UObject Interface
diff --git a/Engine/Source/Editor/UnrealEd/Private/Editor/EditorGameAgnosticSettings.cpp b/Engine/Source/Editor/UnrealEd/Private/Editor/EditorGameAgnosticSettings.cpp
--- a/Engine/Source/Editor/UnrealEd/Private/Editor/EditorGameAgnosticSettings.cpp
+++ b/Engine/Source/Editor/UnrealEd/Private/Editor/EditorGameAgnosticSettings.cpp
@@ -6,7 +6,6 @@ UEditorGameAgnosticSettings::UEditorGameAgnosticSettings(const class FPostConstr
 	: Super(PCIP)
 {
 	bCopyStarterContentPreference = true;
-	bApplyAutoScalabilityOnStartup = false;
 }
 
 void UEditorGameAgnosticSettings::PostEditChangeProperty( struct FPropertyChangedEvent& PropertyChangedEvent)
@@ -30,4 +29,47 @@ void UEditorGameAgnosticSettings::PostEditChangeProperty( struct FPropertyChange
 	}
 
 	GEditor->SaveGameAgnosticSettings();
+}
+
+void UEditorGameAgnosticSettings::LoadScalabilityBenchmark()
+{
+	const TCHAR* Section = TEXT("EngineBenchmarkResult");
+
+	bool bIsValid = false;
+	bIsValid = GConfig->GetBool(Section, TEXT("Valid"), bIsValid, GEditorGameAgnosticIni) && bIsValid;
+
+	Scalability::FQualityLevels Temporary;
+
+	if (bIsValid)
+	{
+		GConfig->GetInt(Section, TEXT("ResolutionQuality"),		Temporary.ResolutionQuality,	GEditorGameAgnosticIni);
+		GConfig->GetInt(Section, TEXT("ViewDistanceQuality"),		Temporary.ViewDistanceQuality,	GEditorGameAgnosticIni);
+		GConfig->GetInt(Section, TEXT("AntiAliasingQuality"),		Temporary.AntiAliasingQuality,	GEditorGameAgnosticIni);
+		GConfig->GetInt(Section, TEXT("ShadowQuality"),			Temporary.ShadowQuality,		GEditorGameAgnosticIni);
+		GConfig->GetInt(Section, TEXT("PostProcessQuality"),		Temporary.PostProcessQuality,	GEditorGameAgnosticIni);
+		GConfig->GetInt(Section, TEXT("TextureQuality"),			Temporary.TextureQuality,		GEditorGameAgnosticIni);
+		GConfig->GetInt(Section, TEXT("EffectsQuality"),			Temporary.EffectsQuality,		GEditorGameAgnosticIni);
+		EngineBenchmarkResult = Temporary;
+	}
+	else
+	{
+		GWarn->StatusUpdate(0, 1, NSLOCTEXT("UnrealEd", "RunningEngineBenchmark", "Running engine benchmark..."));
+		GWarn->PushStatus();
+		
+		Temporary = Scalability::BenchmarkQualityLevels();
+
+		GConfig->SetBool(Section,TEXT("Valid"),					true,							GEditorGameAgnosticIni);
+		GConfig->SetInt(Section, TEXT("ResolutionQuality"),		Temporary.ResolutionQuality,	GEditorGameAgnosticIni);
+		GConfig->SetInt(Section, TEXT("ViewDistanceQuality"),		Temporary.ViewDistanceQuality,	GEditorGameAgnosticIni);
+		GConfig->SetInt(Section, TEXT("AntiAliasingQuality"),		Temporary.AntiAliasingQuality,	GEditorGameAgnosticIni);
+		GConfig->SetInt(Section, TEXT("ShadowQuality"),			Temporary.ShadowQuality,		GEditorGameAgnosticIni);
+		GConfig->SetInt(Section, TEXT("PostProcessQuality"),		Temporary.PostProcessQuality,	GEditorGameAgnosticIni);
+		GConfig->SetInt(Section, TEXT("TextureQuality"),			Temporary.TextureQuality,		GEditorGameAgnosticIni);
+		GConfig->SetInt(Section, TEXT("EffectsQuality"),			Temporary.EffectsQuality,		GEditorGameAgnosticIni);
+
+		Scalability::SetQualityLevels(Temporary);
+		Scalability::SaveState(GEditorGameAgnosticIni);
+		
+		GWarn->PopStatus();
+	}
 }
\ No newline at end of file
diff --git a/Engine/Source/Editor/UnrealEd/Private/PerformanceMonitor.cpp b/Engine/Source/Editor/UnrealEd/Private/PerformanceMonitor.cpp
--- a/Engine/Source/Editor/UnrealEd/Private/PerformanceMonitor.cpp
+++ b/Engine/Source/Editor/UnrealEd/Private/PerformanceMonitor.cpp
@@ -92,7 +92,7 @@ public:
 	/** Called to get the "Show notification" check box state */
 	static ESlateCheckBoxState::Type IsMonitoringPerformance()
 	{
-		const bool bMonitorEditorPerformance = GetDefault<UEditorUserSettings>()->bMonitorEditorPerformance;
+		const bool bMonitorEditorPerformance = GEditor->GetEditorUserSettings().bMonitorEditorPerformance;
 		return bMonitorEditorPerformance ? ESlateCheckBoxState::Checked : ESlateCheckBoxState::Unchecked;
 	}
 
@@ -101,26 +101,18 @@ public:
 	{
 		const bool bNewEnabledState = (NewState == ESlateCheckBoxState::Checked);
 
-		auto* Settings = GetMutableDefault<UEditorUserSettings>();
-		Settings->bMonitorEditorPerformance = bNewEnabledState;
-		Settings->PostEditChange();
+		auto& Settings = GEditor->AccessEditorUserSettings();
+		Settings.bMonitorEditorPerformance = bNewEnabledState;
+		Settings.PostEditChange();
 	}
 };
 
 FPerformanceMonitor::FPerformanceMonitor()
 {
 	LastEnableTime = 0;
+	NotificationTimeout = 5;
 	bIsNotificationAllowed = true;
 
-	UEditorGameAgnosticSettings& Settings = GEditor->AccessGameAgnosticSettings();
-	if (Settings.bApplyAutoScalabilityOnStartup && FApp::HasGameName())
-	{
-		Settings.bApplyAutoScalabilityOnStartup = false;
-		Settings.PostEditChange();
-
-		AutoApplyScalability();
-	}
-
 	// Register the console variables we need
 	IConsoleVariable* CVar = nullptr;
 
@@ -135,7 +127,7 @@ FPerformanceMonitor::FPerformanceMonitor()
 	IConsoleManager::Get().RegisterConsoleVariable(TEXT("PerfWarn.FinePercentThreshold"), 80, TEXT("The percentage of samples that fall below min FPS above which we warn for."), ECVF_Default);
 	IConsoleManager::Get().RegisterConsoleVariable(TEXT("PerfWarn.CoarsePercentThreshold"), 80, TEXT("The percentage of samples that fall below min FPS above which we warn for."), ECVF_Default);
 
-	// Set up the moving average sasmplers
+	// Set up the moving average samplers
 	UpdateMovingAverageSamplers();
 }
 
@@ -148,78 +140,49 @@ FPerformanceMonitor::~FPerformanceMonitor()
 
 void FPerformanceMonitor::AutoApplyScalability()
 {
-	GWarn->BeginSlowTask(LOCTEXT("RunningBenchmark", "Running benchmark to determine optimum engine settings"), true);
-	//GWarn->UpdateProgress(-1, -1);
-	
-	Scalability::FQualityLevels BenchmarkLevels = Scalability::BenchmarkQualityLevels();
 	const Scalability::FQualityLevels ExistingLevels = Scalability::GetQualityLevels();
+	Scalability::FQualityLevels NewLevels = GEditor->GetGameAgnosticSettings().EngineBenchmarkResult;
 	
 	// Make sure we don't turn settings up if the user has turned them down for any reason
-	BenchmarkLevels.ResolutionQuality		= FMath::Min(BenchmarkLevels.ResolutionQuality,	ExistingLevels.ResolutionQuality);
-	BenchmarkLevels.ViewDistanceQuality	= FMath::Min(BenchmarkLevels.ViewDistanceQuality,	ExistingLevels.ViewDistanceQuality);
-	BenchmarkLevels.AntiAliasingQuality	= FMath::Min(BenchmarkLevels.AntiAliasingQuality,	ExistingLevels.AntiAliasingQuality);
-	BenchmarkLevels.ShadowQuality			= FMath::Min(BenchmarkLevels.ShadowQuality,		ExistingLevels.ShadowQuality);
-	BenchmarkLevels.PostProcessQuality	= FMath::Min(BenchmarkLevels.PostProcessQuality,	ExistingLevels.PostProcessQuality);
-	BenchmarkLevels.TextureQuality		= FMath::Min(BenchmarkLevels.TextureQuality,		ExistingLevels.TextureQuality);
-	BenchmarkLevels.EffectsQuality		= FMath::Min(BenchmarkLevels.EffectsQuality,		ExistingLevels.EffectsQuality);
+	NewLevels.ResolutionQuality	= FMath::Min(NewLevels.ResolutionQuality,		ExistingLevels.ResolutionQuality);
+	NewLevels.ViewDistanceQuality	= FMath::Min(NewLevels.ViewDistanceQuality,	ExistingLevels.ViewDistanceQuality);
+	NewLevels.AntiAliasingQuality	= FMath::Min(NewLevels.AntiAliasingQuality,	ExistingLevels.AntiAliasingQuality);
+	NewLevels.ShadowQuality		= FMath::Min(NewLevels.ShadowQuality,			ExistingLevels.ShadowQuality);
+	NewLevels.PostProcessQuality	= FMath::Min(NewLevels.PostProcessQuality,		ExistingLevels.PostProcessQuality);
+	NewLevels.TextureQuality		= FMath::Min(NewLevels.TextureQuality,			ExistingLevels.TextureQuality);
+	NewLevels.EffectsQuality		= FMath::Min(NewLevels.EffectsQuality,			ExistingLevels.EffectsQuality);
 	
-	Scalability::SetQualityLevels(BenchmarkLevels);
-	Scalability::SaveState(GEditorUserSettingsIni);
+	Scalability::SetQualityLevels(NewLevels);
+	Scalability::SaveState(GEditorGameAgnosticIni);
 
 	const bool bAutoApplied = true;
 	Scalability::RecordQualityLevelsAnalytics(bAutoApplied);
 
 	GEditor->DisableRealtimeViewports();
 
-	HidePerformanceWarning();
-	GWarn->EndSlowTask();
+	// Reset the timers so as not to skew the data with the time it took to do the benchmark	
+	FineMovingAverage.Reset();
+	CoarseMovingAverage.Reset();
 }
 
 void FPerformanceMonitor::ShowPerformanceWarning(FText MessageText)
 {
-	auto AlreadyOpenItem = PerformanceWarningNotificationPtr.Pin();
-	if (AlreadyOpenItem.IsValid())
-	{
-		AlreadyOpenItem->SetText(MessageText);
-	}
-	else if ((FPlatformTime::Seconds() - LastEnableTime) > 30.0)
+	static double MinNotifyTime = 30;
+	if ((FPlatformTime::Seconds() - LastEnableTime) > MinNotifyTime)
 	{
 		// Only show a new one if we've not shown one for a while
 		LastEnableTime = FPlatformTime::Seconds();
+		NotificationTimeout = 5;
 
 		// Create notification item
 		FNotificationInfo Info(MessageText);
 		Info.bFireAndForget = false;
-		Info.FadeOutDuration = 1.0f;
+		Info.FadeOutDuration = 4.0f;
 		Info.ExpireDuration = 0.0f;
 		Info.bUseLargeFont = false;
 
-		Info.ButtonDetails.Add(
-			FNotificationButtonInfo(
-			LOCTEXT("PerformanceWarningScalability", "Tweak Manually"),
-			LOCTEXT("PerformanceWarningScalabilityToolTip", "Opens the Scalability UI for manual adjustment."),
-			FSimpleDelegate::CreateRaw(this, &FPerformanceMonitor::ShowScalabilityDialog),
-			SNotificationItem::CS_None
-			)
-			);
-
-		Info.ButtonDetails.Add(
-			FNotificationButtonInfo(
-			LOCTEXT("PerformanceWarningApplyNow", "Apply Now"),
-			LOCTEXT("PerformanceWarningApplyNowToolTip", "Run a benchmark now and apply settings immediately."),
-			FSimpleDelegate::CreateRaw(this, &FPerformanceMonitor::AutoApplyScalability),
-			SNotificationItem::CS_None
-			)
-			);
-
-		Info.ButtonDetails.Add(
-			FNotificationButtonInfo(
-			LOCTEXT("PerformanceWarningDismiss", "Ok"),
-			LOCTEXT("PerformanceWarningDismissToolTip", "Close this notification and apply settings on restart."),
-			FSimpleDelegate::CreateRaw(this, &FPerformanceMonitor::DismissAndDisableNotification),
-			SNotificationItem::CS_None
-			)
-			);
+		Info.Hyperlink = FSimpleDelegate::CreateRaw(this, &FPerformanceMonitor::ShowScalabilityDialog);
+		Info.HyperlinkText = LOCTEXT("PerformanceWarningScalability", "Tweak Manually");
 
 		PerformanceWarningNotificationPtr = FSlateNotificationManager::Get().AddNotification(Info);
 	}
@@ -239,17 +202,6 @@ void FPerformanceMonitor::HidePerformanceWarning()
 	}
 }
 
-
-bool FPerformanceMonitor::IsTickable() const
-{
-#if UE_BUILD_DEBUG
-	// We don't do this in debug
-	return false;
-#else
-	return true;
-#endif
-}
-
 void FPerformanceMonitor::Tick(float DeltaTime)
 {
 	if (GEngine->ShouldThrottleCPUUsage())
@@ -307,22 +259,29 @@ void FPerformanceMonitor::Tick(float DeltaTime)
 		}
 	}
 
+	auto AlreadyOpenItem = PerformanceWarningNotificationPtr.Pin();
+
 	if (!bLowFramerate)
 	{
-		HidePerformanceWarning();
+		// Framerate is back up again - just reset everything and hide the timeout
+		if (AlreadyOpenItem.IsValid())
+		{
+			Reset();
+		}
 	}
 	else
 	{
+		// Choose an appropriate message
 		enum MessagesEnum { Seconds, SecondsPercent, Minute, MinutePercent, Minutes, MinutesPercent };
 		const FText Messages[] = {
-			LOCTEXT("PerformanceWarningInProgress_Seconds",			"Your framerate has been under {Framerate} FPS for the past {SampleTime} seconds.\nYour graphics settings will be reduced next time you restart the editor."),
-			LOCTEXT("PerformanceWarningInProgress_Seconds_Percent", "Your framerate has been under {Framerate} FPS for {Percentage}% of the past {SampleTime} seconds.\nYour graphics settings will be reduced next time you restart the editor."),
+			LOCTEXT("PerformanceWarningInProgress_Seconds",			"Your framerate has been under {Framerate} FPS for the past {SampleTime} seconds.\nApplying optimum engine settings in {TimeRemaining}s."),
+			LOCTEXT("PerformanceWarningInProgress_Seconds_Percent", "Your framerate has been under {Framerate} FPS for {Percentage}% of the past {SampleTime} seconds.\nApplying optimum engine settings in {TimeRemaining}s."),
 
-			LOCTEXT("PerformanceWarningInProgress_Minute",			"Your framerate has been under {Framerate} FPS for the past minute.\nYour graphics settings will be reduced next time you restart the editor."),
-			LOCTEXT("PerformanceWarningInProgress_Minute_Percent",	"Your framerate has been under {Framerate} FPS for {Percentage}% of the last minute.\nYour graphics settings will be reduced next time you restart the editor."),
+			LOCTEXT("PerformanceWarningInProgress_Minute",			"Your framerate has been under {Framerate} FPS for the past minute.\nApplying optimum engine settings in {TimeRemaining}s."),
+			LOCTEXT("PerformanceWarningInProgress_Minute_Percent",	"Your framerate has been under {Framerate} FPS for {Percentage}% of the last minute.\nApplying optimum engine settings in {TimeRemaining}s."),
 
-			LOCTEXT("PerformanceWarningInProgress_Minutes",			"Your framerate has been below {Framerate} FPS for the past {SampleTime} minutes.\nYour graphics settings will be reduced next time you restart the editor."),
-			LOCTEXT("PerformanceWarningInProgress_Minutes_Percent", "Your framerate has been below {Framerate} FPS for {Percentage}% of the past {SampleTime} minutes.\nYour graphics settings will be reduced next time you restart the editor."),
+			LOCTEXT("PerformanceWarningInProgress_Minutes",			"Your framerate has been below {Framerate} FPS for the past {SampleTime} minutes.\nApplying optimum engine settings in {TimeRemaining}s."),
+			LOCTEXT("PerformanceWarningInProgress_Minutes_Percent", "Your framerate has been below {Framerate} FPS for {Percentage}% of the past {SampleTime} minutes.\nApplying optimum engine settings in {TimeRemaining}s."),
 		};
 
 		int32 Message;
@@ -347,21 +306,32 @@ void FPerformanceMonitor::Tick(float DeltaTime)
 			++Message;
 		}
 
-		ShowPerformanceWarning(FText::Format(Messages[Message], Arguments));
-	}
-}
 
-void FPerformanceMonitor::DismissAndDisableNotification()
-{
-	bIsNotificationAllowed = false;
-
-	// Set the flag to say we need to run a benchmark on startup
-	GEditor->AccessGameAgnosticSettings().bApplyAutoScalabilityOnStartup = true;
-	GEditor->AccessGameAgnosticSettings().PostEditChange();
-
-	GEditor->DisableRealtimeViewports();
+		// Now update the notification or create a new one
+		if (AlreadyOpenItem.IsValid())
+		{
+			NotificationTimeout -= DeltaTime;
+			Arguments.Add(TEXT("TimeRemaining"), FMath::Max(1, FMath::CeilToInt(NotificationTimeout)));
+
+			if (NotificationTimeout <= 0)
+			{
+				// Timed-out. Apply the settings.
+				AutoApplyScalability();
+				Reset();
+			}
+			else
+			{
+				AlreadyOpenItem->SetText(FText::Format(Messages[Message], Arguments));
+			}
+		}
+		else
+		{
+			NotificationTimeout = 5;
+			Arguments.Add(TEXT("TimeRemaining"), int(NotificationTimeout));
 
-	HidePerformanceWarning();
+			ShowPerformanceWarning(FText::Format(Messages[Message], Arguments));
+		}
+	}
 }
 
 void FPerformanceMonitor::Reset()
@@ -370,6 +340,7 @@ void FPerformanceMonitor::Reset()
 	CoarseMovingAverage.Reset();
 
 	HidePerformanceWarning();
+	bIsNotificationAllowed = true;
 }
 
 void FPerformanceMonitor::UpdateMovingAverageSamplers(IConsoleVariable* Unused)
@@ -386,6 +357,9 @@ void FPerformanceMonitor::UpdateMovingAverageSamplers(IConsoleVariable* Unused)
 
 void FPerformanceMonitor::ShowScalabilityDialog()
 {
+	Reset();
+	bIsNotificationAllowed = false;
+
 	auto ExistingWindow = ScalabilitySettingsWindowPtr.Pin();
 
 	if (ExistingWindow.IsValid())
@@ -394,10 +368,6 @@ void FPerformanceMonitor::ShowScalabilityDialog()
 	}
 	else
 	{
-		// Turn off the auto-application of settings on start up if the user is manually tweaking them
-		GEditor->AccessGameAgnosticSettings().bApplyAutoScalabilityOnStartup = false;
-		GEditor->AccessGameAgnosticSettings().PostEditChange();
-
 		// Create the window
 		ScalabilitySettingsWindowPtr = ExistingWindow = SNew(SWindow)
 			.Title(LOCTEXT("PerformanceWarningDialogTitle", "Scalability Options"))
diff --git a/Engine/Source/Editor/UnrealEd/Private/PerformanceMonitor.h b/Engine/Source/Editor/UnrealEd/Private/PerformanceMonitor.h
--- a/Engine/Source/Editor/UnrealEd/Private/PerformanceMonitor.h
+++ b/Engine/Source/Editor/UnrealEd/Private/PerformanceMonitor.h
@@ -136,7 +136,7 @@ public:
 protected:
 
 	/** FTickableEditorObject interface */
-	virtual bool IsTickable() const override;
+	virtual bool IsTickable() const override { return true; }
 	virtual void Tick(float DeltaTime) override;
 	virtual TStatId GetStatId() const override { RETURN_QUICK_DECLARE_CYCLE_STAT(FPerformanceMonitor, STATGROUP_Tickables); }
 
@@ -154,9 +154,6 @@ private:
 	/** Display the scalability dialog. */
 	void ShowScalabilityDialog();
 
-	/** Dismiss the notification and disable it for this session. */
-	void DismissAndDisableNotification();
-
 	/** Resets the performance warning data and hides the warning */
 	void Reset();
 
@@ -169,11 +166,14 @@ private:
 	/** Tracks the last time the notification was started, used to avoid spamming. */
 	double LastEnableTime;
 
+	double NotificationTimeout;
+
 	/** The notification window ptr */
 	TWeakPtr<SNotificationItem> PerformanceWarningNotificationPtr;
 
 	/** The Scalability Setting window we may be using */
 	TWeakPtr<SWindow> ScalabilitySettingsWindowPtr;
 
+	/** Whether the notification is allowed to pop up this session */
 	bool bIsNotificationAllowed;
 };
\ No newline at end of file
diff --git a/Engine/Source/Editor/UnrealEd/Private/SScalabilitySettings.cpp b/Engine/Source/Editor/UnrealEd/Private/SScalabilitySettings.cpp
--- a/Engine/Source/Editor/UnrealEd/Private/SScalabilitySettings.cpp
+++ b/Engine/Source/Editor/UnrealEd/Private/SScalabilitySettings.cpp
@@ -31,7 +31,7 @@ void SScalabilitySettings::OnGroupQualityLevelChanged(ESlateCheckBoxState::Type
 	else if (FCString::Strcmp(InGroupName, TEXT("EffectsQuality")) == 0) CachedQualityLevels.EffectsQuality = InQualityLevel;
 
 	Scalability::SetQualityLevels(CachedQualityLevels);
-	Scalability::SaveState(GEditorUserSettingsIni);
+	Scalability::SaveState(GEditorGameAgnosticIni);
 	GEditor->RedrawAllViewports();
 }
 
@@ -40,7 +40,7 @@ void SScalabilitySettings::OnResolutionScaleChanged(float InValue)
 	CachedQualityLevels.ResolutionQuality = (int32)(FMath::Lerp(Scalability::MinResolutionScale, Scalability::MaxResolutionScale, InValue));
 
 	Scalability::SetQualityLevels(CachedQualityLevels);
-	Scalability::SaveState(GEditorUserSettingsIni);
+	Scalability::SaveState(GEditorGameAgnosticIni);
 	GEditor->RedrawAllViewports();
 }
 
@@ -96,16 +96,16 @@ FReply SScalabilitySettings::OnHeaderClicked(int32 InQualityLevel)
 {
 	CachedQualityLevels.SetFromSingleQualityLevel(InQualityLevel);
 	Scalability::SetQualityLevels(CachedQualityLevels);
-	Scalability::SaveState(GEditorUserSettingsIni);
+	Scalability::SaveState(GEditorGameAgnosticIni);
 	GEditor->RedrawAllViewports();
 	return FReply::Handled();
 }
 
 FReply SScalabilitySettings::OnAutoClicked()
 {
-	CachedQualityLevels = Scalability::BenchmarkQualityLevels();
+	CachedQualityLevels = GEditor->GetGameAgnosticSettings().EngineBenchmarkResult;
 	SetQualityLevels(CachedQualityLevels);
-	Scalability::SaveState(GEditorUserSettingsIni);
+	Scalability::SaveState(GEditorGameAgnosticIni);
 	GEditor->RedrawAllViewports();
 	return FReply::Handled();
 }
diff --git a/Engine/Source/Editor/UnrealEd/Private/UnrealEd.cpp b/Engine/Source/Editor/UnrealEd/Private/UnrealEd.cpp
--- a/Engine/Source/Editor/UnrealEd/Private/UnrealEd.cpp
+++ b/Engine/Source/Editor/UnrealEd/Private/UnrealEd.cpp
@@ -90,6 +90,9 @@ int32 EditorInit( IEngineLoop& EngineLoop )
 	// Set up the actor folders singleton
 	FActorFolders::Init();
 
+	GEditor->AccessGameAgnosticSettings().LoadScalabilityBenchmark();
+	GEditor->SaveGameAgnosticSettings();
+
 	// =================== CORE EDITOR INIT FINISHED ===================
 
 	// Hide the splash screen now that everything is ready to go
diff --git a/Engine/Source/Runtime/Engine/Private/Scalability.cpp b/Engine/Source/Runtime/Engine/Private/Scalability.cpp
--- a/Engine/Source/Runtime/Engine/Private/Scalability.cpp
+++ b/Engine/Source/Runtime/Engine/Private/Scalability.cpp
@@ -293,7 +293,7 @@ void ProcessCommand(const TCHAR* Cmd, FOutputDevice& Ar)
 		{
 			FQualityLevels State = Scalability::BenchmarkQualityLevels();
 			Scalability::SetQualityLevels(State);
-			Scalability::SaveState(GIsEditor ? GEditorUserSettingsIni : GGameUserSettingsIni);
+			Scalability::SaveState(GIsEditor ? GEditorGameAgnosticIni : GGameUserSettingsIni);
 			bPrintUsage = false;
 			bPrintCurrentSettings = true;
 		}
@@ -310,7 +310,7 @@ void ProcessCommand(const TCHAR* Cmd, FOutputDevice& Ar)
 			int32 RequestedQualityLevel = FCString::Atoi(*Token);
 			QualityLevels.SetFromSingleQualityLevel(RequestedQualityLevel);
 			SetQualityLevels(QualityLevels);
-			Scalability::SaveState(GIsEditor ? GEditorUserSettingsIni : GGameUserSettingsIni);
+			Scalability::SaveState(GIsEditor ? GEditorGameAgnosticIni : GGameUserSettingsIni);
 
 			bPrintUsage = false;
 		}
