Some restyling for ProxyMesh dialog

diff --git a/Engine/Source/Editor/UnrealEd/Private/Dialogs/SMeshProxyDialog.cpp b/Engine/Source/Editor/UnrealEd/Private/Dialogs/SMeshProxyDialog.cpp
--- a/Engine/Source/Editor/UnrealEd/Private/Dialogs/SMeshProxyDialog.cpp
+++ b/Engine/Source/Editor/UnrealEd/Private/Dialogs/SMeshProxyDialog.cpp
@@ -66,16 +66,9 @@ protected:
 	/** TextureResolution accessors */
 	void SetTextureResolution(TSharedPtr<FString> NewSelection, ESelectInfo::Type SelectInfo);
 
-	/** ErrorMessage accessor */
-	FString GetErrorMessage() const;
-
-	/** Delegates for Merge/Remerge/Unmerge button states */
-	EVisibility GetMergeButtonVisibility() const;
-	bool GetButtonEnabledState() const;
-
-	/** Returns the visibility state of the error message */
-	EVisibility GetErrorMessageVisibility() const;
-
+	/** Delegates for Merge button state */
+	bool GetMergeButtonEnabledState() const;
+	
 	/** Called when the Cancel button is clicked */
 	FReply OnCancelClicked();
 
@@ -111,15 +104,9 @@ private:
 	TArray< TSharedPtr<FString> >	CuttingPlaneOptions;
 	TArray< TSharedPtr<FString> >	TextureResolutionOptions;
 
-	/** The error message displayed when there's a problem with the selected meshes */
-	FString ErrorMessage;
-
-	/** Dictates the visibility of the Merge/Remerge/Unmerge buttons */
+	/** Dictates the availability of the Merge button */
 	bool bCanMerge;
 
-	/** Dictates the visibility of the error message */
-	bool bShowErrorMessage;
-
 	/** If true, the selected group needs to be re-evaluated */
 	bool bDirty;
 	
@@ -141,14 +128,7 @@ SMeshProxyDialog::~SMeshProxyDialog()
 
 void SMeshProxyDialog::AssignWindowContent(TSharedPtr<SWindow> Window)
 {
-	TSharedPtr<SBorder> Controls =
-	SNew(SBorder)
-	.BorderImage(FEditorStyle::GetBrush("ToolPanel.GroupBorder"))
-	[
-		SharedThis(this)
-	];
-
-	Window->SetContent(Controls.ToSharedRef());
+	Window->SetContent(this->AsShared());
 }
 
 TSharedPtr<SWindow> SMeshProxyDialog::GetParentWindow()
@@ -199,21 +179,23 @@ void SMeshProxyDialog::CreateLayout()
 {
 	this->ChildSlot
 	[
-		SNew(SBorder)
-		.BorderImage(FEditorStyle::GetBrush("ToolPanel.GroupBorder"))
+		SNew(SVerticalBox)
+			
+		+SVerticalBox::Slot()
+		.FillHeight(0.2f)
+		.HAlign(HAlign_Center)
+		.VAlign(VAlign_Center)
 		[
-			SNew(SVerticalBox)
-			+SVerticalBox::Slot()
-			.FillHeight(0.2f)
-			.HAlign(HAlign_Center)
-			.VAlign(VAlign_Center)
-			[
-				// Simplygon logo
-				SNew(SImage)
-				.Image(FEditorStyle::GetBrush("MeshProxy.SimplygonLogo"))
-			]
-			+SVerticalBox::Slot()
-			.FillHeight(0.65f)
+			// Simplygon logo
+			SNew(SImage)
+			.Image(FEditorStyle::GetBrush("MeshProxy.SimplygonLogo"))
+		]
+			
+		+SVerticalBox::Slot()
+		.FillHeight(0.65f)
+		[
+			SNew(SBorder)
+			.BorderImage(FEditorStyle::GetBrush("ToolPanel.GroupBorder"))
 			[
 				// Proxy options
 				SNew(SHorizontalBox)
@@ -257,7 +239,7 @@ void SMeshProxyDialog::CreateLayout()
 					.FillHeight(1.0f)
 					.VAlign(VAlign_Center)
 					[
-						SNew(STextBlock).Text(LOCTEXT("DestinationPackageLabel", "Destination Package"))
+						SNew(STextBlock).Text(LOCTEXT("DestinationPackageLabel", "Base Asset Name"))
 					]
 				]
 				+SHorizontalBox::Slot()
@@ -421,46 +403,34 @@ void SMeshProxyDialog::CreateLayout()
 					]
 				]
 			]
-			+SVerticalBox::Slot()
-			.FillHeight(0.15f)
-			.Padding(0.0f, 0.0f, 8.0f, 4.0f)
+		]
+
+		+SVerticalBox::Slot()
+		.AutoHeight()
+		.HAlign(HAlign_Right)
+		.VAlign(VAlign_Bottom)
+		.AutoHeight()
+		.Padding(0,4,0,8)
+		[
+			// Merge, Cancel buttons
+			SNew(SUniformGridPanel)
+			.SlotPadding(FEditorStyle::GetMargin("StandardDialog.SlotPadding"))
+			.MinDesiredSlotWidth(FEditorStyle::GetFloat("StandardDialog.MinDesiredSlotWidth"))
+			.MinDesiredSlotHeight(FEditorStyle::GetFloat("StandardDialog.MinDesiredSlotHeight"))
+			+SUniformGridPanel::Slot(0,0)
 			[
-				// Error message
-				SNew(SHorizontalBox)
-				+SHorizontalBox::Slot()
-				.FillWidth(1.0f)
-				.HAlign(HAlign_Left)
-				.VAlign(VAlign_Center)
-				[
-					SNew(STextBlock)
-					.Text(this, &SMeshProxyDialog::GetErrorMessage)
-					.Visibility(this, &SMeshProxyDialog::GetErrorMessageVisibility)
-				]
-				+SHorizontalBox::Slot()
-				.FillWidth(1.0f)
-				.HAlign(HAlign_Right)
-				.VAlign(VAlign_Center)
-				[
-					SNew(SHorizontalBox)
-					+SHorizontalBox::Slot()
-					.AutoWidth()
-					.Padding(0.0f, 0.0f, 4.0f, 0.0f)
-					[
-						SNew(SButton)
-						.Text(LOCTEXT("MergeLabel", "Merge"))
-						.Visibility(this, &SMeshProxyDialog::GetMergeButtonVisibility)
-						.IsEnabled(this, &SMeshProxyDialog::GetButtonEnabledState)
-						.OnClicked(this,&SMeshProxyDialog::OnMergeClicked)
-					]
-					+SHorizontalBox::Slot()
-					.AutoWidth()
-					.Padding(4.0f, 0.0f, 0.0f, 0.0f)
-					[
-						SNew(SButton)
-						.Text(LOCTEXT("CancelLabel", "Cancel"))
-						.OnClicked(this, &SMeshProxyDialog::OnCancelClicked)
-					]
-				]
+				SNew(SButton)
+				.Text(LOCTEXT("MergeLabel", "Merge"))
+				.IsEnabled(this, &SMeshProxyDialog::GetMergeButtonEnabledState)
+				.OnClicked(this,&SMeshProxyDialog::OnMergeClicked)
+
+			]
+			+SUniformGridPanel::Slot(1,0)
+			[
+				SNew(SButton)
+				.Text(LOCTEXT("Cancel", "Cancel"))
+				.ContentPadding(FEditorStyle::GetMargin("StandardDialog.ContentPadding"))
+				.OnClicked(this, &SMeshProxyDialog::OnCancelClicked)
 			]
 		]
 	];
@@ -469,9 +439,6 @@ void SMeshProxyDialog::CreateLayout()
 void SMeshProxyDialog::Reset(const bool bRefresh)
 {
 	bCanMerge = true;
-	bShowErrorMessage = false;
-	ErrorMessage.Empty();
-
 	if (bRefresh)
 	{
 		FEditorDelegates::RefreshEditor.Broadcast();
@@ -592,25 +559,9 @@ void SMeshProxyDialog::SetTextureResolution(TSharedPtr<FString> NewSelection, ES
 }
 
 
-EVisibility SMeshProxyDialog::GetMergeButtonVisibility() const
-{
-	return bCanMerge? EVisibility::Visible: EVisibility::Collapsed;
-}
-
-
-bool SMeshProxyDialog::GetButtonEnabledState() const
-{
-	return !bShowErrorMessage;
-}
-
-FString SMeshProxyDialog::GetErrorMessage() const
-{
-	return ErrorMessage;
-}
-
-EVisibility SMeshProxyDialog::GetErrorMessageVisibility() const
+bool SMeshProxyDialog::GetMergeButtonEnabledState() const
 {
-	return bShowErrorMessage? EVisibility::Visible: EVisibility::Collapsed;
+	return bCanMerge;
 }
 
 FReply SMeshProxyDialog::OnCancelClicked()
