More fixes to MeshProxy dialog

diff --git a/Engine/Source/Editor/UnrealEd/Private/Dialogs/SMeshProxyDialog.cpp b/Engine/Source/Editor/UnrealEd/Private/Dialogs/SMeshProxyDialog.cpp
--- a/Engine/Source/Editor/UnrealEd/Private/Dialogs/SMeshProxyDialog.cpp
+++ b/Engine/Source/Editor/UnrealEd/Private/Dialogs/SMeshProxyDialog.cpp
@@ -421,6 +421,7 @@ void SMeshProxyDialog::CreateLayout()
 			[
 				SNew(SButton)
 				.Text(LOCTEXT("MergeLabel", "Merge"))
+				.HAlign(HAlign_Center)
 				.IsEnabled(this, &SMeshProxyDialog::GetMergeButtonEnabledState)
 				.OnClicked(this,&SMeshProxyDialog::OnMergeClicked)
 
@@ -429,6 +430,7 @@ void SMeshProxyDialog::CreateLayout()
 			[
 				SNew(SButton)
 				.Text(LOCTEXT("Cancel", "Cancel"))
+				.HAlign(HAlign_Center)
 				.ContentPadding(FEditorStyle::GetMargin("StandardDialog.ContentPadding"))
 				.OnClicked(this, &SMeshProxyDialog::OnCancelClicked)
 			]
@@ -694,7 +696,7 @@ FReply SMeshProxyDialog::OnSelectPackageNameClicked()
 {
 	TSharedRef<SDlgPickAssetPath> NewLayerDlg = 
 		SNew(SDlgPickAssetPath)
-		.Title(LOCTEXT("SelectProxyPackage", "Select destination package"))
+		.Title(LOCTEXT("SelectProxyPackage", "Base Asset Name"))
 		.DefaultAssetPath(FText::FromString(ProxyPackageName));
 
 	if (NewLayerDlg->ShowModal() != EAppReturnType::Cancel)
