Fix to ensure that game agnostic benchmark is only run if a game agnostic config exists

diff --git a/Engine/Source/Editor/UnrealEd/Private/Editor/EditorGameAgnosticSettings.cpp b/Engine/Source/Editor/UnrealEd/Private/Editor/EditorGameAgnosticSettings.cpp
--- a/Engine/Source/Editor/UnrealEd/Private/Editor/EditorGameAgnosticSettings.cpp
+++ b/Engine/Source/Editor/UnrealEd/Private/Editor/EditorGameAgnosticSettings.cpp
@@ -33,6 +33,8 @@ void UEditorGameAgnosticSettings::PostEditChangeProperty( struct FPropertyChange
 
 void UEditorGameAgnosticSettings::LoadScalabilityBenchmark()
 {
+	check(!GEditorGameAgnosticIni.IsEmpty());
+
 	const TCHAR* Section = TEXT("EngineBenchmarkResult");
 
 	bool bIsValid = false;
diff --git a/Engine/Source/Editor/UnrealEd/Private/UnrealEd.cpp b/Engine/Source/Editor/UnrealEd/Private/UnrealEd.cpp
--- a/Engine/Source/Editor/UnrealEd/Private/UnrealEd.cpp
+++ b/Engine/Source/Editor/UnrealEd/Private/UnrealEd.cpp
@@ -90,8 +90,12 @@ int32 EditorInit( IEngineLoop& EngineLoop )
 	// Set up the actor folders singleton
 	FActorFolders::Init();
 
-	GEditor->AccessGameAgnosticSettings().LoadScalabilityBenchmark();
-	GEditor->SaveGameAgnosticSettings();
+	if ( !GEditorGameAgnosticIni.IsEmpty() )
+	{
+		// If we have a game agnostic ini config, ensure that the benchmark has been loaded for it
+		GEditor->AccessGameAgnosticSettings().LoadScalabilityBenchmark();
+		GEditor->SaveGameAgnosticSettings();
+	}
 
 	// =================== CORE EDITOR INIT FINISHED ===================
 
