More Blueprint validation in automated tests.
ttp339664, 339652

diff --git a/Engine/Source/Editor/UnrealEd/Private/Tests/BlueprintAutomationTests.cpp b/Engine/Source/Editor/UnrealEd/Private/Tests/BlueprintAutomationTests.cpp
--- a/Engine/Source/Editor/UnrealEd/Private/Tests/BlueprintAutomationTests.cpp
+++ b/Engine/Source/Editor/UnrealEd/Private/Tests/BlueprintAutomationTests.cpp
@@ -839,6 +839,12 @@ bool FBlueprintCompileOnLoadTest::RunTest(const FString& BlueprintAssetPath)
 		return false;
 	}
 
+	if (!InitialBlueprint->SkeletonGeneratedClass || !InitialBlueprint->GeneratedClass)
+	{
+		AddError(*FString::Printf(TEXT("Unable to load blueprint for: '%s'. Probably it derives from an invalid class."), *BlueprintAssetPath));
+		return false;
+	}
+
 #if WITH_EDITOR
 	TArray<UClass*> BlueprintDependencies = FScopedClassDependencyGather::GetCachedDependencies();
 #else 
@@ -1230,6 +1236,11 @@ bool FBlueprintRenameAndCloneTest::RunTest(const FString& BlueprintAssetPath)
 		bTestFailed = true;
 		AddError(FString::Printf(TEXT("Failed to load '%s' (has it been renamed?)."), *BlueprintAssetPath));
 	}
+	if (!OriginalBlueprint->SkeletonGeneratedClass || !OriginalBlueprint->GeneratedClass)
+	{
+		bTestFailed = true;
+		AddError(*FString::Printf(TEXT("Unable to load blueprint for: '%s'. Probably it derives from an invalid class."), *BlueprintAssetPath));
+	}
 	else if (!bTestFailed)
 	{
 		// duplicate
