CIS fix for failures relating to performance monitor being ran without a GEditorGameAgnosticIni.

diff --git a/Engine/Source/Editor/UnrealEd/Private/UnrealEdEngine.cpp b/Engine/Source/Editor/UnrealEd/Private/UnrealEdEngine.cpp
--- a/Engine/Source/Editor/UnrealEd/Private/UnrealEdEngine.cpp
+++ b/Engine/Source/Editor/UnrealEd/Private/UnrealEdEngine.cpp
@@ -40,7 +40,11 @@ void UUnrealEdEngine::Init(IEngineLoop* InEngineLoop)
 	PackageAutoSaver->LoadRestoreFile();
 
 #if !UE_BUILD_DEBUG
-	PerformanceMonitor = new FPerformanceMonitor;
+	if( GEditorGameAgnosticIni )
+	{
+		// We need the game agnostic ini for this code
+		PerformanceMonitor = new FPerformanceMonitor;
+	}
 #endif
 
 	// Register for the package dirty state updated callback to catch packages that have been modified and need to be checked out.
