Syntax error fix.

Somehow perforce managed to only check in half of my 1 line change.

diff --git a/Engine/Source/Editor/UnrealEd/Private/UnrealEdEngine.cpp b/Engine/Source/Editor/UnrealEd/Private/UnrealEdEngine.cpp
--- a/Engine/Source/Editor/UnrealEd/Private/UnrealEdEngine.cpp
+++ b/Engine/Source/Editor/UnrealEd/Private/UnrealEdEngine.cpp
@@ -40,7 +40,7 @@ void UUnrealEdEngine::Init(IEngineLoop* InEngineLoop)
 	PackageAutoSaver->LoadRestoreFile();
 
 #if !UE_BUILD_DEBUG
-	if( GEditorGameAgnosticIni )
+	if( !GEditorGameAgnosticIni.IsEmpty() )
 	{
 		// We need the game agnostic ini for this code
 		PerformanceMonitor = new FPerformanceMonitor;
