Fixed Snap Pivot to Floor in Brushes.
#ttp 340014 - EDITOR: GEOMETRY: Snap Pivot to Floor doesn't work properly
#branch UE4
#proj Editor.UnrealEd
#add In UEditorEngine::SnapActorTo, in the code block dealing with brushes, handle the case where UsePivot is requested (before it was being ignored).

diff --git a/Engine/Source/Editor/UnrealEd/Private/EditorServer.cpp b/Engine/Source/Editor/UnrealEd/Private/EditorServer.cpp
--- a/Engine/Source/Editor/UnrealEd/Private/EditorServer.cpp
+++ b/Engine/Source/Editor/UnrealEd/Private/EditorServer.cpp
@@ -4273,11 +4273,20 @@ bool UEditorEngine::SnapActorTo( AActor* InActor, const bool InAlign, const bool
 	bool UseLineTrace = Brush ? true: InUseLineTrace;
 	bool UseBounds = Brush ? true: InUseBounds;
 
-	if( UseLineTrace && UseBounds ) // Will do a line trace from the center bottom of the bounds through the world. Will begin at the bottom center of the component's bounds.
+	if( UseLineTrace && UseBounds )
 	{
 		check(InActor->GetRootComponent()->IsRegistered());
-		StartLocation = InActor->GetRootComponent()->Bounds.Origin;
-		StartLocation.Z -= InActor->GetRootComponent()->Bounds.BoxExtent.Z;
+		if (InUsePivot)
+		{
+			// Will do a line trace from the pivot location.
+			StartLocation = GetPivotLocation();
+		}
+		else
+		{
+			// Will do a line trace from the center bottom of the bounds through the world. Will begin at the bottom center of the component's bounds.
+			StartLocation = InActor->GetRootComponent()->Bounds.Origin;
+			StartLocation.Z -= InActor->GetRootComponent()->Bounds.BoxExtent.Z;
+		}
 
 		// Forces a line trace.
 		Extent = FVector::ZeroVector;
