Re worked on editor visibility issue with drag/drop. Instead of using visibility channel, now it's just getting all objects, and then it will filter out what isn't needed.
And disabled response to visibility channel.

diff --git a/Engine/Source/Editor/UnrealEd/Private/LevelEditorViewport.cpp b/Engine/Source/Editor/UnrealEd/Private/LevelEditorViewport.cpp
--- a/Engine/Source/Editor/UnrealEd/Private/LevelEditorViewport.cpp
+++ b/Engine/Source/Editor/UnrealEd/Private/LevelEditorViewport.cpp
@@ -897,6 +897,45 @@ bool FLevelEditorViewportClient::DropObjectsOnActor(FViewportCursorLocation& Cur
 	return bResult;
 }
 
+bool FLevelEditorViewportClient::GetVisibleActors(TArray<struct FHitResult>& OutHits,const FVector& Start,const FVector& End,const struct FCollisionQueryParams& Params) const
+{
+	FCollisionObjectQueryParams ObjectParam(FCollisionObjectQueryParams::InitType::AllObjects);
+	if ( GetWorld()->LineTraceMulti(OutHits, Start, End, Params, ObjectParam) )
+	{
+		// see if any volume is around
+		for (int32 HitIdx=0; HitIdx < OutHits.Num(); ++HitIdx)
+		{
+			if ( OutHits[HitIdx].Component!= NULL &&
+				// ignore volume or trigger boxes
+				(OutHits[HitIdx].Component->IsA(UShapeComponent::StaticClass())))
+			{
+				// remove
+				OutHits.RemoveAt(HitIdx);
+				--HitIdx;
+			}
+			else if(OutHits[HitIdx].Actor!= NULL &&
+			// ignore volume or trigger boxes
+			(OutHits[HitIdx].Actor->IsA(AVolume::StaticClass())))
+			{
+				// remove
+				OutHits.RemoveAt(HitIdx);
+				--HitIdx;
+			}
+		}
+
+		 // now print all list
+//		for (int32 HitIdx=0; HitIdx < OutHits.Num(); ++HitIdx)
+//		{
+//			const AActor * Actor = OutHits[HitIdx].Actor.Get();
+//			UE_LOG(LogLevelTools, Warning, TEXT("%d) Actor: %s, ClassType : %s, Component : %s "), HitIdx+1, *GetNameSafe(Actor), (Actor)? *GetNameSafe(Actor->GetClass()): TEXT("None"),  *GetNameSafe(OutHits[HitIdx].Component.Get()));
+//		}
+
+		return (OutHits.Num() > 0);
+	}
+
+	return false;
+}
+
 bool FLevelEditorViewportClient::DropSingleObjectOnActor(FViewportCursorLocation& Cursor, UObject* DroppedObject, AActor* DroppedUponActor, int32 DroppedUponSlot, const FVector& DroppedLocation, EObjectFlags ObjectFlags, TArray<AActor*>& OutNewActors, bool bUsedHitProxy /*= true*/, bool bSelectActors /*= true*/, class UActorFactory* FactoryToUse /*= NULL*/ )
 {
 	if (ensure( DroppedObject ))
@@ -922,7 +961,7 @@ bool FLevelEditorViewportClient::DropSingleObjectOnActor(FViewportCursorLocation
 			case LVT_OrthoYZ:	WorldDistanceMultiplier = DUAExtent.X * 1.1; break; // Side
 			}
 
-			if ( GetWorld()->LineTraceMulti(Hits, Cursor.GetOrigin() - Cursor.GetDirection() * WorldDistanceMultiplier, Cursor.GetOrigin() + Cursor.GetDirection() * HALF_WORLD_MAX, ECC_Visibility, Param) )
+			if ( GetVisibleActors(Hits, Cursor.GetOrigin() - Cursor.GetDirection() * WorldDistanceMultiplier, Cursor.GetOrigin() + Cursor.GetDirection() * HALF_WORLD_MAX, Param) )
 			{
 				bool FoundMatch = false;
 
@@ -1003,7 +1042,7 @@ bool FLevelEditorViewportClient::DropObjectsOnBSPSurface( FSceneView* View, FVie
 	TArray<FHitResult> Hits;
 
 	FCollisionQueryParams Param(TEXT("DragDropTrace"), true);
-	if ( GetWorld()->LineTraceMulti(Hits, Cursor.GetOrigin(), Cursor.GetOrigin() + Cursor.GetDirection() * HALF_WORLD_MAX, ECC_Visibility, Param) )
+	if ( GetVisibleActors(Hits, Cursor.GetOrigin(), Cursor.GetOrigin() + Cursor.GetDirection() * HALF_WORLD_MAX, Param) )
 	{
 		bool FoundMatch = false;
 
@@ -1256,7 +1295,7 @@ bool FLevelEditorViewportClient::UpdateDropPreviewActors(int32 MouseX, int32 Mou
 				case LVT_OrthoYZ:	WorldDistanceMultiplier = DUAExtent.X * 1.1; break; // Side
 				}
 
-				if ( World->LineTraceMulti(Hits, Cursor.GetOrigin() - Cursor.GetDirection() * WorldDistanceMultiplier, Cursor.GetOrigin() + Cursor.GetDirection() * HALF_WORLD_MAX, ECC_Visibility, Param) )
+				if ( GetVisibleActors(Hits, Cursor.GetOrigin() - Cursor.GetDirection() * WorldDistanceMultiplier, Cursor.GetOrigin() + Cursor.GetDirection() * HALF_WORLD_MAX, Param) )
 				{
 					bool FoundMatch = false;
 
@@ -1309,8 +1348,8 @@ bool FLevelEditorViewportClient::UpdateDropPreviewActors(int32 MouseX, int32 Mou
 
 			TArray<FHitResult> Hits;
 			FCollisionQueryParams Param(TEXT("DragDropTrace"), true);
-			Param.AddIgnoredActors(DraggingActors);
-			if ( GetWorld()->LineTraceMulti(Hits, Cursor.GetOrigin(), Cursor.GetOrigin() + Cursor.GetDirection() * HALF_WORLD_MAX, ECC_Visibility, Param) )
+			Param.AddIgnoredActors(DraggingActors);		
+			if ( GetVisibleActors(Hits, Cursor.GetOrigin(), Cursor.GetOrigin() + Cursor.GetDirection() * HALF_WORLD_MAX, Param) )
 			{
 				bool FoundMatch = false;
 
diff --git a/Engine/Source/Editor/UnrealEd/Public/LevelEditorViewport.h b/Engine/Source/Editor/UnrealEd/Public/LevelEditorViewport.h
--- a/Engine/Source/Editor/UnrealEd/Public/LevelEditorViewport.h
+++ b/Engine/Source/Editor/UnrealEd/Public/LevelEditorViewport.h
@@ -571,6 +571,9 @@ private:
 		return EnabledStats.Num();
 	}
 
+	/** Get the list of visible actors for editor viewport op. It filters out certain actor or component types */
+	bool GetVisibleActors(TArray<struct FHitResult>& OutHits,const FVector& Start,const FVector& End,const struct FCollisionQueryParams& Params) const;
+
 protected:
 	/** 
 	 * Checks the viewport to see if the given blueprint asset can be dropped on the viewport.
diff --git a/Engine/Source/Runtime/Engine/Private/PhysicsEngine/BodyInstance.cpp b/Engine/Source/Runtime/Engine/Private/PhysicsEngine/BodyInstance.cpp
--- a/Engine/Source/Runtime/Engine/Private/PhysicsEngine/BodyInstance.cpp
+++ b/Engine/Source/Runtime/Engine/Private/PhysicsEngine/BodyInstance.cpp
@@ -652,7 +652,6 @@ void FBodyInstance::UpdatePhysicsFilterData()
 			//UE_LOG(LogPhysics, Warning, TEXT("Enabling collision %s : %s"), *GetNameSafe(Owner), *GetNameSafe(OwnerComponent.Get()));
 			// clear all other channel just in case other people using those channels to do something
 			UseResponse.SetAllChannels(ECR_Ignore);
-			UseResponse.SetResponse(ECC_Visibility, ECR_Block);
 			UseCollisionEnabled = ECollisionEnabled::QueryOnly;
 		}
 	}
