BP: Renaming some variable types will reset its default value after compile ttp337095

diff --git a/Engine/Source/Editor/UnrealEd/Private/Kismet2/BlueprintEditorUtils.cpp b/Engine/Source/Editor/UnrealEd/Private/Kismet2/BlueprintEditorUtils.cpp
--- a/Engine/Source/Editor/UnrealEd/Private/Kismet2/BlueprintEditorUtils.cpp
+++ b/Engine/Source/Editor/UnrealEd/Private/Kismet2/BlueprintEditorUtils.cpp
@@ -3265,7 +3265,7 @@ void FBlueprintEditorUtils::RenameMemberVariable(UBlueprint* Blueprint, const FN
 				{
 					// if there is a property for variable, it means the original default value was already copied, so it can be safely overridden
 					Variable.DefaultValue.Empty();
-					TargetProperty->ExportTextItem(Variable.DefaultValue, OldPropertyAddr, OldPropertyAddr, NULL, PPF_None);
+					PropertyValueToString(TargetProperty, reinterpret_cast<const uint8*>(GeneratedCDO), Variable.DefaultValue);
 				}
 			}
 
