Automation: Milestone 1
 -FPS Perf capture test
 -New Asset Creation and Duplication test
 -Import and Export test

diff --git a/Engine/Source/Editor/UnrealEd/Private/Fbx/FbxMainImport.cpp b/Engine/Source/Editor/UnrealEd/Private/Fbx/FbxMainImport.cpp
--- a/Engine/Source/Editor/UnrealEd/Private/Fbx/FbxMainImport.cpp
+++ b/Engine/Source/Editor/UnrealEd/Private/Fbx/FbxMainImport.cpp
@@ -144,6 +144,13 @@ FBXImportOptions* GetImportOptions( UnFbx::FFbxImporter* FbxImporter, UFbxImport
 			OutOperationCanceled = true;
 		}
 	}
+	else if (GIsAutomationTesting)
+	{
+		//Automation tests set ImportUI settings directly.  Just copy them over
+		UnFbx::FBXImportOptions* ImportOptions = FbxImporter->GetImportOptions();
+		ApplyImportUIToImportOptions(ImportUI, *ImportOptions);
+		return ImportOptions;
+	}
 	else
 	{
 		return FbxImporter->GetImportOptions();
diff --git a/Engine/Source/Editor/UnrealEd/Private/Tests/EditorAssetAutomationTests.cpp b/Engine/Source/Editor/UnrealEd/Private/Tests/EditorAssetAutomationTests.cpp
new file mode 100644
--- /dev/null
+++ b/Engine/Source/Editor/UnrealEd/Private/Tests/EditorAssetAutomationTests.cpp
@@ -0,0 +1,1183 @@
+// Copyright 1998-2014 Epic Games, Inc. All Rights Reserved.
+
+#include "UnrealEd.h"
+#include "Tests/AutomationTestSettings.h"
+#include "ObjectTools.h"
+#include "IAssetTools.h"
+#include "PackageTools.h"
+#include "AssetRegistryModule.h"
+#include "AssetToolsModule.h"
+#include "AssetSelection.h"
+#include "AssetEditorManager.h"
+#include "ScopedTransaction.h"
+#include "LevelEditor.h"
+#include "ModuleManager.h"
+#include "ContentBrowserModule.h"
+#include "MainFrame.h"
+
+#include "Particles/ParticleSystem.h"
+#include "Slate/SlateBrushAsset.h"
+#include "Materials/MaterialFunction.h"
+#include "Matinee/InterpData.h"
+#include "Styling/ButtonWidgetStyle.h"
+
+#include "SlateWordWrapper.h"
+#include "AutomationCommon.h"
+#include "AutomationEditorCommon.h"
+
+#define LOCTEXT_NAMESPACE "EditorAssetAutomationTests"
+
+DEFINE_LOG_CATEGORY_STATIC(LogEditorAssetAutomationTests, Log, All);
+
+namespace AssetAutomationCommon
+{
+	/**
+	* Nulls out references to a given object
+	*
+	* @param InObject - Object to null references to
+	*/
+	static void NullReferencesToObject(UObject* InObject)
+	{
+		TArray<UObject*> ReplaceableObjects;
+		TMap<UObject*, UObject*> ReplacementMap;
+		ReplacementMap.Add(InObject, NULL);
+		ReplacementMap.GenerateKeyArray(ReplaceableObjects);
+
+		// Find all the properties (and their corresponding objects) that refer to any of the objects to be replaced
+		TMap< UObject*, TArray<UProperty*> > ReferencingPropertiesMap;
+		for (FObjectIterator ObjIter; ObjIter; ++ObjIter)
+		{
+			UObject* CurObject = *ObjIter;
+
+			// Find the referencers of the objects to be replaced
+			FFindReferencersArchive FindRefsArchive(CurObject, ReplaceableObjects);
+
+			// Inform the object referencing any of the objects to be replaced about the properties that are being forcefully
+			// changed, and store both the object doing the referencing as well as the properties that were changed in a map (so that
+			// we can correctly call PostEditChange later)
+			TMap<UObject*, int32> CurNumReferencesMap;
+			TMultiMap<UObject*, UProperty*> CurReferencingPropertiesMMap;
+			if (FindRefsArchive.GetReferenceCounts(CurNumReferencesMap, CurReferencingPropertiesMMap) > 0)
+			{
+				TArray<UProperty*> CurReferencedProperties;
+				CurReferencingPropertiesMMap.GenerateValueArray(CurReferencedProperties);
+				ReferencingPropertiesMap.Add(CurObject, CurReferencedProperties);
+				for (TArray<UProperty*>::TConstIterator RefPropIter(CurReferencedProperties); RefPropIter; ++RefPropIter)
+				{
+					CurObject->PreEditChange(*RefPropIter);
+				}
+			}
+
+		}
+
+		// Iterate over the map of referencing objects/changed properties, forcefully replacing the references and then
+		// alerting the referencing objects the change has completed via PostEditChange
+		int32 NumObjsReplaced = 0;
+		for (TMap< UObject*, TArray<UProperty*> >::TConstIterator MapIter(ReferencingPropertiesMap); MapIter; ++MapIter)
+		{
+			++NumObjsReplaced;
+
+			UObject* CurReplaceObj = MapIter.Key();
+			const TArray<UProperty*>& RefPropArray = MapIter.Value();
+
+			FArchiveReplaceObjectRef<UObject> ReplaceAr(CurReplaceObj, ReplacementMap, false, true, false);
+
+			for (TArray<UProperty*>::TConstIterator RefPropIter(RefPropArray); RefPropIter; ++RefPropIter)
+			{
+				FPropertyChangedEvent PropertyEvent(*RefPropIter);
+				CurReplaceObj->PostEditChangeProperty(PropertyEvent);
+			}
+
+			if (!CurReplaceObj->HasAnyFlags(RF_Transient) && CurReplaceObj->GetOutermost() != GetTransientPackage())
+			{
+				if (!CurReplaceObj->RootPackageHasAnyFlags(PKG_CompiledIn))
+				{
+					CurReplaceObj->MarkPackageDirty();
+				}
+			}
+		}
+	}
+}
+
+/**
+* Container for items related to the create asset test
+*/
+namespace CreateAssetHelper
+{
+	/**
+	* Gets the base path for this asset
+	*/
+	static FString GetGamePath()
+	{
+		return TEXT("/Temp/Automation/Transient/Automation_AssetCreationDuplication");
+	}
+
+	/**
+	* Gets the full path to the folder on disk
+	*/
+	static FString GetFullPath()
+	{
+		return FPackageName::FilenameToLongPackageName(FPaths::AutomationTransientDir() + TEXT("Automation_AssetCreationDuplication"));
+	}
+
+	/** 
+	* Tracks success counts for each stage of the create / duplicate asset test
+	*/
+	struct FCreateAssetStats
+	{
+		//Total number of assets
+		uint32 NumTotalAssets;
+
+		//Number of assets created
+		uint32 NumCreated;
+
+		//Number of assets saved
+		uint32 NumSaved;
+
+		//Number of duplicates saved
+		uint32 NumDuplicatesSaved;
+
+		//Number of assets duplicated
+		uint32 NumDuplicated;
+
+		//Number of assets deleted
+		uint32 NumDeleted;
+
+		/**
+		* Constructor
+		*/
+		FCreateAssetStats() :
+			NumTotalAssets(0),
+			NumCreated(0),
+			NumSaved(0),
+			NumDuplicatesSaved(0),
+			NumDuplicated(0),
+			NumDeleted(0)
+		{
+		}
+	};
+
+	/**
+	* Handles creating, duplicating, and saving assets
+	*/
+	struct FCreateAssetInfo
+	{
+		//The name to use for this asset
+		FString AssetName;
+
+		//The location this asset will be created at
+		FString AssetPath;
+
+		//The class of the asset
+		UClass* Class;
+
+		//The factory to use to create this asset
+		UFactory* Factory;
+
+		//The asset that was created
+		UObject* CreatedAsset;
+
+		//The package that contains the asset
+		UPackage* AssetPackage;
+
+		//The duplicated asset
+		UObject* DuplicatedAsset;
+
+		//The package that contains the duplicated asset
+		UPackage* DuplicatedPackage;
+
+		//Pointer to the asset test stats
+		TSharedPtr<FCreateAssetStats> TestStats;
+
+		/**
+		* Constructor
+		*/
+		FCreateAssetInfo(const FString& InAssetName, const FString& InAssetPath, UClass* InClass, UFactory* InFactory, TSharedPtr<FCreateAssetStats> InStats) :
+			AssetName(InAssetName),
+			AssetPath(InAssetPath),
+			Class(InClass),
+			Factory(InFactory),
+			CreatedAsset(NULL),
+			AssetPackage(NULL),
+			DuplicatedAsset(NULL),
+			DuplicatedPackage(NULL),
+			TestStats(InStats)
+		{
+		}
+
+		/**
+		* Creates the new item
+		*/
+		void CreateAsset()
+		{
+			const FString PackageName = AssetPath + TEXT("/") + AssetName;
+			AssetPackage = CreatePackage(NULL, *PackageName);
+			EObjectFlags Flags = RF_Public | RF_Standalone;
+
+			CreatedAsset = Factory->FactoryCreateNew(Class, AssetPackage, FName(*AssetName), Flags, NULL, GWarn);
+
+			if (CreatedAsset)
+			{
+				// Notify the asset registry
+				FAssetRegistryModule::AssetCreated(CreatedAsset);
+
+				// Mark the package dirty...
+				AssetPackage->MarkPackageDirty();
+
+				TestStats->NumCreated++;
+				UE_LOG(LogEditorAssetAutomationTests, Display, TEXT("Created asset %s (%s)"), *AssetName, *Class->GetName());
+			}
+			else
+			{
+				UE_LOG(LogEditorAssetAutomationTests, Error, TEXT("Unable to create asset of type %s"), *Class->GetName());
+			}
+		}
+
+		/**
+		* Saves the created asset
+		*/
+		void SaveNewAsset()
+		{
+			if (AssetPackage && CreatedAsset)
+			{
+				AssetPackage->SetDirtyFlag(true);
+				const FString PackagePath = FString::Printf(TEXT("%s/%s"), *GetGamePath(), *AssetName);
+				if (UPackage::SavePackage(AssetPackage, NULL, RF_Standalone, *FPackageName::LongPackageNameToFilename(PackagePath, FPackageName::GetAssetPackageExtension()), GError, nullptr, false, true, SAVE_NoError))
+				{
+					TestStats->NumSaved++;
+					UE_LOG(LogEditorAssetAutomationTests, Display, TEXT("Saved asset %s (%s)"), *CreatedAsset->GetName(), *Class->GetName());
+				}
+				else
+				{
+					UE_LOG(LogEditorAssetAutomationTests, Display, TEXT("Unable to save asset %s (%s)"), *CreatedAsset->GetName(), *Class->GetName());
+				}
+			}
+		}
+
+		/**
+		* Saves the duplicated asset
+		*/
+		void SaveDuplicatedAsset()
+		{
+			if (DuplicatedPackage && DuplicatedAsset)
+			{
+				DuplicatedPackage->SetDirtyFlag(true);
+				const FString PackagePath = FString::Printf(TEXT("%s/%s_Copy"), *GetGamePath(), *AssetName);
+				if (UPackage::SavePackage(DuplicatedPackage, NULL, RF_Standalone, *FPackageName::LongPackageNameToFilename(PackagePath, FPackageName::GetAssetPackageExtension()), GError, nullptr, false, true, SAVE_NoError))
+				{
+					TestStats->NumDuplicatesSaved++;
+					UE_LOG(LogEditorAssetAutomationTests, Display, TEXT("Saved asset %s (%s)"), *DuplicatedAsset->GetName(), *Class->GetName());
+				}
+				else
+				{
+					UE_LOG(LogEditorAssetAutomationTests, Display, TEXT("Unable to save asset %s (%s)"), *DuplicatedAsset->GetName(), *Class->GetName());
+				}
+			}
+		}
+
+		/**
+		* Duplicates the asset
+		*/
+		void DuplicateAsset()
+		{
+			if (AssetPackage && CreatedAsset)
+			{
+				const FString NewObjectName = FString::Printf(TEXT("%s_Copy"), *AssetName);
+				const FString NewPackageName = FString::Printf(TEXT("%s/%s"), *GetGamePath(), *NewObjectName);
+
+				// Make sure the referenced object is deselected before duplicating it.
+				GEditor->GetSelectedObjects()->Deselect(CreatedAsset);
+
+				// Duplicate the asset
+				DuplicatedPackage = CreatePackage(NULL, *NewPackageName);
+				DuplicatedAsset = StaticDuplicateObject(CreatedAsset, DuplicatedPackage, *NewObjectName);
+
+				if (DuplicatedAsset)
+				{
+					DuplicatedAsset->MarkPackageDirty();
+
+					// Notify the asset registry
+					FAssetRegistryModule::AssetCreated(DuplicatedAsset);
+
+					TestStats->NumDuplicated++;
+					UE_LOG(LogEditorAssetAutomationTests, Display, TEXT("Duplicated asset %s to %s (%s)"), *AssetName, *NewObjectName, *Class->GetName());
+				}
+				else
+				{
+					UE_LOG(LogEditorAssetAutomationTests, Error, TEXT("Failed to duplicate asset %s (%s)"), *AssetName, *Class->GetName());
+				}
+			}
+		}
+
+		/**
+		* Deletes the asset
+		*/
+		void DeleteAsset()
+		{
+			if (CreatedAsset)
+			{
+				bool bSuccessful = false;
+
+				bSuccessful = ObjectTools::DeleteSingleObject(CreatedAsset, false);
+
+				//If we failed to delete this object manually clear any references and try again
+				if (!bSuccessful)
+				{
+					//Clear references to the object so we can delete it
+					AssetAutomationCommon::NullReferencesToObject(CreatedAsset);
+
+					bSuccessful = ObjectTools::DeleteSingleObject(CreatedAsset, false);
+				}
+
+				//Delete the package
+				if (bSuccessful)
+				{
+					FString PackageFilename;
+					if (FPackageName::DoesPackageExist(AssetPackage->GetName(), NULL, &PackageFilename))
+					{
+						TArray<UPackage*> PackagesToDelete;
+						PackagesToDelete.Add(AssetPackage);
+						// Let the package auto-saver know that it needs to ignore the deleted packages
+						GUnrealEd->GetPackageAutoSaver().OnPackagesDeleted(PackagesToDelete);
+
+						AssetPackage->SetDirtyFlag(false);
+
+						// Unload the packages and collect garbage.
+						PackageTools::UnloadPackages(PackagesToDelete);
+
+						IFileManager::Get().Delete(*PackageFilename);
+
+						TestStats->NumDeleted++;
+						UE_LOG(LogEditorAssetAutomationTests, Display, TEXT("Deleted asset %s (%s)"), *AssetName, *Class->GetName());
+					}
+				}
+				else
+				{
+					UE_LOG(LogEditorAssetAutomationTests, Error, TEXT("Unable to delete asset: %s (%s)"), *AssetName, *Class->GetName());
+				}
+			}
+		}
+	};
+}
+
+/**
+* Latent command to create an asset
+*/
+DEFINE_LATENT_AUTOMATION_COMMAND_ONE_PARAMETER(FCreateNewAssetCommand, TSharedPtr<CreateAssetHelper::FCreateAssetInfo>, AssetInfo);
+bool FCreateNewAssetCommand::Update()
+{
+	AssetInfo->CreateAsset();
+	return true;
+}
+
+
+/**
+* Latent command to save an asset
+*/
+DEFINE_LATENT_AUTOMATION_COMMAND_ONE_PARAMETER(FSaveNewAssetCommand, TSharedPtr<CreateAssetHelper::FCreateAssetInfo>, AssetInfo);
+bool FSaveNewAssetCommand::Update()
+{
+	AssetInfo->SaveNewAsset();
+	return true;
+}
+
+/**
+* Latent command to save a duplicated asset
+*/
+DEFINE_LATENT_AUTOMATION_COMMAND_ONE_PARAMETER(FSaveDuplicateAssetCommand, TSharedPtr<CreateAssetHelper::FCreateAssetInfo>, AssetInfo);
+bool FSaveDuplicateAssetCommand::Update()
+{
+	AssetInfo->SaveDuplicatedAsset();
+	return true;
+}
+
+/**
+* Latent command to duplicate an asset
+*/
+DEFINE_LATENT_AUTOMATION_COMMAND_ONE_PARAMETER(FDuplicateAssetCommand, TSharedPtr<CreateAssetHelper::FCreateAssetInfo>, AssetInfo);
+bool FDuplicateAssetCommand::Update()
+{
+	AssetInfo->DuplicateAsset();
+	return true;
+}
+
+/**
+* Latent command to delete an asset
+*/
+DEFINE_LATENT_AUTOMATION_COMMAND_ONE_PARAMETER(FDeleteAssetCommand, TSharedPtr<CreateAssetHelper::FCreateAssetInfo>, AssetInfo);
+bool FDeleteAssetCommand::Update()
+{
+	AssetInfo->DeleteAsset();
+	return true;
+}
+
+/**
+* Latent command to clear editor references to temporary objects
+*/
+DEFINE_LATENT_AUTOMATION_COMMAND(FClearEditorReferencesCommand);
+bool FClearEditorReferencesCommand::Update()
+{
+	// Deselect all
+	GEditor->SelectNone(false, true);
+
+	// Clear the transaction buffer so we aren't referencing the new objects
+	GUnrealEd->ResetTransaction(FText::FromString(TEXT("FAssetEditorTest")));
+
+	return true;
+}
+
+/**
+* Latent command log the asset creation stats
+*/
+DEFINE_LATENT_AUTOMATION_COMMAND_ONE_PARAMETER(FLogAssetCreationStatsCommand, TSharedPtr<CreateAssetHelper::FCreateAssetStats>, BuildStats);
+bool FLogAssetCreationStatsCommand::Update()
+{
+	UE_LOG(LogEditorAssetAutomationTests, Display, TEXT(" "));
+	UE_LOG(LogEditorAssetAutomationTests, Display, TEXT("Test Summary:"));
+	UE_LOG(LogEditorAssetAutomationTests, Display, TEXT("%i of %i assets were created successfully"),			BuildStats->NumCreated,			BuildStats->NumTotalAssets);
+	UE_LOG(LogEditorAssetAutomationTests, Display, TEXT("%i of %i assets were saved successfully"),				BuildStats->NumSaved,			BuildStats->NumTotalAssets);
+	UE_LOG(LogEditorAssetAutomationTests, Display, TEXT("%i of %i assets were duplicated successfully"),		BuildStats->NumDuplicated,		BuildStats->NumTotalAssets);
+	UE_LOG(LogEditorAssetAutomationTests, Display, TEXT("%i of %i duplicated assets were saved successfully"),	BuildStats->NumDuplicatesSaved, BuildStats->NumTotalAssets);
+	UE_LOG(LogEditorAssetAutomationTests, Display, TEXT("%i of %i assets were deleted successfully"),			BuildStats->NumDeleted,			BuildStats->NumTotalAssets);
+	return true;
+}
+
+//Macro to create a factory and queue up a command to create a given asset type
+#define ASSET_TEST_CREATE( TAssetClass, TFactoryClass, NamePrefix, ExtraCommands ) \
+{ \
+	FString NameString = FString::Printf(TEXT("%s_%s"), TEXT(#NamePrefix), *CurrentTimestamp); \
+	TFactoryClass* FactoryInst = ConstructObject<TFactoryClass>(TFactoryClass::StaticClass()); \
+	ExtraCommands \
+	TSharedPtr<CreateAssetHelper::FCreateAssetInfo> CreateInfo = MakeShareable(new CreateAssetHelper::FCreateAssetInfo(NameString, GamePath, TAssetClass::StaticClass(), FactoryInst, BuildStats)); \
+	AssetInfos.Add(CreateInfo); \
+	ADD_LATENT_AUTOMATION_COMMAND(FCreateNewAssetCommand(CreateInfo)); \
+}
+
+//Macro to create a factory by name and queue up a command to create an asset
+#define ASSET_TEST_CREATE_BY_NAME(TAssetClassName, TFactoryClassName, NamePrefix, ExtraCommands) \
+{ \
+	FString NameString = FString::Printf(TEXT("%s_%s"), TEXT(#NamePrefix), *CurrentTimestamp); \
+	UClass* FactoryClass = StaticLoadClass(UFactory::StaticClass(), NULL, TEXT(#TFactoryClassName), NULL, LOAD_None, NULL); \
+	if (FactoryClass) \
+	{ \
+		UFactory* FactoryInst = ConstructObject<UFactory>(FactoryClass); \
+		ExtraCommands \
+		UClass* AssetClass = StaticLoadClass(UObject::StaticClass(), NULL, TEXT(#TAssetClassName), NULL, LOAD_None, NULL); \
+		TSharedPtr<CreateAssetHelper::FCreateAssetInfo> CreateInfo = MakeShareable(new CreateAssetHelper::FCreateAssetInfo(NameString, GamePath, AssetClass, FactoryInst, BuildStats)); \
+		AssetInfos.Add(CreateInfo); \
+		ADD_LATENT_AUTOMATION_COMMAND(FCreateNewAssetCommand(CreateInfo)); \
+	} \
+	else \
+	{ \
+		UE_LOG(LogEditorAssetAutomationTests,Error,TEXT("Couldn't find factory class %s"),TEXT(#TFactoryClassName)); \
+	} \
+}
+
+
+/**
+* Automation test for creating, saving, and duplicating assets
+*/
+IMPLEMENT_SIMPLE_AUTOMATION_TEST(FAssetEditorTest, "Editor.Content.Asset Creation and Duplication", EAutomationTestFlags::ATF_Editor);
+
+bool FAssetEditorTest::RunTest(const FString& Parameters)
+{
+	FAssetRegistryModule& AssetRegistryModule = FModuleManager::LoadModuleChecked<FAssetRegistryModule>(TEXT("AssetRegistry"));
+	FContentBrowserModule& ContentBrowserModule = FModuleManager::LoadModuleChecked<FContentBrowserModule>("ContentBrowser");
+
+	const FString FullPath = CreateAssetHelper::GetFullPath();
+	const FString GamePath = CreateAssetHelper::GetGamePath();
+
+	//Create the folder if it doesn't already exist
+	if (!IFileManager::Get().DirectoryExists(*FullPath))
+	{
+		//Make the new folder
+		IFileManager::Get().MakeDirectory(*FullPath, true);
+
+		// Add the path to the asset registry
+		AssetRegistryModule.Get().AddPath(GamePath);
+
+		// Notify 'asset path changed' delegate
+		FContentBrowserModule::FOnAssetPathChanged& PathChangedDelegate = ContentBrowserModule.GetOnAssetPathChanged();
+		if (PathChangedDelegate.IsBound())
+		{
+			PathChangedDelegate.Broadcast(GamePath);
+		}
+	}
+
+	//Timestamp
+	FString CurrentTimestamp = FPlatformTime::StrTimestamp();
+	CurrentTimestamp = CurrentTimestamp.Replace(TEXT("/"), TEXT(""));
+	CurrentTimestamp = CurrentTimestamp.Replace(TEXT(":"), TEXT(""));
+	CurrentTimestamp = CurrentTimestamp.Replace(TEXT(" "), TEXT("_"));
+
+	//Skeleton - Grab the first available loaded skeleton
+	TArray<FAssetData> AllSkeletons;
+	AssetRegistryModule.Get().GetAssetsByClass(USkeleton::StaticClass()->GetFName(), AllSkeletons);
+	USkeleton* FirstSkeleton = NULL;
+	for (int32 SkelIndex = 0; SkelIndex < AllSkeletons.Num(); ++SkelIndex)
+	{
+		if (AllSkeletons[SkelIndex].IsAssetLoaded())
+		{
+			FirstSkeleton = (USkeleton*)(AllSkeletons[SkelIndex].GetAsset());
+			break;
+		}
+	}
+
+	//Holds info on each asset we are creating
+	TArray< TSharedPtr<CreateAssetHelper::FCreateAssetInfo> > AssetInfos;
+	TSharedPtr<CreateAssetHelper::FCreateAssetStats> BuildStats = (MakeShareable(new CreateAssetHelper::FCreateAssetStats()));
+
+	//Queue creating the different kinds of assets
+	ASSET_TEST_CREATE(UBlueprint, UBlueprintFactory, BP, FactoryInst->ParentClass = AActor::StaticClass();)
+	ASSET_TEST_CREATE(UMaterial, UMaterialFactoryNew, MAT, )
+	ASSET_TEST_CREATE(UParticleSystem, UParticleSystemFactoryNew, PS, )
+	ASSET_TEST_CREATE(UAimOffsetBlendSpace, UAimOffsetBlendSpaceFactoryNew, AO, FactoryInst->TargetSkeleton = FirstSkeleton;)
+	ASSET_TEST_CREATE(UAimOffsetBlendSpace1D, UAimOffsetBlendSpaceFactory1D, AO1D, FactoryInst->TargetSkeleton = FirstSkeleton;)
+	ASSET_TEST_CREATE(UAnimBlueprint, UAnimBlueprintFactory, AB, FactoryInst->ParentClass = UAnimInstance::StaticClass(); FactoryInst->TargetSkeleton = FirstSkeleton;)
+	ASSET_TEST_CREATE(UAnimComposite, UAnimCompositeFactory, AC, FactoryInst->TargetSkeleton = FirstSkeleton;)
+	ASSET_TEST_CREATE(UAnimMontage, UAnimMontageFactory, AM, FactoryInst->TargetSkeleton = FirstSkeleton;)
+	ASSET_TEST_CREATE(UBlendSpace, UBlendSpaceFactoryNew, BS, FactoryInst->TargetSkeleton = FirstSkeleton;)
+	ASSET_TEST_CREATE(UBlendSpace1D, UBlendSpaceFactory1D, BS1D, FactoryInst->TargetSkeleton = FirstSkeleton;)
+	ASSET_TEST_CREATE(UTextureRenderTargetCube, UTextureRenderTargetCubeFactoryNew, CRT, )
+	ASSET_TEST_CREATE(UFont, UTrueTypeFontFactory, F, )
+	ASSET_TEST_CREATE(UMaterialFunction, UMaterialFunctionFactoryNew, MF, )
+	ASSET_TEST_CREATE(UMaterialInstanceConstant, UMaterialInstanceConstantFactoryNew, MI, )
+	ASSET_TEST_CREATE(UMaterialParameterCollection, UMaterialParameterCollectionFactoryNew, MPC, )
+	ASSET_TEST_CREATE(UTextureRenderTarget2D, UTextureRenderTargetFactoryNew, RT, )
+	ASSET_TEST_CREATE(UDialogueVoice, UDialogueVoiceFactory, DV, )
+	ASSET_TEST_CREATE(UDialogueWave, UDialogueWaveFactory, DW, )
+	ASSET_TEST_CREATE(UReverbEffect, UReverbEffectFactory, RE, )
+	ASSET_TEST_CREATE(USoundAttenuation, USoundAttenuationFactory, SA, )
+	ASSET_TEST_CREATE(USoundClass, USoundClassFactory, SC, )
+	ASSET_TEST_CREATE(USoundCue, USoundCueFactoryNew, Scue, )
+	ASSET_TEST_CREATE(USoundMix, USoundMixFactory, SM, )
+	ASSET_TEST_CREATE(UPhysicalMaterial, UPhysicalMaterialFactoryNew, PM, )
+	ASSET_TEST_CREATE(USlateBrushAsset, USlateBrushAssetFactory, SB, )
+	ASSET_TEST_CREATE(USlateWidgetStyleAsset, USlateWidgetStyleAssetFactory, SWS, FactoryInst->StyleType = UButtonWidgetStyle::StaticClass();)
+	ASSET_TEST_CREATE_BY_NAME(AIModule.BehaviorTree, BehaviorTreeEditor.BehaviorTreeFactory, BT, )
+	ASSET_TEST_CREATE(UBlueprint, UBlueprintFunctionLibraryFactory, BFL, )
+	ASSET_TEST_CREATE(UBlueprint, UBlueprintMacroFactory, MPL, FactoryInst->ParentClass = AActor::StaticClass();)
+	ASSET_TEST_CREATE(UCameraAnim, UCameraAnimFactory, CA, )
+	ASSET_TEST_CREATE(UCurveBase, UCurveFactory, C, FactoryInst->CurveClass = UCurveFloat::StaticClass();)
+	UClass* GameplayAbilityClass = StaticLoadClass(UObject::StaticClass(), NULL, TEXT("GameplayAbilities.GameplayAbilitySet"), NULL, LOAD_None, NULL);
+	ASSET_TEST_CREATE(UDataAsset, UDataAssetFactory, DA, FactoryInst->DataAssetClass = GameplayAbilityClass;)
+	ASSET_TEST_CREATE(UUserDefinedEnum, UEnumFactory, Enum, )
+	ASSET_TEST_CREATE(UForceFeedbackEffect, UForceFeedbackEffectFactory, FFE, )
+	ASSET_TEST_CREATE(UInterpData, UInterpDataFactoryNew, MD, )
+	ASSET_TEST_CREATE(UObjectLibrary, UObjectLibraryFactory, OL, )
+	ASSET_TEST_CREATE(UUserDefinedStruct, UStructureFactory, S, )
+	ASSET_TEST_CREATE(UTouchInterface, UTouchInterfaceFactory, TIS, )
+
+	//Record how many assets we are testing
+	BuildStats->NumTotalAssets = AssetInfos.Num();
+
+	//Save new assets
+	for (int32 i = 0; i < AssetInfos.Num(); i++)
+	{
+		ADD_LATENT_AUTOMATION_COMMAND(FSaveNewAssetCommand(AssetInfos[i]));
+	}
+
+	//Duplicate new assets
+	for (int32 i = 0; i < AssetInfos.Num(); i++)
+	{
+		ADD_LATENT_AUTOMATION_COMMAND(FDuplicateAssetCommand(AssetInfos[i]));
+	}
+
+	//Save duplicates
+	for (int32 i = 0; i < AssetInfos.Num(); i++)
+	{
+		ADD_LATENT_AUTOMATION_COMMAND(FSaveDuplicateAssetCommand(AssetInfos[i]));
+	}
+
+	ADD_LATENT_AUTOMATION_COMMAND(FClearEditorReferencesCommand());
+
+	//Delete Original
+	for (int32 i = 0; i < AssetInfos.Num(); i++)
+	{
+		ADD_LATENT_AUTOMATION_COMMAND(FDeleteAssetCommand(AssetInfos[i]));
+	}
+
+	ADD_LATENT_AUTOMATION_COMMAND(FLogAssetCreationStatsCommand(BuildStats));
+
+	return true;
+}
+
+#undef ASSET_TEST_CREATE
+#undef ASSET_TEST_CREATE_BY_NAME
+
+/**
+* Namespace for helper items for the import / export asset test
+*/
+namespace ImportExportAssetHelper
+{
+	// How long to wait for the asset editor window to open.
+	static int32 MaxWaitForEditorTicks = 5;
+
+	//State flags for the FAssetInfo class
+	namespace EState
+	{
+		enum Type
+		{
+			Import,
+			OpenEditor,
+			WaitForEditor,
+			Screenshot,
+			Export,
+			Done
+		};
+	}
+
+	/**
+	* gets a factory class based off an asset file extension
+	*
+	* @param AssetExtension - The file extension to use to find a supporting UFactory
+	*/
+	static UClass* GetFactoryClassForType(const FString& AssetExtension)
+	{
+		// First instantiate one factory for each file extension encountered that supports the extension
+		for (TObjectIterator<UClass> ClassIt; ClassIt; ++ClassIt)
+		{
+			if ((*ClassIt)->IsChildOf(UFactory::StaticClass()) && !((*ClassIt)->HasAnyClassFlags(CLASS_Abstract)))
+			{
+				UFactory* Factory = Cast<UFactory>((*ClassIt)->GetDefaultObject());
+				if (Factory->bEditorImport && Factory->ValidForCurrentGame())
+				{
+					TArray<FString> FactoryExtensions;
+					Factory->GetSupportedFileExtensions(FactoryExtensions);
+
+					// Case insensitive string compare with supported formats of this factory
+					if (FactoryExtensions.Contains(AssetExtension))
+					{
+						return *ClassIt;
+					}
+				}
+			}
+		}
+
+		return NULL;
+	}
+
+	/**
+	* Applies settings to an object by finding UProperties by name and calling ImportText
+	*
+	* @param InObject - The object to search for matching properties
+	* @param PropertyChain - The list UProperty names recursively to search through
+	* @param Value - The value to import on the found property
+	*/
+	static void ApplyCustomFactorySetting(UObject* InObject, TArray<FString>& PropertyChain, const FString& Value)
+	{
+		const FString PropertyName = PropertyChain[0];
+		PropertyChain.RemoveAt(0);
+
+		UProperty* TargetProperty = FindField<UProperty>(InObject->GetClass(), *PropertyName);
+		if (TargetProperty)
+		{
+			if (PropertyChain.Num() == 0)
+			{
+				TargetProperty->ImportText(*Value, TargetProperty->ContainerPtrToValuePtr<uint8>(InObject), 0, InObject);
+			}
+			else
+			{
+				UStructProperty* StructProperty = Cast<UStructProperty>(TargetProperty);
+				UObjectProperty* ObjectProperty = Cast<UObjectProperty>(TargetProperty);
+
+				UObject* SubObject = NULL;
+				bool bValidPropertyType = true;
+
+				if (StructProperty)
+				{
+					SubObject = StructProperty->Struct;
+				}
+				else if (ObjectProperty)
+				{
+					SubObject = ObjectProperty->GetObjectPropertyValue(ObjectProperty->ContainerPtrToValuePtr<UObject>(InObject));
+				}
+				else
+				{
+					//Unknown nested object type
+					bValidPropertyType = false;
+					UE_LOG(LogEditorAssetAutomationTests, Error, TEXT("ERROR: Unknown nested object type for property: %s"), *PropertyName);
+				}
+
+				if (SubObject)
+				{
+					ApplyCustomFactorySetting(SubObject, PropertyChain, Value);
+				}
+				else if (bValidPropertyType)
+				{
+					UE_LOG(LogEditorAssetAutomationTests, Error, TEXT("Error accessing null property: %s"), *PropertyName);
+				}
+			}
+		}
+		else
+		{
+			UE_LOG(LogEditorAssetAutomationTests, Error, TEXT("ERROR: Could not find factory property: %s"), *PropertyName);
+		}
+	}
+
+	/**
+	* Applies the custom factory settings
+	*
+	* @param InFactory - The factory to apply custom settings to
+	* @param FactorySettings - An array of custom settings to apply to the factory
+	*/
+	static void ApplyCustomFactorySettings(UFactory* InFactory, const TArray<FImportFactorySettingValues>& FactorySettings)
+	{
+		for (int32 i = 0; i < FactorySettings.Num(); ++i)
+		{
+			if (FactorySettings[i].SettingName.Len() > 0 && FactorySettings[i].Value.Len() > 0)
+			{
+				TArray<FString> PropertyChain;
+				FactorySettings[i].SettingName.ParseIntoArray(&PropertyChain, TEXT("."), false);
+				ApplyCustomFactorySetting(InFactory, PropertyChain, FactorySettings[i].Value);
+			}
+		}
+	}
+
+	/**
+	* Import test report for a single asset
+	*/
+	struct FAssetImportReport
+	{
+		//The Asset file name
+		FString AssetName;
+
+		//If the asset imported successfuly
+		bool bImportSuccessful;
+
+		//If the export step was skipped
+		bool bSkippedExport;
+
+		//If the asset exported successfuly
+		bool bExportSuccessful;
+
+		//The size of the exported file
+		int32 FileSize;
+
+		/**
+		* Constructor
+		*/
+		FAssetImportReport() :
+			AssetName(),
+			bImportSuccessful(false),
+			bSkippedExport(false),
+			bExportSuccessful(false),
+			FileSize(0)
+		{
+		}
+	};
+
+	/**
+	* Import test stats
+	*/
+	struct FAssetImportStats
+	{
+		//List of import reports
+		TArray<FAssetImportReport> Reports;
+	};
+
+	/**
+	* Namespace for helper items for the import / export asset test
+	*/
+	struct FAssetInfo
+	{
+		//Path to the file we are importing
+		FString ImportPath;
+		
+		//The file extension to use when exporting this asset
+		FString ExportExtension;
+
+		//The current state this asset is in
+		EState::Type State;
+
+		//A pointer to the asset we imported
+		UObject* ImportedAsset;
+
+		//A list of custom settings to apply to our import factory
+		TArray<FImportFactorySettingValues> FactorySettings;
+
+		//How many frames we have waited for the asset editor
+		int32 WaitingForEditorCount;
+
+		//If we should skip the export step
+		bool bSkipExport;
+
+		//Pointer to the execution info of this test
+		FAutomationTestExecutionInfo* TestExecutionInfo;
+
+		//The number of existing errors, warnings, and logs when the command started
+		int32 ErrorStartCount;
+		int32 WarningStartCount;
+		int32 LogStartCount;
+
+		//Shared list of test results
+		TSharedPtr<FAssetImportStats> TestStats;
+
+		//Test report for this asset
+		FAssetImportReport TestReport;
+
+		/**
+		* Constructor
+		*/
+		FAssetInfo(const FEditorImportExportTestDefinition& InTestDef, FAutomationTestExecutionInfo* InExecutionInfo, TSharedPtr<FAssetImportStats> InStats) :
+			ImportPath(InTestDef.ImportFilePath.FilePath),
+			ExportExtension(InTestDef.ExportFileExtension),
+			State(EState::Import),
+			ImportedAsset(NULL),
+			FactorySettings(InTestDef.FactorySettings),
+			WaitingForEditorCount(0),
+			bSkipExport(InTestDef.bSkipExport),
+			TestExecutionInfo(InExecutionInfo),
+			ErrorStartCount(0),
+			WarningStartCount(0),
+			LogStartCount(0),
+			TestStats(InStats),
+			TestReport()
+		{
+		}
+
+		/**
+		* Updates the import state
+		*/
+		bool Update()
+		{
+			switch (State)
+			{
+			case EState::Import:
+				ImportAsset();
+				break;
+			case EState::OpenEditor:
+				OpenEditor();
+				break;
+			case EState::WaitForEditor:
+				CheckEditor();
+				break;
+			case EState::Screenshot:
+				TakeScreenshot();
+				break;
+			case EState::Export:
+				ExportAsset();
+				break;
+			default:
+				break;
+			}
+
+			//Clean up the asset if we are done
+			if (State == EState::Done)
+			{
+				if (ImportedAsset)
+				{
+					DeleteAsset();
+				}
+
+				//Report the result
+				TestStats->Reports.Add(TestReport);
+
+				UpdateExecutionLogs();
+
+				return true;
+			}
+
+			return false;
+		}
+
+		/**
+		* Imports the asset from disk
+		*/
+		void ImportAsset()
+		{
+			//Default to failed
+			State = EState::Done;
+
+			//Set the starting log counts
+			if (TestExecutionInfo)
+			{
+				ErrorStartCount = TestExecutionInfo->Errors.Num();
+				WarningStartCount = TestExecutionInfo->Warnings.Num();
+				LogStartCount = TestExecutionInfo->LogItems.Num();
+			}
+
+			TestReport.AssetName = FPaths::GetCleanFilename(ImportPath);
+
+			//Get the factory
+			const FString FileExtension = FPaths::GetExtension(ImportPath);
+			UClass* FactoryClass = GetFactoryClassForType(FileExtension);
+
+			if (FactoryClass)
+			{
+				GWarn->BeginSlowTask(LOCTEXT("ImportSlowTask", "Importing"), true);
+
+				UFactory* ImportFactory = ConstructObject<UFactory>(FactoryClass);
+
+				ImportFactory->ConfigureProperties();
+
+				//Apply any custom settings to the factory
+				ApplyCustomFactorySettings(ImportFactory, FactorySettings);
+
+				FString Name = ObjectTools::SanitizeObjectName(FPaths::GetBaseFilename(ImportPath));
+				FString PackageName = FString::Printf(TEXT("/Game/Automation_Imports/%s"), *Name);
+
+				UPackage* Pkg = CreatePackage(NULL, *PackageName);
+				if (Pkg)
+				{
+					// Make sure the destination package is loaded
+					Pkg->FullyLoad();
+
+					UClass* ImportAssetType = ImportFactory->ResolveSupportedClass();
+					bool bDummy = false;
+					ImportedAsset = UFactory::StaticImportObject(ImportAssetType, Pkg, FName(*Name), RF_Public | RF_Standalone, bDummy, *ImportPath, NULL, ImportFactory, NULL, GWarn, 0);
+
+					if (ImportedAsset)
+					{
+						TestReport.bImportSuccessful = true;
+						State = EState::OpenEditor;
+					}
+					else
+					{
+						UE_LOG(LogEditorAssetAutomationTests, Error, TEXT("Failed to import asset %s using factory %s!"), *Name, *ImportFactory->GetName());
+					}
+				}
+				else
+				{
+					UE_LOG(LogEditorAssetAutomationTests, Error, TEXT("Failed to create a package for %s!"), *Name);
+				}
+
+				GWarn->EndSlowTask();
+			}
+			else
+			{
+				UE_LOG(LogEditorAssetAutomationTests, Error, TEXT("Failed to find an import factory for %s!"), *FPaths::GetBaseFilename(ImportPath));
+			}
+		}
+
+		/**
+		* Opens the asset editor
+		*/
+		void OpenEditor()
+		{
+			State = EState::Done;
+
+			if (ImportedAsset)
+			{
+				if (FAssetEditorManager::Get().OpenEditorForAsset(ImportedAsset))
+				{
+					State = EState::WaitForEditor;
+				}
+				else
+				{
+					UE_LOG(LogEditorAssetAutomationTests, Error, TEXT("Failed to open the asset editor for %s!"), *ImportedAsset->GetName());
+				}
+			}
+		}
+
+		/**
+		* Wait for the asset editor window
+		*/
+		void CheckEditor()
+		{
+			TSharedPtr<SWindow> ActiveWindow = FSlateApplication::Get().GetActiveTopLevelWindow();
+			if (ActiveWindow.IsValid())
+			{
+				FString ActiveWindowTitle = ActiveWindow->GetTitle().ToString();
+
+				//Check that we have the right window (Tutorial may have opened on top of the editor)
+				if (!ActiveWindowTitle.StartsWith(ImportedAsset->GetName()))
+				{
+					//Bring the asset editor to the front
+					FAssetEditorManager::Get().FindEditorForAsset(ImportedAsset, true);
+				}
+
+				State = EState::Screenshot;
+			}
+			else
+			{
+				WaitingForEditorCount++;
+				if (WaitingForEditorCount > MaxWaitForEditorTicks)
+				{
+
+					UE_LOG(LogEditorAssetAutomationTests, Error, TEXT("Timed out waiting for editor window: %s"), *ImportedAsset->GetName());
+					State = EState::Done;
+				}
+			}
+		}
+
+		/**
+		* Take a screenshot of the editor window
+		*/
+		void TakeScreenshot()
+		{
+			TSharedPtr<SWindow> ActiveWindow = FSlateApplication::Get().GetActiveTopLevelWindow();
+			if (ActiveWindow.IsValid())
+			{
+				FString ScreenshotName;
+				const FString TestName = FString::Printf(TEXT("ImportExportTest/%s"), *ImportedAsset->GetName());
+				AutomationCommon::GetScreenshotPath(TestName, ScreenshotName, false);
+
+				TSharedRef<SWidget> WindowRef = ActiveWindow.ToSharedRef();
+
+				TArray<FColor> OutImageData;
+				FIntVector OutImageSize;
+				FSlateApplication::Get().TakeScreenshot(WindowRef, OutImageData, OutImageSize);
+				FAutomationTestFramework::GetInstance().OnScreenshotCaptured().ExecuteIfBound(OutImageSize.X, OutImageSize.Y, OutImageData, ScreenshotName);
+
+				//Close the editor
+				FAssetEditorManager::Get().CloseAllAssetEditors();
+
+				State = EState::Export;
+			}
+			else
+			{
+				UE_LOG(LogEditorAssetAutomationTests, Error, TEXT("No asset editor window found: %s"), *ImportedAsset->GetName());
+				State = EState::Done;
+			}
+		}
+
+		/**
+		* Export the asset based of the export extension
+		*/
+		void ExportAsset()
+		{
+			if (!bSkipExport)
+			{
+				FString Extension = ExportExtension;
+				if (Extension.Len() == 0)
+				{
+					//Get the extension from the imported filename
+					Extension = FPaths::GetExtension(ImportPath);
+				}
+
+				if (Extension.StartsWith(TEXT(".")))
+				{
+					Extension = Extension.RightChop(1);
+				}
+
+				//Export the asset
+				const FString ExportPath = FString::Printf(TEXT("../../../QAGame/Saved/Exports/%s.%s"), *ImportedAsset->GetName(), *Extension);
+				UExporter* ExporterToUse = UExporter::FindExporter(ImportedAsset, *Extension);
+				
+				UExporter::FExportToFileParams Params;
+				Params.Object = ImportedAsset;
+				Params.Exporter = ExporterToUse;
+				Params.Filename = *ExportPath;
+				Params.InSelectedOnly = false;
+				Params.NoReplaceIdentical = false;
+				Params.Prompt = false;
+				Params.bUseFileArchive = ImportedAsset->IsA(UPackage::StaticClass());
+				Params.WriteEmptyFiles = false;
+				if (UExporter::ExportToFileEx(Params) != 0)  //1 - success, 0 - fatal error, -1 - non fatal error
+				{
+					//Success
+					TestReport.bExportSuccessful = true;
+					TestReport.FileSize = IFileManager::Get().FileSize(*ExportPath);
+				}
+				else
+				{
+					UE_LOG(LogEditorAssetAutomationTests, Error, TEXT("Failed to export asset: %s"), *ImportedAsset->GetName());
+				}
+			}
+			else
+			{
+				TestReport.bSkippedExport = true;
+			}
+
+			State = EState::Done;
+		}
+
+		/**
+		* Delete the assset
+		*/
+		void DeleteAsset()
+		{
+			//Clear references to the object so we can delete it
+			AssetAutomationCommon::NullReferencesToObject(ImportedAsset);
+
+			//Delete the object
+			TArray<UObject*> ObjList;
+			ObjList.Add(ImportedAsset);
+			ObjectTools::ForceDeleteObjects(ObjList, false);
+
+			ImportedAsset = NULL;
+		}
+
+		void UpdateExecutionLogs()
+		{
+			if (TestExecutionInfo)
+			{
+				const FString CleanFilename = FPaths::GetCleanFilename(ImportPath);
+				for (int32 ErrorIndex = ErrorStartCount; ErrorIndex < TestExecutionInfo->Errors.Num(); ++ErrorIndex)
+				{
+					TestExecutionInfo->Errors[ErrorIndex] = FString::Printf(TEXT("%s: %s"), *CleanFilename, *TestExecutionInfo->Errors[ErrorIndex]);
+				}
+				for (int32 WarningIndex = WarningStartCount; WarningIndex < TestExecutionInfo->Warnings.Num(); ++WarningIndex)
+				{
+					TestExecutionInfo->Warnings[WarningIndex] = FString::Printf(TEXT("%s: %s"), *CleanFilename, *TestExecutionInfo->Warnings[WarningIndex]);
+				}
+				for (int32 LogIndex = LogStartCount; LogIndex < TestExecutionInfo->LogItems.Num(); ++LogIndex)
+				{
+					TestExecutionInfo->LogItems[LogIndex] = FString::Printf(TEXT("%s: %s"), *CleanFilename, *TestExecutionInfo->LogItems[LogIndex]);
+				}
+			}
+		}
+	};
+}
+
+/**
+* Latent command to update the asset helper
+*/
+DEFINE_LATENT_AUTOMATION_COMMAND_ONE_PARAMETER(FImportExportAssetCommand, TSharedPtr<ImportExportAssetHelper::FAssetInfo>, AssetHelper);
+bool FImportExportAssetCommand::Update()
+{
+	return AssetHelper->Update();
+}
+
+/**
+* Latent command to update the asset helper
+*/
+DEFINE_LATENT_AUTOMATION_COMMAND_ONE_PARAMETER(FLogImportExportTestResultsCommand, TSharedPtr<ImportExportAssetHelper::FAssetImportStats>, BuildStats);
+bool FLogImportExportTestResultsCommand::Update()
+{
+	UE_LOG(LogEditorAssetAutomationTests, Display, TEXT(" "));
+	UE_LOG(LogEditorAssetAutomationTests, Display, TEXT("Test Summary:"));
+	for (int32 i = 0; i < BuildStats->Reports.Num(); ++i)
+	{
+		UE_LOG(LogEditorAssetAutomationTests, Display, TEXT("AssetName: %s"), *BuildStats->Reports[i].AssetName);
+		UE_LOG(LogEditorAssetAutomationTests, Display, TEXT("  Import: %s"), BuildStats->Reports[i].bImportSuccessful ? TEXT("SUCCESS") : TEXT("FAILED"));
+		if (BuildStats->Reports[i].bSkippedExport)
+		{
+			UE_LOG(LogEditorAssetAutomationTests, Display, TEXT("  Export: SKIPPED"));
+		}
+		else if (BuildStats->Reports[i].bExportSuccessful)
+		{
+			UE_LOG(LogEditorAssetAutomationTests, Display, TEXT("  Export: SUCCESS FileSize: %i"), BuildStats->Reports[i].FileSize);
+		}
+		else
+		{
+			UE_LOG(LogEditorAssetAutomationTests, Display, TEXT("  Export: FAILED"));
+		}
+	}
+	return true;
+}
+
+/**
+* Automation test to import, open, screenshot, and export assets
+*/
+IMPLEMENT_SIMPLE_AUTOMATION_TEST(FAssetImportEditorTest, "Editor.Content.Asset Import and Export", EAutomationTestFlags::ATF_Editor);
+bool FAssetImportEditorTest::RunTest(const FString& Parameters)
+{
+	UAutomationTestSettings const* AutomationTestSettings = GetDefault<UAutomationTestSettings>();
+	check(AutomationTestSettings);
+
+	TSharedPtr<ImportExportAssetHelper::FAssetImportStats> BuildStats = MakeShareable(new ImportExportAssetHelper::FAssetImportStats());
+
+	for (int32 TestIdx = 0; TestIdx < AutomationTestSettings->ImportExportTestDefinitions.Num(); ++TestIdx)
+	{
+		TSharedPtr<ImportExportAssetHelper::FAssetInfo> AssetInfo = MakeShareable(new ImportExportAssetHelper::FAssetInfo(AutomationTestSettings->ImportExportTestDefinitions[TestIdx], &ExecutionInfo, BuildStats));
+		ADD_LATENT_AUTOMATION_COMMAND(FImportExportAssetCommand(AssetInfo));
+	}
+
+	ADD_LATENT_AUTOMATION_COMMAND(FLogImportExportTestResultsCommand(BuildStats));
+
+	return true;
+}
+
+#undef LOCTEXT_NAMESPACE
diff --git a/Engine/Source/Runtime/AutomationWorker/Private/AutomationWorkerModule.cpp b/Engine/Source/Runtime/AutomationWorker/Private/AutomationWorkerModule.cpp
--- a/Engine/Source/Runtime/AutomationWorker/Private/AutomationWorkerModule.cpp
+++ b/Engine/Source/Runtime/AutomationWorker/Private/AutomationWorkerModule.cpp
@@ -498,11 +498,45 @@ static struct FQueueTests
 		{
 			GLog->Logf(ELogVerbosity::Log, TEXT("...Automation Test Queue Empty %d tests performed."), NumTestsRun);
 			NumTestsRun = 0;
+
+			//Quit when we are done if we have the command line option
+			const FString CommandLine(FCommandLine::Get());
+			if (CommandLine.Contains(TEXT("-AutomationTests=")))
+			{
+				//Let the current callstack complete then quit
+				GEngine->DeferredCommands.Add(TEXT("QUIT"));
+			}
 		}
 	}
 
 } QueueTests;
 
+bool GenerateTestNamesFromCommandLine(const TCHAR* InTestCommands, TArray<FString>& OutTestNames)
+{
+	//Split the test names up
+	TArray<FString> TestList;
+	const FString StringCommand = InTestCommands;
+	StringCommand.ParseIntoArray(&TestList, TEXT(","), true);
+
+	//Get the list of valid names
+	TArray<FAutomationTestInfo> TestInfo;
+	FAutomationTestFramework::GetInstance().GetValidTestNames(TestInfo);
+
+	for (int32 TestListIndex = 0; TestListIndex < TestList.Num(); ++TestListIndex)
+	{
+		const FString CleanTestName = TestList[TestListIndex].Trim();
+		for (int TestIndex = 0; TestIndex < TestInfo.Num(); ++TestIndex)
+		{
+			if (TestInfo[TestIndex].GetDisplayName().StartsWith(CleanTestName))
+			{
+				OutTestNames.Add(TestInfo[TestIndex].GetTestName());
+			}
+		}
+	}
+
+	return OutTestNames.Num() > 0;
+}
+
 
 bool DirectAutomationCommand(const TCHAR* Cmd, FOutputDevice* Ar = GLog)
 {
@@ -527,6 +561,24 @@ bool DirectAutomationCommand(const TCHAR* Cmd, FOutputDevice* Ar = GLog)
 			Test = Test.Trim();
 			QueueTests.NewTest(Test);
 		}
+		else if (FParse::Command(&TempCmd, TEXT("CommandLineTests")))
+		{
+			TArray<FString> TestNames;
+			if (GenerateTestNamesFromCommandLine(TempCmd, TestNames))
+			{
+				//Get the number of times to loop
+				int32 NumTestLoops = 1;
+				FParse::Value(FCommandLine::Get(), TEXT("TestLoops="), NumTestLoops);
+
+				for (int32 LoopIndex = 0; LoopIndex < NumTestLoops; ++LoopIndex)
+				{
+					for (int32 TestIndex = 0; TestIndex < TestNames.Num(); ++TestIndex)
+					{
+						QueueTests.NewTest(TestNames[TestIndex]);
+					}
+				}
+			}
+		}
 		else if( FParse::Command(&TempCmd,TEXT("runall")) )
 		{
 			int32 Mod = 0;
diff --git a/Engine/Source/Runtime/Engine/Classes/Tests/AutomationTestSettings.h b/Engine/Source/Runtime/Engine/Classes/Tests/AutomationTestSettings.h
--- a/Engine/Source/Runtime/Engine/Classes/Tests/AutomationTestSettings.h
+++ b/Engine/Source/Runtime/Engine/Classes/Tests/AutomationTestSettings.h
@@ -61,6 +61,49 @@ struct FExternalToolDefinition
 	FDirectoryPath ScriptDirectory;
 };
 
+/**
+ * Holds UProperty names and values to customize factory settings
+ */
+USTRUCT()
+struct FImportFactorySettingValues
+{
+	GENERATED_USTRUCT_BODY()
+
+	/* The name of the UProperty to change */
+	UPROPERTY(config, EditAnywhere, Category = Automation, meta = (ToolTip = "Name of the property to change.  Nested settings can be modified using \"Outer.Property\""))
+	FString SettingName;
+
+	/* The value to apply to the UProperty */
+	UPROPERTY(config, EditAnywhere, Category = Automation, meta = (ToolTip = "Value to import for the specified property."))
+	FString Value;
+};
+
+
+/**
+ * Holds settings for the asset import / export automation test
+ */
+USTRUCT()
+struct FEditorImportExportTestDefinition
+{
+	GENERATED_USTRUCT_BODY()
+
+	/* The file to import */
+	UPROPERTY(config, EditAnywhere, Category = Automation, meta = (FilePathFilter = "*"))
+	FFilePath ImportFilePath;
+
+	/* The file extension to use when exporting */
+	UPROPERTY(config, EditAnywhere, Category = Automation, meta = (ToolTip = "The file extension to use when exporting this asset.  Used to find a supporting exporter"))
+	FString ExportFileExtension;
+
+	/* If true, the export step will be skipped */
+	UPROPERTY(config, EditAnywhere, Category = Automation)
+	bool bSkipExport;
+
+	/* Settings for the import factory */
+	UPROPERTY(config, EditAnywhere, Category=Automation)
+	TArray<FImportFactorySettingValues> FactorySettings;
+};
+
 /**
  * Implements the Editor's user settings.
  */
@@ -106,4 +149,10 @@ public:
 	 */
 	UPROPERTY(EditAnywhere, config, Category=ExternalTools)
 	TArray<FExternalToolDefinition> ExternalTools;
+
+	/**
+	 * Asset import / Export test settings
+	 */
+	UPROPERTY(EditAnywhere, config, Category = Automation)
+	TArray<FEditorImportExportTestDefinition> ImportExportTestDefinitions;
 };
diff --git a/Engine/Source/Runtime/Engine/Private/Tests/EngineAutomationTests.cpp b/Engine/Source/Runtime/Engine/Private/Tests/EngineAutomationTests.cpp
--- a/Engine/Source/Runtime/Engine/Private/Tests/EngineAutomationTests.cpp
+++ b/Engine/Source/Runtime/Engine/Private/Tests/EngineAutomationTests.cpp
@@ -26,6 +26,39 @@ namespace
 
 		return GEngine->GetWorldContexts()[0].World();
 	}
+
+	/**
+	* Populates the test names and commands for complex tests that are ran on all available maps
+	*
+	* @param OutBeautifiedNames - The list of map names
+	* @param OutTestCommands - The list of commands for each test (The file names in this case)
+	*/
+	void PopulateTestsForAllAvailableMaps(TArray<FString>& OutBeautifiedNames, TArray <FString>& OutTestCommands)
+	{
+		TArray<FString> FileList;
+#if WITH_EDITOR
+		FEditorFileUtils::FindAllPackageFiles(FileList);
+#else
+		// Look directly on disk. Very slow!
+		FPackageName::FindPackagesInDirectory(FileList, *FPaths::GameContentDir());
+#endif
+
+		// Iterate over all files, adding the ones with the map extension..
+		for (int32 FileIndex = 0; FileIndex < FileList.Num(); FileIndex++)
+		{
+			const FString& Filename = FileList[FileIndex];
+
+			// Disregard filenames that don't have the map extension if we're in MAPSONLY mode.
+			if (FPaths::GetExtension(Filename, true) == FPackageName::GetMapPackageExtension())
+			{
+				if (FAutomationTestFramework::GetInstance().ShouldTestContent(Filename))
+				{
+					OutBeautifiedNames.Add(FPaths::GetBaseFilename(Filename));
+					OutTestCommands.Add(Filename);
+				}
+			}
+		}
+	}
 }
 
 /**
@@ -143,29 +176,7 @@ IMPLEMENT_COMPLEX_AUTOMATION_TEST( FLoadAllMapsInGameTest, "Maps.Load All In Gam
  */
 void FLoadAllMapsInGameTest::GetTests(TArray<FString>& OutBeautifiedNames, TArray <FString>& OutTestCommands) const
 {
-	TArray<FString> FileList;
-#if WITH_EDITOR
-	FEditorFileUtils::FindAllPackageFiles(FileList);
-#else
-	// Look directly on disk. Very slow!
-	FPackageName::FindPackagesInDirectory(FileList, *FPaths::GameContentDir());
-#endif
-
-	// Iterate over all files, adding the ones with the map extension..
-	for( int32 FileIndex=0; FileIndex< FileList.Num(); FileIndex++ )
-	{
-		const FString& Filename = FileList[FileIndex];
-
-		// Disregard filenames that don't have the map extension if we're in MAPSONLY mode.
-		if ( FPaths::GetExtension(Filename, true) == FPackageName::GetMapPackageExtension() ) 
-		{
-			if (FAutomationTestFramework::GetInstance().ShouldTestContent(Filename))
-			{
-				OutBeautifiedNames.Add(FPaths::GetBaseFilename(Filename));
-				OutTestCommands.Add(Filename);
-			}
-		}
-	}
+	PopulateTestsForAllAvailableMaps(OutBeautifiedNames, OutTestCommands);
 }
 
 /** 
@@ -269,6 +280,70 @@ bool FLoadGameMapCommand::Update()
 	return true;
 }
 
+/**
+ * Latent command to run an exec command that also requires a UWorld.
+ */
+DEFINE_LATENT_AUTOMATION_COMMAND_ONE_PARAMETER(FExecWorldStringLatentCommand, FString, ExecCommand);
+
+bool FExecWorldStringLatentCommand::Update()
+{
+	check(GEngine->GetWorldContexts().Num() == 1);
+	check(GEngine->GetWorldContexts()[0].WorldType == EWorldType::Game);
+
+	GEngine->Exec(GEngine->GetWorldContexts()[0].World(), *ExecCommand);
+	return true;
+}
+
+/**
+ * Automation test to load a map and capture FPS performance charts
+ */
+IMPLEMENT_COMPLEX_AUTOMATION_TEST(FCinematicFPSPerfTest, "Engine.Cinematic FPS Perf Capture", EAutomationTestFlags::ATF_Game);
+
+void FCinematicFPSPerfTest::GetTests(TArray<FString>& OutBeautifiedNames, TArray <FString>& OutTestCommands) const
+{
+	PopulateTestsForAllAvailableMaps(OutBeautifiedNames, OutTestCommands);
+}
+
+bool FCinematicFPSPerfTest::RunTest(const FString& Parameters)
+{
+	//Check we are running from commandline
+	const FString CommandLine(FCommandLine::Get());
+	if( CommandLine.Contains(TEXT("AutomationTests")) )
+	{
+		//Get the name of the console event to trigger the cinematic
+		FString CinematicEventCommand;
+		if( !FParse::Value(*CommandLine, TEXT("CE="), CinematicEventCommand) )
+		{
+			CinematicEventCommand = TEXT("CE Start");
+		}
+
+		//Get the length of time to let the cinematic run
+		float RunTime;
+		if( !FParse::Value(*CommandLine, TEXT("RunTime="), RunTime) )
+		{
+			RunTime=5.f;
+		}
+
+		//Load map
+		const FString MapName = Parameters;
+		ADD_LATENT_AUTOMATION_COMMAND(FLoadGameMapCommand(MapName));
+		ADD_LATENT_AUTOMATION_COMMAND(FEngineWaitLatentCommand(1.0f));
+
+		//Start the matinee and perform the FPS Chart
+		ADD_LATENT_AUTOMATION_COMMAND(FExecWorldStringLatentCommand(CinematicEventCommand));
+		ADD_LATENT_AUTOMATION_COMMAND(FExecWorldStringLatentCommand(TEXT("StartFPSChart")));
+		ADD_LATENT_AUTOMATION_COMMAND(FEngineWaitLatentCommand(RunTime));
+		ADD_LATENT_AUTOMATION_COMMAND(FExecWorldStringLatentCommand(TEXT("StopFPSChart")));
+	}
+	else
+	{
+		UE_LOG(LogEngineAutomationTests, Warning, TEXT("FCinematicFPSPerfTest is a Commandline test.  Please use -AutomationTests=\"Engine.Cinematic FPS Perf Capture\""));
+		return false;
+	}
+
+	return true;
+}
+
 
 /* UAutomationTestSettings interface
  *****************************************************************************/
diff --git a/Engine/Source/Runtime/Engine/Private/UnrealEngine.cpp b/Engine/Source/Runtime/Engine/Private/UnrealEngine.cpp
--- a/Engine/Source/Runtime/Engine/Private/UnrealEngine.cpp
+++ b/Engine/Source/Runtime/Engine/Private/UnrealEngine.cpp
@@ -722,6 +722,13 @@ void UEngine::Init(IEngineLoop* InEngineLoop)
 		}
 	}
 
+	// Optionally queue automation tests
+	FString AutomationCmds;
+	if (FParse::Value(FCommandLine::Get(), TEXT("AutomationTests="), AutomationCmds, false))
+	{
+		new(GEngine->DeferredCommands) FString(FString(TEXT("Automation CommandLineTests ")) + AutomationCmds);
+	}
+
 	// optionally set the vsync console variable
 	if( FParse::Param(FCommandLine::Get(), TEXT("vsync")) )
 	{
