UE4: Reviving the ParticleSystemAudit Commandlet. Fixed it up to use the Asset Registry instead of a database.

diff --git a/Engine/Source/Editor/UnrealEd/Classes/Commandlets/ParticleSystemAuditCommandlet.h b/Engine/Source/Editor/UnrealEd/Classes/Commandlets/ParticleSystemAuditCommandlet.h
new file mode 100644
--- /dev/null
+++ b/Engine/Source/Editor/UnrealEd/Classes/Commandlets/ParticleSystemAuditCommandlet.h
@@ -0,0 +1,156 @@
+// Copyright 1998-2014 Epic Games, Inc. All Rights Reserved.
+
+#include "Runtime/Engine/Classes/Particles/ParticleSystem.h"
+#include "ParticleDefinitions.h"
+#include "ParticleSystemAuditCommandlet.generated.h"
+
+UCLASS()
+class UParticleSystemAuditCommandlet : public UCommandlet
+{
+	GENERATED_UCLASS_BODY()
+
+	/** Helper struct for tracking disabled LOD level issues */
+	struct FParticleSystemLODInfo
+	{
+	public:
+		/** The method of LOD updating on the particle system */
+		ParticleSystemLODMethod LODMethod;
+		/** Indices of emitters with disabled LODLevel mismatches */
+		TArray<int32>	EmittersWithDisableLODMismatch;
+		/** true if the PSys has emitters w/ inter looping mismatches */
+		bool bHasInterLoopingMismatch;
+		/** true if the PSys has emitters w/ intra looping mismatches */
+		bool bHasIntraLoopingMismatch;
+
+		FParticleSystemLODInfo() :
+			  bHasInterLoopingMismatch(false)
+			, bHasIntraLoopingMismatch(false)
+		{
+		}
+
+		inline bool operator==(const FParticleSystemLODInfo& Src)
+		{
+			return (
+				(LODMethod == Src.LODMethod) && 
+				(EmittersWithDisableLODMismatch == Src.EmittersWithDisableLODMismatch) && 
+				(bHasInterLoopingMismatch == Src.bHasInterLoopingMismatch) && 
+				(bHasIntraLoopingMismatch == Src.bHasIntraLoopingMismatch)
+				);
+		}
+
+		inline FParticleSystemLODInfo& operator=(const FParticleSystemLODInfo& Src)
+		{
+			LODMethod = Src.LODMethod;
+			EmittersWithDisableLODMismatch.Empty();
+			EmittersWithDisableLODMismatch = Src.EmittersWithDisableLODMismatch;
+			bHasInterLoopingMismatch = Src.bHasInterLoopingMismatch;
+			bHasIntraLoopingMismatch = Src.bHasIntraLoopingMismatch;
+			return *this;
+		}
+	};
+
+	/** Helper struct for tracking duplicate module info */
+	struct FParticleSystemDuplicateModuleInfo
+	{
+	public:
+		/** The number of modules the particle system has */
+		int32 ModuleCount;
+		/** The amount of memory used by those modules */
+		int32 ModuleMemory;
+		/** The number of modules the particle system *could* have w/ duplicates removed */
+		int32 RemovedDuplicateCount;
+		/** The amount of memory there would be w/ duplicates removed */
+		int32 RemovedDuplicateMemory;
+
+		FParticleSystemDuplicateModuleInfo() :
+		ModuleCount(0)
+			, ModuleMemory(0)
+			, RemovedDuplicateCount(0)
+			, RemovedDuplicateMemory(0)
+		{
+		}
+
+		inline bool operator==(const FParticleSystemDuplicateModuleInfo& Src)
+		{
+			return (
+				(ModuleCount == Src.ModuleCount) &&
+				(ModuleMemory == Src.ModuleMemory) &&
+				(RemovedDuplicateCount == Src.RemovedDuplicateCount) &&
+				(RemovedDuplicateMemory == Src.RemovedDuplicateMemory)
+				);
+		}
+
+		inline FParticleSystemDuplicateModuleInfo& operator=(const FParticleSystemDuplicateModuleInfo& Src)
+		{
+			ModuleCount = Src.ModuleCount;
+			ModuleMemory = Src.ModuleMemory;
+			RemovedDuplicateCount = Src.RemovedDuplicateCount;
+			RemovedDuplicateMemory = Src.RemovedDuplicateMemory;
+			return *this;
+		}
+	};
+
+	/** All particle systems w/ a NO LOD levels */
+	TSet<FString> ParticleSystemsWithNoLODs;
+	/** All particle systems w/ a single LOD level */
+	TSet<FString> ParticleSystemsWithSingleLOD;
+	/** All particle systems w/out fixed bounds set */
+	TSet<FString> ParticleSystemsWithoutFixedBounds;
+	/** All particle systems w/ LOD Method of Automatic & a check time of 0.0 */
+	TSet<FString> ParticleSystemsWithBadLODCheckTimes;
+	/** All particle systems w/ bOrientZAxisTowardCamera enabled */
+	TSet<FString> ParticleSystemsWithOrientZAxisTowardCamera;
+	/** All particle systems w/ missing materials */
+	TSet<FString> ParticleSystemsWithMissingMaterials;
+	/** All particle systems w/ no emitters */
+	TSet<FString> ParticleSystemsWithNoEmitters;
+	/** All particle systems w/ collision on in at least one emitter */
+	TSet<FString> ParticleSystemsWithCollisionEnabled;
+	/** All particle systems w/ a color scale over life that is a constant */
+	TSet<FString> ParticleSystemsWithConstantColorScaleOverLife;
+	/** All particle systems w/ a color scale over life that is a constant */
+	TMap<FString,int32> ParticleSystemsWithConstantColorScaleOverLifeCounts;
+	/** Particle systems w/ disabled LOD level matches */
+	TMap<FString,FParticleSystemLODInfo> ParticleSystemsWithLODLevelIssues;
+	/** Particle system duplicate module information */
+	TMap<FString,FParticleSystemDuplicateModuleInfo> PSysDuplicateModuleInfo;
+
+	/** Entry point */
+	int32 Main(const FString& Params) override;
+
+	/** Process all referenced particle systems */
+	bool ProcessParticleSystems();
+
+	/** 
+	 *	Check the given ParticleSystem for LOD mismatch issues
+	 *
+	 *	@param	InPSys		The particle system to check
+	 */
+	void CheckPSysForLODMismatches(UParticleSystem* InPSys);
+
+	/**
+	 *	Determine the given ParticleSystems duplicate module information 
+	 *
+	 *	@param	InPSys		The particle system to check
+	 */
+	void CheckPSysForDuplicateModules(UParticleSystem* InPSys);
+
+	/** Dump the results of the audit */
+	virtual void DumpResults();
+
+	/**
+	 *	Dump the give list of particle systems to an audit CSV file...
+	 *
+	 *	@param	InPSysSet		The particle system set to dump
+	 *	@param	InFilename		The name for the output file (short name)
+	 *
+	 *	@return	bool			true if successful, false if not
+	 */
+	bool DumpSimplePSysSet(TSet<FString>& InPSysSet, const TCHAR* InShortFilename);
+
+	/** Generic function to handle dumping values to a CSV file */
+	bool DumpSimpleSet(TSet<FString>& InSet, const TCHAR* InShortFilename, const TCHAR* InObjectClassName);
+
+	/** Gets an archive to write to an output file */
+	FArchive* GetOutputFile(const TCHAR* InFolderName, const TCHAR* InShortFilename);
+};
diff --git a/Engine/Source/Editor/UnrealEd/Private/Commandlets/ParticleSystemAuditCommandlet.cpp b/Engine/Source/Editor/UnrealEd/Private/Commandlets/ParticleSystemAuditCommandlet.cpp
new file mode 100644
--- /dev/null
+++ b/Engine/Source/Editor/UnrealEd/Private/Commandlets/ParticleSystemAuditCommandlet.cpp
@@ -0,0 +1,676 @@
+// Copyright 1998-2014 Epic Games, Inc. All Rights Reserved.
+
+#include "UnrealEd.h"
+#include "Runtime/Engine/Public/ObjectEditorUtils.h"
+#include "Runtime/Engine/Classes/Particles/ParticleEmitter.h"
+#include "Runtime/Engine/Classes/Particles/ParticleModuleRequired.h"
+#include "Runtime/Engine/Classes/Particles/Color/ParticleModuleColorScaleOverLife.h"
+#include "Runtime/Engine/Classes/Particles/Collision/ParticleModuleCollision.h"
+#include "Runtime/Engine/Classes/Particles/Spawn/ParticleModuleSpawn.h"
+#include "Runtime/AssetRegistry/Public/AssetRegistryModule.h"
+
+DEFINE_LOG_CATEGORY_STATIC(LogParticleSystemAuditCommandlet, Log, All);
+
+UParticleSystemAuditCommandlet::UParticleSystemAuditCommandlet(const class FPostConstructInitializeProperties& PCIP)
+	: Super(PCIP)
+{
+}
+
+int32 UParticleSystemAuditCommandlet::Main(const FString& Params)
+{
+	ProcessParticleSystems();
+	DumpResults();
+
+	return 0;
+}
+
+bool UParticleSystemAuditCommandlet::ProcessParticleSystems()
+{
+	FAssetRegistryModule& AssetRegistryModule = FModuleManager::GetModuleChecked<FAssetRegistryModule>(TEXT("AssetRegistry"));
+	IAssetRegistry& AssetRegistry = AssetRegistryModule.Get();
+	AssetRegistry.SearchAllAssets(true);
+
+	TArray<FAssetData> AssetList;
+	AssetRegistry.GetAssetsByClass(UParticleSystem::StaticClass()->GetFName(), AssetList);
+
+	double StartProcessParticleSystemsTime = FPlatformTime::Seconds();
+
+	// Find all level placed particle systems with:
+	//	- Single LOD level
+	//	- No fixed bounds
+	//	- LODLevel Mismatch 
+	//	- Kismet referenced & auto-activate set
+	// Iterate over the list and check each system for *no* lod
+	// 
+	const FString DevelopersFolder = FPackageName::FilenameToLongPackageName(FPaths::GameDevelopersDir().LeftChop(1));
+	FString LastPackageName = TEXT("");
+	int32 PackageSwitches = 0;
+	UPackage* CurrentPackage = NULL;
+	for (const FAssetData& AssetIt : AssetList)
+	{
+		const FString PSysName = AssetIt.ObjectPath.ToString();
+		const FString PackageName = AssetIt.PackageName.ToString();
+
+		if ( PackageName.StartsWith(DevelopersFolder) )
+		{
+			// Skip developer folders
+			continue;
+		}
+
+		if (PackageName != LastPackageName)
+		{
+			UPackage* Package = ::LoadPackage(NULL, *PackageName, LOAD_None);
+			if (Package != NULL)
+			{
+				LastPackageName = PackageName;
+				Package->FullyLoad();
+				CurrentPackage = Package;
+			}
+			else
+			{
+				UE_LOG(LogParticleSystemAuditCommandlet, Warning, TEXT("Failed to load package %s processing %s"), *PackageName, *PSysName);
+				CurrentPackage = NULL;
+			}
+		}
+
+		const FString ShorterPSysName = AssetIt.AssetName.ToString();
+		UParticleSystem* PSys = FindObject<UParticleSystem>(CurrentPackage, *ShorterPSysName);
+		if (PSys != NULL)
+		{
+			bool bInvalidLOD = false;
+			bool bSingleLOD = false;
+			bool bFoundEmitter = false;
+			bool bMissingMaterial = false;
+			bool bHasConstantColorScaleOverLife = false;
+			bool bHasCollisionEnabled = false;
+			int32 ConstantColorScaleOverLifeCount = 0;
+			for (int32 EmitterIdx = 0; EmitterIdx < PSys->Emitters.Num(); EmitterIdx++)
+			{
+				UParticleEmitter* Emitter = PSys->Emitters[EmitterIdx];
+				if (Emitter != NULL)
+				{
+					if (Emitter->LODLevels.Num() == 0)
+					{
+						bInvalidLOD = true;
+					}
+					else if (Emitter->LODLevels.Num() == 1)
+					{
+						bSingleLOD = true;
+					}
+					bFoundEmitter = true;
+					for (int32 LODIdx = 0; LODIdx < Emitter->LODLevels.Num(); LODIdx++)
+					{
+						UParticleLODLevel* LODLevel = Emitter->LODLevels[LODIdx];
+						if (LODLevel != NULL)
+						{
+							if (LODLevel->RequiredModule != NULL)
+							{
+								if (LODLevel->RequiredModule->Material == NULL)
+								{
+									bMissingMaterial = true;
+								}
+							}
+
+							for (int32 ModuleIdx = 0; ModuleIdx < LODLevel->Modules.Num(); ModuleIdx++)
+							{
+								UParticleModule* Module = LODLevel->Modules[ModuleIdx];
+
+								UParticleModuleColorScaleOverLife* CSOLModule = Cast<UParticleModuleColorScaleOverLife>(Module);
+								if (CSOLModule != NULL)
+								{
+									UDistributionFloatConstant* FloatConst = Cast<UDistributionFloatConstant>(CSOLModule->AlphaScaleOverLife.Distribution);
+									UDistributionVectorConstant* VectorConst = Cast<UDistributionVectorConstant>(CSOLModule->ColorScaleOverLife.Distribution);
+									if ((FloatConst != NULL) && (VectorConst != NULL))
+									{
+										bHasConstantColorScaleOverLife = true;
+										ConstantColorScaleOverLifeCount++;
+									}
+								}
+								else if (bHasCollisionEnabled == false)
+								{
+									UParticleModuleCollision* CollisionModule = Cast<UParticleModuleCollision>(Module);
+									if (CollisionModule != NULL)
+									{
+										if (CollisionModule->bEnabled == true)
+										{
+											bHasCollisionEnabled = true;
+										}
+									}
+								}
+							}
+						}
+					}
+				}
+			}
+
+			// Note all PSystems w/ no emitters...
+			if (PSys->Emitters.Num() == 0)
+			{
+				ParticleSystemsWithNoEmitters.Add(PSys->GetPathName());
+			}
+
+			// Note all missing material case PSystems...
+			if (bMissingMaterial == true)
+			{
+				ParticleSystemsWithMissingMaterials.Add(PSys->GetPathName());
+			}
+
+			// Note all PSystems that have at least one emitter w/ constant ColorScaleOverLife modules...
+			if (bHasConstantColorScaleOverLife == true)
+			{
+				ParticleSystemsWithConstantColorScaleOverLife.Add(PSys->GetPathName());
+				ParticleSystemsWithConstantColorScaleOverLifeCounts.Add(PSys->GetPathName(), ConstantColorScaleOverLifeCount);
+			}
+
+			// Note all PSystems that have at least one emitter w/ an enabled collision module...
+			if (bHasCollisionEnabled == true)
+			{
+				ParticleSystemsWithCollisionEnabled.Add(PSys->GetPathName());
+			}
+
+			// Note all 0 LOD case PSystems...
+			if (bInvalidLOD == true)
+			{
+				ParticleSystemsWithNoLODs.Add(PSys->GetPathName());
+			}
+			// Note all single LOD case PSystems...
+			if (bSingleLOD == true)
+			{
+				ParticleSystemsWithSingleLOD.Add(PSys->GetPathName());
+			}
+
+			// Note all non-fixed bound PSystems...
+			if (PSys->bUseFixedRelativeBoundingBox == false)
+			{
+				ParticleSystemsWithoutFixedBounds.Add(PSys->GetPathName());
+			}
+
+			// Note all bOrientZAxisTowardCamera systems
+			if (PSys->bOrientZAxisTowardCamera == true)
+			{
+				ParticleSystemsWithOrientZAxisTowardCamera.Add(PSys->GetPathName());
+			}
+
+			if ((PSys->LODMethod == PARTICLESYSTEMLODMETHOD_Automatic) &&
+				(bInvalidLOD == false) && (bSingleLOD == false) &&
+				(PSys->LODDistanceCheckTime == 0.0f))
+			{
+				ParticleSystemsWithBadLODCheckTimes.Add(PSys->GetPathName());
+			}
+
+			// Find LOD mistmatches - looping & non looping, etc.
+			CheckPSysForLODMismatches(PSys);
+
+			// Find duplicate module information
+			CheckPSysForDuplicateModules(PSys);
+
+			if (LastPackageName.Len() > 0)
+			{
+				if (LastPackageName != PSys->GetOutermost()->GetName())
+				{
+					LastPackageName = PSys->GetOutermost()->GetName();
+					PackageSwitches++;
+				}
+			}
+			else
+			{
+				LastPackageName = PSys->GetOutermost()->GetName();
+			}
+
+			if (PackageSwitches > 10)
+			{
+				::CollectGarbage(RF_Native);
+				PackageSwitches = 0;
+			}
+
+		}
+		else
+		{
+			UE_LOG(LogParticleSystemAuditCommandlet, Warning, TEXT("Failed to load particle system %s"), *PSysName);
+		}
+	}
+
+	// Probably don't need to do this, but just in case we have any 'hanging' packages 
+	// and more processing steps are added later, let's clean up everything...
+	::CollectGarbage(RF_Native);
+
+	double ProcessParticleSystemsTime = FPlatformTime::Seconds() - StartProcessParticleSystemsTime;
+	UE_LOG(LogParticleSystemAuditCommandlet, Warning, TEXT("Took %5.3f seconds to process referenced particle systems..."), ProcessParticleSystemsTime);
+
+	return true;
+}
+
+/** 
+ *	Check the given ParticleSystem for LOD mismatch issues
+ *
+ *	@param	InPSys		The particle system to check
+ */
+void UParticleSystemAuditCommandlet::CheckPSysForLODMismatches(UParticleSystem* InPSys)
+{
+	if (InPSys == NULL)
+	{
+		return;
+	}
+
+	// Process it...
+	bool bHasLoopingMismatch = false;
+	bool bHasIntraLoopingMismatch = false;
+	bool bHasInterLoopingMismatch = false;
+	bool bHasMultipleLODLevels = true;
+
+	TArray<int32> InterEmitterCompare;
+
+	TArray<bool> EmitterLODLevelEnabledFlags;
+	FParticleSystemLODInfo* LODInfo = NULL;
+
+	EmitterLODLevelEnabledFlags.Empty(InPSys->Emitters.Num());
+	EmitterLODLevelEnabledFlags.AddZeroed(InPSys->Emitters.Num());
+	for (int32 EmitterIdx = 0; (EmitterIdx < InPSys->Emitters.Num()) && !bHasLoopingMismatch; EmitterIdx++)
+	{
+		UParticleEmitter* Emitter = InPSys->Emitters[EmitterIdx];
+		if (Emitter != NULL)
+		{
+			int32 EmitterLooping = -1;
+			for (int32 LODIdx = 0; (LODIdx < Emitter->LODLevels.Num()) && !bHasLoopingMismatch; LODIdx++)
+			{
+				UParticleLODLevel* LODLevel = Emitter->LODLevels[LODIdx];
+				if (LODLevel != NULL)
+				{
+					if (LODIdx == 0)
+					{
+						EmitterLODLevelEnabledFlags[EmitterIdx] = LODLevel->bEnabled;
+					}
+					else
+					{
+						if (EmitterLODLevelEnabledFlags[EmitterIdx] != LODLevel->bEnabled)
+						{
+							// MISMATCH
+							if (LODInfo == NULL)
+							{
+								LODInfo = new FParticleSystemLODInfo();
+								LODInfo->LODMethod = ParticleSystemLODMethod(InPSys->LODMethod);
+							}
+							LODInfo->EmittersWithDisableLODMismatch.AddUnique(EmitterIdx);
+						}
+					}
+				}
+				if ((LODLevel != NULL) && (LODLevel->bEnabled == true))
+				{
+					int32 CheckEmitterLooping = FMath::Min<int32>(LODLevel->RequiredModule->EmitterLoops, 1);
+					if (EmitterIdx == 0)
+					{
+						InterEmitterCompare.Add(CheckEmitterLooping);
+					}
+					else
+					{
+						if (InterEmitterCompare[LODIdx] != -1)
+						{
+							if (InterEmitterCompare[LODIdx] != CheckEmitterLooping)
+							{
+								if (LODInfo == NULL)
+								{
+									LODInfo = new FParticleSystemLODInfo();
+									LODInfo->LODMethod = ParticleSystemLODMethod(InPSys->LODMethod);
+								}
+								LODInfo->bHasInterLoopingMismatch = true;
+							}
+						}
+						else
+						{
+							InterEmitterCompare[LODIdx] = CheckEmitterLooping;
+						}
+					}
+
+					if (EmitterLooping == -1)
+					{
+						EmitterLooping = CheckEmitterLooping;
+					}
+					else
+					{
+						if (EmitterLooping != CheckEmitterLooping)
+						{
+							if (LODInfo == NULL)
+							{
+								LODInfo = new FParticleSystemLODInfo();
+								LODInfo->LODMethod = ParticleSystemLODMethod(InPSys->LODMethod);
+							}
+							LODInfo->bHasIntraLoopingMismatch = true;
+						}
+					}
+				}
+				else
+				{
+					if (EmitterIdx == 0)
+					{
+						InterEmitterCompare.Add(-1);
+					}
+				}
+			}
+		}
+	}
+
+	// If there was an LOD info mismatch, add it to the list...
+	if (LODInfo != NULL)
+	{
+		ParticleSystemsWithLODLevelIssues.Add(InPSys->GetPathName(), *LODInfo);
+		delete LODInfo;
+	}
+}
+
+/**
+ *	Determine the given ParticleSystems duplicate module information 
+ *
+ *	@param	InPSys		The particle system to check
+ */
+void UParticleSystemAuditCommandlet::CheckPSysForDuplicateModules(UParticleSystem* InPSys)
+{
+	if (InPSys == NULL)
+	{
+		return;
+	}
+
+	FParticleSystemDuplicateModuleInfo* DupInfo = PSysDuplicateModuleInfo.Find(InPSys->GetPathName());
+	if (DupInfo == NULL)
+	{
+		// Compare all the particle modules in the array
+		TMap<UClass*,TMap<UParticleModule*,int32> > ClassToModulesMap;
+		TMap<UParticleModule*,int32> AllModulesArray;
+		for (int32 EmitterIdx = 0; EmitterIdx < InPSys->Emitters.Num(); EmitterIdx++)
+		{
+			UParticleEmitter* Emitter = InPSys->Emitters[EmitterIdx];
+			if (Emitter != NULL)
+			{
+				for (int32 LODIdx = 0; LODIdx < Emitter->LODLevels.Num(); LODIdx++)
+				{
+					UParticleLODLevel* LODLevel = Emitter->LODLevels[LODIdx];
+					if (LODLevel != NULL)
+					{
+						for (int32 ModuleIdx = -1; ModuleIdx < LODLevel->Modules.Num(); ModuleIdx++)
+						{
+							UParticleModule* Module = NULL;
+							if (ModuleIdx == -1)
+							{
+								Module = LODLevel->SpawnModule;
+							}
+							else
+							{
+								Module = LODLevel->Modules[ModuleIdx];
+							}
+							if (Module != NULL)
+							{
+								if (AllModulesArray.Find(Module) == NULL)
+								{
+									FArchiveCountMem ModuleMemCount(Module);
+									int32 ModuleSize = ModuleMemCount.GetMax();
+									AllModulesArray.Add(Module, ModuleSize);
+								}
+
+								TMap<UParticleModule*,int32>* ModuleList = ClassToModulesMap.Find(Module->GetClass());
+								if (ModuleList == NULL)
+								{
+									TMap<UParticleModule*,int32> TempModuleList;
+									ClassToModulesMap.Add(Module->GetClass(), TempModuleList);
+									ModuleList = ClassToModulesMap.Find(Module->GetClass());
+								}
+								check(ModuleList);
+								int32* ModuleCount = ModuleList->Find(Module);
+								if (ModuleCount == NULL)
+								{
+									int32 TempModuleCount = 0;
+									ModuleList->Add(Module, TempModuleCount);
+									ModuleCount = ModuleList->Find(Module);
+								}
+								check(ModuleCount);
+								*ModuleCount++;
+							}
+						}
+					}
+				}
+			}
+		}
+
+		// Now we have a list of module classes and the modules they contain
+		TMap<UParticleModule*, TArray<UParticleModule*> > DuplicateModules;
+		TMap<UParticleModule*,bool> FoundAsADupeModules;
+		for (TMap<UClass*,TMap<UParticleModule*,int32> >::TIterator ModClassIt(ClassToModulesMap); ModClassIt; ++ModClassIt)
+		{
+			UClass* ModuleClass = ModClassIt.Key();
+			TMap<UParticleModule*,int32>& ModuleMap = ModClassIt.Value();
+			if (ModuleMap.Num() > 1)
+			{
+				// There is more than one of this module, so see if there are dupes...
+				TArray<UParticleModule*> ModuleArray;
+				for (TMap<UParticleModule*,int32>::TIterator ModuleIt(ModuleMap); ModuleIt; ++ModuleIt)
+				{
+					ModuleArray.Add(ModuleIt.Key());
+				}
+
+				// For each module, see if it it a duplicate of another
+				for (int32 ModuleIdx = 0; ModuleIdx < ModuleArray.Num(); ModuleIdx++)
+				{
+					UParticleModule* SourceModule = ModuleArray[ModuleIdx];
+					if (FoundAsADupeModules.Find(SourceModule) == NULL)
+					{
+						for (int32 InnerModuleIdx = ModuleIdx + 1; InnerModuleIdx < ModuleArray.Num(); InnerModuleIdx++)
+						{
+							UParticleModule* CheckModule = ModuleArray[InnerModuleIdx];
+							bool bIsDifferent = false;
+							if (FoundAsADupeModules.Find(CheckModule) == NULL)
+							{
+								FName CascadeCategory(TEXT("Cascade"));
+								// Copy non component properties from the old actor to the new actor
+								for (UProperty* Property = ModuleClass->PropertyLink; Property != NULL; Property = Property->PropertyLinkNext)
+								{
+									bool bIsTransient = (Property->PropertyFlags & CPF_Transient) != 0;
+									bool bIsEditorOnly = (Property->PropertyFlags & CPF_EditorOnly) != 0;
+									bool bIsCascade = (FObjectEditorUtils::GetCategoryFName(Property) == CascadeCategory);
+									if (!bIsTransient && !bIsEditorOnly && !bIsCascade)
+									{
+										bool bIsIdentical = Property->Identical_InContainer(SourceModule, CheckModule, 0, PPF_DeepComparison);
+										if (bIsIdentical == false)
+										{
+											bIsDifferent = true;
+										}
+									}
+									else
+									{
+
+									}
+								}
+							}
+							if (bIsDifferent == false)
+							{
+								TArray<UParticleModule*>* DupedModules = DuplicateModules.Find(SourceModule);
+								if (DupedModules == NULL)
+								{
+									TArray<UParticleModule*> TempDupedModules;
+									DuplicateModules.Add(SourceModule, TempDupedModules);
+									DupedModules = DuplicateModules.Find(SourceModule);
+								}
+								check(DupedModules);
+								DupedModules->AddUnique(CheckModule);
+								FoundAsADupeModules.Add(CheckModule, true);
+							}
+						}
+					}
+				}
+			}
+		}
+
+		if (DuplicateModules.Num() > 0)
+		{
+			FParticleSystemDuplicateModuleInfo TempDupInfo;
+			PSysDuplicateModuleInfo.Add(InPSys->GetPathName(), TempDupInfo);
+			DupInfo = PSysDuplicateModuleInfo.Find(InPSys->GetPathName());
+
+			DupInfo->ModuleCount = AllModulesArray.Num();
+			DupInfo->ModuleMemory = 0;
+			int32 DupeMemory = 0;
+			for (TMap<UParticleModule*,int32>::TIterator ModuleIt(AllModulesArray); ModuleIt; ++ModuleIt)
+			{
+				int32 ModuleSize = ModuleIt.Value();
+				DupInfo->ModuleMemory += ModuleSize;
+				if (FoundAsADupeModules.Find(ModuleIt.Key()) != NULL)
+				{
+					DupeMemory += ModuleSize;
+				}
+			}
+
+			DupInfo->RemovedDuplicateCount = FoundAsADupeModules.Num();
+			DupInfo->RemovedDuplicateMemory = DupeMemory;
+		}
+	}
+	else
+	{
+		UE_LOG(LogParticleSystemAuditCommandlet, Warning, TEXT("Already processed PSys for duplicate module information: %s"), *(InPSys->GetPathName()));
+	}
+}
+
+/** Dump the results of the audit */
+void UParticleSystemAuditCommandlet::DumpResults()
+{
+	// Dump all the simple mappings...
+	DumpSimplePSysSet(ParticleSystemsWithNoLODs, TEXT("PSysNoLOD"));
+	DumpSimplePSysSet(ParticleSystemsWithSingleLOD, TEXT("PSysSingleLOD"));
+	DumpSimplePSysSet(ParticleSystemsWithoutFixedBounds, TEXT("PSysNoFixedBounds"));
+	DumpSimplePSysSet(ParticleSystemsWithBadLODCheckTimes, TEXT("PSysBadLODCheckTimes"));
+	DumpSimplePSysSet(ParticleSystemsWithMissingMaterials, TEXT("PSysMissingMaterial"));
+	DumpSimplePSysSet(ParticleSystemsWithNoEmitters, TEXT("PSysNoEmitters"));
+	DumpSimplePSysSet(ParticleSystemsWithCollisionEnabled, TEXT("PSysCollisionEnabled"));
+	DumpSimplePSysSet(ParticleSystemsWithConstantColorScaleOverLife, TEXT("PSysConstantColorScale"));
+	DumpSimplePSysSet(ParticleSystemsWithOrientZAxisTowardCamera, TEXT("PSysOrientZTowardsCamera"));
+
+	FArchive* OutputStream;
+
+	// Dump out the particle systems w/ disabled LOD level mismatches...
+	const TCHAR* ConstColorScaleCountsName = TEXT("PSysConstColorScaleCounts");
+	OutputStream = GetOutputFile(TEXT("Audit"), ConstColorScaleCountsName);
+	if (OutputStream != NULL)
+	{
+		UE_LOG(LogParticleSystemAuditCommandlet, Log, TEXT("Dumping '%s' results..."), ConstColorScaleCountsName);
+		OutputStream->Logf(TEXT("Particle System,ModuleCount"));
+		for (TMap<FString,int32>::TIterator DumpIt(ParticleSystemsWithConstantColorScaleOverLifeCounts); DumpIt; ++DumpIt)
+		{
+			FString PSysName = DumpIt.Key();
+			int32& Count = DumpIt.Value();
+			OutputStream->Logf(TEXT("%s,%d"), *PSysName, Count);
+		}
+		OutputStream->Close();
+		delete OutputStream;
+	}
+	else
+	{
+		UE_LOG(LogParticleSystemAuditCommandlet, Warning, TEXT("Failed to open ConstColorScaleCounts file %s"), ConstColorScaleCountsName);
+	}
+
+	// Dump out the particle systems w/ disabled LOD level mismatches...
+	const TCHAR* LODIssuesName = TEXT("PSysLODIssues");
+	OutputStream = GetOutputFile(TEXT("Audit"), LODIssuesName);
+	if (OutputStream != NULL)
+	{
+		UE_LOG(LogParticleSystemAuditCommandlet, Log, TEXT("Dumping '%s' results..."), LODIssuesName);
+		OutputStream->Logf(TEXT("Particle System,LOD Method,InterLoop,IntraLoop,Emitters with mismatched LOD enabled"));
+		for (TMap<FString,FParticleSystemLODInfo>::TIterator DumpIt(ParticleSystemsWithLODLevelIssues); DumpIt; ++DumpIt)
+		{
+			FString PSysName = DumpIt.Key();
+			FParticleSystemLODInfo& LODInfo = DumpIt.Value();
+
+			FString Output = FString::Printf(TEXT("%s,%s,%s,%s"), *PSysName, 
+				(LODInfo.LODMethod == PARTICLESYSTEMLODMETHOD_Automatic) ? TEXT("AUTO") :
+				((LODInfo.LODMethod == PARTICLESYSTEMLODMETHOD_ActivateAutomatic) ? TEXT("AUTOACTIVATE") : TEXT("DIRECTSET")),
+				LODInfo.bHasInterLoopingMismatch ? TEXT("Y") : TEXT("N"),
+				LODInfo.bHasIntraLoopingMismatch ? TEXT("Y") : TEXT("N"));
+			for (int32 EmitterIdx = 0; EmitterIdx < LODInfo.EmittersWithDisableLODMismatch.Num(); EmitterIdx++)
+			{
+				Output += FString::Printf(TEXT(",%d"), LODInfo.EmittersWithDisableLODMismatch[EmitterIdx]);
+			}
+			OutputStream->Logf(*Output);
+		}
+		OutputStream->Close();
+		delete OutputStream;
+	}
+	else
+	{
+		UE_LOG(LogParticleSystemAuditCommandlet, Warning, TEXT("Failed to open LODIssues file %s"), LODIssuesName);
+	}
+
+	// Dump out the duplicate module findings...
+	const TCHAR* DuplicateModulesName = TEXT("PSysDuplicateModules");
+	OutputStream = GetOutputFile(TEXT("Audit"), DuplicateModulesName);
+	if (OutputStream != NULL)
+	{
+		UE_LOG(LogParticleSystemAuditCommandlet, Log, TEXT("Dumping '%s' results..."), DuplicateModulesName);
+		OutputStream->Logf(TEXT("Particle System,Module Count,Module Memory,Duplicate Count,Duplicate Memory"));
+		for (TMap<FString,FParticleSystemDuplicateModuleInfo>::TIterator DupeIt(PSysDuplicateModuleInfo); DupeIt; ++DupeIt)
+		{
+			FString PSysName = DupeIt.Key();
+			FParticleSystemDuplicateModuleInfo& DupeInfo = DupeIt.Value();
+
+			OutputStream->Logf(TEXT("%s,%d,%d,%d,%d"), *PSysName, 
+				DupeInfo.ModuleCount, DupeInfo.ModuleMemory,
+				DupeInfo.RemovedDuplicateCount, DupeInfo.RemovedDuplicateMemory);
+		}
+		OutputStream->Close();
+		delete OutputStream;
+	}
+	else
+	{
+		UE_LOG(LogParticleSystemAuditCommandlet, Warning, TEXT("Failed to open DuplicateModule file %s"), DuplicateModulesName);
+	}
+}
+
+/**
+ *	Dump the give list of particle systems to an audit CSV file...
+ *
+ *	@param	InPSysMap		The particle system map to dump
+ *	@param	InFilename		The name for the output file (short name)
+ *
+ *	@return	bool			true if successful, false if not
+ */
+bool UParticleSystemAuditCommandlet::DumpSimplePSysSet(TSet<FString>& InPSysSet, const TCHAR* InShortFilename)
+{
+	return DumpSimpleSet(InPSysSet, InShortFilename, TEXT("ParticleSystem"));
+}
+
+bool UParticleSystemAuditCommandlet::DumpSimpleSet(TSet<FString>& InSet,
+	const TCHAR* InShortFilename, const TCHAR* InObjectClassName)
+{
+	if (InSet.Num() > 0)
+	{
+		check(InShortFilename != NULL);
+		check(InObjectClassName != NULL);
+
+		FArchive* OutputStream = GetOutputFile(TEXT("Audit"), InShortFilename);
+		if (OutputStream != NULL)
+		{
+			UE_LOG(LogParticleSystemAuditCommandlet, Log, TEXT("Dumping '%s' results..."), InShortFilename);
+			OutputStream->Logf(TEXT("%s,..."), InObjectClassName);
+			for (TSet<FString>::TIterator DumpIt(InSet); DumpIt; ++DumpIt)
+			{
+				FString ObjName = *DumpIt;
+				OutputStream->Logf(TEXT("%s"), *ObjName);
+			}
+
+			OutputStream->Close();
+			delete OutputStream;
+		}
+		else
+		{
+			return false;
+		}
+	}
+	return true;
+}
+
+FArchive* UParticleSystemAuditCommandlet::GetOutputFile(const TCHAR* InFolderName, const TCHAR* InShortFilename)
+{
+	// Place in the <UE4>\<GAME>\Saved\<InFolderName> folder
+	FString Filename = FString::Printf(TEXT("%s%s/%s-%s.csv"), *FPaths::GameSavedDir(), InFolderName, InShortFilename, *FDateTime::Now().ToString());
+
+	FArchive* OutputStream = IFileManager::Get().CreateDebugFileWriter(*Filename);
+	if (OutputStream == NULL)
+	{
+		UE_LOG(LogParticleSystemAuditCommandlet, Warning, TEXT("Failed to create output stream %s"), *Filename);
+	}
+	return OutputStream;
+}
\ No newline at end of file
