ttp340827 SOUL: AUTOMATION: CRASH: Running StaticMeshUV editor test

diff --git a/Engine/Source/Editor/UnrealEd/Private/Kismet2/BlueprintEditorUtils.cpp b/Engine/Source/Editor/UnrealEd/Private/Kismet2/BlueprintEditorUtils.cpp
--- a/Engine/Source/Editor/UnrealEd/Private/Kismet2/BlueprintEditorUtils.cpp
+++ b/Engine/Source/Editor/UnrealEd/Private/Kismet2/BlueprintEditorUtils.cpp
@@ -682,7 +682,8 @@ struct FRegenerationHelper
 		for (FBPInterfaceDescription const& InterfaceDesc : Blueprint->ImplementedInterfaces)
 		{
 			UClass* InterfaceClass = InterfaceDesc.Interface;
-			if (UBlueprint* InterfaceBlueprint = Cast<UBlueprint>(InterfaceClass->ClassGeneratedBy))
+			UBlueprint* InterfaceBlueprint = InterfaceClass ? Cast<UBlueprint>(InterfaceClass->ClassGeneratedBy) : NULL;
+			if (InterfaceBlueprint)
 			{
 				ForcedLoadMembers(InterfaceBlueprint);
 				if (InterfaceBlueprint->HasAnyFlags(RF_BeingRegenerated))
