Pulling CanAccessComponentsMode() into Blueprint utils (so we can use it in other places, like the upcoming menu refactor).

diff --git a/Engine/Source/Editor/Kismet/Private/BlueprintEditor.cpp b/Engine/Source/Editor/Kismet/Private/BlueprintEditor.cpp
--- a/Engine/Source/Editor/Kismet/Private/BlueprintEditor.cpp
+++ b/Engine/Source/Editor/Kismet/Private/BlueprintEditor.cpp
@@ -2618,10 +2618,7 @@ bool FBlueprintEditor::CanAccessComponentsMode() const
 	if(IsEditingSingleBlueprint())
 	{
 		UBlueprint* Blueprint = GetBlueprintObj();
-		bCanAccess = Blueprint->SimpleConstructionScript != NULL		// An SCS must be present (otherwise there is nothing valid to edit)
-			&& FBlueprintEditorUtils::IsActorBased(Blueprint)			// Must be parented to an AActor-derived class (some older BPs may have an SCS but may not be Actor-based)
-			&& Blueprint->BlueprintType != BPTYPE_MacroLibrary			// Must not be a macro-type Blueprint
-			&& Blueprint->BlueprintType != BPTYPE_FunctionLibrary;		// Must not be a function library
+		bCanAccess = FBlueprintEditorUtils::DoesSupportComponents(Blueprint);
 	}
 	
 	return bCanAccess;
diff --git a/Engine/Source/Editor/UnrealEd/Private/Kismet2/BlueprintEditorUtils.cpp b/Engine/Source/Editor/UnrealEd/Private/Kismet2/BlueprintEditorUtils.cpp
--- a/Engine/Source/Editor/UnrealEd/Private/Kismet2/BlueprintEditorUtils.cpp
+++ b/Engine/Source/Editor/UnrealEd/Private/Kismet2/BlueprintEditorUtils.cpp
@@ -2350,6 +2350,14 @@ bool FBlueprintEditorUtils::DoesSupportImplementingInterfaces(const UBlueprint*
 		&& Blueprint->BlueprintType != BPTYPE_FunctionLibrary;
 }
 
+bool FBlueprintEditorUtils::DoesSupportComponents(UBlueprint const* Blueprint)
+{
+	return (Blueprint->SimpleConstructionScript != nullptr)      // An SCS must be present (otherwise there is nothing valid to edit)
+		&& FBlueprintEditorUtils::IsActorBased(Blueprint)        // Must be parented to an AActor-derived class (some older BPs may have an SCS but may not be Actor-based)
+		&& (Blueprint->BlueprintType != BPTYPE_MacroLibrary)     // Must not be a macro-type Blueprint
+		&& (Blueprint->BlueprintType != BPTYPE_FunctionLibrary); // Must not be a function library
+}
+
 // Returns a descriptive name of the type of blueprint passed in
 FString FBlueprintEditorUtils::GetBlueprintTypeDescription(const UBlueprint* Blueprint)
 {
diff --git a/Engine/Source/Editor/UnrealEd/Public/Kismet2/BlueprintEditorUtils.h b/Engine/Source/Editor/UnrealEd/Public/Kismet2/BlueprintEditorUtils.h
--- a/Engine/Source/Editor/UnrealEd/Public/Kismet2/BlueprintEditorUtils.h
+++ b/Engine/Source/Editor/UnrealEd/Public/Kismet2/BlueprintEditorUtils.h
@@ -365,6 +365,9 @@ public:
 	/** Returns whether or not the blueprint supports implementing interfaces */
 	static bool DoesSupportImplementingInterfaces(const UBlueprint* Blueprint);
 
+	/** Returns whether or not the blueprint supports components */
+	static bool DoesSupportComponents(UBlueprint const* Blueprint);
+
 	// Returns a descriptive name of the type of blueprint passed in
 	static FString GetBlueprintTypeDescription(const UBlueprint* Blueprint);
 
