Editor - The scalability monitoring no longer warns if the new settings will be of no benefit.  We also give the user more time to decide what to do when the notification appears.  There are now buttons for handling the notification differently, Apply Now, Tweak Manually, and Cancel & Ignore.  The monitor check box now appears in the scalability settings area in the Quick View.

TTP 340925

diff --git a/Engine/Source/Editor/UnrealEd/Private/PerformanceMonitor.cpp b/Engine/Source/Editor/UnrealEd/Private/PerformanceMonitor.cpp
--- a/Engine/Source/Editor/UnrealEd/Private/PerformanceMonitor.cpp
+++ b/Engine/Source/Editor/UnrealEd/Private/PerformanceMonitor.cpp
@@ -7,6 +7,8 @@
 
 #define LOCTEXT_NAMESPACE "PerformanceMonitor"
 
+const double AutoApplyScalabilityTimeout = 10;
+
 /** Scalability dialog widget */
 class SScalabilitySettingsDialog : public SCompoundWidget
 {
@@ -31,19 +33,23 @@ public:
 				.WidthOverride(500.0f)
 				[
 					SNew(SVerticalBox)
+					
 					+ SVerticalBox::Slot()
 					.AutoHeight()
 					.Padding(5.f)
 					[
-						SNew(STextBlock).Text(LOCTEXT("PerformanceWarningDescription",
-						"The current performance of the editor seems to be low.\nUse the options below to reduce the amount of detail and increase performance.").ToString())
+						SNew(STextBlock)
+						.Text(LOCTEXT("PerformanceWarningDescription",
+							"The current performance of the editor seems to be low.\nUse the options below to reduce the amount of detail and increase performance."))
 					]
+					
 					+ SVerticalBox::Slot()
 					.AutoHeight()
 					.Padding(5.f)
 					[
 						SNew(SScalabilitySettings)
 					]
+
 					+ SVerticalBox::Slot()
 					.AutoHeight()
 					.Padding(5.f)
@@ -58,26 +64,18 @@ public:
 						)
 						.AutoWrapText(true)
 						.Text(LOCTEXT("PerformanceWarningChangeLater",
-							"You can modify these settings in future via \"Quick Settings\" button on the level editor toolbar and choosing \"Engine Scalability Settings\".").ToString())
+							"You can modify these settings in future via \"Quick Settings\" button on the level editor toolbar and choosing \"Engine Scalability Settings\"."))
 					]
+
 					+ SVerticalBox::Slot()
 					.AutoHeight()
 					.Padding(5.f)
 					[
 						SNew(SHorizontalBox)
+
 						+ SHorizontalBox::Slot()
-						[
-							SNew(SCheckBox)
-							.OnCheckStateChanged(FOnCheckStateChanged::CreateStatic(&OnMonitorPerformanceChanged))
-							.IsChecked(TAttribute<ESlateCheckBoxState::Type>::Create(&IsMonitoringPerformance))
-							.Content()
-							[
-								SNew(STextBlock)
-								.Text(LOCTEXT("PerformanceWarningEnableDisableCheckbox", "Keep monitoring editor performance?"))
-							]
-						]
-						+ SHorizontalBox::Slot()
-						.AutoWidth()
+						.FillWidth(1.0)
+						.HAlign(HAlign_Right)
 						[
 							SNew(SButton)
 							.Text(LOCTEXT("ScalabilityDone", "Done"))
@@ -88,29 +86,12 @@ public:
 			]
 		];
 	}
-
-	/** Called to get the "Show notification" check box state */
-	static ESlateCheckBoxState::Type IsMonitoringPerformance()
-	{
-		const bool bMonitorEditorPerformance = GEditor->GetEditorUserSettings().bMonitorEditorPerformance;
-		return bMonitorEditorPerformance ? ESlateCheckBoxState::Checked : ESlateCheckBoxState::Unchecked;
-	}
-
-	/** Called when the state of the "Show notification" check box changes */
-	static void OnMonitorPerformanceChanged(ESlateCheckBoxState::Type NewState)
-	{
-		const bool bNewEnabledState = (NewState == ESlateCheckBoxState::Checked);
-
-		auto& Settings = GEditor->AccessEditorUserSettings();
-		Settings.bMonitorEditorPerformance = bNewEnabledState;
-		Settings.PostEditChange();
-	}
 };
 
 FPerformanceMonitor::FPerformanceMonitor()
 {
 	LastEnableTime = 0;
-	NotificationTimeout = 5;
+	NotificationTimeout = AutoApplyScalabilityTimeout;
 	bIsNotificationAllowed = true;
 
 	// Register the console variables we need
@@ -138,20 +119,46 @@ FPerformanceMonitor::~FPerformanceMonitor()
 	Console.FindConsoleVariable(TEXT("PerfWarn.CoarseSampleTime"))->SetOnChangedCallback(FConsoleVariableDelegate());
 }
 
-void FPerformanceMonitor::AutoApplyScalability()
+bool FPerformanceMonitor::WillAutoScalabilityHelp() const
+{
+	Scalability::FQualityLevels CurrentLevels = Scalability::GetQualityLevels();
+	Scalability::FQualityLevels NewLevels = GetAutoScalabilityQualityLevels();
+
+	bool IsAutoScaleLower = false;
+	IsAutoScaleLower |= NewLevels.ResolutionQuality < CurrentLevels.ResolutionQuality;
+	IsAutoScaleLower |= NewLevels.ViewDistanceQuality < CurrentLevels.ViewDistanceQuality;
+	IsAutoScaleLower |= NewLevels.AntiAliasingQuality < CurrentLevels.AntiAliasingQuality;
+	IsAutoScaleLower |= NewLevels.ShadowQuality < CurrentLevels.ShadowQuality;
+	IsAutoScaleLower |= NewLevels.PostProcessQuality < CurrentLevels.PostProcessQuality;
+	IsAutoScaleLower |= NewLevels.TextureQuality < CurrentLevels.TextureQuality;
+	IsAutoScaleLower |= NewLevels.EffectsQuality < CurrentLevels.EffectsQuality;
+
+	// We don't check things like real-time, because the user may have enabled it temporarily.
+
+	return IsAutoScaleLower;
+}
+
+Scalability::FQualityLevels FPerformanceMonitor::GetAutoScalabilityQualityLevels() const
 {
 	const Scalability::FQualityLevels ExistingLevels = Scalability::GetQualityLevels();
 	Scalability::FQualityLevels NewLevels = GEditor->GetGameAgnosticSettings().EngineBenchmarkResult;
-	
+
 	// Make sure we don't turn settings up if the user has turned them down for any reason
-	NewLevels.ResolutionQuality	= FMath::Min(NewLevels.ResolutionQuality,		ExistingLevels.ResolutionQuality);
-	NewLevels.ViewDistanceQuality	= FMath::Min(NewLevels.ViewDistanceQuality,	ExistingLevels.ViewDistanceQuality);
-	NewLevels.AntiAliasingQuality	= FMath::Min(NewLevels.AntiAliasingQuality,	ExistingLevels.AntiAliasingQuality);
-	NewLevels.ShadowQuality		= FMath::Min(NewLevels.ShadowQuality,			ExistingLevels.ShadowQuality);
-	NewLevels.PostProcessQuality	= FMath::Min(NewLevels.PostProcessQuality,		ExistingLevels.PostProcessQuality);
-	NewLevels.TextureQuality		= FMath::Min(NewLevels.TextureQuality,			ExistingLevels.TextureQuality);
-	NewLevels.EffectsQuality		= FMath::Min(NewLevels.EffectsQuality,			ExistingLevels.EffectsQuality);
-	
+	NewLevels.ResolutionQuality		= FMath::Min(NewLevels.ResolutionQuality, ExistingLevels.ResolutionQuality);
+	NewLevels.ViewDistanceQuality	= FMath::Min(NewLevels.ViewDistanceQuality, ExistingLevels.ViewDistanceQuality);
+	NewLevels.AntiAliasingQuality	= FMath::Min(NewLevels.AntiAliasingQuality, ExistingLevels.AntiAliasingQuality);
+	NewLevels.ShadowQuality			= FMath::Min(NewLevels.ShadowQuality, ExistingLevels.ShadowQuality);
+	NewLevels.PostProcessQuality	= FMath::Min(NewLevels.PostProcessQuality, ExistingLevels.PostProcessQuality);
+	NewLevels.TextureQuality		= FMath::Min(NewLevels.TextureQuality, ExistingLevels.TextureQuality);
+	NewLevels.EffectsQuality		= FMath::Min(NewLevels.EffectsQuality, ExistingLevels.EffectsQuality);
+
+	return NewLevels;
+}
+
+void FPerformanceMonitor::AutoApplyScalability()
+{
+	Scalability::FQualityLevels NewLevels = FPerformanceMonitor::GetAutoScalabilityQualityLevels();
+
 	Scalability::SetQualityLevels(NewLevels);
 	Scalability::SaveState(GEditorGameAgnosticIni);
 
@@ -172,22 +179,34 @@ void FPerformanceMonitor::ShowPerformanceWarning(FText MessageText)
 	{
 		// Only show a new one if we've not shown one for a while
 		LastEnableTime = FPlatformTime::Seconds();
-		NotificationTimeout = 5;
+		NotificationTimeout = AutoApplyScalabilityTimeout;
 
 		// Create notification item
 		FNotificationInfo Info(MessageText);
 		Info.bFireAndForget = false;
-		Info.FadeOutDuration = 4.0f;
+		Info.FadeOutDuration = 3.0f;
 		Info.ExpireDuration = 0.0f;
 		Info.bUseLargeFont = false;
 
-		Info.Hyperlink = FSimpleDelegate::CreateRaw(this, &FPerformanceMonitor::ShowScalabilityDialog);
-		Info.HyperlinkText = LOCTEXT("PerformanceWarningScalability", "Tweak Manually");
+		Info.ButtonDetails.Add( FNotificationButtonInfo( LOCTEXT("ApplyNow", "Apply Now"), FText::GetEmpty(), FSimpleDelegate::CreateRaw( this, &FPerformanceMonitor::AutoApplyScalability ) ) );
+		Info.ButtonDetails.Add( FNotificationButtonInfo( LOCTEXT("TweakManually", "Tweak Manually"), FText::GetEmpty(), FSimpleDelegate::CreateRaw( this, &FPerformanceMonitor::ShowScalabilityDialog ) ) );
+		Info.ButtonDetails.Add( FNotificationButtonInfo( LOCTEXT("DontRemindMe", "Cancel & Ignore"), FText::GetEmpty(), FSimpleDelegate::CreateRaw( this, &FPerformanceMonitor::CancelPerformanceNotification ) ) );
 
 		PerformanceWarningNotificationPtr = FSlateNotificationManager::Get().AddNotification(Info);
+		PerformanceWarningNotificationPtr.Pin()->SetCompletionState(SNotificationItem::CS_Pending);
 	}
 }
 
+void FPerformanceMonitor::CancelPerformanceNotification()
+{
+	UEditorUserSettings& EditorUserSettings = GEditor->AccessEditorUserSettings();
+	EditorUserSettings.bMonitorEditorPerformance = false;
+	EditorUserSettings.PostEditChange();
+	EditorUserSettings.SaveConfig();
+
+	Reset();
+}
+
 void FPerformanceMonitor::HidePerformanceWarning()
 {
 	// Finished! Notify the UI.
@@ -213,7 +232,7 @@ void FPerformanceMonitor::Tick(float DeltaTime)
 	FineMovingAverage.Tick(FPlatformTime::Seconds(), GAverageFPS);
 	CoarseMovingAverage.Tick(FPlatformTime::Seconds(), GAverageFPS);
 
-	const bool bMonitorEditorPerformance = GetDefault<UEditorUserSettings>()->bMonitorEditorPerformance;
+	bool bMonitorEditorPerformance = GEditor->GetEditorUserSettings().bMonitorEditorPerformance;
 	if( !bMonitorEditorPerformance || !bIsNotificationAllowed )
 	{
 		return;
@@ -241,7 +260,6 @@ void FPerformanceMonitor::Tick(float DeltaTime)
 		}
 	}
 
-
 	if (!bLowFramerate && CoarseMovingAverage.IsReliable())
 	{
 		static IConsoleVariable* CVarMinFPS = IConsoleManager::Get().FindConsoleVariable(TEXT("PerfWarn.CoarseMinFPS"));
@@ -271,17 +289,24 @@ void FPerformanceMonitor::Tick(float DeltaTime)
 	}
 	else
 	{
+		// Before we warn the user lets check if we already have auto-scalability enabled, no sense in warning them
+		// if we can't do anything about it.
+		if ( !WillAutoScalabilityHelp() )
+		{
+			return;
+		}
+
 		// Choose an appropriate message
 		enum MessagesEnum { Seconds, SecondsPercent, Minute, MinutePercent, Minutes, MinutesPercent };
 		const FText Messages[] = {
-			LOCTEXT("PerformanceWarningInProgress_Seconds",			"Your framerate has been under {Framerate} FPS for the past {SampleTime} seconds.\nApplying optimum engine settings in {TimeRemaining}s."),
-			LOCTEXT("PerformanceWarningInProgress_Seconds_Percent", "Your framerate has been under {Framerate} FPS for {Percentage}% of the past {SampleTime} seconds.\nApplying optimum engine settings in {TimeRemaining}s."),
+			LOCTEXT("PerformanceWarningInProgress_Seconds",			"Your framerate has been under {Framerate} FPS for the past {SampleTime} seconds.\n\nApplying optimum engine settings in {TimeRemaining}s."),
+			LOCTEXT("PerformanceWarningInProgress_Seconds_Percent", "Your framerate has been under {Framerate} FPS for {Percentage}% of the past {SampleTime} seconds.\n\nApplying optimum engine settings in {TimeRemaining}s."),
 
-			LOCTEXT("PerformanceWarningInProgress_Minute",			"Your framerate has been under {Framerate} FPS for the past minute.\nApplying optimum engine settings in {TimeRemaining}s."),
-			LOCTEXT("PerformanceWarningInProgress_Minute_Percent",	"Your framerate has been under {Framerate} FPS for {Percentage}% of the last minute.\nApplying optimum engine settings in {TimeRemaining}s."),
+			LOCTEXT("PerformanceWarningInProgress_Minute",			"Your framerate has been under {Framerate} FPS for the past minute.\n\nApplying optimum engine settings in {TimeRemaining}s."),
+			LOCTEXT("PerformanceWarningInProgress_Minute_Percent",	"Your framerate has been under {Framerate} FPS for {Percentage}% of the last minute.\n\nApplying optimum engine settings in {TimeRemaining}s."),
 
-			LOCTEXT("PerformanceWarningInProgress_Minutes",			"Your framerate has been below {Framerate} FPS for the past {SampleTime} minutes.\nApplying optimum engine settings in {TimeRemaining}s."),
-			LOCTEXT("PerformanceWarningInProgress_Minutes_Percent", "Your framerate has been below {Framerate} FPS for {Percentage}% of the past {SampleTime} minutes.\nApplying optimum engine settings in {TimeRemaining}s."),
+			LOCTEXT("PerformanceWarningInProgress_Minutes",			"Your framerate has been below {Framerate} FPS for the past {SampleTime} minutes.\n\nApplying optimum engine settings in {TimeRemaining}s."),
+			LOCTEXT("PerformanceWarningInProgress_Minutes_Percent", "Your framerate has been below {Framerate} FPS for {Percentage}% of the past {SampleTime} minutes.\n\nApplying optimum engine settings in {TimeRemaining}s."),
 		};
 
 		int32 Message;
@@ -326,7 +351,7 @@ void FPerformanceMonitor::Tick(float DeltaTime)
 		}
 		else
 		{
-			NotificationTimeout = 5;
+			NotificationTimeout = AutoApplyScalabilityTimeout;
 			Arguments.Add(TEXT("TimeRemaining"), int(NotificationTimeout));
 
 			ShowPerformanceWarning(FText::Format(Messages[Message], Arguments));
@@ -383,14 +408,15 @@ void FPerformanceMonitor::ShowScalabilityDialog()
 		ExistingWindow->SetContent(
 			SNew(SScalabilitySettingsDialog)
 			.OnDoneClicked(
-				FOnClicked::CreateStatic([](TWeakPtr<SWindow> Window){
+				FOnClicked::CreateStatic([](TWeakPtr<SWindow> Window, FPerformanceMonitor* PerfWarn){
 					auto WindowPin = Window.Pin();
 					if (WindowPin.IsValid())
 					{
+						PerfWarn->bIsNotificationAllowed = true;
 						WindowPin->RequestDestroyWindow();
 					}
 					return FReply::Handled();
-				}, ScalabilitySettingsWindowPtr)
+				}, ScalabilitySettingsWindowPtr, this)
 			)
 		);
 
diff --git a/Engine/Source/Editor/UnrealEd/Private/PerformanceMonitor.h b/Engine/Source/Editor/UnrealEd/Private/PerformanceMonitor.h
--- a/Engine/Source/Editor/UnrealEd/Private/PerformanceMonitor.h
+++ b/Engine/Source/Editor/UnrealEd/Private/PerformanceMonitor.h
@@ -145,9 +145,18 @@ private:
 	/** Starts the notification. */
 	void ShowPerformanceWarning(FText MessageText);
 
+	/** Gets the quality levels that would be applied with auto-scalability. */
+	Scalability::FQualityLevels GetAutoScalabilityQualityLevels() const;
+
+	/** Returns true if the quality settings match the ones that would be applied with auto scalability. */
+	bool WillAutoScalabilityHelp() const;
+
 	/** Run a benchmark and auto apply scalability settings */
 	void AutoApplyScalability();
 
+	/** Adjusts the performance monitor settings to stop monitoring. */
+	void CancelPerformanceNotification();
+
 	/** Ends the notification. */
 	void HidePerformanceWarning();
 
@@ -166,6 +175,7 @@ private:
 	/** Tracks the last time the notification was started, used to avoid spamming. */
 	double LastEnableTime;
 
+	/** The time remaining before the auto scalability settings are automatically applied. */
 	double NotificationTimeout;
 
 	/** The notification window ptr */
diff --git a/Engine/Source/Editor/UnrealEd/Private/SScalabilitySettings.cpp b/Engine/Source/Editor/UnrealEd/Private/SScalabilitySettings.cpp
--- a/Engine/Source/Editor/UnrealEd/Private/SScalabilitySettings.cpp
+++ b/Engine/Source/Editor/UnrealEd/Private/SScalabilitySettings.cpp
@@ -120,6 +120,22 @@ SGridPanel::FSlot& SScalabilitySettings::MakeGridSlot(int32 InCol, int32 InRow,
 		.ColumnSpan(InColSpan);
 }
 
+ESlateCheckBoxState::Type SScalabilitySettings::IsMonitoringPerformance() const
+{
+	const bool bMonitorEditorPerformance = GEditor->GetEditorUserSettings().bMonitorEditorPerformance;
+	return bMonitorEditorPerformance ? ESlateCheckBoxState::Checked : ESlateCheckBoxState::Unchecked;
+}
+
+void SScalabilitySettings::OnMonitorPerformanceChanged(ESlateCheckBoxState::Type NewState)
+{
+	const bool bNewEnabledState = ( NewState == ESlateCheckBoxState::Checked );
+
+	auto& Settings = GEditor->AccessEditorUserSettings();
+	Settings.bMonitorEditorPerformance = bNewEnabledState;
+	Settings.PostEditChange();
+	Settings.SaveConfig();
+}
+
 void SScalabilitySettings::Construct( const FArguments& InArgs )
 {
 	const FText NamesLow(LOCTEXT("QualityLowLabel", "Low"));
@@ -200,6 +216,20 @@ void SScalabilitySettings::Construct( const FArguments& InArgs )
 				+MakeGridSlot(3,7) [ MakeButtonWidget(NamesHigh, TEXT("EffectsQuality"), 2, LOCTEXT("EffectsQualityHigh", "Set effects quality to high")) ]
 				+MakeGridSlot(4,7) [ MakeButtonWidget(NamesEpic, TEXT("EffectsQuality"), 3, LOCTEXT("EffectsQualityEpic", "Set effects quality to epic")) ]
 			]
+
+			+ SVerticalBox::Slot()
+			.AutoHeight()
+			.Padding(5.f)
+			[
+				SNew(SCheckBox)
+				.OnCheckStateChanged(this, &SScalabilitySettings::OnMonitorPerformanceChanged)
+				.IsChecked(this, &SScalabilitySettings::IsMonitoringPerformance)
+				.Content()
+				[
+					SNew(STextBlock)
+					.Text(LOCTEXT("PerformanceWarningEnableDisableCheckbox", "Monitor Editor Performance?"))
+				]
+			]
 		];
 }
 
diff --git a/Engine/Source/Editor/UnrealEd/Public/SScalabilitySettings.h b/Engine/Source/Editor/UnrealEd/Public/SScalabilitySettings.h
--- a/Engine/Source/Editor/UnrealEd/Public/SScalabilitySettings.h
+++ b/Engine/Source/Editor/UnrealEd/Public/SScalabilitySettings.h
@@ -56,6 +56,12 @@ private:
 	// Create a gridslot for the group quality level with all the required formatting
 	SGridPanel::FSlot& MakeGridSlot(int32 InCol, int32 InRow, int32 InColSpan = 1, int32 InRowSpan = 1);
 
+	/** Called to get the "Show notification" check box state */
+	ESlateCheckBoxState::Type IsMonitoringPerformance() const;
+
+	/** Called when the state of the "Show notification" check box changes */
+	void OnMonitorPerformanceChanged(ESlateCheckBoxState::Type NewState);
+
 private:
 	/* The state of scalability settings at the point of opening the menu*/
 	Scalability::FQualityLevels InitialQualityLevels;
