Workaround for an issue which arises when an empty FStringAssetReference is serialized to ini and the cooker attempts to interpret "None" as a package.

The error was:
LogInit:Display: LogCookCommandlet:Warning: Unable to find package file for: None
LogInit:Display: LogPackageName:Error: DoesPackageExist: DoesPackageExist FAILED: 'None' is not a standard unreal filename or a long path name.

diff --git a/Engine/Source/Editor/UnrealEd/Private/Commandlets/CookCommandlet.cpp b/Engine/Source/Editor/UnrealEd/Private/Commandlets/CookCommandlet.cpp
--- a/Engine/Source/Editor/UnrealEd/Private/Commandlets/CookCommandlet.cpp
+++ b/Engine/Source/Editor/UnrealEd/Private/Commandlets/CookCommandlet.cpp
@@ -894,19 +894,31 @@ void UCookCommandlet::CollectFilesToCook(TArray<FString>& FilesInPath)
 		FString Obj;
 		if (PlatformEngineIni.GetString(TEXT("/Script/EngineSettings.GameMapsSettings"), TEXT("GameDefaultMap"), Obj))
 		{
-			FilesInPath.AddUnique(Obj);
+			if (Obj != FName(NAME_None).ToString())
+			{
+				FilesInPath.AddUnique(Obj);
+			}
 		}
 		if (PlatformEngineIni.GetString(TEXT("/Script/EngineSettings.GameMapsSettings"), TEXT("ServerDefaultMap"), Obj))
 		{
-			FilesInPath.AddUnique(Obj);
+			if (Obj != FName(NAME_None).ToString())
+			{
+				FilesInPath.AddUnique(Obj);
+			}
 		}
 		if (PlatformEngineIni.GetString(TEXT("/Script/EngineSettings.GameMapsSettings"), TEXT("GlobalDefaultGameMode"), Obj))
 		{
-			FilesInPath.AddUnique(Obj);
+			if (Obj != FName(NAME_None).ToString())
+			{
+				FilesInPath.AddUnique(Obj);
+			}
 		}
 		if (PlatformEngineIni.GetString(TEXT("/Script/EngineSettings.GameMapsSettings"), TEXT("GlobalDefaultServerGameMode"), Obj))
 		{
-			FilesInPath.AddUnique(Obj);
+			if (Obj != FName(NAME_None).ToString())
+			{
+				FilesInPath.AddUnique(Obj);
+			}
 		}
 	}
 
