Fix for localization gather from assets reporting errors for missing keys on empty texts that wouldn't be gathered anyways. Refactored gathering of localization data from assets to issue warnings and errors more appropriately, separate from actual committing the gathered data. Archive/package flagging for localization requirements are also improved to flag less false-positives.

diff --git a/Engine/Source/Editor/UnrealEd/Classes/Commandlets/GatherTextFromAssetsCommandlet.h b/Engine/Source/Editor/UnrealEd/Classes/Commandlets/GatherTextFromAssetsCommandlet.h
--- a/Engine/Source/Editor/UnrealEd/Classes/Commandlets/GatherTextFromAssetsCommandlet.h
+++ b/Engine/Source/Editor/UnrealEd/Classes/Commandlets/GatherTextFromAssetsCommandlet.h
@@ -6,13 +6,18 @@
 
 namespace EAssetTextGatherStatus
 {
+	enum 
+	{
+		BadBitMask = 1
+	};
 	enum Type
 	{
-		None,
-		MissingKey,
-		MissingKey_Resolved,
-		IdentityConflict,
-		IdentityConflict_Resolved,
+		//							ErrorType << 1 + BadBit
+		None						= (0 << 1) + 0,
+		MissingKey					= (1 << 1) + 1,
+		MissingKey_Resolved			= (1 << 1) + 0,
+		IdentityConflict			= (2 << 1) + 1,
+		IdentityConflict_Resolved	= (2 << 1) + 0,
 	};
 }
 
@@ -28,7 +33,7 @@ class UGatherTextFromAssetsCommandlet : public UGatherTextCommandletBase
 
 	void ProcessPackages( const TArray< UPackage* >& PackagesToProcess );
 	void ProcessDialogueWave( const UDialogueWave* DialogueWave );
-	bool ProcessTextProperty( UTextProperty* TextProp, FText* Data, UObject* Object, bool& OutRepaired );
+	void ProcessTextProperty( UTextProperty* TextProp, FText* Data, UObject* Object, bool& OutRepaired );
 
 public:
 	// Begin UCommandlet Interface
diff --git a/Engine/Source/Editor/UnrealEd/Private/Commandlets/GatherTextFromAssetsCommandlet.cpp b/Engine/Source/Editor/UnrealEd/Private/Commandlets/GatherTextFromAssetsCommandlet.cpp
--- a/Engine/Source/Editor/UnrealEd/Private/Commandlets/GatherTextFromAssetsCommandlet.cpp
+++ b/Engine/Source/Editor/UnrealEd/Private/Commandlets/GatherTextFromAssetsCommandlet.cpp
@@ -411,12 +411,10 @@ namespace
 				}
 
 				bool Fixed = false;
-				if( Commandlet->ProcessTextProperty(TextProp, Text, Object, /*OUT*/Fixed ) )
+				Commandlet->ProcessTextProperty(TextProp, Text, Object, /*OUT*/Fixed );
+				if( Fixed )
 				{
-					if( Fixed )
-					{
-						ObjectPackage->MarkPackageDirty();
-					}
+					ObjectPackage->MarkPackageDirty();
 				}
 			}
 		}
@@ -499,143 +497,117 @@ void UGatherTextFromAssetsCommandlet::ProcessDialogueWave( const UDialogueWave*
 	}
 }
 
-bool UGatherTextFromAssetsCommandlet::ProcessTextProperty(UTextProperty* InTextProp, FText* Data, UObject* Object, bool& OutFixed)
+void UGatherTextFromAssetsCommandlet::ProcessTextProperty(UTextProperty* InTextProp, FText* Data, UObject* Object, bool& OutFixed)
 {
-	bool TextPropertyWasValid = true;
-
 	OutFixed = false;
 
-	// Transient check.
-	if( Data->Flags & ETextFlag::Transient )
+	FConflictTracker::FEntry NewEntry;
+	NewEntry.ObjectPath = Object->GetPathName();
+	NewEntry.SourceString = Data->GetSourceString();
+	NewEntry.Status = EAssetTextGatherStatus::None;
+
+	TSharedPtr< FString, ESPMode::ThreadSafe > Namespace;
+	TSharedPtr< FString, ESPMode::ThreadSafe > Key;
+	FTextLocalizationManager::Get().FindKeyNamespaceFromDisplayString(Data->DisplayString, Namespace, Key);
+
+	// Check if text is localizable and check for missing key. Unlocalizable texts don't need a key.
+	if( ( !Key.IsValid() || Key->IsEmpty() ) && Data->ShouldGatherForLocalization() )
 	{
-		UE_LOG(LogGatherTextFromAssetsCommandlet, Warning, TEXT("Transient text found set to %s in %s  -  %s."), *InTextProp->GetName(), *Object->GetPathName(), *Object->GetName());
-		TextPropertyWasValid = false;
+		NewEntry.Status = EAssetTextGatherStatus::MissingKey;
+
+		// Fix missing key if allowed.
+		if (bFixBroken)
+		{
+			// Create key if needed.
+			if( !( Key.IsValid() ) )
+			{
+				Key = MakeShareable( new FString() );
+			}
+			// Generate new GUID for key.
+			*(Key) = FGuid::NewGuid().ToString();
+
+			// Fixed.
+			NewEntry.Status = EAssetTextGatherStatus::MissingKey_Resolved;
+		}
 	}
-	else
+
+	// Must have valid key to test for identity conflicts. Even if a text doesn't require localization, if it conflicts, it should be warned about.
+	if( Key.IsValid() && !( Key->IsEmpty() ) )
 	{
-		FConflictTracker::FEntry NewEntry;
-		NewEntry.ObjectPath = Object->GetPathName();
-		NewEntry.SourceString = Data->GetSourceString();
-		NewEntry.Status = EAssetTextGatherStatus::None;
+		// Find existing entry from manifest or manifest dependencies.
+		FContext SearchContext;
+		SearchContext.Key = *Key;
+		TSharedPtr< FManifestEntry > ExistingEntry = ManifestInfo->GetManifest()->FindEntryByContext( Namespace.IsValid() ? *Namespace : TEXT(""), SearchContext );
 
-		// Fix missing key if broken and allowed.
-		TSharedPtr< FString, ESPMode::ThreadSafe > Namespace;
-		TSharedPtr< FString, ESPMode::ThreadSafe > Key;
-		FTextLocalizationManager::Get().FindKeyNamespaceFromDisplayString(Data->DisplayString, Namespace, Key);
+		if( !ExistingEntry.IsValid() )
+		{
+			FString FileInfo;
+			ExistingEntry = ManifestInfo->FindDependencyEntrybyContext( Namespace.IsValid() ? *Namespace : TEXT(""), SearchContext, FileInfo );
+		}
 
-		if( !( Key.IsValid() ) || Key->IsEmpty() )
+		// Entry already exists, check for conflict.
+		if( ExistingEntry.IsValid() && ExistingEntry->Source.Text != ( NewEntry.SourceString.IsValid() ? **(NewEntry.SourceString) : TEXT("") ))
 		{
-			// Key fix.
+			NewEntry.Status = EAssetTextGatherStatus::IdentityConflict;
+
+			// Fix conflict if allowed.
 			if (bFixBroken)
 			{
-				// Create key if needed.
-				if( !( Key.IsValid() ) )
-				{
-					Key = MakeShareable( new FString() );
-				}
 				// Generate new GUID for key.
 				*(Key) = FGuid::NewGuid().ToString();
 
 				// Fixed.
-				NewEntry.Status = EAssetTextGatherStatus::MissingKey_Resolved;
-			}
-			else
-			{
-				NewEntry.Status = EAssetTextGatherStatus::MissingKey;
-				TextPropertyWasValid = false;
-			}
-		}
+				NewEntry.Status = EAssetTextGatherStatus::IdentityConflict_Resolved;
 
-		// Must have valid key.
-		if( Key.IsValid() && !( Key->IsEmpty() ) )
-		{
-			FContext SearchContext;
-			SearchContext.Key = Key.IsValid() ? *Key : TEXT("");
-
-			// Find existing entry from manifest or manifest dependencies.
-			TSharedPtr< FManifestEntry > ExistingEntry = ManifestInfo->GetManifest()->FindEntryByContext( Namespace.IsValid() ? *Namespace : TEXT(""), SearchContext );
-			if( !ExistingEntry.IsValid() )
-			{
-				FString FileInfo;
-				ExistingEntry = ManifestInfo->FindDependencyEntrybyContext( Namespace.IsValid() ? *Namespace : TEXT(""), SearchContext, FileInfo );
+				// Conflict resolved, no existing entry.
+				ExistingEntry.Reset();
 			}
+		}
+	}
 
-			// Entry already exists, check for conflict.
-			if( ExistingEntry.IsValid() )
-			{
-				// Fix conflict if present and allowed.
-				if( ExistingEntry->Source.Text != ( NewEntry.SourceString.IsValid() ? **(NewEntry.SourceString) : TEXT("") ) )
-				{
-					if (bFixBroken)
-					{
-						// Generate new GUID for key.
-						*(Key) = FGuid::NewGuid().ToString();
-
-						// Fixed.
-						NewEntry.Status = EAssetTextGatherStatus::IdentityConflict_Resolved;
+	// Gather if it requires gathering and isn't in a bad state.
+	if( Data->ShouldGatherForLocalization() && (NewEntry.Status & EAssetTextGatherStatus::BadBitMask) == false)
+	{
+		FString SrcLocation = Object->GetPathName() + TEXT(".") + InTextProp->GetName();
 
-						// Conflict resolved, no existing entry.
-						ExistingEntry.Reset();
-					}
-					else
-					{
-						NewEntry.Status = EAssetTextGatherStatus::IdentityConflict;
-						TextPropertyWasValid = false;
-					}
-				}
-			}
+		// Adjust the source location if needed.
+		{
+			UClass* Class = Object->GetClass();
+			UObject* CDO = Class ? Class->GetDefaultObject() : NULL;
 
-			// Only add an entry to the manifest if no existing entry exists.
-			if( !( ExistingEntry.IsValid() ) )
+			if( CDO && CDO != Object )
 			{
-				// Check for valid string.
-				if( NewEntry.SourceString.IsValid() && !( NewEntry.SourceString->IsEmpty() ) )
+				for( TFieldIterator<UTextProperty> PropIt(CDO->GetClass(), EFieldIteratorFlags::IncludeSuper); PropIt; ++PropIt )
 				{
-					FString SrcLocation = Object->GetPathName() + TEXT(".") + InTextProp->GetName();
+					UTextProperty* TextProp	=  Cast<UTextProperty>( *(PropIt) );
+					FText* DataCDO = TextProp->ContainerPtrToValuePtr<FText>( CDO );
+					TSharedPtr< FString, ESPMode::ThreadSafe > NamespaceCDO;
+					TSharedPtr< FString, ESPMode::ThreadSafe > KeyCDO;
+					FTextLocalizationManager::Get().FindKeyNamespaceFromDisplayString(DataCDO->DisplayString, NamespaceCDO, KeyCDO);
 
-					// Adjust the source location if needed.
+					if( KeyCDO == Key || ( KeyCDO.Get() && Key.Get() && ( *KeyCDO == *Key ) ) )
 					{
-						UClass* Class = Object->GetClass();
-						UObject* CDO = Class ? Class->GetDefaultObject() : NULL;
-
-						if( CDO && CDO != Object )
-						{
-							for( TFieldIterator<UTextProperty> PropIt(CDO->GetClass(), EFieldIteratorFlags::IncludeSuper); PropIt; ++PropIt )
-							{
-								UTextProperty* TextProp	=  Cast<UTextProperty>( *(PropIt) );
-								FText* DataCDO = TextProp->ContainerPtrToValuePtr<FText>( CDO );
-								TSharedPtr< FString, ESPMode::ThreadSafe > NamespaceCDO;
-								TSharedPtr< FString, ESPMode::ThreadSafe > KeyCDO;
-								FTextLocalizationManager::Get().FindKeyNamespaceFromDisplayString(DataCDO->DisplayString, NamespaceCDO, KeyCDO);
-
-								if( KeyCDO == Key || ( KeyCDO.Get() && Key.Get() && ( *KeyCDO == *Key ) ) )
-								{
-									SrcLocation = CDO->GetPathName() + TEXT(".") + TextProp->GetName();
-									break;
-								}
-							}
-						}
+						SrcLocation = CDO->GetPathName() + TEXT(".") + TextProp->GetName();
+						break;
 					}
-
-					FContext Context;
-					Context.Key = Key.IsValid() ? *Key : TEXT("");
-					Context.SourceLocation = SrcLocation;
-
-					FString EntryDescription = FString::Printf( TEXT("In %s"), *Object->GetFullName());
-					ManifestInfo->AddEntry(EntryDescription, Namespace.Get() ? *Namespace : TEXT(""), NewEntry.SourceString.Get() ? *(NewEntry.SourceString) : TEXT(""), Context );
 				}
 			}
 		}
 
-		// Add to conflict tracker.
-		FConflictTracker::FKeyTable& KeyTable = ConflictTracker.Namespaces.FindOrAdd( Namespace.IsValid() ? *Namespace : TEXT("") );
-		FConflictTracker::FEntryArray& EntryArray = KeyTable.FindOrAdd( Key.IsValid() ? *Key : TEXT("") );
-		EntryArray.Add(NewEntry);
+		FContext Context;
+		Context.Key = *Key;
+		Context.SourceLocation = SrcLocation;
 
-		OutFixed = (NewEntry.Status == EAssetTextGatherStatus::MissingKey_Resolved || NewEntry.Status == EAssetTextGatherStatus::IdentityConflict_Resolved);
+		FString EntryDescription = FString::Printf( TEXT("In %s"), *Object->GetFullName());
+		ManifestInfo->AddEntry(EntryDescription, Namespace.Get() ? *Namespace : TEXT(""), NewEntry.SourceString.Get() ? *(NewEntry.SourceString) : TEXT(""), Context );
 	}
 
-	return TextPropertyWasValid;
+	// Add to conflict tracker.
+	FConflictTracker::FKeyTable& KeyTable = ConflictTracker.Namespaces.FindOrAdd( Namespace.IsValid() ? *Namespace : TEXT("") );
+	FConflictTracker::FEntryArray& EntryArray = KeyTable.FindOrAdd( Key.IsValid() ? *Key : TEXT("") );
+	EntryArray.Add(NewEntry);
+
+	OutFixed = (NewEntry.Status == EAssetTextGatherStatus::MissingKey_Resolved || NewEntry.Status == EAssetTextGatherStatus::IdentityConflict_Resolved);
 }
 
 int32 UGatherTextFromAssetsCommandlet::Main(const FString& Params)
diff --git a/Engine/Source/Runtime/Core/Private/Internationalization/Text.cpp b/Engine/Source/Runtime/Core/Private/Internationalization/Text.cpp
--- a/Engine/Source/Runtime/Core/Private/Internationalization/Text.cpp
+++ b/Engine/Source/Runtime/Core/Private/Internationalization/Text.cpp
@@ -23,6 +23,11 @@ FArchive& operator<<( FArchive& Ar, FFormatArgumentData& Value )
 	return Ar;
 }
 
+bool FTextInspector::ShouldGatherForLocalization(const FText& Text)
+{
+	return Text.ShouldGatherForLocalization();
+}
+
 const FString* FTextInspector::GetNamespace(const FText& Text)
 {
 	TSharedPtr< FString, ESPMode::ThreadSafe > Namespace;
@@ -507,7 +512,10 @@ bool FText::FindText( const FString& Namespace, const FString& Key, FText& OutTe
 
 CORE_API FArchive& operator<<( FArchive& Ar, FText& Value )
 {
-	Ar.ThisRequiresLocalizationGather();
+	if( Value.ShouldGatherForLocalization() )
+	{
+		Ar.ThisRequiresLocalizationGather();
+	}
 
 	//When duplicating, the CDO is used as the template, then values for the instance are assigned.
 	//If we don't duplicate the string, the CDO and the instance are both pointing at the same thing.
@@ -724,6 +732,25 @@ void FText::Rebuild() const
 	}
 }
 
+bool FText::ShouldGatherForLocalization() const
+{
+	auto SourceString = GetSourceString();
+
+	auto IsAllWhitespace = [](FString& String) -> bool
+	{
+		for(int32 i = 0; i < String.Len(); ++i)
+		{
+			if( !FChar::IsWhitespace( String[i] ) )
+			{
+				return false;
+			}
+		}
+		return true;
+	};
+
+	return !((Flags & ETextFlag::CultureInvariant) || (Flags & ETextFlag::Transient)) && SourceString.IsValid() && !SourceString->IsEmpty() && !IsAllWhitespace(*SourceString);
+}
+
 TSharedPtr< FString, ESPMode::ThreadSafe > FText::GetSourceString() const
 {
 	if(History.IsValid())
diff --git a/Engine/Source/Runtime/Core/Public/Internationalization/Text.h b/Engine/Source/Runtime/Core/Public/Internationalization/Text.h
--- a/Engine/Source/Runtime/Core/Public/Internationalization/Text.h
+++ b/Engine/Source/Runtime/Core/Public/Internationalization/Text.h
@@ -295,6 +295,8 @@ private:
 	static FText FormatInternal(const FText& Pattern, const FFormatOrderedArguments& Arguments, bool bInRebuildText, bool bInRebuildAsSource);
 	static FText FormatInternal(const FText& Pattern, const TArray< struct FFormatArgumentData > InArguments, bool bInRebuildText, bool bInRebuildAsSource);
 
+	bool ShouldGatherForLocalization() const;
+
 private:
 	template<typename T1, typename T2>
 	static FText AsNumberTemplate(T1 Val, const FNumberFormattingOptions* const Options = NULL, const TSharedPtr<FCulture, ESPMode::ThreadSafe>& TargetCulture = NULL);
@@ -353,6 +355,7 @@ private:
 	~FTextInspector() {}
 
 public:
+	static bool ShouldGatherForLocalization(const FText& Text);
 	static const FString* GetNamespace(const FText& Text);
 	static const FString* GetKey(const FText& Text);
 	static const FString* GetSourceString(const FText& Text);
