FBoolConfigValueHelper in CoreMisc - Refactor

diff --git a/Engine/Source/Editor/BlueprintGraph/Private/EdGraphSchema_K2.cpp b/Engine/Source/Editor/BlueprintGraph/Private/EdGraphSchema_K2.cpp
--- a/Engine/Source/Editor/BlueprintGraph/Private/EdGraphSchema_K2.cpp
+++ b/Engine/Source/Editor/BlueprintGraph/Private/EdGraphSchema_K2.cpp
@@ -632,17 +632,8 @@ bool UEdGraphSchema_K2::IsAllowableBlueprintVariableType(const class UClass* InC
 			return true;
 		}
 
-		struct FBlueprintIsNotBlueprintTypeHelper
-		{
-			bool bNotBlueprintType;
-
-			FBlueprintIsNotBlueprintTypeHelper() : bNotBlueprintType(false)
-			{
-				GConfig->GetBool(TEXT("EditoronlyBP"), TEXT("bBlueprintIsNotBlueprintType"), bNotBlueprintType, GEditorIni);
-			}
-		};
-		static const FBlueprintIsNotBlueprintTypeHelper BPHelper;
-		if (BPHelper.bNotBlueprintType && InClass->IsChildOf(UBlueprint::StaticClass()))
+		static const FBoolConfigValueHelper NotBlueprintType(TEXT("EditoronlyBP"), TEXT("bBlueprintIsNotBlueprintType"));
+		if (NotBlueprintType && InClass->IsChildOf(UBlueprint::StaticClass()))
 		{
 			return false;
 		}
diff --git a/Engine/Source/Editor/Kismet/Private/UserDefinedStructureEditor.cpp b/Engine/Source/Editor/Kismet/Private/UserDefinedStructureEditor.cpp
--- a/Engine/Source/Editor/Kismet/Private/UserDefinedStructureEditor.cpp
+++ b/Engine/Source/Editor/Kismet/Private/UserDefinedStructureEditor.cpp
@@ -275,16 +275,9 @@ TSharedRef<SDockTab> FUserDefinedStructureEditor::SpawnStructureTab(const FSpawn
 	}
 
 	DefaultValueView = NULL;
-	struct FShowDefaultValuePropertyEditor
-	{
-		bool bShowDefaultValuePropertyEditor;
-		FShowDefaultValuePropertyEditor() : bShowDefaultValuePropertyEditor(false)
-		{
-			GConfig->GetBool(TEXT("UserDefinedStructure"), TEXT("bShowDefaultValuePropertyEditor"), bShowDefaultValuePropertyEditor, GEditorIni);
-		}
-	};
-	static FShowDefaultValuePropertyEditor Helper;
-	if (Helper.bShowDefaultValuePropertyEditor)
+
+	static FBoolConfigValueHelper ShowDefaultValuePropertyEditor(TEXT("UserDefinedStructure"), TEXT("bShowDefaultValuePropertyEditor"));
+	if (ShowDefaultValuePropertyEditor)
 	{
 		DefaultValueView = MakeShareable(new FStructureDefaultValueView(EditedStruct));
 		DefaultValueView->Initialize();
diff --git a/Engine/Source/Editor/UnrealEd/Private/Factories/EditorFactories.cpp b/Engine/Source/Editor/UnrealEd/Private/Factories/EditorFactories.cpp
--- a/Engine/Source/Editor/UnrealEd/Private/Factories/EditorFactories.cpp
+++ b/Engine/Source/Editor/UnrealEd/Private/Factories/EditorFactories.cpp
@@ -6281,16 +6281,8 @@ BlueprintFunctionLibraryFactory implementation.
 UBlueprintFunctionLibraryFactory::UBlueprintFunctionLibraryFactory(const class FPostConstructInitializeProperties& PCIP)
 : Super(PCIP)
 {
-	struct FCanCreateNewHelper
-	{
-		bool bCanCreateNew;
-		FCanCreateNewHelper() : bCanCreateNew(false)
-		{
-			GConfig->GetBool(TEXT("CustomBlueprintFunctionLibrary"), TEXT("bCanCreateNew"), bCanCreateNew, GEditorIni);
-		}
-	};
-	static const FCanCreateNewHelper Helper;
-	bCreateNew = Helper.bCanCreateNew;
+	static FBoolConfigValueHelper CanCreateNewHelper(TEXT("CustomBlueprintFunctionLibrary"), TEXT("bCanCreateNew"));
+	bCreateNew = CanCreateNewHelper;
 	bEditAfterNew = true;
 	SupportedClass = UBlueprint::StaticClass();
 	ParentClass = UBlueprintFunctionLibrary::StaticClass();
diff --git a/Engine/Source/Editor/UnrealEd/Private/Kismet2/BlueprintEditorUtils.cpp b/Engine/Source/Editor/UnrealEd/Private/Kismet2/BlueprintEditorUtils.cpp
--- a/Engine/Source/Editor/UnrealEd/Private/Kismet2/BlueprintEditorUtils.cpp
+++ b/Engine/Source/Editor/UnrealEd/Private/Kismet2/BlueprintEditorUtils.cpp
@@ -1055,19 +1055,8 @@ UClass* FBlueprintEditorUtils::RegenerateBlueprintClass(UBlueprint* Blueprint, U
 
 		FKismetEditorUtilities::GenerateBlueprintSkeleton(Blueprint);
 
-		struct FReplaceBlueprintWithClassHelper
-		{
-			bool bShouldReplace;
-
-			FReplaceBlueprintWithClassHelper() : bShouldReplace(false)
-			{
-				GConfig->GetBool(TEXT("EditoronlyBP"), TEXT("bReplaceBlueprintWithClass"), bShouldReplace, GEditorIni);
-			}
-		};
-		static FReplaceBlueprintWithClassHelper ReplaceBlueprintWithClassHelper;
-		const bool bReplaceBlueprintWithClass = ReplaceBlueprintWithClassHelper.bShouldReplace;
-
-		if (bReplaceBlueprintWithClass)
+		static FBoolConfigValueHelper ReplaceBlueprintWithClass(TEXT("EditoronlyBP"), TEXT("bReplaceBlueprintWithClass"));
+		if (ReplaceBlueprintWithClass)
 		{
 			FEditoronlyBlueprintHelper::ChangeBlueprint(Blueprint);
 		}
@@ -1125,7 +1114,7 @@ UClass* FBlueprintEditorUtils::RegenerateBlueprintClass(UBlueprint* Blueprint, U
 		}
 #endif
 
-		if (bReplaceBlueprintWithClass)
+		if (ReplaceBlueprintWithClass)
 		{
 			FEditoronlyBlueprintHelper::ShouldBeFixed(Blueprint, true);
 		}
diff --git a/Engine/Source/Editor/UnrealEd/Private/Kismet2/StructureEditorUtils.cpp b/Engine/Source/Editor/UnrealEd/Private/Kismet2/StructureEditorUtils.cpp
--- a/Engine/Source/Editor/UnrealEd/Private/Kismet2/StructureEditorUtils.cpp
+++ b/Engine/Source/Editor/UnrealEd/Private/Kismet2/StructureEditorUtils.cpp
@@ -414,16 +414,8 @@ FString FStructureEditorUtils::GetVariableDisplayName(const UUserDefinedStruct*
 
 bool FStructureEditorUtils::UserDefinedStructEnabled()
 {
-	struct FUserDefinedStructEnabled
-	{
-		bool bUseUserDefinedStructure;
-		FUserDefinedStructEnabled() : bUseUserDefinedStructure(false)
-		{
-			GConfig->GetBool(TEXT("UserDefinedStructure"), TEXT("bUseUserDefinedStructure"), bUseUserDefinedStructure, GEditorIni);
-		}
-	};
-	static FUserDefinedStructEnabled Helper;
-	return Helper.bUseUserDefinedStructure;
+	static FBoolConfigValueHelper UseUserDefinedStructure(TEXT("UserDefinedStructure"), TEXT("bUseUserDefinedStructure"));
+	return UseUserDefinedStructure;
 }
 
 bool FStructureEditorUtils::Fill_MakeStructureDefaultValue(const UUserDefinedStruct* Struct, uint8* StructData)
diff --git a/Engine/Source/Runtime/Core/Private/Misc/CoreMisc.cpp b/Engine/Source/Runtime/Core/Private/Misc/CoreMisc.cpp
--- a/Engine/Source/Runtime/Core/Private/Misc/CoreMisc.cpp
+++ b/Engine/Source/Runtime/Core/Private/Misc/CoreMisc.cpp
@@ -951,4 +951,13 @@ void GenerateConvenientWindowedResolutions(const FDisplayMetrics& InDisplayMetri
 			}
 		}
 	}
+}
+
+/*----------------------------------------------------------------------------
+FBoolConfigValueHelper
+----------------------------------------------------------------------------*/
+FBoolConfigValueHelper::FBoolConfigValueHelper(const TCHAR* Section, const TCHAR* Key, const FString& Filename)
+	: bValue(false)
+{
+	GConfig->GetBool(Section, Key, bValue, Filename);
 }
\ No newline at end of file
diff --git a/Engine/Source/Runtime/Core/Public/Misc/CoreMisc.h b/Engine/Source/Runtime/Core/Public/Misc/CoreMisc.h
--- a/Engine/Source/Runtime/Core/Public/Misc/CoreMisc.h
+++ b/Engine/Source/Runtime/Core/Public/Misc/CoreMisc.h
@@ -363,4 +363,18 @@ struct FScriptTraceStackNode
 
 private:
 	FScriptTraceStackNode() {};
+};
+
+/** Helper structure for boolean values in config */
+struct CORE_API FBoolConfigValueHelper
+{
+private:
+	bool bValue;
+public:
+	FBoolConfigValueHelper(const TCHAR* Section, const TCHAR* Key, const FString& Filename = GEditorIni);
+
+	operator bool() const
+	{
+		return bValue;
+	}
 };
\ No newline at end of file
diff --git a/Engine/Source/Runtime/CoreUObject/Private/UObject/Class.cpp b/Engine/Source/Runtime/CoreUObject/Private/UObject/Class.cpp
--- a/Engine/Source/Runtime/CoreUObject/Private/UObject/Class.cpp
+++ b/Engine/Source/Runtime/CoreUObject/Private/UObject/Class.cpp
@@ -1244,16 +1244,8 @@ void UStruct::Serialize( FArchive& Ar )
 
 				bool bSkipByteCodeSerialization = false;
 #if WITH_EDITOR
-				struct FSkipByteCodeSerializationHelper
-				{
-					bool bValue;
-					FSkipByteCodeSerializationHelper() : bValue(false)
-					{
-						GConfig->GetBool(TEXT("StructSerialization"), TEXT("SkipByteCodeSerialization"), bValue, GEditorIni);
-					}
-				};
-				static FSkipByteCodeSerializationHelper SkipByteCodeHelper;
-				bSkipByteCodeSerialization = SkipByteCodeHelper.bValue;
+				static const FBoolConfigValueHelper SkipByteCodeHelper(TEXT("StructSerialization"), TEXT("SkipByteCodeSerialization"));
+				bSkipByteCodeSerialization = SkipByteCodeHelper;
 #endif // WITH_EDITOR
 				if (bSkipByteCodeSerialization || (Ar.UE4Ver() < VER_MIN_SCRIPTVM_UE4) || (Ar.LicenseeUE4Ver() < VER_MIN_SCRIPTVM_LICENSEEUE4))
 				{
