Modified internationalization export to add source file information as an extracted comment as well as a reference comment in order to workaround an issue with OneSky.

diff --git a/Engine/Source/Editor/UnrealEd/Classes/Commandlets/InternationalizationExportCommandlet.h b/Engine/Source/Editor/UnrealEd/Classes/Commandlets/InternationalizationExportCommandlet.h
--- a/Engine/Source/Editor/UnrealEd/Classes/Commandlets/InternationalizationExportCommandlet.h
+++ b/Engine/Source/Editor/UnrealEd/Classes/Commandlets/InternationalizationExportCommandlet.h
@@ -142,10 +142,17 @@ public:
 		return ( MsgId == Other.MsgId && MsgIdPlural == Other.MsgIdPlural && MsgCtxt == Other.MsgCtxt ); 
 	}
 	
+	/**
+	 * Helper function that adds to the extracted comments.
+	 *
+	 * @param InReference	String representing extracted comment entries of the following form: "File/path/file.cpp:20 File/path/file.cpp:21 File/path/file2.cpp:5".
+	 */
+	void AddExtractedComment( const FString& InComment );
+
 	/**
 	 * Helper function that adds to the reference comments.
 	 *
-	 * @param InReference	String representing reference comment entries of the following form: "File/path/file.cpp:20 File/path/file.cpp:21 File/path/file2.cpp:5".
+	 * @param InReference	String representing reference comment entries.
 	 */
 	void AddReference( const FString& InReference );
 
@@ -154,6 +161,13 @@ public:
 	 *
 	 * @param InReferences	String array representing reference comment entries of the following form: "File/path/file.cpp:20 File/path/file.cpp:21 File/path/file2.cpp:5".
 	 */
+	void AddExtractedComments( const TArray<FString>& InComments );
+
+	/**
+	 * Helper function that adds to the extracted comments.
+	 *
+	 * @param InReferences	String array representing extracted comment entries
+	 */
 	void AddReferences( const TArray<FString>& InReferences );
 
 	/**
diff --git a/Engine/Source/Editor/UnrealEd/Private/Commandlets/InternationalizationExportCommandlet.cpp b/Engine/Source/Editor/UnrealEd/Private/Commandlets/InternationalizationExportCommandlet.cpp
--- a/Engine/Source/Editor/UnrealEd/Private/Commandlets/InternationalizationExportCommandlet.cpp
+++ b/Engine/Source/Editor/UnrealEd/Private/Commandlets/InternationalizationExportCommandlet.cpp
@@ -806,13 +806,9 @@ bool FPortableObjectFormatDOM::AddEntry( const TSharedRef< FPortableObjectEntry>
 	if( ExistingEntry.IsValid()  )
 	{
 		ExistingEntry->AddReferences( LocEntry->ReferenceComments );
-		
-		// Checks for a situation that we don't know how to handle yet.  ex. What happens when we match all the IDs for an entry but the other params differ?
-		check( LocEntry->MsgId == ExistingEntry->MsgId );
-		check( LocEntry->MsgCtxt == ExistingEntry->MsgCtxt );
-		check( LocEntry->MsgIdPlural == ExistingEntry->MsgIdPlural );
-		check( LocEntry->ExtractedComments == ExistingEntry->ExtractedComments );
-		check( LocEntry->Flags == ExistingEntry->Flags );
+		ExistingEntry->AddExtractedComments( LocEntry->ExtractedComments );
+
+		// Checks for a situation that we don't know how to handle yet.  ex. What happens when we match all the IDs for an entry but another param differs?
 		check( LocEntry->TranslatorComments == ExistingEntry->TranslatorComments );
 	}
 	else
@@ -847,6 +843,21 @@ TSharedPtr<FPortableObjectEntry> FPortableObjectFormatDOM::FindEntry( const FStr
 	return FindEntry( TempEntry );
 }
 
+void FPortableObjectEntry::AddExtractedComment( const FString& InComment )
+{
+	if(!InComment.IsEmpty())
+	{
+		ExtractedComments.AddUnique(InComment);
+	}
+
+	//// Extracted comments can contain multiple references in a single line so we parse those out.
+	//TArray<FString> CommentsToProcess;
+	//InComment.ParseIntoArray( &CommentsToProcess, TEXT(" "), true );
+	//for( const FString& ExtractedComment : CommentsToProcess )
+	//{
+	//	ExtractedComments.AddUnique( ExtractedComment );
+	//}
+}
 
 void FPortableObjectEntry::AddReference( const FString& InReference )
 {
@@ -864,6 +875,14 @@ void FPortableObjectEntry::AddReference( const FString& InReference )
 	//}
 }
 
+void FPortableObjectEntry::AddExtractedComments( const TArray<FString>& InComments )
+{
+	for( const FString& ExtractedComment : InComments )
+	{
+		AddExtractedComment( ExtractedComment );
+	}
+}
+
 void FPortableObjectEntry::AddReferences( const TArray<FString>& InReferences )
 {
 	for( const FString& Reference : InReferences )
@@ -1032,7 +1051,9 @@ bool UInternationalizationExportCommandlet::DoExport( const FString& SourcePath,
 								PoEntry->MsgCtxt = Namespace;
 								PoEntry->MsgStr.Add( ConditionedArchiveTranslation );
 
-								PoEntry->AddReference( ConvertSrcLocationToPORef( ContextIter->SourceLocation ) );
+								FString PORefString = ConvertSrcLocationToPORef( ContextIter->SourceLocation );
+								PoEntry->AddReference( PORefString ); // The appropriate comment type.
+								PoEntry->AddExtractedComment( PORefString ); // Added to extracted comments to workaround what seems to be an issue with OneSky.
 								PortableObj.AddEntry( PoEntry );
 							}
 						}
