Local variables can be named the same as local variables in other functions, nothing global in a Blueprint can be named the same as a local variable.

Will display an error notification on the editable text box when attempting to change a function's local variable or parameter to be the same.

#ttp 336502 - BP: Local variables can have the same name as function parameters and returns, which causes a duplicate property error.

diff --git a/Engine/Source/Editor/BlueprintGraph/Private/EdGraphSchema_K2.cpp b/Engine/Source/Editor/BlueprintGraph/Private/EdGraphSchema_K2.cpp
--- a/Engine/Source/Editor/BlueprintGraph/Private/EdGraphSchema_K2.cpp
+++ b/Engine/Source/Editor/BlueprintGraph/Private/EdGraphSchema_K2.cpp
@@ -4737,7 +4737,7 @@ void UEdGraphSchema_K2::ConfigureVarNode(UK2Node_Variable* InVarNode, FName InVa
 	}
 	else if (InVariableSource->IsA(UStruct::StaticClass()))
 	{
-		InVarNode->VariableReference.SetLocalMember(InVariableName, InVariableSource, FBlueprintEditorUtils::FindLocalVariableGuidByName(InTargetBlueprint, InVariableName));
+		InVarNode->VariableReference.SetLocalMember(InVariableName, InVariableSource, FBlueprintEditorUtils::FindLocalVariableGuidByName(InTargetBlueprint, InVariableSource, InVariableName));
 	}
 	else
 	{
diff --git a/Engine/Source/Editor/BlueprintGraph/Private/K2Node_VariableGet.cpp b/Engine/Source/Editor/BlueprintGraph/Private/K2Node_VariableGet.cpp
--- a/Engine/Source/Editor/BlueprintGraph/Private/K2Node_VariableGet.cpp
+++ b/Engine/Source/Editor/BlueprintGraph/Private/K2Node_VariableGet.cpp
@@ -96,7 +96,7 @@ FString UK2Node_VariableGet::GetTooltip() const
 	if (VarName != NAME_None)
 	{
 		FString BPMetaData;
-		FBlueprintEditorUtils::GetBlueprintVariableMetaData(GetBlueprint(), VarName, TEXT("tooltip"), BPMetaData);
+		FBlueprintEditorUtils::GetBlueprintVariableMetaData(GetBlueprint(), VarName, VariableReference.GetMemberScope(this), TEXT("tooltip"), BPMetaData);
 
 		if( !BPMetaData.IsEmpty() )
 		{
diff --git a/Engine/Source/Editor/Kismet/Private/BPVariableDragDropAction.cpp b/Engine/Source/Editor/Kismet/Private/BPVariableDragDropAction.cpp
--- a/Engine/Source/Editor/Kismet/Private/BPVariableDragDropAction.cpp
+++ b/Engine/Source/Editor/Kismet/Private/BPVariableDragDropAction.cpp
@@ -198,7 +198,7 @@ void FKismetVariableDragDropAction::HoverTargetChanged()
 
 		if(Blueprint != NULL)
 		{
-			Category = FBlueprintEditorUtils::GetBlueprintVariableCategory(Blueprint, VariableProperty->GetFName() );
+			Category = FBlueprintEditorUtils::GetBlueprintVariableCategory(Blueprint, VariableProperty->GetFName(), GetLocalVariableScope() );
 		}
 
 		FName NewCategory = FName(*HoveredCategoryName,  FNAME_Find);
@@ -505,11 +505,11 @@ FReply FKismetVariableDragDropAction::DroppedOnAction(TSharedRef<FEdGraphSchemaA
 			if(bMoved)
 			{
 				// Change category of var to match the one we dragged on to as well
-				FName MovedVarCategory = FBlueprintEditorUtils::GetBlueprintVariableCategory(BP, VariableName);
-				FName TargetVarCategory = FBlueprintEditorUtils::GetBlueprintVariableCategory(BP, TargetVarName);
+				FName MovedVarCategory = FBlueprintEditorUtils::GetBlueprintVariableCategory(BP, VariableName, GetLocalVariableScope());
+				FName TargetVarCategory = FBlueprintEditorUtils::GetBlueprintVariableCategory(BP, TargetVarName, GetLocalVariableScope());
 				if(MovedVarCategory != TargetVarCategory)
 				{
-					FBlueprintEditorUtils::SetBlueprintVariableCategory(BP, VariableName, TargetVarCategory, true);
+					FBlueprintEditorUtils::SetBlueprintVariableCategory(BP, VariableName, GetLocalVariableScope(), TargetVarCategory, true);
 				}
 
 				// Update Blueprint after changes so they reflect in My Blueprint tab.
@@ -539,10 +539,10 @@ FReply FKismetVariableDragDropAction::DroppedOnCategory(FString Category)
 	if(BP != NULL)
 	{
 		// Check this is actually a different category
-		FName CurrentCategory = FBlueprintEditorUtils::GetBlueprintVariableCategory(BP, VariableName);
+		FName CurrentCategory = FBlueprintEditorUtils::GetBlueprintVariableCategory(BP, VariableName, GetLocalVariableScope());
 		if(FName(*Category) != CurrentCategory)
 		{
-			FBlueprintEditorUtils::SetBlueprintVariableCategory(BP, VariableName, FName(*Category), false);
+			FBlueprintEditorUtils::SetBlueprintVariableCategory(BP, VariableName, GetLocalVariableScope(), FName(*Category), false);
 		}
 	}
 
@@ -566,4 +566,13 @@ bool FKismetVariableDragDropAction::CanVariableBeDropped(const UProperty* InVari
 	return bCanVariableBeDropped;
 }
 
+UStruct* FKismetVariableDragDropAction::GetLocalVariableScope() const
+{
+	if( VariableSource->IsA(UFunction::StaticClass()) )
+	{
+		return VariableSource.Get();
+	}
+	return NULL;
+}
+
 #undef LOCTEXT_NAMESPACE
diff --git a/Engine/Source/Editor/Kismet/Private/BPVariableDragDropAction.h b/Engine/Source/Editor/Kismet/Private/BPVariableDragDropAction.h
--- a/Engine/Source/Editor/Kismet/Private/BPVariableDragDropAction.h
+++ b/Engine/Source/Editor/Kismet/Private/BPVariableDragDropAction.h
@@ -109,6 +109,10 @@ protected:
 	 */
 	bool CanVariableBeDropped(const UProperty* InVariableProperty, const UEdGraph& InGraph) const;
 
+
+	/** Returns the local variable's scope, if any */
+	UStruct* GetLocalVariableScope() const;
+
 protected:
 	/** Name of variable being dragged */
 	FName VariableName;
diff --git a/Engine/Source/Editor/Kismet/Private/BlueprintDetailsCustomization.cpp b/Engine/Source/Editor/Kismet/Private/BlueprintDetailsCustomization.cpp
--- a/Engine/Source/Editor/Kismet/Private/BlueprintDetailsCustomization.cpp
+++ b/Engine/Source/Editor/Kismet/Private/BlueprintDetailsCustomization.cpp
@@ -640,6 +640,16 @@ bool FBlueprintVarActionDetails::IsALocalVariable(UProperty* VariableProperty) c
 	return VariableProperty && (Cast<UFunction>(VariableProperty->GetOuter()) != NULL);
 }
 
+UStruct* FBlueprintVarActionDetails::GetLocalVariableScope(UProperty* VariableProperty) const
+{
+	if(IsALocalVariable(VariableProperty))
+	{
+		return Cast<UFunction>(VariableProperty->GetOuter());
+	}
+
+	return NULL;
+}
+
 bool FBlueprintVarActionDetails::GetVariableNameChangeEnabled() const
 {
 	bool bIsReadOnly = true;
@@ -707,7 +717,7 @@ void FBlueprintVarActionDetails::OnVarNameChanged(const FText& InNewText)
 		}
 	}
 
-	TSharedPtr<INameValidatorInterface> NameValidator = MakeShareable(new FKismetNameValidator(Blueprint, GetVariableName()));
+	TSharedPtr<INameValidatorInterface> NameValidator = MakeShareable(new FKismetNameValidator(Blueprint, GetVariableName(), GetLocalVariableScope(VariableProperty)));
 
 	EValidatorResult ValidatorResult = NameValidator->IsValid(InNewText.ToString());
 	if(ValidatorResult == EValidatorResult::AlreadyInUse)
@@ -722,6 +732,10 @@ void FBlueprintVarActionDetails::OnVarNameChanged(const FText& InNewText)
 	{
 		VarNameEditableTextBox->SetError(LOCTEXT("RenameFailed_NameTooLong", "Names must have fewer than 100 characters!"));
 	}
+	else if(ValidatorResult == EValidatorResult::LocallyInUse)
+	{
+		VarNameEditableTextBox->SetError(LOCTEXT("ConflictsWithProperty", "Conflicts with another another local variable or function parameter!"));
+	}
 	else
 	{
 		bIsVarNameInvalid = false;
@@ -757,7 +771,8 @@ void FBlueprintVarActionDetails::OnVarNameCommitted(const FText& InNewText, ETex
 		}
 		else if(IsALocalVariable(VariableProperty))
 		{
-			FBlueprintEditorUtils::RenameLocalVariable(GetBlueprintObj(), GetVariableName(), NewVarName);
+			UFunction* LocalVarScope = Cast<UFunction>(VariableProperty->GetOuter());
+			FBlueprintEditorUtils::RenameLocalVariable(GetBlueprintObj(), LocalVarScope, GetVariableName(), NewVarName);
 		}
 		else
 		{
@@ -812,7 +827,7 @@ bool FBlueprintVarActionDetails::GetVariableTypeChangeEnabled() const
 			}
 		}
 
-		bool bIsAVarInThisBlueprint = FBlueprintEditorUtils::FindLocalVariable(GetBlueprintObj(), LocalVarAction->GetVariableName()) != NULL;
+		bool bIsAVarInThisBlueprint = FBlueprintEditorUtils::FindLocalVariable(GetBlueprintObj(), LocalVarAction->GetVariableScope(), LocalVarAction->GetVariableName()) != NULL;
 		return !bNodesPendingDeletion && bIsAVarInThisBlueprint;
 	}
 	FEdGraphSchemaAction_K2LocalVar* VarLocalAction = MyBlueprintSelectionAsLocalVar();
@@ -877,7 +892,7 @@ void FBlueprintVarActionDetails::OnVarTypeChanged(const FEdGraphPinType& NewPinT
 			FEdGraphSchemaAction_K2LocalVar* VarLocalAction = MyBlueprintSelectionAsLocalVar();
 			if(VarLocalAction)
 			{
-				FBlueprintEditorUtils::ChangeLocalVariableType(GetBlueprintObj(), VarName, NewPinType);
+				FBlueprintEditorUtils::ChangeLocalVariableType(GetBlueprintObj(), VarLocalAction->GetVariableScope(), VarName, NewPinType);
 			}
 		}
 	}
@@ -889,7 +904,7 @@ FText FBlueprintVarActionDetails::OnGetTooltipText() const
 	if (VarName != NAME_None)
 	{
 		FString Result;
-		FBlueprintEditorUtils::GetBlueprintVariableMetaData(GetBlueprintObj(), VarName, TEXT("tooltip"), Result);
+		FBlueprintEditorUtils::GetBlueprintVariableMetaData(GetBlueprintObj(), VarName, GetLocalVariableScope(SelectionAsProperty()), TEXT("tooltip"), Result);
 		return FText::FromString(Result);
 	}
 	return FText();
@@ -897,7 +912,7 @@ FText FBlueprintVarActionDetails::OnGetTooltipText() const
 
 void FBlueprintVarActionDetails::OnTooltipTextCommitted(const FText& NewText, ETextCommit::Type InTextCommit, FName VarName)
 {
-	FBlueprintEditorUtils::SetBlueprintVariableMetaData(GetBlueprintObj(), VarName, TEXT("tooltip"), NewText.ToString() );
+	FBlueprintEditorUtils::SetBlueprintVariableMetaData(GetBlueprintObj(), VarName, GetLocalVariableScope(SelectionAsProperty()), TEXT("tooltip"), NewText.ToString() );
 }
 
 void FBlueprintVarActionDetails::PopulateCategories(SMyBlueprint* MyBlueprint, TArray<TSharedPtr<FString>>& CategorySource)
@@ -926,7 +941,7 @@ void FBlueprintVarActionDetails::PopulateCategories(SMyBlueprint* MyBlueprint, T
 	CategorySource.Add(MakeShareable(new FString(TEXT("Default"))));
 	for (int32 i = 0; i < VisibleVariables.Num(); ++i)
 	{
-		FName Category = FBlueprintEditorUtils::GetBlueprintVariableCategory(Blueprint, VisibleVariables[i]);
+		FName Category = FBlueprintEditorUtils::GetBlueprintVariableCategory(Blueprint, VisibleVariables[i], NULL);
 		if (Category != NAME_None && Category != Blueprint->GetFName())
 		{
 			bool bNewCategory = true;
@@ -1032,7 +1047,7 @@ FText FBlueprintVarActionDetails::OnGetCategoryText() const
 	{
 		const UEdGraphSchema_K2* K2Schema = GetDefault<UEdGraphSchema_K2>();
 
-		FName Category = FBlueprintEditorUtils::GetBlueprintVariableCategory(GetBlueprintObj(), VarName);
+		FName Category = FBlueprintEditorUtils::GetBlueprintVariableCategory(GetBlueprintObj(), VarName, GetLocalVariableScope(SelectionAsProperty()));
 
 		// Older blueprints will have their name as the default category
 		if( Category == GetBlueprintObj()->GetFName() )
@@ -1055,7 +1070,7 @@ void FBlueprintVarActionDetails::OnCategoryTextCommitted(const FText& NewText, E
 		FString NewCategory = NewText.ToString();
 		if(NewCategory.Len() <= NAME_SIZE)
 		{
-			FBlueprintEditorUtils::SetBlueprintVariableCategory(GetBlueprintObj(), VarName, FName( *NewCategory ));
+			FBlueprintEditorUtils::SetBlueprintVariableCategory(GetBlueprintObj(), VarName, GetLocalVariableScope(SelectionAsProperty()), FName( *NewCategory ));
 			check(MyBlueprint.IsValid());
 			PopulateCategories(MyBlueprint.Pin().Get(), CategorySource);
 			MyBlueprint.Pin()->ExpandCategory(NewCategory);
@@ -1078,7 +1093,7 @@ void FBlueprintVarActionDetails::OnCategorySelectionChanged( TSharedPtr<FString>
 	{
 		FString NewCategory = *ProposedSelection.Get();
 
-		FBlueprintEditorUtils::SetBlueprintVariableCategory(GetBlueprintObj(), VarName, FName( *NewCategory ));
+		FBlueprintEditorUtils::SetBlueprintVariableCategory(GetBlueprintObj(), VarName, GetLocalVariableScope(SelectionAsProperty()), FName( *NewCategory ));
 		CategoryListView.Pin()->ClearSelection();
 		CategoryComboButton.Pin()->SetIsOpen(false);
 		MyBlueprint.Pin()->ExpandCategory(NewCategory);
@@ -1138,11 +1153,11 @@ void FBlueprintVarActionDetails::OnCreateWidgetChanged(ESlateCheckBoxState::Type
 	{
 		if (InNewState == ESlateCheckBoxState::Checked)
 		{
-			FBlueprintEditorUtils::SetBlueprintVariableMetaData(GetBlueprintObj(), VarName, FEdMode::MD_MakeEditWidget, TEXT("true"));
+			FBlueprintEditorUtils::SetBlueprintVariableMetaData(GetBlueprintObj(), VarName, GetLocalVariableScope(SelectionAsProperty()), FEdMode::MD_MakeEditWidget, TEXT("true"));
 		}
 		else
 		{
-			FBlueprintEditorUtils::RemoveBlueprintVariableMetaData(GetBlueprintObj(), VarName, FEdMode::MD_MakeEditWidget);
+			FBlueprintEditorUtils::RemoveBlueprintVariableMetaData(GetBlueprintObj(), VarName, GetLocalVariableScope(SelectionAsProperty()), FEdMode::MD_MakeEditWidget);
 		}
 	}
 }
@@ -1188,11 +1203,11 @@ void FBlueprintVarActionDetails::OnExposedToSpawnChanged(ESlateCheckBoxState::Ty
 		const bool bExposeOnSpawn = (InNewState == ESlateCheckBoxState::Checked);
 		if(bExposeOnSpawn)
 		{
-			FBlueprintEditorUtils::SetBlueprintVariableMetaData(GetBlueprintObj(), VarName, FBlueprintMetadata::MD_ExposeOnSpawn, TEXT("true"));
+			FBlueprintEditorUtils::SetBlueprintVariableMetaData(GetBlueprintObj(), VarName, NULL, FBlueprintMetadata::MD_ExposeOnSpawn, TEXT("true"));
 		}
 		else
 		{
-			FBlueprintEditorUtils::RemoveBlueprintVariableMetaData(GetBlueprintObj(), VarName, FBlueprintMetadata::MD_ExposeOnSpawn);
+			FBlueprintEditorUtils::RemoveBlueprintVariableMetaData(GetBlueprintObj(), VarName, NULL, FBlueprintMetadata::MD_ExposeOnSpawn);
 		} 
 	}
 }
@@ -1233,11 +1248,11 @@ void FBlueprintVarActionDetails::OnPrivateChanged(ESlateCheckBoxState::Type InNe
 		const bool bExposeOnSpawn = (InNewState == ESlateCheckBoxState::Checked);
 		if(bExposeOnSpawn)
 		{
-			FBlueprintEditorUtils::SetBlueprintVariableMetaData(GetBlueprintObj(), VarName, FBlueprintMetadata::MD_Private, TEXT("true"));
+			FBlueprintEditorUtils::SetBlueprintVariableMetaData(GetBlueprintObj(), VarName, NULL, FBlueprintMetadata::MD_Private, TEXT("true"));
 		}
 		else
 		{
-			FBlueprintEditorUtils::RemoveBlueprintVariableMetaData(GetBlueprintObj(), VarName, FBlueprintMetadata::MD_Private);
+			FBlueprintEditorUtils::RemoveBlueprintVariableMetaData(GetBlueprintObj(), VarName, NULL, FBlueprintMetadata::MD_Private);
 		}
 	}
 }
@@ -1304,7 +1319,7 @@ FText FBlueprintVarActionDetails::OnGetMetaKeyValue(FName Key) const
 	if (VarName != NAME_None)
 	{
 		FString Result;
-		FBlueprintEditorUtils::GetBlueprintVariableMetaData(GetBlueprintObj(), VarName, Key, /*out*/ Result);
+		FBlueprintEditorUtils::GetBlueprintVariableMetaData(GetBlueprintObj(), VarName, GetLocalVariableScope(SelectionAsProperty()), Key, /*out*/ Result);
 
 		return FText::FromString(Result);
 	}
@@ -1318,7 +1333,7 @@ void FBlueprintVarActionDetails::OnMetaKeyValueChanged(const FText& NewMinValue,
 	{
 		if ((CommitInfo == ETextCommit::OnEnter) || (CommitInfo == ETextCommit::OnUserMovedFocus))
 		{
-			FBlueprintEditorUtils::SetBlueprintVariableMetaData(GetBlueprintObj(), VarName, Key, NewMinValue.ToString());
+			FBlueprintEditorUtils::SetBlueprintVariableMetaData(GetBlueprintObj(), VarName, GetLocalVariableScope(SelectionAsProperty()), Key, NewMinValue.ToString());
 		}
 	}
 }
@@ -2389,7 +2404,7 @@ FText FBlueprintDelegateActionDetails::OnGetTooltipText() const
 	if (UMulticastDelegateProperty* DelegateProperty = GetDelegatePoperty())
 	{
 		FString Result;
-		FBlueprintEditorUtils::GetBlueprintVariableMetaData(GetBlueprintObj(), DelegateProperty->GetFName(), TEXT("tooltip"), Result);
+		FBlueprintEditorUtils::GetBlueprintVariableMetaData(GetBlueprintObj(), DelegateProperty->GetFName(), NULL, TEXT("tooltip"), Result);
 		return FText::FromString(Result);
 	}
 	return FText();
@@ -2399,7 +2414,7 @@ void FBlueprintDelegateActionDetails::OnTooltipTextCommitted(const FText& NewTex
 {
 	if (UMulticastDelegateProperty* DelegateProperty = GetDelegatePoperty())
 	{
-		FBlueprintEditorUtils::SetBlueprintVariableMetaData(GetBlueprintObj(), DelegateProperty->GetFName(), TEXT("tooltip"), NewText.ToString() );
+		FBlueprintEditorUtils::SetBlueprintVariableMetaData(GetBlueprintObj(), DelegateProperty->GetFName(), NULL, TEXT("tooltip"), NewText.ToString() );
 	}
 }
 
@@ -2409,7 +2424,7 @@ FText FBlueprintDelegateActionDetails::OnGetCategoryText() const
 	{
 		const UEdGraphSchema_K2* K2Schema = GetDefault<UEdGraphSchema_K2>();
 		FName DelegateName = DelegateProperty->GetFName();
-		FName Category = FBlueprintEditorUtils::GetBlueprintVariableCategory(GetBlueprintObj(), DelegateName);
+		FName Category = FBlueprintEditorUtils::GetBlueprintVariableCategory(GetBlueprintObj(), DelegateName, NULL);
 
 		// Older blueprints will have their name as the default category
 		if( Category == GetBlueprintObj()->GetFName() )
@@ -2433,7 +2448,7 @@ void FBlueprintDelegateActionDetails::OnCategoryTextCommitted(const FText& NewTe
 		{
 			FString NewCategory = NewText.ToString();
 			
-			FBlueprintEditorUtils::SetBlueprintVariableCategory(GetBlueprintObj(), DelegateProperty->GetFName(), FName( *NewCategory ));
+			FBlueprintEditorUtils::SetBlueprintVariableCategory(GetBlueprintObj(), DelegateProperty->GetFName(), NULL, FName( *NewCategory ));
 			check(MyBlueprint.IsValid());
 			FBlueprintVarActionDetails::PopulateCategories(MyBlueprint.Pin().Get(), CategorySource);
 			MyBlueprint.Pin()->ExpandCategory(NewCategory);
@@ -2456,7 +2471,7 @@ void FBlueprintDelegateActionDetails::OnCategorySelectionChanged( TSharedPtr<FSt
 	{
 		FString NewCategory = *ProposedSelection.Get();
 
-		FBlueprintEditorUtils::SetBlueprintVariableCategory(GetBlueprintObj(), DelegateProperty->GetFName(), FName( *NewCategory ));
+		FBlueprintEditorUtils::SetBlueprintVariableCategory(GetBlueprintObj(), DelegateProperty->GetFName(), NULL, FName( *NewCategory ));
 		CategoryListView.Pin()->ClearSelection();
 		CategoryComboButton.Pin()->SetIsOpen(false);
 		MyBlueprint.Pin()->ExpandCategory(NewCategory);
@@ -4224,7 +4239,7 @@ FText FBlueprintComponentDetails::OnGetTooltipText() const
 	if (VarName != NAME_None)
 	{
 		FString Result;
-		FBlueprintEditorUtils::GetBlueprintVariableMetaData(GetBlueprintObj(), VarName, TEXT("tooltip"), Result);
+		FBlueprintEditorUtils::GetBlueprintVariableMetaData(GetBlueprintObj(), VarName, NULL, TEXT("tooltip"), Result);
 		return FText::FromString(Result);
 	}
 
@@ -4233,7 +4248,7 @@ FText FBlueprintComponentDetails::OnGetTooltipText() const
 
 void FBlueprintComponentDetails::OnTooltipTextCommitted(const FText& NewText, ETextCommit::Type InTextCommit, FName VarName)
 {
-	FBlueprintEditorUtils::SetBlueprintVariableMetaData(GetBlueprintObj(), VarName, TEXT("tooltip"), NewText.ToString() );
+	FBlueprintEditorUtils::SetBlueprintVariableMetaData(GetBlueprintObj(), VarName, NULL, TEXT("tooltip"), NewText.ToString() );
 }
 
 bool FBlueprintComponentDetails::OnVariableCategoryChangeEnabled() const
@@ -4252,7 +4267,7 @@ FText FBlueprintComponentDetails::OnGetVariableCategoryText() const
 	{
 		const UEdGraphSchema_K2* K2Schema = GetDefault<UEdGraphSchema_K2>();
 
-		FName Category = FBlueprintEditorUtils::GetBlueprintVariableCategory(GetBlueprintObj(), VarName);
+		FName Category = FBlueprintEditorUtils::GetBlueprintVariableCategory(GetBlueprintObj(), VarName, NULL);
 
 		// Older blueprints will have their name as the default category
 		if( Category == GetBlueprintObj()->GetFName() )
@@ -4276,7 +4291,7 @@ void FBlueprintComponentDetails::OnVariableCategoryTextCommitted(const FText& Ne
 	{
 		FString NewCategory = NewText.ToString();
 
-		FBlueprintEditorUtils::SetBlueprintVariableCategory(GetBlueprintObj(), CachedNodePtr->GetVariableName(), FName( *NewCategory ));
+		FBlueprintEditorUtils::SetBlueprintVariableCategory(GetBlueprintObj(), CachedNodePtr->GetVariableName(), NULL, FName( *NewCategory ));
 		PopulateVariableCategories();
 	}
 }
@@ -4289,7 +4304,7 @@ void FBlueprintComponentDetails::OnVariableCategorySelectionChanged( TSharedPtr<
 	if (ProposedSelection.IsValid() && VarName != NAME_None)
 	{
 		FString NewCategory = *ProposedSelection.Get();
-		FBlueprintEditorUtils::SetBlueprintVariableCategory(GetBlueprintObj(), VarName, FName( *NewCategory ));
+		FBlueprintEditorUtils::SetBlueprintVariableCategory(GetBlueprintObj(), VarName, NULL, FName( *NewCategory ));
 
 		check(VariableCategoryListView.IsValid());
 		check(VariableCategoryComboButton.IsValid());
@@ -4332,7 +4347,7 @@ void FBlueprintComponentDetails::PopulateVariableCategories()
 	VariableCategorySource.Add(MakeShareable(new FString(TEXT("Default"))));
 	for (int32 i = 0; i < VisibleVariables.Num(); ++i)
 	{
-		FName Category = FBlueprintEditorUtils::GetBlueprintVariableCategory(Blueprint, VisibleVariables[i]);
+		FName Category = FBlueprintEditorUtils::GetBlueprintVariableCategory(Blueprint, VisibleVariables[i], NULL);
 		if (Category != NAME_None && Category != Blueprint->GetFName())
 		{
 			bool bNewCategory = true;
diff --git a/Engine/Source/Editor/Kismet/Private/BlueprintDetailsCustomization.h b/Engine/Source/Editor/Kismet/Private/BlueprintDetailsCustomization.h
--- a/Engine/Source/Editor/Kismet/Private/BlueprintDetailsCustomization.h
+++ b/Engine/Source/Editor/Kismet/Private/BlueprintDetailsCustomization.h
@@ -48,6 +48,7 @@ private:
 	bool IsAComponentVariable(UProperty* VariableProperty) const;
 	bool IsABlueprintVariable(UProperty* VariableProperty) const;
 	bool IsALocalVariable(UProperty* VariableProperty) const;
+	UStruct* GetLocalVariableScope(UProperty* VariableProperty) const;
 
 	// Callbacks for uproperty details customization
 	bool GetVariableNameChangeEnabled() const;
diff --git a/Engine/Source/Editor/Kismet/Private/BlueprintEditor.cpp b/Engine/Source/Editor/Kismet/Private/BlueprintEditor.cpp
--- a/Engine/Source/Editor/Kismet/Private/BlueprintEditor.cpp
+++ b/Engine/Source/Editor/Kismet/Private/BlueprintEditor.cpp
@@ -5586,7 +5586,7 @@ void FBlueprintEditor::OnAddNewLocalVariable()
 		// Now create new variable
 		FBPVariableDescription NewVar;
 
-		NewVar.VarName = FBlueprintEditorUtils::FindUniqueKismetName(GetBlueprintObj(), TEXT("NewLocalVar"));
+		NewVar.VarName = FBlueprintEditorUtils::FindUniqueKismetName(GetBlueprintObj(), TEXT("NewLocalVar"), FindField<UFunction>(GetBlueprintObj()->SkeletonGeneratedClass, TargetGraph->GetFName()));
 		NewVar.VarGuid = FGuid::NewGuid();
 		NewVar.VarType = GetMyBlueprintWidget()->GetLastPinTypeUsed();
 		NewVar.FriendlyName = FName::NameToDisplayString( NewVar.VarName.ToString(), (NewVar.VarType.PinCategory == K2Schema->PC_Boolean) ? true : false );
@@ -5691,12 +5691,7 @@ void FBlueprintEditor::NewDocument_OnClicked(ECreatedDocumentType GraphType)
 	FName DocumentName = FName(*DocumentNameText.ToString());
 
 	// Make sure the new name is valid
-	int32 Index = 0;
-	while (!FBlueprintEditorUtils::IsGraphNameUnique(GetBlueprintObj(), DocumentName))
-	{
-		DocumentName = FName(*FString::Printf(TEXT("%s%i"), *DocumentNameText.ToString(), Index));
-		++Index;
-	}
+	DocumentName = FBlueprintEditorUtils::FindUniqueKismetName(GetBlueprintObj(), DocumentNameText.ToString());
 
 	check(IsEditingSingleBlueprint());
 
diff --git a/Engine/Source/Editor/Kismet/Private/FindInBlueprintManager.cpp b/Engine/Source/Editor/Kismet/Private/FindInBlueprintManager.cpp
--- a/Engine/Source/Editor/Kismet/Private/FindInBlueprintManager.cpp
+++ b/Engine/Source/Editor/Kismet/Private/FindInBlueprintManager.cpp
@@ -290,7 +290,11 @@ namespace BlueprintSearchMetaDataHelpers
 
 		// Find the variable's tooltip
 		FString TooltipResult;
-		FBlueprintEditorUtils::GetBlueprintVariableMetaData(InBlueprint, InVariableDescription.VarName, FBlueprintMetadata::MD_Tooltip, TooltipResult);
+		
+		if(InVariableDescription.HasMetaData(FBlueprintMetadata::MD_Tooltip))
+		{
+			TooltipResult = InVariableDescription.GetMetaData(FBlueprintMetadata::MD_Tooltip);
+		}
 		InWriter->WriteValue(FFindInBlueprintSearchTags::FiB_Tooltip, TooltipResult);
 
 		// Save the variable's pin type
diff --git a/Engine/Source/Editor/Kismet/Private/SBlueprintPalette.cpp b/Engine/Source/Editor/Kismet/Private/SBlueprintPalette.cpp
--- a/Engine/Source/Editor/Kismet/Private/SBlueprintPalette.cpp
+++ b/Engine/Source/Editor/Kismet/Private/SBlueprintPalette.cpp
@@ -117,7 +117,7 @@ static FString GetVarTooltip(UBlueprint* InBlueprint, UClass* VarClass, FName Va
 			}
 			else
 			{
-				FBlueprintEditorUtils::GetBlueprintVariableMetaData(InBlueprint, VarName, TEXT("tooltip"), ResultTooltip);
+				FBlueprintEditorUtils::GetBlueprintVariableMetaData(InBlueprint, VarName, NULL, TEXT("tooltip"), ResultTooltip);
 			}
 		}
 	}
@@ -784,7 +784,7 @@ private:
 			TSharedPtr<FEdGraphSchemaAction_K2Var> VarAction = StaticCastSharedPtr<FEdGraphSchemaAction_K2Var>(ActionPtr.Pin());
 
 			FString Result;
-			FBlueprintEditorUtils::GetBlueprintVariableMetaData( BlueprintEditorPtr.Pin()->GetBlueprintObj(), VarAction->GetVariableName(), TEXT("tooltip"), Result);
+			FBlueprintEditorUtils::GetBlueprintVariableMetaData( BlueprintEditorPtr.Pin()->GetBlueprintObj(), VarAction->GetVariableName(), NULL, TEXT("tooltip"), Result);
 
 			if(!Result.IsEmpty())
 			{
@@ -816,7 +816,7 @@ private:
 			TSharedPtr<FEdGraphSchemaAction_K2Var> VarAction = StaticCastSharedPtr<FEdGraphSchemaAction_K2Var>(ActionPtr.Pin());
 
 			FString Result;
-			FBlueprintEditorUtils::GetBlueprintVariableMetaData( BlueprintEditorPtr.Pin()->GetBlueprintObj(), VarAction->GetVariableName(), TEXT("tooltip"), Result);
+			FBlueprintEditorUtils::GetBlueprintVariableMetaData( BlueprintEditorPtr.Pin()->GetBlueprintObj(), VarAction->GetVariableName(), NULL, TEXT("tooltip"), Result);
 			if(!Result.IsEmpty())
 			{
 				ToolTip = LOCTEXT("VariablePrivacy_is_public_Tooltip", "Variable is public and is editable on each instance of this Blueprint.").ToString();
@@ -1017,6 +1017,8 @@ bool SBlueprintPaletteItem::OnNameTextVerifyChanged(const FText& InNewText, FTex
 
 	FName OriginalName;
 
+	UStruct* ValidationScope = NULL;
+
 	// Check if certain action names are unchanged.
 	if (ActionPtr.Pin()->GetTypeId() == FEdGraphSchemaAction_K2Var::StaticGetTypeId())
 	{
@@ -1027,6 +1029,8 @@ bool SBlueprintPaletteItem::OnNameTextVerifyChanged(const FText& InNewText, FTex
 	{
 		FEdGraphSchemaAction_K2LocalVar* LocalVarAction = (FEdGraphSchemaAction_K2LocalVar*)ActionPtr.Pin().Get();
 		OriginalName = (LocalVarAction->GetVariableName());
+		
+		ValidationScope = LocalVarAction->GetVariableScope();
 	}
 	else
 	{
@@ -1066,7 +1070,7 @@ bool SBlueprintPaletteItem::OnNameTextVerifyChanged(const FText& InNewText, FTex
 		}
 	}
 
-	TSharedPtr<INameValidatorInterface> NameValidator = MakeShareable(new FKismetNameValidator(Blueprint, OriginalName));
+	TSharedPtr<INameValidatorInterface> NameValidator = MakeShareable(new FKismetNameValidator(Blueprint, OriginalName, ValidationScope));
 
 	EValidatorResult ValidatorResult = NameValidator->IsValid(TextAsString);
 	if(ValidatorResult == EValidatorResult::AlreadyInUse)
@@ -1081,6 +1085,10 @@ bool SBlueprintPaletteItem::OnNameTextVerifyChanged(const FText& InNewText, FTex
 	{
 		OutErrorMessage = LOCTEXT("RenameFailed_NameTooLong", "Names must have fewer than 100 characters!");
 	}
+	else if(ValidatorResult == EValidatorResult::LocallyInUse)
+	{
+		OutErrorMessage = LOCTEXT("ConflictsWithProperty", "Conflicts with another another local variable or function parameter!");
+	}
 
 	if(OutErrorMessage.IsEmpty())
 	{
@@ -1237,7 +1245,7 @@ void SBlueprintPaletteItem::OnNameTextCommitted(const FText& NewText, ETextCommi
 		BlueprintEditorPtr.Pin()->GetBlueprintObj()->Modify();
 
 		FName NewVarName = FName(*NewText.ToString());
-		FBlueprintEditorUtils::RenameLocalVariable(BlueprintEditorPtr.Pin()->GetBlueprintObj(), LocalVarAction->GetVariableName(), NewVarName);
+		FBlueprintEditorUtils::RenameLocalVariable(BlueprintEditorPtr.Pin()->GetBlueprintObj(), LocalVarAction->GetVariableScope(), LocalVarAction->GetVariableName(), NewVarName);
 	}
 	BlueprintEditorPtr.Pin()->GetMyBlueprintWidget()->SelectItemByName(FName(*NewText.ToString()), ESelectInfo::OnMouseClick);
 }
@@ -1328,7 +1336,7 @@ FText SBlueprintPaletteItem::GetToolTipText() const
 			else
 			{
 				FString Result;
-				FBlueprintEditorUtils::GetBlueprintVariableMetaData(BlueprintEditorPtr.Pin()->GetBlueprintObj(), LocalVarAction->GetVariableName(), TEXT("tooltip"), Result);
+				FBlueprintEditorUtils::GetBlueprintVariableMetaData(BlueprintEditorPtr.Pin()->GetBlueprintObj(), LocalVarAction->GetVariableName(), LocalVarAction->GetVariableScope(), TEXT("tooltip"), Result);
 				// Only use the variable tooltip if it has been filled out.
 				ToolTipText = !Result.IsEmpty() ? Result : GetVarType(LocalVarAction->GetVariableScope(), LocalVarAction->GetVariableName(), true, true);
 			}
diff --git a/Engine/Source/Editor/Kismet/Private/SMyBlueprint.cpp b/Engine/Source/Editor/Kismet/Private/SMyBlueprint.cpp
--- a/Engine/Source/Editor/Kismet/Private/SMyBlueprint.cpp
+++ b/Engine/Source/Editor/Kismet/Private/SMyBlueprint.cpp
@@ -295,7 +295,13 @@ void SMyBlueprint::OnCategoryNameCommitted(const FText& InNewText, ETextCommit::
 			{
 				FEdGraphSchemaAction_K2Var* VarAction = (FEdGraphSchemaAction_K2Var*)Actions[i].Get();
 
-				FBlueprintEditorUtils::SetBlueprintVariableCategory(GetBlueprintObj(), VarAction->GetVariableName(), FName( *InNewText.ToString() ), true);
+				FBlueprintEditorUtils::SetBlueprintVariableCategory(GetBlueprintObj(), VarAction->GetVariableName(), NULL, FName( *InNewText.ToString() ), true);
+			}
+			else if (Actions[i]->GetTypeId() == FEdGraphSchemaAction_K2LocalVar::StaticGetTypeId())
+			{
+				FEdGraphSchemaAction_K2LocalVar* LocalVarAction = (FEdGraphSchemaAction_K2LocalVar*)Actions[i].Get();
+
+				FBlueprintEditorUtils::SetBlueprintVariableCategory(GetBlueprintObj(), LocalVarAction->GetVariableName(), LocalVarAction->GetVariableScope(), FName( *InNewText.ToString() ), true);
 			}
 		}
 		FBlueprintEditorUtils::MarkBlueprintAsStructurallyModified(GetBlueprintObj());
@@ -1570,7 +1576,7 @@ void SMyBlueprint::OnDeleteEntry()
 		check(FunctionEntryNodes.Num() == 1);
 		FunctionEntryNodes[0]->Modify();
 
-		FBlueprintEditorUtils::RemoveLocalVariable(GetBlueprintObj(), LocalVarAction->GetVariableName());
+		FBlueprintEditorUtils::RemoveLocalVariable(GetBlueprintObj(), LocalVarAction->GetVariableScope(), LocalVarAction->GetVariableName());
 	}
 	else if (FEdGraphSchemaAction_K2Event* EventAction = SelectionAsEvent())
 	{
diff --git a/Engine/Source/Editor/UnrealEd/Private/Kismet2/BlueprintEditorUtils.cpp b/Engine/Source/Editor/UnrealEd/Private/Kismet2/BlueprintEditorUtils.cpp
--- a/Engine/Source/Editor/UnrealEd/Private/Kismet2/BlueprintEditorUtils.cpp
+++ b/Engine/Source/Editor/UnrealEd/Private/Kismet2/BlueprintEditorUtils.cpp
@@ -2563,7 +2563,7 @@ void FBlueprintEditorUtils::SetBlueprintOnlyEditableFlag(UBlueprint* Blueprint,
 
 	if (bNewBlueprintOnly)
 	{
-		FBlueprintEditorUtils::RemoveBlueprintVariableMetaData(Blueprint, VarName, FEdMode::MD_MakeEditWidget);
+		FBlueprintEditorUtils::RemoveBlueprintVariableMetaData(Blueprint, VarName, NULL, FEdMode::MD_MakeEditWidget);
 	}
 
 	if (VarIndex != INDEX_NONE)
@@ -2637,82 +2637,99 @@ void FBlueprintEditorUtils::SetVariableSaveGameFlag(UBlueprint* InBlueprint, con
 	FBlueprintEditorUtils::MarkBlueprintAsStructurallyModified(InBlueprint);
 }
 
-void FBlueprintEditorUtils::SetBlueprintVariableMetaData(UBlueprint* Blueprint, const FName& VarName, const FName& MetaDataKey, const FString& MetaDataValue)
+void FBlueprintEditorUtils::SetBlueprintVariableMetaData(UBlueprint* Blueprint, const FName& VarName, const UStruct* InLocalVarScope, const FName& MetaDataKey, const FString& MetaDataValue)
 {
-	const int32 VarIndex = FBlueprintEditorUtils::FindNewVariableIndex(Blueprint, VarName);
-	if (VarIndex == INDEX_NONE)
+	// If there is a local var scope, we know we are looking at a local variable
+	if(InLocalVarScope)
 	{
-		//Not a NewVariable is the VarName from a Timeline?
-		const int32 TimelineIndex = FBlueprintEditorUtils::FindTimelineIndex(Blueprint,VarName);
-
-		if (TimelineIndex == INDEX_NONE)
+		if(FBPVariableDescription* LocalVariable = FindLocalVariable(Blueprint, InLocalVarScope, VarName))
+		{
+			LocalVariable->SetMetaData(MetaDataKey, *MetaDataValue);
+		}
+	}
+	else
+	{
+		const int32 VarIndex = FBlueprintEditorUtils::FindNewVariableIndex(Blueprint, VarName);
+		if (VarIndex == INDEX_NONE)
 		{
-			//Not a Timeline is this a SCS Node?
-			const int32 SCS_NodeIndex = FBlueprintEditorUtils::FindSCS_Node(Blueprint,VarName);
+			//Not a NewVariable is the VarName from a Timeline?
+			const int32 TimelineIndex = FBlueprintEditorUtils::FindTimelineIndex(Blueprint,VarName);
 
-			if (SCS_NodeIndex == INDEX_NONE)
+			if (TimelineIndex == INDEX_NONE)
 			{
-				// Check to see if it is a local variable
-				if(FBPVariableDescription* LocalVariable = FindLocalVariable(Blueprint, VarName))
+				//Not a Timeline is this a SCS Node?
+				const int32 SCS_NodeIndex = FBlueprintEditorUtils::FindSCS_Node(Blueprint,VarName);
+
+				if (SCS_NodeIndex != INDEX_NONE)
 				{
-					LocalVariable->SetMetaData(MetaDataKey, *MetaDataValue);
+					Blueprint->SimpleConstructionScript->GetAllNodes()[SCS_NodeIndex]->SetMetaData(MetaDataKey, MetaDataValue);
 				}
+
 			}
 			else
 			{
-				Blueprint->SimpleConstructionScript->GetAllNodes()[SCS_NodeIndex]->SetMetaData(MetaDataKey, MetaDataValue);
+				Blueprint->Timelines[TimelineIndex]->SetMetaData(MetaDataKey, MetaDataValue);
 			}
-
 		}
 		else
 		{
-			Blueprint->Timelines[TimelineIndex]->SetMetaData(MetaDataKey, MetaDataValue);
-		}
-	}
-	else
-	{
-		Blueprint->NewVariables[VarIndex].SetMetaData(MetaDataKey, MetaDataValue);
-		UProperty* Property = FindField<UProperty>(Blueprint->SkeletonGeneratedClass, VarName);
-		Property->SetMetaData(MetaDataKey, *MetaDataValue);
-		
-		Property = FindField<UProperty>(Blueprint->GeneratedClass, VarName);
-		if (Property)
-		{
+			Blueprint->NewVariables[VarIndex].SetMetaData(MetaDataKey, MetaDataValue);
+			UProperty* Property = FindField<UProperty>(Blueprint->SkeletonGeneratedClass, VarName);
 			Property->SetMetaData(MetaDataKey, *MetaDataValue);
+
+			Property = FindField<UProperty>(Blueprint->GeneratedClass, VarName);
+			if (Property)
+			{
+				Property->SetMetaData(MetaDataKey, *MetaDataValue);
+			}
 		}
 	}
 	FBlueprintEditorUtils::MarkBlueprintAsStructurallyModified(Blueprint);
 }
 
-bool FBlueprintEditorUtils::GetBlueprintVariableMetaData(const UBlueprint* Blueprint, const FName& VarName, const FName& MetaDataKey, FString& OutMetaDataValue)
+bool FBlueprintEditorUtils::GetBlueprintVariableMetaData(const UBlueprint* Blueprint, const FName& VarName, const UStruct* InLocalVarScope, const FName& MetaDataKey, FString& OutMetaDataValue)
 {
-	const int32 VarIndex = FBlueprintEditorUtils::FindNewVariableIndex(Blueprint, VarName);
-	if (VarIndex == INDEX_NONE)
+	// If there is a local var scope, we know we are looking at a local variable
+	if(InLocalVarScope)
 	{
-		//Not a NewVariable is the VarName from a Timeline?
-		const int32 TimelineIndex = FBlueprintEditorUtils::FindTimelineIndex(Blueprint,VarName);
-
-		if (TimelineIndex == INDEX_NONE)
+		if(FBPVariableDescription* LocalVariable = FindLocalVariable(Blueprint, InLocalVarScope, VarName))
+		{
+			int32 EntryIndex = LocalVariable->FindMetaDataEntryIndexForKey(MetaDataKey);
+			if (EntryIndex != INDEX_NONE)
+			{
+				OutMetaDataValue = LocalVariable->GetMetaData(MetaDataKey);
+				return true;
+			}
+		}
+	}
+	else
+	{
+		const int32 VarIndex = FBlueprintEditorUtils::FindNewVariableIndex(Blueprint, VarName);
+		if (VarIndex == INDEX_NONE)
 		{
-			//Not a Timeline is this a SCS Node?
-			const int32 SCS_NodeIndex = FBlueprintEditorUtils::FindSCS_Node(Blueprint,VarName);
+			//Not a NewVariable is the VarName from a Timeline?
+			const int32 TimelineIndex = FBlueprintEditorUtils::FindTimelineIndex(Blueprint,VarName);
 
-			if (SCS_NodeIndex == INDEX_NONE)
+			if (TimelineIndex == INDEX_NONE)
 			{
-				// Check to see if it is a local variable
-				if(FBPVariableDescription* LocalVariable = FindLocalVariable(Blueprint, VarName))
+				//Not a Timeline is this a SCS Node?
+				const int32 SCS_NodeIndex = FBlueprintEditorUtils::FindSCS_Node(Blueprint,VarName);
+
+				if (SCS_NodeIndex != INDEX_NONE)
 				{
-					int32 EntryIndex = LocalVariable->FindMetaDataEntryIndexForKey(MetaDataKey);
+					USCS_Node& Desc = *Blueprint->SimpleConstructionScript->GetAllNodes()[SCS_NodeIndex];
+
+					int32 EntryIndex = Desc.FindMetaDataEntryIndexForKey(MetaDataKey);
 					if (EntryIndex != INDEX_NONE)
 					{
-						OutMetaDataValue = LocalVariable->GetMetaData(MetaDataKey);
+						OutMetaDataValue = Desc.GetMetaData(MetaDataKey);
 						return true;
 					}
 				}
 			}
 			else
 			{
-				USCS_Node& Desc = *Blueprint->SimpleConstructionScript->GetAllNodes()[SCS_NodeIndex];
+				UTimelineTemplate& Desc = *Blueprint->Timelines[TimelineIndex];
 
 				int32 EntryIndex = Desc.FindMetaDataEntryIndexForKey(MetaDataKey);
 				if (EntryIndex != INDEX_NONE)
@@ -2724,7 +2741,7 @@ bool FBlueprintEditorUtils::GetBlueprintVariableMetaData(const UBlueprint* Bluep
 		}
 		else
 		{
-			UTimelineTemplate& Desc = *Blueprint->Timelines[TimelineIndex];
+			const FBPVariableDescription& Desc = Blueprint->NewVariables[VarIndex];
 
 			int32 EntryIndex = Desc.FindMetaDataEntryIndexForKey(MetaDataKey);
 			if (EntryIndex != INDEX_NONE)
@@ -2734,55 +2751,55 @@ bool FBlueprintEditorUtils::GetBlueprintVariableMetaData(const UBlueprint* Bluep
 			}
 		}
 	}
-	else
-	{
-		const FBPVariableDescription& Desc = Blueprint->NewVariables[VarIndex];
-
-		int32 EntryIndex = Desc.FindMetaDataEntryIndexForKey(MetaDataKey);
-		if (EntryIndex != INDEX_NONE)
-		{
-			OutMetaDataValue = Desc.GetMetaData(MetaDataKey);
-			return true;
-		}
-	}
 
 	OutMetaDataValue.Empty();
 	return false;
 }
 
-void FBlueprintEditorUtils::RemoveBlueprintVariableMetaData(UBlueprint* Blueprint, const FName& VarName, const FName& MetaDataKey)
+void FBlueprintEditorUtils::RemoveBlueprintVariableMetaData(UBlueprint* Blueprint, const FName& VarName, const UStruct* InLocalVarScope, const FName& MetaDataKey)
 {
-	const int32 VarIndex = FBlueprintEditorUtils::FindNewVariableIndex(Blueprint, VarName);
-	if (VarIndex == INDEX_NONE)
+	// If there is a local var scope, we know we are looking at a local variable
+	if(InLocalVarScope)
 	{
-		//Not a NewVariable is the VarName from a Timeline?
-		const int32 TimelineIndex = FBlueprintEditorUtils::FindTimelineIndex(Blueprint,VarName);
-
-		if (TimelineIndex == INDEX_NONE)
+		if(FBPVariableDescription* LocalVariable = FindLocalVariable(Blueprint, InLocalVarScope, VarName))
+		{
+			LocalVariable->RemoveMetaData(MetaDataKey);
+		}
+	}
+	else
+	{
+		const int32 VarIndex = FBlueprintEditorUtils::FindNewVariableIndex(Blueprint, VarName);
+		if (VarIndex == INDEX_NONE)
 		{
-			//Not a Timeline is this a SCS Node?
-			const int32 SCS_NodeIndex = FBlueprintEditorUtils::FindSCS_Node(Blueprint,VarName);
+			//Not a NewVariable is the VarName from a Timeline?
+			const int32 TimelineIndex = FBlueprintEditorUtils::FindTimelineIndex(Blueprint,VarName);
 
-			if (SCS_NodeIndex != INDEX_NONE)
+			if (TimelineIndex == INDEX_NONE)
 			{
-				Blueprint->SimpleConstructionScript->GetAllNodes()[SCS_NodeIndex]->RemoveMetaData(MetaDataKey);
-			}
+				//Not a Timeline is this a SCS Node?
+				const int32 SCS_NodeIndex = FBlueprintEditorUtils::FindSCS_Node(Blueprint,VarName);
+
+				if (SCS_NodeIndex != INDEX_NONE)
+				{
+					Blueprint->SimpleConstructionScript->GetAllNodes()[SCS_NodeIndex]->RemoveMetaData(MetaDataKey);
+				}
 
+			}
+			else
+			{
+				Blueprint->Timelines[TimelineIndex]->RemoveMetaData(MetaDataKey);
+			}
 		}
 		else
 		{
-			Blueprint->Timelines[TimelineIndex]->RemoveMetaData(MetaDataKey);
+			Blueprint->NewVariables[VarIndex].RemoveMetaData(MetaDataKey);
 		}
 	}
-	else
-	{
-		Blueprint->NewVariables[VarIndex].RemoveMetaData(MetaDataKey);
-	}
 
 	FBlueprintEditorUtils::MarkBlueprintAsStructurallyModified(Blueprint);
 }
 
-void FBlueprintEditorUtils::SetBlueprintVariableCategory(UBlueprint* Blueprint, const FName& VarName, const FName& NewCategory, bool bDontRecompile)
+void FBlueprintEditorUtils::SetBlueprintVariableCategory(UBlueprint* Blueprint, const FName& VarName, const UStruct* InLocalVarScope, const FName& NewCategory, bool bDontRecompile)
 {
 	// Ensure we always set a category
 	const UEdGraphSchema_K2* K2Schema = GetDefault<UEdGraphSchema_K2>();
@@ -2821,10 +2838,9 @@ void FBlueprintEditorUtils::SetBlueprintVariableCategory(UBlueprint* Blueprint,
 			}
 		}
 	}
-	else
+	else if(InLocalVarScope)
 	{
-		// Check to see if it is a local variable
-		if(FBPVariableDescription* LocalVariable = FindLocalVariable(Blueprint, VarName))
+		if(FBPVariableDescription* LocalVariable = FindLocalVariable(Blueprint, InLocalVarScope, VarName))
 		{
 			LocalVariable->SetMetaData(TEXT("Category"), *SetCategory.ToString());
 			LocalVariable->Category = SetCategory;
@@ -2837,7 +2853,7 @@ void FBlueprintEditorUtils::SetBlueprintVariableCategory(UBlueprint* Blueprint,
 	}
 }
 
-FName FBlueprintEditorUtils::GetBlueprintVariableCategory(UBlueprint* Blueprint, const FName& VarName)
+FName FBlueprintEditorUtils::GetBlueprintVariableCategory(UBlueprint* Blueprint, const FName& VarName, const UStruct* InLocalVarScope)
 {
 	FName CategoryName = NAME_None;
 	UClass* SkeletonGeneratedClass = Blueprint->SkeletonGeneratedClass;
@@ -2846,10 +2862,10 @@ FName FBlueprintEditorUtils::GetBlueprintVariableCategory(UBlueprint* Blueprint,
 	{
 		CategoryName = FObjectEditorUtils::GetCategoryFName(TargetProperty);
 	}
-	else
+	else if(InLocalVarScope)
 	{
 		// Check to see if it is a local variable
-		if(FBPVariableDescription* LocalVariable = FindLocalVariable(Blueprint, VarName))
+		if(FBPVariableDescription* LocalVariable = FindLocalVariable(Blueprint, InLocalVarScope, VarName))
 		{
 			return LocalVariable->Category;
 		}
@@ -3363,25 +3379,39 @@ void FBlueprintEditorUtils::ChangeMemberVariableType(UBlueprint* Blueprint, cons
 	}
 }
 
-void FBlueprintEditorUtils::RemoveLocalVariable(UBlueprint* InBlueprint, const FName& InVarName)
+void FBlueprintEditorUtils::RemoveLocalVariable(UBlueprint* InBlueprint, const UStruct* InScope, const FName& InVarName)
 {
-	TArray<UK2Node_FunctionEntry*> FunctionEntryNodes;
-	GetAllNodesOfClass<UK2Node_FunctionEntry>(InBlueprint, FunctionEntryNodes);
+	UEdGraph* ScopeGraph = FindScopeGraph(InBlueprint, InScope);
 
-	for (UK2Node_FunctionEntry* const FunctionEntry : FunctionEntryNodes)
+	if(ScopeGraph)
 	{
-		for( int32 VarIdx = 0; VarIdx < FunctionEntry->LocalVariables.Num(); ++VarIdx )
+		TArray<UK2Node_FunctionEntry*> GraphNodes;
+		ScopeGraph->GetNodesOfClass<UK2Node_FunctionEntry>(GraphNodes);
+
+		bool bFoundLocalVariable = false;
+
+		// There is only ever 1 function entry
+		check(GraphNodes.Num() == 1)
+		for( int32 VarIdx = 0; VarIdx < GraphNodes[0]->LocalVariables.Num(); ++VarIdx )
 		{
-			if(FunctionEntry->LocalVariables[VarIdx].VarName == InVarName)
+			if(GraphNodes[0]->LocalVariables[VarIdx].VarName == InVarName)
 			{
-				FunctionEntry->LocalVariables.RemoveAt(VarIdx);
+				GraphNodes[0]->LocalVariables.RemoveAt(VarIdx);
 				FBlueprintEditorUtils::RemoveVariableNodes(InBlueprint, InVarName);
 				FBlueprintEditorUtils::MarkBlueprintAsStructurallyModified(InBlueprint);
 
+				bFoundLocalVariable = true;
+
 				// No other local variables will match, we are done
-				return;
+				break;
 			}
 		}
+
+		// Check if we found the local variable, it is a problem if we do not.
+		if(!bFoundLocalVariable)
+		{
+			UE_LOG(LogBlueprint, Warning, TEXT("Could not find local variable '%s'!"), *InVarName.ToString());
+		}
 	}
 }
 
@@ -3411,14 +3441,34 @@ UFunction* FFunctionFromNodeHelper::FunctionFromNode(UK2Node* Node)
 	return Function;
 }
 
-void FBlueprintEditorUtils::RenameLocalVariable(UBlueprint* InBlueprint, const FName& InOldName, const FName& InNewName)
+UEdGraph* FBlueprintEditorUtils::FindScopeGraph(const UBlueprint* InBlueprint, const UStruct* InScope)
+{
+	UEdGraph* ScopeGraph = NULL;
+
+	TArray<UEdGraph*> AllGraphs;
+	InBlueprint->GetAllGraphs(AllGraphs);
+
+	for(const auto Graph : AllGraphs)
+	{
+		check(Graph != NULL);
+		if(Graph->GetName() == InScope->GetName())
+		{
+			// This graph should always be a function graph
+			check(Graph->GetSchema()->GetGraphType(Graph) == GT_Function);
+			
+			ScopeGraph = Graph;
+			break;
+		}
+	}
+	return ScopeGraph;
+}
+
+void FBlueprintEditorUtils::RenameLocalVariable(UBlueprint* InBlueprint, const UStruct* InScope, const FName& InOldName, const FName& InNewName)
 {
 	if ((InOldName != InNewName) && (InNewName != NAME_None))
 	{
-		UK2Node_FunctionEntry* FoundFunctionEntry = NULL;
-		FBPVariableDescription* LocalVariable = FindLocalVariable(InBlueprint, InOldName, &FoundFunctionEntry);
-		const auto FoundFunction = FFunctionFromNodeHelper::FunctionFromNode(FoundFunctionEntry);
-		const auto ExistingProperty = FindField<const UProperty>(FoundFunction, InNewName);
+		FBPVariableDescription* LocalVariable = FindLocalVariable(InBlueprint, InScope, InOldName);
+		const auto ExistingProperty = FindField<const UProperty>(InScope, InNewName);
 		if (ExistingProperty)
 		{
 			UE_LOG(LogBlueprint, Warning, TEXT("Cannot name local variable '%s'. The name is already used."), *InNewName.ToString());
@@ -3446,32 +3496,36 @@ void FBlueprintEditorUtils::RenameLocalVariable(UBlueprint* InBlueprint, const F
 	}
 }
 
-FBPVariableDescription* FBlueprintEditorUtils::FindLocalVariable(const UBlueprint* InBlueprint, const FName& InVariableName, class UK2Node_FunctionEntry** OutFunctionEntry)
+FBPVariableDescription* FBlueprintEditorUtils::FindLocalVariable(const UBlueprint* InBlueprint, const UStruct* InScope, const FName& InVariableName, class UK2Node_FunctionEntry** OutFunctionEntry)
 {
-	// Search through all function entry nodes for a local variable with the passed name
-	TArray<UK2Node_FunctionEntry*> FunctionEntryNodes;
-	GetAllNodesOfClass<UK2Node_FunctionEntry>(InBlueprint, FunctionEntryNodes);
+	UEdGraph* ScopeGraph = FindScopeGraph(InBlueprint, InScope);
 
 	FBPVariableDescription* ReturnVariable = NULL;
-
-	for (UK2Node_FunctionEntry* const FunctionEntry : FunctionEntryNodes)
+	if(ScopeGraph)
 	{
-		for( FBPVariableDescription& Variable : FunctionEntry->LocalVariables )
+		TArray<UK2Node_FunctionEntry*> GraphNodes;
+		ScopeGraph->GetNodesOfClass<UK2Node_FunctionEntry>(GraphNodes);
+
+		bool bFoundLocalVariable = false;
+
+		// There is only ever 1 function entry
+		check(GraphNodes.Num() == 1)
+		for( int32 VarIdx = 0; VarIdx < GraphNodes[0]->LocalVariables.Num(); ++VarIdx )
 		{
-			if(InVariableName == Variable.VarName)
+			if(GraphNodes[0]->LocalVariables[VarIdx].VarName == InVariableName)
 			{
 				if (OutFunctionEntry)
 				{
-					*OutFunctionEntry = FunctionEntry;
+					*OutFunctionEntry = GraphNodes[0];
 				}
-				ReturnVariable = &Variable;
+				ReturnVariable = &GraphNodes[0]->LocalVariables[VarIdx];
 				break;
 			}
 		}
 
-		if(ReturnVariable)
+		if(!ReturnVariable)
 		{
-			break;
+			UE_LOG(LogBlueprint, Warning, TEXT("Could not find local variable '%s'!"), *InVariableName.ToString());
 		}
 	}
 
@@ -3505,40 +3559,24 @@ FName FBlueprintEditorUtils::FindLocalVariableNameByGuid(UBlueprint* InBlueprint
 	return ReturnVariableName;
 }
 
-FGuid FBlueprintEditorUtils::FindLocalVariableGuidByName(UBlueprint* InBlueprint, const FName InVariableName)
+FGuid FBlueprintEditorUtils::FindLocalVariableGuidByName(UBlueprint* InBlueprint, const UStruct* InScope, const FName InVariableName)
 {
-	// Search through all function entry nodes for a local variable with the passed name
-	TArray<UK2Node_FunctionEntry*> FunctionEntryNodes;
-	GetAllNodesOfClass<UK2Node_FunctionEntry>(InBlueprint, FunctionEntryNodes);
-
 	FGuid ReturnGuid;
-	for (UK2Node_FunctionEntry* const FunctionEntry : FunctionEntryNodes)
+	if(FBPVariableDescription* LocalVariable = FindLocalVariable(InBlueprint, InScope, InVariableName))
 	{
-		for( FBPVariableDescription& Variable : FunctionEntry->LocalVariables )
-		{
-			if(Variable.VarName == InVariableName)
-			{
-				ReturnGuid = Variable.VarGuid;
-				break;
-			}
-		}
-
-		if(ReturnGuid.IsValid())
-		{
-			break;
-		}
+		ReturnGuid = LocalVariable->VarGuid;
 	}
 
 	return ReturnGuid;
 }
 
-void FBlueprintEditorUtils::ChangeLocalVariableType(UBlueprint* InBlueprint, const FName& InVariableName, const FEdGraphPinType& NewPinType)
+void FBlueprintEditorUtils::ChangeLocalVariableType(UBlueprint* InBlueprint, const UStruct* InScope, const FName& InVariableName, const FEdGraphPinType& NewPinType)
 {
 	if (InVariableName != NAME_None)
 	{
 		FString ActionCategory;
 		const UEdGraphSchema_K2* K2Schema = GetDefault<UEdGraphSchema_K2>();
-		FBPVariableDescription* VariablePtr = FindLocalVariable(InBlueprint, InVariableName);
+		FBPVariableDescription* VariablePtr = FindLocalVariable(InBlueprint, InScope, InVariableName);
 
 		if(VariablePtr)
 		{
@@ -4714,13 +4752,13 @@ void FBlueprintEditorUtils::UpdateStalePinWatches( UBlueprint* Blueprint )
 	}
 }
 
-FName FBlueprintEditorUtils::FindUniqueKismetName(const UBlueprint* InBlueprint, const FString& InBaseName)
+FName FBlueprintEditorUtils::FindUniqueKismetName(const UBlueprint* InBlueprint, const FString& InBaseName, UStruct* InScope/* = NULL*/)
 {
 	int32 Count = 0;
 	FString KismetName;
 	FString BaseName = InBaseName;
 
-	TSharedPtr<FKismetNameValidator> NameValidator = MakeShareable(new FKismetNameValidator(InBlueprint));
+	TSharedPtr<FKismetNameValidator> NameValidator = MakeShareable(new FKismetNameValidator(InBlueprint, NAME_None, InScope));
 	while(NameValidator->IsValid(KismetName) != EValidatorResult::Ok)
 	{
 		// Calculate the number of digits in the number, adding 2 (1 extra to correctly count digits, another to account for the '_' that will be added to the name
diff --git a/Engine/Source/Editor/UnrealEd/Private/Kismet2/Kismet2NameValidators.cpp b/Engine/Source/Editor/UnrealEd/Private/Kismet2/Kismet2NameValidators.cpp
--- a/Engine/Source/Editor/UnrealEd/Private/Kismet2/Kismet2NameValidators.cpp
+++ b/Engine/Source/Editor/UnrealEd/Private/Kismet2/Kismet2NameValidators.cpp
@@ -87,7 +87,7 @@ namespace BlueprintNameConstants
 	 int32 NameMaxLength = 100;
 }
 
-FKismetNameValidator::FKismetNameValidator(const class UBlueprint* Blueprint, FName InExistingName/* = NAME_None*/)
+FKismetNameValidator::FKismetNameValidator(const class UBlueprint* Blueprint, FName InExistingName/* = NAME_None*/, UStruct* InScope/* = NULL*/)
 {
 	ExistingName = InExistingName;
 	BlueprintObject = Blueprint;
@@ -95,6 +95,8 @@ FKismetNameValidator::FKismetNameValidator(const class UBlueprint* Blueprint, FN
 	FBlueprintEditorUtils::GetAllGraphNames(BlueprintObject, Names);
 	FBlueprintEditorUtils::GetSCSVariableNameList(Blueprint, Names);
 	FBlueprintEditorUtils::GetImplementingBlueprintsFunctionNameList(Blueprint, Names);
+
+	Scope = InScope;
 }
 
 int32 FKismetNameValidator::GetMaximumNameLength()
@@ -141,23 +143,33 @@ EValidatorResult FKismetNameValidator::IsValid(const FName& Name, bool /* bOrigi
 
 		if(ValidatorResult == EValidatorResult::Ok)
 		{
-			// Search through all functions for their local variables and prevent duplicate names
-			TArray<UK2Node_FunctionEntry*> FunctionEntryNodes;
-			FBlueprintEditorUtils::GetAllNodesOfClass(BlueprintObject, FunctionEntryNodes);
-			for (UK2Node_FunctionEntry* const FunctionEntry : FunctionEntryNodes)
+			if(Scope == NULL)
 			{
-				for( const FBPVariableDescription& Variable : FunctionEntry->LocalVariables )
+				// Search through all functions for their local variables and prevent duplicate names
+				TArray<UK2Node_FunctionEntry*> FunctionEntryNodes;
+				FBlueprintEditorUtils::GetAllNodesOfClass(BlueprintObject, FunctionEntryNodes);
+				for (UK2Node_FunctionEntry* const FunctionEntry : FunctionEntryNodes)
 				{
-					if(Variable.VarName == Name)
+					for( const FBPVariableDescription& Variable : FunctionEntry->LocalVariables )
+					{
+						if(Variable.VarName == Name)
+						{
+							ValidatorResult = EValidatorResult::AlreadyInUse;
+							break;
+						}
+					}
+
+					if(ValidatorResult != EValidatorResult::Ok)
 					{
-						ValidatorResult = EValidatorResult::AlreadyInUse;
 						break;
 					}
 				}
-
-				if(ValidatorResult != EValidatorResult::Ok)
+			}
+			else
+			{
+				if(FindField<const UProperty>(Scope, *Name.ToString()) != NULL)
 				{
-					break;
+					ValidatorResult = EValidatorResult::LocallyInUse;
 				}
 			}
 		}
diff --git a/Engine/Source/Editor/UnrealEd/Public/Kismet2/BlueprintEditorUtils.h b/Engine/Source/Editor/UnrealEd/Public/Kismet2/BlueprintEditorUtils.h
--- a/Engine/Source/Editor/UnrealEd/Public/Kismet2/BlueprintEditorUtils.h
+++ b/Engine/Source/Editor/UnrealEd/Public/Kismet2/BlueprintEditorUtils.h
@@ -508,28 +508,39 @@ public:
 	/** Changes the type of a member variable */
 	static void ChangeMemberVariableType(UBlueprint* Blueprint, const FName& VariableName, const FEdGraphPinType& NewPinType);
 
+	/**
+	 * Finds the scope's associated graph for local variables (or any passed UFunction)
+	 *
+	 * @param	InBlueprint			The Blueprint the local variable can be found in
+	 * @param	InScope				Local variable's scope
+	 */
+	static UEdGraph* FindScopeGraph(const UBlueprint* InBlueprint, const UStruct* InScope);
+
 	/**
 	 * Removes a member variable if it was declared in this blueprint and not in a base class.
 	 *
 	 * @param	InBlueprint			The Blueprint the local variable can be found in
+	 * @param	InScope				Local variable's scope
 	 * @param	InVarName			Name of the variable to be removed.
 	 */
-	static void RemoveLocalVariable(UBlueprint* InBlueprint, const FName& InVarName);
+	static void RemoveLocalVariable(UBlueprint* InBlueprint, const UStruct* InScope, const FName& InVarName);
 
 	/**
 	 * Returns a local variable
 	 *
 	 * @param InBlueprint		Blueprint to search for the local variable
+	 * @param InScope			Local variable's scope
 	 * @param InVariableName	Name of the variable to search for
 	 * @param OutFunctionEntry	Optional output parameter. If not null, the found function entry is returned.
 	 * @return					The local variable description
 	 */
-	static FBPVariableDescription* FindLocalVariable(const UBlueprint* InBlueprint, const FName& InVariableName, class UK2Node_FunctionEntry** OutFunctionEntry = NULL);
+	static FBPVariableDescription* FindLocalVariable(const UBlueprint* InBlueprint, const UStruct* InScope, const FName& InVariableName, class UK2Node_FunctionEntry** OutFunctionEntry = NULL);
 
 	/**
 	 * Finds a local variable name using the variable's Guid
 	 *
 	 * @param InBlueprint		Blueprint to search for the local variable
+	 * @param InScope			Local variable's scope
 	 * @param InVariableGuid	Guid to identify the local variable with
 	 * @return					Local variable's name
 	 */
@@ -539,28 +550,31 @@ public:
 	 * Finds a local variable Guid using the variable's name
 	 *
 	 * @param InBlueprint		Blueprint to search for the local variable
+	 * @param InScope			Local variable's scope
 	 * @param InVariableGuid	Local variable's name to search for
 	 * @return					The Guid associated with the local variable
 	 */
-	static FGuid FindLocalVariableGuidByName(UBlueprint* InBlueprint, const FName InVariableName);
+	static FGuid FindLocalVariableGuidByName(UBlueprint* InBlueprint, const UStruct* InScope, const FName InVariableName);
 
 	/**
 	 * Rename a local variable
 	 *
 	 * @param InBlueprint		Blueprint to search for the local variable
-	 * @param InOldName The name of the local variable to change
-	 * @param InNewName	The new name of the local variable
+	 * @param InScope			Local variable's scope
+	 * @param InOldName			The name of the local variable to change
+	 * @param InNewName			The new name of the local variable
 	 */
-	static void RenameLocalVariable(UBlueprint* InBlueprint, const FName& InOldName, const FName& InNewName);
+	static void RenameLocalVariable(UBlueprint* InBlueprint, const UStruct* InScope, const FName& InOldName, const FName& InNewName);
 
 	/**
 	 * Changes the type of a local variable
 	 *
 	 * @param InBlueprint		Blueprint to search for the local variable
+	 * @param InScope			Local variable's scope
 	 * @param InVariableName	Name of the local variable to change the type of
 	 * @param InNewPinType		The pin type to change the local variable type to
 	 */
-	static void ChangeLocalVariableType(UBlueprint* InBlueprint, const FName& InVariableName, const FEdGraphPinType& InNewPinType);
+	static void ChangeLocalVariableType(UBlueprint* InBlueprint, const UStruct* InScope, const FName& InVariableName, const FEdGraphPinType& InNewPinType);
 
 	/** Replaces all variable references in the specified blueprint */
 	static void ReplaceVariableReferences(UBlueprint* Blueprint, const FName& OldName, const FName& NewName);
@@ -606,32 +620,54 @@ public:
 	 */
 	static void SetVariableSaveGameFlag(UBlueprint* InBlueprint, const FName& InVarName, const bool bInIsSaveGame);
 
-	/** Sets a metadata key/value on the specified variable */
-	static void SetBlueprintVariableMetaData(UBlueprint* Blueprint, const FName& VarName, const FName& MetaDataKey, const FString& MetaDataValue);
+	/** Sets a metadata key/value on the specified variable
+	 *
+	 * @param Blueprint				The Blueprint to find the variable in
+	 * @param VarName				Name of the variable
+	 * @param InLocalVarScope		Local variable's scope, if looking to modify a local variable
+	 * @param MetaDataKey			Key name for the metadata to change
+	 * @param MetaDataValue			Value to change the metadata to
+	 */
+	static void SetBlueprintVariableMetaData(UBlueprint* Blueprint, const FName& VarName, const UStruct* InLocalVarScope, const FName& MetaDataKey, const FString& MetaDataValue);
 
-	/** Get a metadata key/value on the specified variable, or timeline if it exists, returning false if it does not exist */
-	static bool GetBlueprintVariableMetaData(const UBlueprint* Blueprint, const FName& VarName, const FName& MetaDataKey, FString& OutMetaDataValue);
+	/** Get a metadata key/value on the specified variable, or timeline if it exists, returning false if it does not exist
+	 *
+	 * @param Blueprint				The Blueprint to find the variable in
+	 * @param VarName				Name of the variable
+	 * @param InLocalVarScope		Local variable's scope, if looking to modify a local variable
+	 * @param MetaDataKey			Key name for the metadata to change
+	 * @param OutMetaDataValue		Value of the metadata
+	 * @return						TRUE if finding the metadata was successful
+	 */
+	static bool GetBlueprintVariableMetaData(const UBlueprint* Blueprint, const FName& VarName, const UStruct* InLocalVarScope, const FName& MetaDataKey, FString& OutMetaDataValue);
 
-	/** Clear metadata key on specified variable, or timeline */
-	static void RemoveBlueprintVariableMetaData(UBlueprint* Blueprint, const FName& VarName, const FName& MetaDataKey);
+	/** Clear metadata key on specified variable, or timeline
+	 * @param Blueprint				The Blueprint to find the variable in
+	 * @param VarName				Name of the variable
+	 * @param InLocalVarScope		Local variable's scope, if looking to modify a local variable
+	 * @param MetaDataKey			Key name for the metadata to change
+	 */
+	static void RemoveBlueprintVariableMetaData(UBlueprint* Blueprint, const FName& VarName, const UStruct* InLocalVarScope, const FName& MetaDataKey);
 
 	/**
 	 * Sets the custom category on the variable with the specified name.
 	 * @note: Will not change the category for variables defined via native classes.
 	 *
-	 * @param	VarName			Name of the variable
-	 * @param	VarCategory		The new value of the custom category for the variable
-	 * @param	bDontRecompile	If true, the blueprint will not be marked as modified, and will not be recompiled.  
+	 * @param	VarName				Name of the variable
+	 * @param	InLocalVarScope		Local variable's scope, if looking to modify a local variable
+	 * @param	VarCategory			The new value of the custom category for the variable
+	 * @param	bDontRecompile		If true, the blueprint will not be marked as modified, and will not be recompiled.  
 	 */
-	static void SetBlueprintVariableCategory(UBlueprint* Blueprint, const FName& VarName, const FName& NewCategory, bool bDontRecompile=false);
+	static void SetBlueprintVariableCategory(UBlueprint* Blueprint, const FName& VarName, const UStruct* InLocalVarScope, const FName& NewCategory, bool bDontRecompile=false);
 
 	/**
 	 * Gets the custom category on the variable with the specified name.
 	 *
-	 * @param	VarName			Name of the variable
-	 * @return	The custom category (None indicates the name will be the same as the blueprint)
+	 * @param	VarName				Name of the variable
+	 * @param	InLocalVarScope		Local variable's scope, if looking to modify a local variable
+	 * @return						The custom category (None indicates the name will be the same as the blueprint)
 	 */
-	static FName GetBlueprintVariableCategory(UBlueprint* Blueprint, const FName& VarName);
+	static FName GetBlueprintVariableCategory(UBlueprint* Blueprint, const FName& VarName, const UStruct* InLocalVarScope);
 
 	/** Gets pointer to PropertyFlags of variable */
 	static uint64* GetBlueprintVariablePropertyFlags(UBlueprint* Blueprint, const FName& VarName);
@@ -779,10 +815,11 @@ public:
 	 *
 	 * @param InBlueprint		The blueprint the kismet object's name needs to be unique in
 	 * @param InBaseName		The base name to use
+	 * @param InScope			Scope, if any, of the unique kismet name to generate, used for locals
 	 *
 	 * @return					A unique name that will not conflict in the Blueprint
 	 */
-	static FName FindUniqueKismetName(const UBlueprint* InBlueprint, const FString& InBaseName);
+	static FName FindUniqueKismetName(const UBlueprint* InBlueprint, const FString& InBaseName, UStruct* InScope = NULL);
 
 	/** Finds a unique and valid name for a custom event */
 	static FName FindUniqueCustomEventName(const UBlueprint* Blueprint);
diff --git a/Engine/Source/Editor/UnrealEd/Public/Kismet2/Kismet2NameValidators.h b/Engine/Source/Editor/UnrealEd/Public/Kismet2/Kismet2NameValidators.h
--- a/Engine/Source/Editor/UnrealEd/Public/Kismet2/Kismet2NameValidators.h
+++ b/Engine/Source/Editor/UnrealEd/Public/Kismet2/Kismet2NameValidators.h
@@ -24,7 +24,9 @@ enum EValidatorResult
 	/* The entered name matches the current name */
 	ExistingName,
 	/* The entered name is too long */
-	TooLong
+	TooLong,
+	/** The entered is in use locally */
+	LocallyInUse
 };
 
 
@@ -64,7 +66,7 @@ public:
 class UNREALED_API FKismetNameValidator : public INameValidatorInterface
 {
 public:
-	FKismetNameValidator(const class UBlueprint* Blueprint, FName InExistingName = NAME_None);
+	FKismetNameValidator(const class UBlueprint* Blueprint, FName InExistingName = NAME_None, UStruct* InScope = NULL);
 	~FKismetNameValidator() {}
 
 	/** Return the name validator maximum string length */
@@ -81,6 +83,8 @@ private:
 	const UBlueprint* BlueprintObject;
 	/** The current name of the object being validated */
 	FName ExistingName;
+	/** Scope to check against local variables (or other members) */
+	UStruct* Scope;
 };
 
 /////////////////////////////////////////////////////
