Adding a more robust system for editor category mapping (allows for easy remapping and localization). Pulling editor metadata specifics out of UClass (and relocating them).

diff --git a/Engine/Source/Editor/DetailCustomizations/Private/ActorDetails.cpp b/Engine/Source/Editor/DetailCustomizations/Private/ActorDetails.cpp
--- a/Engine/Source/Editor/DetailCustomizations/Private/ActorDetails.cpp
+++ b/Engine/Source/Editor/DetailCustomizations/Private/ActorDetails.cpp
@@ -27,6 +27,7 @@
 #include "IDocumentation.h"
 #include "Runtime/Engine/Classes/Engine/BrushShape.h"
 #include "ActorDetailsDelegates.h"
+#include "EditorCategoryUtils.h"
 
 #define LOCTEXT_NAMESPACE "ActorDetails"
 
@@ -104,7 +105,7 @@ void FActorDetails::CustomizeDetails( IDetailLayoutBuilder& DetailLayout )
 
 		// Get the list of hidden categories
 		TArray<FString> HideCategories;
-		DetailLayout.GetDetailsView().GetBaseClass()->GetHideCategories(HideCategories);
+		FEditorCategoryUtils::GetClassHideCategories(DetailLayout.GetDetailsView().GetBaseClass(), HideCategories);
 
 		// Add Blueprint category, if not being hidden
 		if (!HideCategories.Contains(TEXT("Blueprint")))
diff --git a/Engine/Source/Editor/DetailCustomizations/Private/PrimitiveComponentDetails.cpp b/Engine/Source/Editor/DetailCustomizations/Private/PrimitiveComponentDetails.cpp
--- a/Engine/Source/Editor/DetailCustomizations/Private/PrimitiveComponentDetails.cpp
+++ b/Engine/Source/Editor/DetailCustomizations/Private/PrimitiveComponentDetails.cpp
@@ -9,6 +9,7 @@
 #include "Components/DestructibleComponent.h"
 #include "Components/InstancedStaticMeshComponent.h"
 #include "IDocumentation.h"
+#include "EditorCategoryUtils.h"
 
 
 #define LOCTEXT_NAMESPACE "PrimitiveComponentDetails"
@@ -151,7 +152,7 @@ void FPrimitiveComponentDetails::CustomizeDetails( IDetailLayoutBuilder& DetailB
 
 		// See if we are hiding Physics category
 		TArray<FString> HideCategories;
-		DetailBuilder.GetDetailsView().GetBaseClass()->GetHideCategories(HideCategories);
+		FEditorCategoryUtils::GetClassHideCategories(DetailBuilder.GetDetailsView().GetBaseClass(), HideCategories);
 		if (!HideCategories.Contains(TEXT("Physics")))
 		{
 			IDetailCategoryBuilder& PhysicsCategory = DetailBuilder.EditCategory("Physics");
diff --git a/Engine/Source/Editor/KismetCompiler/Private/KismetCompiler.cpp b/Engine/Source/Editor/KismetCompiler/Private/KismetCompiler.cpp
--- a/Engine/Source/Editor/KismetCompiler/Private/KismetCompiler.cpp
+++ b/Engine/Source/Editor/KismetCompiler/Private/KismetCompiler.cpp
@@ -16,6 +16,7 @@
 #include "MovieSceneBindings.h"
 #include "Kismet2/Kismet2NameValidators.h"
 #include "UserDefinedStructureCompilerUtils.h"
+#include "EditorCategoryUtils.h"
 
 static bool bDebugPropertyPropagation = false;
 
@@ -1615,7 +1616,7 @@ void FKismetCompilerContext::FinishCompilingClass(UClass* Class)
 
 		// Copy the category info from the parent class
 #if WITH_EDITORONLY_DATA
-		ParentClass->GetHideCategories(AllHideCategories);
+		FEditorCategoryUtils::GetClassHideCategories(ParentClass, AllHideCategories);
 		if (ParentClass->HasMetaData(TEXT("ShowCategories")))
 		{
 			Class->SetMetaData(TEXT("ShowCategories"), *ParentClass->GetMetaData("ShowCategories"));
diff --git a/Engine/Source/Editor/PropertyEditor/Private/ObjectPropertyNode.cpp b/Engine/Source/Editor/PropertyEditor/Private/ObjectPropertyNode.cpp
--- a/Engine/Source/Editor/PropertyEditor/Private/ObjectPropertyNode.cpp
+++ b/Engine/Source/Editor/PropertyEditor/Private/ObjectPropertyNode.cpp
@@ -7,6 +7,7 @@
 #include "CategoryPropertyNode.h"
 #include "ItemPropertyNode.h"
 #include "ObjectEditorUtils.h"
+#include "EditorCategoryUtils.h"
 
 FObjectPropertyNode::FObjectPropertyNode(void)
 	: FComplexPropertyNode()
@@ -324,7 +325,7 @@ void FObjectPropertyNode::InternalInitChildNodes( FName SinglePropertyName )
 		{
 			UClass* Class = ClassesToConsider[ ClassIndex ];
 		
-			if( Class->IsCategoryHidden( CategoryName.ToString() ) )
+			if( FEditorCategoryUtils::IsCategoryHiddenFromClass(Class, CategoryName.ToString()) )
 			{
 				HiddenCategories.Add( CategoryName );
 
@@ -454,7 +455,7 @@ void FObjectPropertyNode::InternalInitChildNodes( FName SinglePropertyName )
 		for( TFieldIterator<UProperty> It(BaseClass.Get()); It; ++It )
 		{
 			if( bShouldShowHiddenProperties ||
-				( ( It->PropertyFlags&CPF_Edit ) && ( BaseClass->IsCategoryHidden(FObjectEditorUtils::GetCategory(*It)) == false ) ) )
+				( ( It->PropertyFlags&CPF_Edit ) && ( FEditorCategoryUtils::IsCategoryHiddenFromClass(BaseClass.Get(), FObjectEditorUtils::GetCategory(*It)) == false ) ) )
 			{
 				UProperty* CurProp = *It;
 				if( SinglePropertyName == NAME_None || CurProp->GetFName() == SinglePropertyName )
diff --git a/Engine/Source/Editor/UnrealEd/Private/Commandlets/DumpHiddenCategoriesCommandlet.cpp b/Engine/Source/Editor/UnrealEd/Private/Commandlets/DumpHiddenCategoriesCommandlet.cpp
--- a/Engine/Source/Editor/UnrealEd/Private/Commandlets/DumpHiddenCategoriesCommandlet.cpp
+++ b/Engine/Source/Editor/UnrealEd/Private/Commandlets/DumpHiddenCategoriesCommandlet.cpp
@@ -4,6 +4,8 @@
 //#include "DumpHiddenCategoriesCommandlet.h"
 #include "ObjectEditorUtils.h"
 #include "KismetEditorUtilities.h"
+#include "EditorCategoryUtils.h"
+#include "ScopedTimers.h"
 
 /*******************************************************************************
  * Static Helpers
@@ -28,7 +30,7 @@ static int32 GetHideCategories(uint32 Indent, UClass* Class, FString& JsonOut)
 	JsonOut += IndentString + TEXT("\"HiddenCategories\" : [");
 
 	TArray<FString> HideCategories;
-	Class->GetHideCategories(HideCategories);
+	FEditorCategoryUtils::GetClassHideCategories(Class, HideCategories);
 
 	for (FString const& Category : HideCategories)
 	{
@@ -50,7 +52,7 @@ static int32 GetShowCategories(uint32 Indent, UClass* Class, FString& JsonOut)
 	JsonOut += IndentString + TEXT("\"ShownCategories\" : [");
 
 	TArray<FString> ShowCategories;
-	Class->GetShowCategories(ShowCategories);
+	FEditorCategoryUtils::GetClassShowCategories(Class, ShowCategories);
 
 	for (FString const& Category : ShowCategories)
 	{
@@ -188,6 +190,9 @@ int32 UDumpHiddenCategoriesCommandlet::Main(FString const& Params)
 	{
 		FileOut->Serialize(TCHAR_TO_ANSI(TEXT("{")), 1);
 
+		double Duration = 0.0;
+		FDurationTimer DurationTimer(Duration);
+
 		bool bFirstEntry = true;
 		for (TObjectIterator<UClass> ClassIt; ClassIt; ++ClassIt)
 		{
@@ -234,7 +239,10 @@ int32 UDumpHiddenCategoriesCommandlet::Main(FString const& Params)
 			FileOut->Serialize(TCHAR_TO_ANSI(*ClassEntry), ClassEntry.Len());
 			bFirstEntry = false;
 		}
-		FileOut->Serialize(TCHAR_TO_ANSI(TEXT("\n}")), 2);
+		DurationTimer.Stop();
+
+		FString ClosingStatement = FString::Printf(TEXT(", \"Duration\" : %f\n}"), Duration);
+		FileOut->Serialize(TCHAR_TO_ANSI(*ClosingStatement), ClosingStatement.Len());
 
 		FileOut->Close();
 		delete FileOut;
diff --git a/Engine/Source/Editor/UnrealEd/Private/EditorCategoryUtils.cpp b/Engine/Source/Editor/UnrealEd/Private/EditorCategoryUtils.cpp
new file mode 100644
--- /dev/null
+++ b/Engine/Source/Editor/UnrealEd/Private/EditorCategoryUtils.cpp
@@ -0,0 +1,314 @@
+// Copyright 1998-2014 Epic Games, Inc. All Rights Reserved.
+
+#include "UnrealEd.h"
+#include "EditorCategoryUtils.h"
+
+#define LOCTEXT_NAMESPACE "EditorCategoryUtils"
+
+/*******************************************************************************
+ * Static FEditorCategoryUtils Helpers
+ ******************************************************************************/
+
+namespace FEditorCategoryUtilsImpl
+{
+	/**
+	 * Gets the table that tracks mappings from string keys to qualified 
+	 * category paths. Inits the structure if it hadn't been before (adds 
+	 * default mappings for all FCommonEditorCategory values)
+	 * 
+	 * @return A map from string category keys, to fully qualified category paths.
+	 */
+	static FFormatNamedArguments& GetCategoryTable();
+
+	/**
+	 * Performs a lookup into the category key table, retrieving a fully 
+	 * qualified category path for the specified key.
+	 * 
+	 * @param  Key	The key you want a category path for.
+	 * @return The category display string associated with the specified key (an empty string if an entry wasn't found).
+	 */
+	static FText const& GetCategory(FString const& Key);
+
+
+	/** Metadata tags */
+	static const FName ClassHideCategoriesMetaKey(TEXT("HideCategories"));
+	static const FName ClassShowCategoriesMetaKey(TEXT("ShowCategories"));
+}
+
+#define ENUM_STRING(EnumVal) FString(#EnumVal)
+
+//------------------------------------------------------------------------------
+static FFormatNamedArguments& FEditorCategoryUtilsImpl::GetCategoryTable()
+{
+	static FFormatNamedArguments CategoryLookup;
+
+	static bool bInitialized = false;
+	// this function is reentrant, so we have to guard against recursion
+	if (!bInitialized)
+	{
+		bInitialized = true;
+
+#define ASSIGN_ROOT_CATEGORY(EnumVal, Text) FCommonEditorCategory::EnumVal; \
+		FEditorCategoryUtils::RegisterCategoryKey(ENUM_STRING(EnumVal), LOCTEXT(#EnumVal L"Category", Text))
+#define REGISTER_ROOT_CATEGORY(EnumVal) ASSIGN_ROOT_CATEGORY(EnumVal, #EnumVal)
+
+		REGISTER_ROOT_CATEGORY(Ai);
+		REGISTER_ROOT_CATEGORY(Animation);
+		REGISTER_ROOT_CATEGORY(Audio);
+		REGISTER_ROOT_CATEGORY(Development);
+		REGISTER_ROOT_CATEGORY(Effects);
+		ASSIGN_ROOT_CATEGORY(Gameplay, "Game");
+		REGISTER_ROOT_CATEGORY(Input);
+		REGISTER_ROOT_CATEGORY(Math);
+		REGISTER_ROOT_CATEGORY(Networking);
+		REGISTER_ROOT_CATEGORY(Pawn);
+		REGISTER_ROOT_CATEGORY(Rendering);
+		REGISTER_ROOT_CATEGORY(Utilities);
+
+#undef REGISTER_ROOT_CATEGORY
+#undef ASSIGN_ROOT_CATEGORY
+
+		// NOTE: the root category has to be registered prior to this
+#define REGISTER_SUB_CATEGORY(RootId, EnumVal) FCommonEditorCategory::EnumVal; \
+		FText EnumVal##_Category = FEditorCategoryUtils::BuildCategoryString(FCommonEditorCategory::RootId, LOCTEXT(#EnumVal L"Category", #EnumVal)); \
+		FEditorCategoryUtils::RegisterCategoryKey(ENUM_STRING(EnumVal), EnumVal##_Category)
+
+		REGISTER_SUB_CATEGORY(Utilities, Transformation);
+		REGISTER_SUB_CATEGORY(Utilities, String);
+		REGISTER_SUB_CATEGORY(Utilities, Text);
+		REGISTER_SUB_CATEGORY(Utilities, Name);
+		REGISTER_SUB_CATEGORY(Utilities, Enum);
+		REGISTER_SUB_CATEGORY(Utilities, Struct);
+		REGISTER_SUB_CATEGORY(Utilities, Macro);
+
+#undef REGISTER_SUB_CATEGORY
+	}
+
+	return CategoryLookup;
+}
+
+//------------------------------------------------------------------------------
+static FText const& FEditorCategoryUtilsImpl::GetCategory(FString const& Key)
+{
+	FFormatNamedArguments const& CategoryLookup = GetCategoryTable();
+	if (FFormatArgumentValue const* FoundCategory = CategoryLookup.Find(Key.ToUpper()))
+	{
+		return *FoundCategory->TextValue;
+	}
+	return FText::GetEmpty();
+}
+
+/*******************************************************************************
+ * FEditorCategoryUtils
+ ******************************************************************************/
+
+//------------------------------------------------------------------------------
+void FEditorCategoryUtils::RegisterCategoryKey(FString const& Key, FText const& Category)
+{
+	FEditorCategoryUtilsImpl::GetCategoryTable().Add(Key.ToUpper(), FEditorCategoryUtils::GetCategoryDisplayString(Category));
+}
+
+//------------------------------------------------------------------------------
+FText const& FEditorCategoryUtils::GetCommonCategory(FCommonEditorCategory::EValue CategoryId)
+{
+	static TMap<FCommonEditorCategory::EValue, FString> CommonCategoryKeys;
+	if (CommonCategoryKeys.Num() == 0)
+	{
+#define REGISTER_COMMON_CATEGORY(EnumVal) \
+		CommonCategoryKeys.Add(FCommonEditorCategory::EnumVal, ENUM_STRING(EnumVal));
+
+		REGISTER_COMMON_CATEGORY(Ai);
+		REGISTER_COMMON_CATEGORY(Animation);
+		REGISTER_COMMON_CATEGORY(Audio);
+		REGISTER_COMMON_CATEGORY(Development);
+		REGISTER_COMMON_CATEGORY(Effects);
+		REGISTER_COMMON_CATEGORY(Gameplay);
+		REGISTER_COMMON_CATEGORY(Input);
+		REGISTER_COMMON_CATEGORY(Math);
+		REGISTER_COMMON_CATEGORY(Networking);
+		REGISTER_COMMON_CATEGORY(Pawn);
+		REGISTER_COMMON_CATEGORY(Rendering);
+		REGISTER_COMMON_CATEGORY(Utilities);
+
+		REGISTER_COMMON_CATEGORY(Transformation);
+		REGISTER_COMMON_CATEGORY(String);
+		REGISTER_COMMON_CATEGORY(Text);
+		REGISTER_COMMON_CATEGORY(Name);
+		REGISTER_COMMON_CATEGORY(Enum);
+		REGISTER_COMMON_CATEGORY(Struct);
+		REGISTER_COMMON_CATEGORY(Macro);
+
+#undef REGISTER_COMMON_CATEGORY
+	}
+
+	if (FString* CategoryKey = CommonCategoryKeys.Find(CategoryId))
+	{
+		return FEditorCategoryUtilsImpl::GetCategory(*CategoryKey);
+	}
+	return FText::GetEmpty();
+}
+
+//------------------------------------------------------------------------------
+FText FEditorCategoryUtils::BuildCategoryString(FCommonEditorCategory::EValue RootId, FText const& SubCategory)
+{
+	FText ConstructedCategory;
+
+	FText const& RootCategory = GetCommonCategory(RootId);
+	if (RootCategory.IsEmpty())
+	{
+		ConstructedCategory = SubCategory;
+	}
+	else if (SubCategory.IsEmpty())
+	{
+		ConstructedCategory = RootCategory;
+	}
+	else
+	{
+		ConstructedCategory = FText::Format(LOCTEXT("ConcatedCategory", "{0}|{1}"), RootCategory, SubCategory);
+	}
+	return ConstructedCategory;
+}
+
+//------------------------------------------------------------------------------
+FText FEditorCategoryUtils::GetCategoryDisplayString(FText const& UnsanitizedCategory)
+{
+	return FText::FromString(GetCategoryDisplayString(UnsanitizedCategory.ToString()));
+}
+
+//------------------------------------------------------------------------------
+FString FEditorCategoryUtils::GetCategoryDisplayString(FString const& UnsanitizedCategory)
+{
+	FString DisplayString = UnsanitizedCategory;
+
+	int32 KeyIndex = INDEX_NONE;
+	do
+	{
+		KeyIndex = DisplayString.Find(TEXT("{"), ESearchCase::CaseSensitive, ESearchDir::FromStart, KeyIndex);
+		if (KeyIndex != INDEX_NONE)
+		{
+			int32 EndIndex = DisplayString.Find(TEXT("}"), ESearchCase::CaseSensitive, ESearchDir::FromStart, KeyIndex);
+			if (EndIndex != INDEX_NONE)
+			{
+				FString ToReplaceStr(EndIndex+1 - KeyIndex, *DisplayString + KeyIndex);
+				FString ReplacementStr;
+				
+				int32 KeyLen = EndIndex - (KeyIndex + 1);
+				if (KeyLen > 0)
+				{
+					FString Key(KeyLen, *DisplayString + KeyIndex+1);
+					ReplacementStr = FEditorCategoryUtilsImpl::GetCategory(Key).ToString();
+				}
+				DisplayString.ReplaceInline(*ToReplaceStr, *ReplacementStr);
+			}
+			KeyIndex = EndIndex;
+		}
+
+	} while (KeyIndex != INDEX_NONE);
+
+	DisplayString = FName::NameToDisplayString(DisplayString, /*bIsBool =*/false);
+	DisplayString.ReplaceInline(TEXT("| "), TEXT("|"));
+
+	return DisplayString;
+}
+
+//------------------------------------------------------------------------------
+void FEditorCategoryUtils::GetClassHideCategories(UClass const* Class, TArray<FString>& CategoriesOut)
+{
+	CategoriesOut.Empty();
+
+	using namespace FEditorCategoryUtilsImpl;
+	if (Class->HasMetaData(ClassHideCategoriesMetaKey))
+	{
+		FString const& HideCategories = Class->GetMetaData(ClassHideCategoriesMetaKey);
+
+		HideCategories.ParseIntoArray(&CategoriesOut, TEXT(" "), /*InCullEmpty =*/true);
+		
+		for (FString& Category : CategoriesOut)
+		{
+			Category = GetCategoryDisplayString(FText::FromString(Category)).ToString();
+		}
+	}
+}
+
+//------------------------------------------------------------------------------
+void  FEditorCategoryUtils::GetClassShowCategories(UClass const* Class, TArray<FString>& CategoriesOut)
+{
+	CategoriesOut.Empty();
+
+	using namespace FEditorCategoryUtilsImpl;
+	if (Class->HasMetaData(ClassShowCategoriesMetaKey))
+	{
+		FString const& ShowCategories = Class->GetMetaData(ClassShowCategoriesMetaKey);
+		ShowCategories.ParseIntoArray(&CategoriesOut, TEXT(" "), /*InCullEmpty =*/true);
+
+		for (FString& Category : CategoriesOut)
+		{
+			Category = GetCategoryDisplayString(FText::FromString(Category)).ToString();
+		}
+	}
+}
+
+//------------------------------------------------------------------------------
+bool FEditorCategoryUtils::IsCategoryHiddenFromClass(UClass const* Class, FCommonEditorCategory::EValue CategoryId)
+{
+	return IsCategoryHiddenFromClass(Class, GetCommonCategory(CategoryId));
+}
+
+//------------------------------------------------------------------------------
+bool FEditorCategoryUtils::IsCategoryHiddenFromClass(UClass const* Class, FText const& Category)
+{
+	return IsCategoryHiddenFromClass(Class, Category.ToString());
+}
+
+//------------------------------------------------------------------------------
+bool FEditorCategoryUtils::IsCategoryHiddenFromClass(UClass const* Class, FString const& Category)
+{
+	bool bIsHidden = false;
+
+	TArray<FString> ClassHideCategories;
+	GetClassHideCategories(Class, ClassHideCategories);
+
+	// run the category through sanitization so we can ensure compares will hit
+	FString const DisplayCategory = GetCategoryDisplayString(Category);
+
+	for (FString& HideCategory : ClassHideCategories)
+	{
+		bIsHidden = (HideCategory == DisplayCategory);
+		if (bIsHidden)
+		{
+			TArray<FString> ClassShowCategories;
+			GetClassShowCategories(Class, ClassShowCategories);
+			// if they hid it, and showed it... favor showing (could be a shown in a sub-class, and hid in a super)
+			bIsHidden = (ClassShowCategories.Find(DisplayCategory) == INDEX_NONE);
+		}
+		else // see if the category's root is hidden
+		{
+			TArray<FString> SubCategoryList;
+			DisplayCategory.ParseIntoArray(&SubCategoryList, TEXT("|"), /*InCullEmpty =*/true);
+
+			FString FullSubCategoryPath;
+			for (FString const& SubCategory : SubCategoryList)
+			{
+				FullSubCategoryPath += SubCategory;
+				if ((HideCategory == SubCategory) || (HideCategory == FullSubCategoryPath))
+				{
+					TArray<FString> ClassShowCategories;
+					GetClassShowCategories(Class, ClassShowCategories);
+					// if they hid it, and showed it... favor showing (could be a shown in a sub-class, and hid in a super)
+					bIsHidden = (ClassShowCategories.Find(DisplayCategory) == INDEX_NONE);
+				}
+				FullSubCategoryPath += "|";
+			}
+		}
+
+		if (bIsHidden)
+		{
+			break;
+		}
+	}
+
+	return bIsHidden;
+}
+
+#undef ENUM_STRING
+#undef LOCTEXT_NAMESPACE
diff --git a/Engine/Source/Editor/UnrealEd/Public/EditorCategoryUtils.h b/Engine/Source/Editor/UnrealEd/Public/EditorCategoryUtils.h
new file mode 100644
--- /dev/null
+++ b/Engine/Source/Editor/UnrealEd/Public/EditorCategoryUtils.h
@@ -0,0 +1,155 @@
+// Copyright 1998-2014 Epic Games, Inc. All Rights Reserved.
+
+#pragma once
+
+/**
+ * Wrapper for an enum detailing common editor categories. Users can reference
+ * these categories in metadata using the enum value name in braces, like so:
+ *
+ *     UFUNCTION(Category="{Utilities}|MySubCategory")
+ *
+ * This gives users the ability to reference shared categories across the 
+ * engine, but gives us the freedom to easily remap them as needed (also gives 
+ * us the ability to easily localize these categories). Games can override these 
+ * default mappings with the RegisterCategoryKey() function.
+ */
+struct FCommonEditorCategory
+{
+	enum EValue
+	{
+		// Function categories:
+		Ai,			
+		Animation,	
+		Audio,		
+		Development,
+		Effects,
+		Gameplay,
+		Input,
+		Math,
+		Networking,
+		Pawn,
+		Physics,
+		Rendering,
+		Transformation,
+		Utilities,
+
+		// Type categories:
+		String,
+		Text,
+		Name,
+		Enum,
+		Struct,
+		Macro,
+
+		Components,
+		Events,
+		Variables,
+		Functions,
+	};
+};
+
+/**
+ * Set of utilities for parsing, querying, and resolving class/function/field 
+ * category data.
+ */
+struct UNREALED_API FEditorCategoryUtils
+{
+	/**
+	 * To facilitate simple category renaming/reordering, we offer a key 
+	 * replacement system, where users can specify a key in their category 
+	 * metadata that will evaluate to some fully qualified category. Use this
+	 * function to register key/category mappings, or to override existing ones
+	 * (like those pre-registered for all the "common" categories).
+	 *
+	 * In metadata, keys are denoted by braces, like {Utilities} here: 
+	 *	UFUNCTION(Category="{Utilities}|MySubCategory")
+	 * 
+	 * @param  Key		A string key that people will use in metadata to reflect this category mapping.
+	 * @param  Category	The qualified category path that you want the key expanded to.
+	 */
+	static void RegisterCategoryKey(FString const& Key, FText const& Category);
+
+	/**
+	 * Retrieves a qualified category path for the desired common category.
+	 * 
+	 * @param  CategoryId	The common category you want a path for.
+	 * @return A text string, (empty if the common category was not registered)
+	 */
+	static FText const& GetCommonCategory(FCommonEditorCategory::EValue CategoryId);
+
+	/**
+	 * Utility function that concatenates the supplied sub-category with one 
+	 * that matches the root category id.
+	 * 
+	 * @param  RootCategory	An id denoting the root category that you want prefixed onto the sub-category.
+	 * @param  SubCategory	A sub-category that you want postfixed to the root category.
+	 * @return A concatenated text string, with the two categories separated by a pipe, '|', character.
+	 */
+	static FText BuildCategoryString(FCommonEditorCategory::EValue RootCategory, FText const& SubCategory);
+
+	/**
+	 * Expands any keys found in the category string (any terms found in square 
+	 * brackets), and sanitizes the name (spacing individual words, etc.).
+	 * 
+	 * @param  RootCategory	An id denoting the root category that you want prefixing the result.
+	 * @param  SubCategory	A sub-category that you want postfixing the result.
+	 * @return A concatenated text string, with the two categories separated by a pipe, '|', character.
+	 */
+	static FText GetCategoryDisplayString(FText const& UnsanitizedCategory);
+
+	/**
+	 * Expands any keys found in the category string (any terms found in square 
+	 * brackets), and sanitizes the name (spacing individual words, etc.).
+	 * 
+	 * @param  RootCategory	An id denoting the root category that you want prefixing the result.
+	 * @param  SubCategory	A sub-category that you want postfixing the result.
+	 * @return A concatenated string, with the two categories separated by a pipe, '|', character.
+	 */
+	static FString GetCategoryDisplayString(FString const& UnsanitizedCategory);
+
+	/**
+	 * Parses out the class's "HideCategories" metadata, and returns it 
+	 * segmented and sanitized.
+	 * 
+	 * @param  Class			The class you want to pull data from.
+	 * @param  CategoriesOut	An array that will be filled with a list of hidden categories.
+	 */
+	static void GetClassHideCategories(UClass const* Class, TArray<FString>& CategoriesOut);
+
+	/**
+	 * Parses out the class's "ShowCategories" metadata, and returns it 
+	 * segmented and sanitized.
+	 * 
+	 * @param  Class			The class you want to pull data from.
+	 * @param  CategoriesOut	An array that will be filled with a list of shown categories.
+	 */
+	static void GetClassShowCategories(UClass const* Class, TArray<FString>& CategoriesOut);
+
+	/**
+	 * Checks to see if the category associated with the supplied common 
+	 * category id is hidden from the specified class.
+	 * 
+	 * @param  Class		The class you want to query.
+	 * @param  CategoryId	An id associated with a category that you want to check.
+	 * @return True if the common category is hidden, false if not.
+	 */
+	static bool IsCategoryHiddenFromClass(UClass const* Class, FCommonEditorCategory::EValue CategoryId);
+
+	/**
+	 * Checks to see if the specified category is hidden from the supplied class.
+	 * 
+	 * @param  Class	The class you want to query.
+	 * @param  Category A category path that you want to check.
+	 * @return True if the category is hidden, false if not.
+	 */
+	static bool IsCategoryHiddenFromClass(UClass const* Class, FText const& Category);
+
+	/**
+	 * Checks to see if the specified category is hidden from the supplied class.
+	 * 
+	 * @param  Class	The class you want to query.
+	 * @param  Category A category path that you want to check.
+	 * @return True if the category is hidden, false if not.
+	 */
+	static bool IsCategoryHiddenFromClass(UClass const* Class, FString const& Category);
+};
\ No newline at end of file
diff --git a/Engine/Source/Programs/UnrealHeaderTool/Private/Class.cpp b/Engine/Source/Programs/UnrealHeaderTool/Private/Class.cpp
--- a/Engine/Source/Programs/UnrealHeaderTool/Private/Class.cpp
+++ b/Engine/Source/Programs/UnrealHeaderTool/Private/Class.cpp
@@ -83,3 +83,23 @@ TArray<FClass*> FClass::GetInterfaceTypes() const
 	}
 	return Result;
 }
+
+void FClass::GetHideCategories(TArray<FString>& OutHideCategories) const
+{
+	static const FName NAME_HideCategories(TEXT("HideCategories"));
+	if (HasMetaData(NAME_HideCategories))
+	{
+		const FString& HideCategories = GetMetaData(NAME_HideCategories);
+		HideCategories.ParseIntoArray(&OutHideCategories, TEXT(" "), true);
+	}
+}
+
+void FClass::GetShowCategories(TArray<FString>& OutShowCategories) const
+{
+	static const FName NAME_ShowCategories(TEXT("ShowCategories"));
+	if (HasMetaData(NAME_ShowCategories))
+	{
+		const FString& ShowCategories = GetMetaData(NAME_ShowCategories);
+		ShowCategories.ParseIntoArray(&OutShowCategories, TEXT(" "), true);
+	}
+}
diff --git a/Engine/Source/Programs/UnrealHeaderTool/Private/Class.h b/Engine/Source/Programs/UnrealHeaderTool/Private/Class.h
--- a/Engine/Source/Programs/UnrealHeaderTool/Private/Class.h
+++ b/Engine/Source/Programs/UnrealHeaderTool/Private/Class.h
@@ -51,4 +51,7 @@ public:
 	TArray<FName> GetDependentNames() const;
 
 	TArray<FClass*> GetInterfaceTypes() const;
+
+	void GetHideCategories(TArray<FString>& OutHideCategories) const;
+	void GetShowCategories(TArray<FString>& OutShowCategories) const;
 };
diff --git a/Engine/Source/Runtime/Core/Public/ProfilingDebugging/ScopedTimers.h b/Engine/Source/Runtime/Core/Public/ProfilingDebugging/ScopedTimers.h
--- a/Engine/Source/Runtime/Core/Public/ProfilingDebugging/ScopedTimers.h
+++ b/Engine/Source/Runtime/Core/Public/ProfilingDebugging/ScopedTimers.h
@@ -1,5 +1,7 @@
 // Copyright 1998-2014 Epic Games, Inc. All Rights Reserved.
 
+#pragma once
+
 /**
  * Utility stopwatch class for tracking the duration of some action (tracks 
  * time in seconds and adds it to the specified variable on destruction).
diff --git a/Engine/Source/Runtime/CoreUObject/Private/UObject/Class.cpp b/Engine/Source/Runtime/CoreUObject/Private/UObject/Class.cpp
--- a/Engine/Source/Runtime/CoreUObject/Private/UObject/Class.cpp
+++ b/Engine/Source/Runtime/CoreUObject/Private/UObject/Class.cpp
@@ -3419,77 +3419,6 @@ const FString UClass::GetConfigName() const
 }
 
 #if WITH_EDITOR || HACK_HEADER_GENERATOR
-void UClass::GetHideCategories(TArray<FString>& OutHideCategories) const
-{
-	static const FName NAME_HideCategories(TEXT("HideCategories"));
-	if (HasMetaData(NAME_HideCategories))
-	{
-		const FString& HideCategories = GetMetaData(NAME_HideCategories);
-		HideCategories.ParseIntoArray(&OutHideCategories, TEXT(" "), true);
-	}
-}
-
-void UClass::GetShowCategories(TArray<FString>& OutShowCategories) const
-{
-	static const FName NAME_ShowCategories(TEXT("ShowCategories"));
-	if (HasMetaData(NAME_ShowCategories))
-	{
-		const FString& ShowCategories = GetMetaData(NAME_ShowCategories);
-		ShowCategories.ParseIntoArray(&OutShowCategories, TEXT(" "), true);
-	}
-}
-
-bool UClass::IsCategoryHidden(const FString& InCategory) const
-{
-	bool bHidden = false;
-	static FName const NAME_HideCategories(TEXT("HideCategories"));
-	static FName const NAME_ShowCategories(TEXT("ShowCategories"));
-	if (HasMetaData(NAME_HideCategories))
-	{
-		FString ShowCategories;
-		if (HasMetaData(NAME_ShowCategories))
-		{
-			ShowCategories = GetMetaData(NAME_ShowCategories);
-		}
-
-		FString const& HideCategories = GetMetaData(NAME_HideCategories);
-		bHidden = !!FCString::StrfindDelim(*HideCategories, *InCategory, TEXT(" "));
-		if (!bHidden)
-		{
-			TArray<FString> SubCategoryList;
-			InCategory.ParseIntoArray(&SubCategoryList, TEXT("|"), true);
-
-			FString SubCategoryPath;
-			for (int32 SubCatIndex = 0; SubCatIndex < SubCategoryList.Num(); ++SubCatIndex)
-			{
-				FString SubCategory = SubCategoryList[SubCatIndex];
-				SubCategory = SubCategory.Replace(TEXT(" "), TEXT(""));
-				SubCategoryPath += SubCategory;
-
-				if (!!FCString::StrfindDelim(*HideCategories, *SubCategory, TEXT(" ")))
-				{
-					bHidden = true;
-				}
-				else if (!!FCString::StrfindDelim(*HideCategories, *SubCategoryPath, TEXT(" ")))
-				{
-					bHidden = true;
-				}
-
-				if (bHidden && !ShowCategories.IsEmpty() && !!FCString::StrfindDelim(*ShowCategories, *SubCategoryPath, TEXT(" ")))
-				{
-					bHidden = false;
-				}
-				SubCategoryPath += "|";
-			}	
-		}
-		else if (!ShowCategories.IsEmpty() && !!FCString::StrfindDelim(*ShowCategories, *InCategory, TEXT(" ")))
-		{
-			bHidden = false;
-		}
-	}
-	return bHidden;
-}
-
 void UClass::GetHideFunctions(TArray<FString>& OutHideFunctions) const
 {
 	static const FName NAME_HideFunctions(TEXT("HideFunctions"));
diff --git a/Engine/Source/Runtime/CoreUObject/Public/UObject/Class.h b/Engine/Source/Runtime/CoreUObject/Public/UObject/Class.h
--- a/Engine/Source/Runtime/CoreUObject/Public/UObject/Class.h
+++ b/Engine/Source/Runtime/CoreUObject/Public/UObject/Class.h
@@ -1737,9 +1737,6 @@ public:
 
 #if WITH_EDITOR || HACK_HEADER_GENERATOR 
 	// Editor only properties
-	void GetHideCategories(TArray<FString>& OutHideCategories) const;
-	void GetShowCategories(TArray<FString>& OutShowCategories) const;
-	bool IsCategoryHidden(const FString& InCategory) const;
 	void GetHideFunctions(TArray<FString>& OutHideFunctions) const;
 	bool IsFunctionHidden(const TCHAR* InFunction) const;
 	void GetAutoExpandCategories(TArray<FString>& OutAutoExpandCategories) const;
diff --git a/Engine/Source/Runtime/Engine/Private/ObjectEditorUtils.cpp b/Engine/Source/Runtime/Engine/Private/ObjectEditorUtils.cpp
--- a/Engine/Source/Runtime/Engine/Private/ObjectEditorUtils.cpp
+++ b/Engine/Source/Runtime/Engine/Private/ObjectEditorUtils.cpp
@@ -4,6 +4,8 @@
 #include "ObjectEditorUtils.h"
 
 #if WITH_EDITOR
+#include "EditorCategoryUtils.h"
+
 namespace FObjectEditorUtils
 {
 
@@ -41,8 +43,8 @@ namespace FObjectEditorUtils
 			static const FName FunctionCategory(TEXT("Category")); // FBlueprintMetadata::MD_FunctionCategory
 			if( !bResult && InFunction->HasMetaData( FunctionCategory ) )
 			{
-				const FString& ClassCategory = InFunction->GetMetaData(FunctionCategory);				
-				bResult = Class->IsCategoryHidden( ClassCategory );
+				FString const& FuncCategory = InFunction->GetMetaData(FunctionCategory);
+				bResult = FEditorCategoryUtils::IsCategoryHiddenFromClass(Class, FuncCategory);
 			}
 		}
 		return bResult;
@@ -53,7 +55,7 @@ namespace FObjectEditorUtils
 		bool bResult = false;
 		if( InVariable && Class )
 		{
-			bResult = Class->IsCategoryHidden( GetCategory(InVariable) );
+			bResult = FEditorCategoryUtils::IsCategoryHiddenFromClass(Class, GetCategory(InVariable));
 		}
 		return bResult;
 	}
