Minor clang fixes.

- unused result (++).
- || and && in the same scope.

diff --git a/Engine/Source/Editor/Kismet/Private/SSCSEditor.cpp b/Engine/Source/Editor/Kismet/Private/SSCSEditor.cpp
--- a/Engine/Source/Editor/Kismet/Private/SSCSEditor.cpp
+++ b/Engine/Source/Editor/Kismet/Private/SSCSEditor.cpp
@@ -620,7 +620,7 @@ FSCSEditorTreeNodePtrType FSCSEditorTreeNode::FindChild(const USCS_Node* InSCSNo
 
 	if(OutDepth && Result.IsValid())
 	{
-		*OutDepth++;
+		*OutDepth += 1;
 	}
 
 	return Result;
@@ -649,7 +649,7 @@ FSCSEditorTreeNodePtrType FSCSEditorTreeNode::FindChild(const UActorComponent* I
 
 	if(OutDepth && Result.IsValid())
 	{
-		*OutDepth++;
+		*OutDepth += 1;
 	}
 
 	return Result;
diff --git a/Engine/Source/Editor/UnrealEd/Private/Fbx/FbxMatineeImport.cpp b/Engine/Source/Editor/UnrealEd/Private/Fbx/FbxMatineeImport.cpp
--- a/Engine/Source/Editor/UnrealEd/Private/Fbx/FbxMatineeImport.cpp
+++ b/Engine/Source/Editor/UnrealEd/Private/Fbx/FbxMatineeImport.cpp
@@ -505,7 +505,7 @@ float FFbxImporter::ImportMatineeActor(FbxNode* Node, UInterpGroupInst* MatineeG
 			}
 
 			// apply Pre rotations only on bones / end of chains
-			if(Node->GetNodeAttribute() && Node->GetNodeAttribute()->GetAttributeType() == FbxNodeAttribute::eSkeleton
+			if((Node->GetNodeAttribute() && Node->GetNodeAttribute()->GetAttributeType() == FbxNodeAttribute::eSkeleton)
 				|| (Node->GetMarker() && Node->GetMarker()->GetType() == FbxMarker::eEffectorFK)
 				|| (Node->GetMarker() && Node->GetMarker()->GetType() == FbxMarker::eEffectorIK))
 			{
diff --git a/Engine/Source/Editor/UnrealEd/Private/Fbx/FbxStaticMeshImport.cpp b/Engine/Source/Editor/UnrealEd/Private/Fbx/FbxStaticMeshImport.cpp
--- a/Engine/Source/Editor/UnrealEd/Private/Fbx/FbxStaticMeshImport.cpp
+++ b/Engine/Source/Editor/UnrealEd/Private/Fbx/FbxStaticMeshImport.cpp
@@ -394,7 +394,7 @@ bool UnFbx::FFbxImporter::BuildStaticMeshFromGeometry(FbxMesh* Mesh, UStaticMesh
 	// build collision
 	//
 	bool bEnableCollision = ImportCollisionModels(StaticMesh, GetNodeNameWithoutNamespace( Node ) )
-		|| GBuildStaticMeshCollision && LODIndex == 0 && ImportOptions->bRemoveDegenerates;
+		|| (GBuildStaticMeshCollision && LODIndex == 0 && ImportOptions->bRemoveDegenerates);
 	for(int32 SectionIndex=MaterialIndexOffset; SectionIndex<MaterialIndexOffset+MaterialCount; SectionIndex++)
 	{
 		FMeshSectionInfo Info = StaticMesh->SectionInfoMap.Get(LODIndex, SectionIndex);
