Cleaning up unused categories.

diff --git a/Engine/Source/Editor/UnrealEd/Public/EditorCategoryUtils.h b/Engine/Source/Editor/UnrealEd/Public/EditorCategoryUtils.h
--- a/Engine/Source/Editor/UnrealEd/Public/EditorCategoryUtils.h
+++ b/Engine/Source/Editor/UnrealEd/Public/EditorCategoryUtils.h
@@ -40,11 +40,6 @@ struct FCommonEditorCategory
 		Enum,
 		Struct,
 		Macro,
-
-		Components,
-		Events,
-		Variables,
-		Functions,
 	};
 };
 
