changed NULL to nullptr

diff --git a/Engine/Source/Editor/UnrealEd/Private/Factories/EditorFactories.cpp b/Engine/Source/Editor/UnrealEd/Private/Factories/EditorFactories.cpp
--- a/Engine/Source/Editor/UnrealEd/Private/Factories/EditorFactories.cpp
+++ b/Engine/Source/Editor/UnrealEd/Private/Factories/EditorFactories.cpp
@@ -90,7 +90,7 @@ UObject* UTexture2DFactoryNew::FactoryCreateNew( UClass* InClass, UObject* InPar
 	// Do not create a texture with bad dimensions.
 	if((Width & (Width - 1)) || (Height & (Height - 1)))
 	{
-		return NULL;
+		return nullptr;
 	}
 
 	UTexture2D* Object = CastChecked<UTexture2D>(StaticConstructObject(InClass,InParent,InName,Flags) );
@@ -158,7 +158,7 @@ UObject* UMaterialFactoryNew::FactoryCreateNew(UClass* Class,UObject* InParent,F
 {
 	UMaterial* NewMaterial = ConstructObject<UMaterial>(Class, InParent, Name, Flags);
 
-	if ( InitialTexture != NULL )
+	if ( InitialTexture != nullptr )
 	{
 		// An initial texture was specified, add it and assign it to the BaseColor
 		UMaterialExpressionTextureSample* Expression = ConstructObject<UMaterialExpressionTextureSample>( UMaterialExpressionTextureSample::StaticClass(), NewMaterial );
@@ -224,9 +224,9 @@ UObject* UMaterialParameterCollectionFactoryNew::FactoryCreateNew(UClass* Class,
 
 /**
  * Iterates over an object's properties making sure that any UObjectProperty properties
- * that refer to non-NULL actors refer to valid actors.
+ * that refer to non-nullptr actors refer to valid actors.
  *
- * @return		false if no object references were NULL'd out, true otherwise.
+ * @return		false if no object references were nullptr'd out, true otherwise.
  */
 static bool ForceValidActorRefs(UStruct* Struct, uint8* Data)
 {
@@ -259,7 +259,7 @@ static bool ForceValidActorRefs(UStruct* Struct, uint8* Data)
 					if( !bFound )
 					{
 						UE_LOG(LogEditorFactories, Log,  TEXT("Usurped %s"), *Obj->GetClass()->GetName() );
-						Prop->SetObjectPropertyValue(Value, NULL);
+						Prop->SetObjectPropertyValue(Value, nullptr);
 						bChangedObjectPointer = true;
 					}
 				}
@@ -324,16 +324,16 @@ UObject* ULevelFactory::FactoryCreateText
 				Buffer += FCString::Strlen(TEXT("Name="));
 				Buffer += MapName.Len();
 				// Check to make sure that there are no naming conflicts
-				if( RootMapPackage->Rename(*MapName, NULL, REN_Test | REN_ForceNoResetLoaders) )
+				if( RootMapPackage->Rename(*MapName, nullptr, REN_Test | REN_ForceNoResetLoaders) )
 				{
 					// Rename it!
-					RootMapPackage->Rename(*MapName, NULL, REN_ForceNoResetLoaders);
+					RootMapPackage->Rename(*MapName, nullptr, REN_ForceNoResetLoaders);
 				}
 				else
 				{
 					Warn->Logf(ELogVerbosity::Warning, TEXT("The Root map package name : '%s', conflicts with the existing object : '%s'"), *RootMapPackage->GetFullName(), *MapName);
-					FEditorDelegates::OnAssetPostImport.Broadcast( this, NULL );
-					return NULL;
+					FEditorDelegates::OnAssetPostImport.Broadcast( this, nullptr );
+					return nullptr;
 				}
 				
 				// Stick it in the package map
@@ -423,7 +423,7 @@ UObject* ULevelFactory::FactoryCreateText
 				if (FParse::Value(Str, TEXT("NAME="), LevelName))
 				{
 					// create a new named level
-					World->SetCurrentLevel( new(World->GetOuter(), *LevelName)ULevel(FPostConstructInitializeProperties(),FURL(NULL)) );
+					World->SetCurrentLevel( new(World->GetOuter(), *LevelName)ULevel(FPostConstructInitializeProperties(),FURL(nullptr)) );
 				}
 			}
 #endif
@@ -446,7 +446,7 @@ UObject* ULevelFactory::FactoryCreateText
 				FParse::Value( Str, TEXT("NAME="), ActorSourceName );
 				ActorUniqueName = ActorSourceName;
 				// Make sure this name is unique.
-				AActor* Found=NULL;
+				AActor* Found=nullptr;
 				if( ActorUniqueName!=NAME_None )
 				{
 					// look in the current level for the same named actor
@@ -467,7 +467,7 @@ UObject* ULevelFactory::FactoryCreateText
 
 				// if an archetype was specified in the Begin Object block, use that as the template for the ConstructObject call.
 				FString ArchetypeName;
-				AActor* Archetype = NULL;
+				AActor* Archetype = nullptr;
 				if (FParse::Value(Str, TEXT("Archetype="), ArchetypeName))
 				{
 					// if given a name, break it up along the ' so separate the class from the name
@@ -525,8 +525,8 @@ UObject* ULevelFactory::FactoryCreateText
 				if ( FLevelUtils::IsLevelLocked(World->GetCurrentLevel()) )
 				{
 					UE_LOG(LogEditorFactories, Warning, TEXT("Import actor: The requested operation could not be completed because the level is locked."));
-					FEditorDelegates::OnAssetPostImport.Broadcast( this, NULL );
-					return NULL;
+					FEditorDelegates::OnAssetPostImport.Broadcast( this, nullptr );
+					return nullptr;
 				}
 				else if ( !(bShouldSkipImportSpecialActors && ActorIndex < 2) )
 				{
@@ -539,7 +539,7 @@ UObject* ULevelFactory::FactoryCreateText
 						SpawnInfo.Name = ActorUniqueName;
 						SpawnInfo.Template = Archetype;
 						SpawnInfo.bNoCollisionFail = true;
-						AActor* NewActor = World->SpawnActor( TempClass, NULL, NULL, SpawnInfo );
+						AActor* NewActor = World->SpawnActor( TempClass, nullptr, nullptr, SpawnInfo );
 						
 						if( NewActor )
 						{
@@ -547,7 +547,7 @@ UObject* ULevelFactory::FactoryCreateText
 							{
 								bool bGrouped = false;
 
-								AGroupActor** tmpNewGroup = NULL;
+								AGroupActor** tmpNewGroup = nullptr;
 								// We need to add all the objects we selected into groups with new objects that were in their group before.
 								FString GroupName;
 								if(FParse::Value(Str, TEXT("GroupActor="), GroupName))
@@ -614,7 +614,7 @@ UObject* ULevelFactory::FactoryCreateText
 		}
 		else if( GetBEGIN(&Str,TEXT("SURFACE")) )
 		{
-			UMaterialInterface* SrcMaterial = NULL;
+			UMaterialInterface* SrcMaterial = nullptr;
 			FVector SrcBase, SrcTextureU, SrcTextureV, SrcNormal;
 			uint32 SrcPolyFlags = PF_DefaultFlags;
 			int32 SurfacePropertiesParsed = 0;
@@ -643,7 +643,7 @@ UObject* ULevelFactory::FactoryCreateText
 					bParsedLineSuccessfully = FParse::Line(&Buffer, TextureName, true);
 					if ( TextureName != TEXT("None") )
 					{
-						SrcMaterial = Cast<UMaterialInterface>(StaticLoadObject( UMaterialInterface::StaticClass(), NULL, *TextureName, NULL, LOAD_NoWarn, NULL ));
+						SrcMaterial = Cast<UMaterialInterface>(StaticLoadObject( UMaterialInterface::StaticClass(), nullptr, *TextureName, nullptr, LOAD_NoWarn, nullptr ));
 					}
 					bJustParsedTextureName = true;
 					SurfacePropertiesParsed++;
@@ -757,7 +757,7 @@ UObject* ULevelFactory::FactoryCreateText
 
 		const TCHAR* MapPkg_BufferStart = *MapPackageText;
 		const TCHAR* MapPkg_BufferEnd = MapPkg_BufferStart + MapPackageText.Len();
-		PackageFactory->FactoryCreateText(UPackage::StaticClass(), NULL, NewPackageName, RF_NoFlags, 0, TEXT("T3D"), MapPkg_BufferStart, MapPkg_BufferEnd, Warn);
+		PackageFactory->FactoryCreateText(UPackage::StaticClass(), nullptr, NewPackageName, RF_NoFlags, 0, TEXT("T3D"), MapPkg_BufferStart, MapPkg_BufferEnd, Warn);
 	}
 
 	// Pass 1: Sort out all the properties on the individual actors
@@ -773,7 +773,7 @@ UObject* ULevelFactory::FactoryCreateText
 		{
 			if ( Actor->ShouldImport(PropText, bIsMoveToStreamingLevel) )
 			{
-				Actor->PreEditChange(NULL);
+				Actor->PreEditChange(nullptr);
 				ImportObjectProperties( (uint8*)Actor, **PropText, Actor->GetClass(), Actor, Actor, Warn, 0 );
 				bActorChanged = true;
 
@@ -875,22 +875,22 @@ UObject* ULevelFactory::FactoryCreateText
 
 		// Fixup parenting
 		FAttachmentDetail* ActorAttachmentDetail = NewActorsAttachmentMap.Find( Actor );
-		if( ActorAttachmentDetail != NULL )
+		if( ActorAttachmentDetail != nullptr )
 		{
-			AActor* ActorParent = NULL;
+			AActor* ActorParent = nullptr;
 			// Try to find the new copy of the parent
 			AActor** NewActorParent = NewActorsFNames.Find( ActorAttachmentDetail->ParentName );
-			if ( NewActorParent != NULL )
+			if ( NewActorParent != nullptr )
 			{
 				ActorParent = *NewActorParent;
 			}
 			// Try to find an already existing parent
-			if( ActorParent == NULL )
+			if( ActorParent == nullptr )
 			{
 				ActorParent = FindObject<AActor>( World->GetCurrentLevel(), *ActorAttachmentDetail->ParentName.ToString() );
 			}
 			// Parent the actors
-			if( ActorParent != NULL )
+			if( ActorParent != nullptr )
 			{
 				// Make sure our parent isn't selected (would cause GEditor->ParentActors to fail)
 				const bool bParentWasSelected = ActorParent->IsSelected();
@@ -948,16 +948,16 @@ UObject* UPackageFactory::FactoryCreateText( UClass* Class, UObject* InParent, F
 	GEditor->IsImportingT3D = true;
 	GIsImportingT3D = true;
 
-	if (InParent != NULL)
+	if (InParent != nullptr)
 	{
-		return NULL;
+		return nullptr;
 	}
 
 	TMap<FString, UPackage*> MapPackages;
 	bool bImportingMapPackage = false;
 
-	UPackage* TopLevelPackage = NULL;
-	UPackage* RootMapPackage = NULL;
+	UPackage* TopLevelPackage = nullptr;
+	UPackage* RootMapPackage = nullptr;
 	UWorld* World = GWorld;
 	if (World)
 	{
@@ -1003,15 +1003,15 @@ UObject* UPackageFactory::FactoryCreateText( UClass* Class, UObject* InParent, F
 				if (FindObject<UPackage>(ANY_PACKAGE, *(PackageName.ToString())))
 				{
 					UE_LOG(LogEditorFactories, Warning, TEXT("Package factory can only handle the map package or new packages!"));
-					return NULL;
+					return nullptr;
 				}
-				TopLevelPackage = CreatePackage(NULL, *(PackageName.ToString()));
+				TopLevelPackage = CreatePackage(nullptr, *(PackageName.ToString()));
 				TopLevelPackage->SetFlags(RF_Standalone|RF_Public);
 				MapPackages.Add(TopLevelPackage->GetName(), TopLevelPackage);
 
 				// if an archetype was specified in the Begin Object block, use that as the template for the ConstructObject call.
 				FString ArchetypeName;
-				AActor* Archetype = NULL;
+				AActor* Archetype = nullptr;
 				if (FParse::Value(Str, TEXT("Archetype="), ArchetypeName))
 				{
 				}
@@ -1030,7 +1030,7 @@ UObject* UPackageFactory::FactoryCreateText( UClass* Class, UObject* InParent, F
 
 				// if an archetype was specified in the Begin Object block, use that as the template for the ConstructObject call.
 				FString ArchetypeName;
-				UPackage* Archetype = NULL;
+				UPackage* Archetype = nullptr;
 				if (FParse::Value(Str, TEXT("Archetype="), ArchetypeName))
 				{
 					// if given a name, break it up along the ' so separate the class from the name
@@ -1055,7 +1055,7 @@ UObject* UPackageFactory::FactoryCreateText( UClass* Class, UObject* InParent, F
 						}
 					}
 
-					UPackage* ParentPkg = NULL;
+					UPackage* ParentPkg = nullptr;
 					UPackage** ppParentPkg = MapPackages.Find(ParentPackageName);
 					if (ppParentPkg)
 					{
@@ -1083,7 +1083,7 @@ UObject* UPackageFactory::FactoryCreateText( UClass* Class, UObject* InParent, F
 			FString* PropText = NewPackageObjectMap.Find(LoadObject);
 			if (PropText)
 			{
-				LoadObject->PreEditChange(NULL);
+				LoadObject->PreEditChange(nullptr);
 				ImportObjectProperties((uint8*)LoadObject, **PropText, LoadObject->GetClass(), LoadObject, LoadObject, Warn, 0 );
 				bModifiedObject = true;
 			}
@@ -1190,7 +1190,7 @@ UObject* UPolysFactory::FactoryCreateText
 							// Got a poly definition.
 							if( IsFace ) NewPoly.Reverse();
 							NewPoly.Base = NewPoly.Vertices[0];
-							NewPoly.Finalize(NULL,0);
+							NewPoly.Finalize(nullptr,0);
 							new(Polys->Element)FPoly( NewPoly );
 						}
 						else
@@ -1300,7 +1300,7 @@ UObject* UPolysFactory::FactoryCreateText
 					new(Poly.Vertices)FVector(PointPool[FCString::Atoi(FCString::Strstr(Str,TEXT("B:"))+2)]);
 					new(Poly.Vertices)FVector(PointPool[FCString::Atoi(FCString::Strstr(Str,TEXT("C:"))+2)]);
 					Poly.Base = Poly.Vertices[0];
-					Poly.Finalize(NULL,0);
+					Poly.Finalize(nullptr,0);
 					new(Polys->Element)FPoly(Poly);
 					TempNumPolys++;
 				}
@@ -1330,9 +1330,9 @@ UObject* UPolysFactory::FactoryCreateText
 			{
 				Poly.Material = Cast<UMaterialInterface>(StaticFindObject( UMaterialInterface::StaticClass(), ANY_PACKAGE, *TextureName ) );
 /***
-				if (Poly.Material == NULL)
+				if (Poly.Material == nullptr)
 				{
-					Poly.Material = Cast<UMaterialInterface>(StaticLoadObject( UMaterialInterface::StaticClass(), NULL, *TextureName, NULL,  LOAD_NoWarn, NULL ) );
+					Poly.Material = Cast<UMaterialInterface>(StaticLoadObject( UMaterialInterface::StaticClass(), nullptr, *TextureName, nullptr,  LOAD_NoWarn, nullptr ) );
 				}
 ***/
 			}
@@ -1371,7 +1371,7 @@ UObject* UPolysFactory::FactoryCreateText
 		{
 			if( !GotBase )
 				Poly.Base = Poly.Vertices[0];
-			if( Poly.Finalize(NULL,1)==0 )
+			if( Poly.Finalize(nullptr,1)==0 )
 				new(Polys->Element)FPoly(Poly);
 			GotBase=0;
 		}
@@ -1430,7 +1430,7 @@ UObject* UModelFactory::FactoryCreateText
 		else if( GetBEGIN (&StrPtr,TEXT("POLYLIST")) )
 		{
 			UPolysFactory* PolysFactory = new UPolysFactory(FPostConstructInitializeProperties());
-			Model->Polys = (UPolys*)PolysFactory->FactoryCreateText(UPolys::StaticClass(),Model,NAME_None,RF_Transactional,NULL,Type,Buffer,BufferEnd,Warn);
+			Model->Polys = (UPolys*)PolysFactory->FactoryCreateText(UPolys::StaticClass(),Model,NAME_None,RF_Transactional,nullptr,Type,Buffer,BufferEnd,Warn);
 			check(Model->Polys);
 		}
 		if( TempOwner )
@@ -1617,13 +1617,13 @@ UObject* USoundFactory::FactoryCreateBinary
 			if( !bCuePathIsValid )
 			{
 				FMessageDialog::Open( EAppMsgType::Ok, FText::Format(NSLOCTEXT("UnrealEd", "Error_ImportFailed_f", "Import failed for {0}: {1}"), FText::FromString(CuePackageName), Reason) );
-				FEditorDelegates::OnAssetPostImport.Broadcast( this, NULL );
-				return NULL;
+				FEditorDelegates::OnAssetPostImport.Broadcast( this, nullptr );
+				return nullptr;
 			}
 		}	
 
 		// if we are creating the cue move it when necessary
-		UPackage* CuePackage = bMoveCue ? CreatePackage( NULL, *CuePackageName ) : NULL;
+		UPackage* CuePackage = bMoveCue ? CreatePackage( nullptr, *CuePackageName ) : nullptr;
 
 		// if the sound already exists, remember the user settings
 		USoundWave* ExistingSound = FindObject<USoundWave>( InParent, *Name.ToString() );
@@ -1667,8 +1667,8 @@ UObject* USoundFactory::FactoryCreateBinary
 				}
 			default:
 				{
-					FEditorDelegates::OnAssetPostImport.Broadcast( this, NULL );
-					return NULL;
+					FEditorDelegates::OnAssetPostImport.Broadcast( this, nullptr );
+					return nullptr;
 				}
 			}
 		}
@@ -1690,23 +1690,23 @@ UObject* USoundFactory::FactoryCreateBinary
 			{
 				WaveInfo.ReportImportFailure();
 				Warn->Logf(ELogVerbosity::Error, TEXT( "Currently, only 16 bit WAV files are supported (%s)." ), *Name.ToString() );
-				FEditorDelegates::OnAssetPostImport.Broadcast( this, NULL );
-				return NULL;
+				FEditorDelegates::OnAssetPostImport.Broadcast( this, nullptr );
+				return nullptr;
 			}
 
 			if( *WaveInfo.pChannels != 1 && *WaveInfo.pChannels != 2 )
 			{
 				WaveInfo.ReportImportFailure();
 				Warn->Logf(ELogVerbosity::Error, TEXT("Currently, only mono or stereo WAV files are supported (%s)."), *Name.ToString());
-				FEditorDelegates::OnAssetPostImport.Broadcast( this, NULL );
-				return NULL;
+				FEditorDelegates::OnAssetPostImport.Broadcast( this, nullptr );
+				return nullptr;
 			}
 		}
 		else
 		{
 			Warn->Logf(ELogVerbosity::Error, TEXT( "Unable to read wave file '%s' - \"%s\""), *Name.ToString(), *ErrorMessage );
-			FEditorDelegates::OnAssetPostImport.Broadcast( this, NULL );
-			return NULL;
+			FEditorDelegates::OnAssetPostImport.Broadcast( this, nullptr );
+			return nullptr;
 		}
 
 		// Use pre-existing sound if it exists and we want to keep settings,
@@ -1770,8 +1770,8 @@ UObject* USoundFactory::FactoryCreateBinary
 	{
 		// Unrecognized.
 		Warn->Logf(ELogVerbosity::Error, TEXT("Unrecognized sound format '%s' in %s"), FileType, *Name.ToString() );
-		FEditorDelegates::OnAssetPostImport.Broadcast( this, NULL );
-		return NULL;
+		FEditorDelegates::OnAssetPostImport.Broadcast( this, nullptr );
+		return nullptr;
 	}
 }
 
@@ -1891,7 +1891,7 @@ EReimportResult::Type UReimportSoundFactory::Reimport( UObject* Obj )
 	// Suppress the import overwrite dialog, we want to keep existing settings when re-importing
 	USoundFactory::SuppressImportOverwriteDialog();
 
-	if( UFactory::StaticImportObject( SoundWave->GetClass(), SoundWave->GetOuter(), *SoundWave->GetName(), RF_Public|RF_Standalone, *Filename, NULL, this ) )
+	if( UFactory::StaticImportObject( SoundWave->GetClass(), SoundWave->GetOuter(), *SoundWave->GetName(), RF_Public|RF_Standalone, *Filename, nullptr, this ) )
 	{
 		UE_LOG(LogEditorFactories, Log, TEXT("-- imported successfully") );
 
@@ -1990,20 +1990,20 @@ UObject* USoundSurroundFactory::FactoryCreateBinary
 		if( SpeakerIndex == SPEAKER_Count )
 		{
 			Warn->Logf(ELogVerbosity::Error, TEXT( "Failed to find speaker location; valid extensions are _fl, _fr, _fc, _lf, _sl, _sr, _bl, _br." ) );
-			FEditorDelegates::OnAssetPostImport.Broadcast( this, NULL );
-			return( NULL );
+			FEditorDelegates::OnAssetPostImport.Broadcast( this, nullptr );
+			return( nullptr );
 		}
 
 		// Find existing soundwave
 		USoundWave* Sound = FindObject<USoundWave>(InParent, *BaseName.ToString());
 
 		// Create new sound if necessary
-		if( Sound == NULL )
+		if( Sound == nullptr )
 		{
 			// If This is a single asset package, then create package so that its name will be identical to the asset.
 			if (PackageTools::IsSingleAssetPackage(InParent->GetName()))
 			{
-				InParent = CreatePackage(NULL, *InParent->GetName().LeftChop(3) );
+				InParent = CreatePackage(nullptr, *InParent->GetName().LeftChop(3) );
 
 				// Make sure the destination package is loaded
 				CastChecked<UPackage>(InParent)->FullyLoad();
@@ -2011,7 +2011,7 @@ UObject* USoundSurroundFactory::FactoryCreateBinary
 				Sound = FindObject<USoundWave>(InParent, *BaseName.ToString());
 			}
 
-			if (Sound == NULL)
+			if (Sound == nullptr)
 			{
 				Sound = new( InParent, BaseName, Flags ) USoundWave(FPostConstructInitializeProperties());
 			}
@@ -2038,7 +2038,7 @@ UObject* USoundSurroundFactory::FactoryCreateBinary
 		Sound->InvalidateCompressedData();
 
 		// Delete the old version of the wave from the bulk data
-		uint8 * RawWaveData[SPEAKER_Count] = { NULL };
+		uint8 * RawWaveData[SPEAKER_Count] = { nullptr };
 		uint8 * RawData = ( uint8 * )Sound->RawData.Lock( LOCK_READ_WRITE );
 		int32 RawDataOffset = 0;
 		int32 TotalSize = 0;
@@ -2099,21 +2099,21 @@ UObject* USoundSurroundFactory::FactoryCreateBinary
 			{
 				Warn->Logf(ELogVerbosity::Error, TEXT( "Currently, only 16 bit WAV files are supported (%s)." ), *Name.ToString() );
 				Sound->MarkPendingKill();
-				Sound = NULL;
+				Sound = nullptr;
 			}
 
 			if( *WaveInfo.pChannels != 1 )
 			{
 				Warn->Logf(ELogVerbosity::Error, TEXT( "Currently, only mono WAV files can be imported as channels of surround audio (%s)." ), *Name.ToString() );
 				Sound->MarkPendingKill();
-				Sound = NULL;
+				Sound = nullptr;
 			}
 		}
 		else
 		{
 			Warn->Logf(ELogVerbosity::Error, TEXT( "Unable to read wave file '%s' - \"%s\""), *Name.ToString(), *ErrorReason );
 			Sound->MarkPendingKill();
-			Sound = NULL;
+			Sound = nullptr;
 		}
 		if (Sound)
 		{
@@ -2135,10 +2135,10 @@ UObject* USoundSurroundFactory::FactoryCreateBinary
 	{
 		// Unrecognized.
 		Warn->Logf(ELogVerbosity::Error, TEXT("Unrecognized sound extension '%s' in %s"), FileType, *Name.ToString() );
-		FEditorDelegates::OnAssetPostImport.Broadcast( this, NULL );
+		FEditorDelegates::OnAssetPostImport.Broadcast( this, nullptr );
 	}
 
-	return( NULL );
+	return( nullptr );
 }
 
 /*-----------------------------------------------------------------------------
@@ -2273,7 +2273,7 @@ EReimportResult::Type UReimportSoundSurroundFactory::Reimport( UObject* Obj )
 		FString SpeakerLocation = FPaths::GetBaseFilename(Filename).Right(3);
 		FName ImportName = *(SoundWave->GetName() + SpeakerLocation);
 
-		if (UFactory::StaticImportObject(SoundWave->GetClass(), SoundWave->GetOuter(), ImportName, RF_Public|RF_Standalone, *Filename, NULL, this ))
+		if (UFactory::StaticImportObject(SoundWave->GetClass(), SoundWave->GetOuter(), ImportName, RF_Public|RF_Standalone, *Filename, nullptr, this ))
 		{
 			FFormatNamedArguments Arguments;
 			Arguments.Add(TEXT("NameText"), NameText);
@@ -2370,7 +2370,7 @@ UObject* USoundClassFactory::FactoryCreateNew( UClass* InClass, UObject* InParen
 {
 	USoundClass* SoundClass = ConstructObject<USoundClass>( USoundClass::StaticClass(), InParent, InName, Flags );
 	
-	FAudioDevice* AudioDevice = GEngine ? GEngine->GetAudioDevice() : NULL;
+	FAudioDevice* AudioDevice = GEngine ? GEngine->GetAudioDevice() : nullptr;
 	if (AudioDevice)
 	{
 		AudioDevice->InitSoundClasses();
@@ -2605,8 +2605,8 @@ struct FPSDFileHeader
 
 static bool psd_ReadData( uint8* pOut, const uint8*& pBuffer, FPSDFileHeader& Info )
 {
-	const uint8* pPlane = NULL;
-	const uint8* pRowTable = NULL;
+	const uint8* pPlane = nullptr;
+	const uint8* pRowTable = nullptr;
 	int32         iPlane;
 	int16       CompressionType;
 	int32         iPixel;
@@ -3475,7 +3475,7 @@ UTexture* UTextureFactory::ImportTexture(UClass* Class, UObject* InParent, FName
 		if(TextureFormat == TSF_Invalid)
 		{
 			Warn->Logf(ELogVerbosity::Error, TEXT("PNG file contains data in an unsupported format."));
-			return NULL;
+			return nullptr;
 		}
 
 		UTexture2D* Texture = CreateTexture2D( InParent, Name, Flags );
@@ -3593,12 +3593,12 @@ UTexture* UTextureFactory::ImportTexture(UClass* Class, UObject* InParent, FName
 		// Check the resolution of the imported texture to ensure validity
 		if ( !IsImportResolutionValid(bmhdr->biWidth, bmhdr->biHeight, bAllowNonPowerOfTwo, Warn) )
 		{
-			return NULL;
+			return nullptr;
 		}
 		if( bmhdr->biCompression != BCBI_RGB )
 		{
 			Warn->Logf(ELogVerbosity::Error, TEXT("RLE compression of BMP images not supported") );
-			return NULL;
+			return nullptr;
 		}
 		if( bmhdr->biPlanes==1 && bmhdr->biBitCount==8 )
 		{
@@ -3706,12 +3706,12 @@ UTexture* UTextureFactory::ImportTexture(UClass* Class, UObject* InParent, FName
 		else if( bmhdr->biPlanes==1 && bmhdr->biBitCount==16 )
 		{
 			Warn->Logf(ELogVerbosity::Error, TEXT("BMP 16 bit format no longer supported. Use terrain tools for importing/exporting heightmaps.") );
-			return NULL;
+			return nullptr;
 		}
 		else
 		{
 			Warn->Logf(ELogVerbosity::Error, TEXT("BMP uses an unsupported format (%i/%i)"), bmhdr->biPlanes, bmhdr->biBitCount );
-			return NULL;
+			return nullptr;
 		}
 
 		return Texture;
@@ -3730,7 +3730,7 @@ UTexture* UTextureFactory::ImportTexture(UClass* Class, UObject* InParent, FName
 		// Check the resolution of the imported texture to ensure validity
 		if ( !IsImportResolutionValid(NewU, NewV, bAllowNonPowerOfTwo, Warn) )
 		{
-			return NULL;
+			return nullptr;
 		}
 		else if( PCX->NumPlanes==1 && PCX->BitsPerPixel==8 )
 		{
@@ -3834,7 +3834,7 @@ UTexture* UTextureFactory::ImportTexture(UClass* Class, UObject* InParent, FName
 		else
 		{
 			Warn->Logf(ELogVerbosity::Error, TEXT("PCX uses an unsupported format (%i/%i)"), PCX->NumPlanes, PCX->BitsPerPixel );
-			return NULL;
+			return nullptr;
 		}
 
 		return Texture;
@@ -3854,7 +3854,7 @@ UTexture* UTextureFactory::ImportTexture(UClass* Class, UObject* InParent, FName
 		// Check the resolution of the imported texture to ensure validity
 		if ( !IsImportResolutionValid(TGA->Width, TGA->Height, bAllowNonPowerOfTwo, Warn) )
 		{
-			return NULL;
+			return nullptr;
 		}
 
 		return DecompressTGA(TGA, this, Class, InParent, Name, Flags, Warn);
@@ -3872,12 +3872,12 @@ UTexture* UTextureFactory::ImportTexture(UClass* Class, UObject* InParent, FName
 		// Check the resolution of the imported texture to ensure validity
 		if ( !IsImportResolutionValid(psdhdr.Width, psdhdr.Height, bAllowNonPowerOfTwo, Warn) )
 		{
-			return NULL;
+			return nullptr;
 		}
 		if (!psdhdr.IsSupported())
 		{
 			Warn->Logf( TEXT("Format of this PSD is not supported") );
-			return NULL;
+			return nullptr;
 		}
 
 		// Select the texture's source format
@@ -3894,7 +3894,7 @@ UTexture* UTextureFactory::ImportTexture(UClass* Class, UObject* InParent, FName
 		if(TextureFormat == TSF_Invalid)
 		{
 			Warn->Logf(ELogVerbosity::Error, TEXT("PSD file contains data in an unsupported format."));
-			return NULL;
+			return nullptr;
 		}
 
 		UTexture2D* Texture = CreateTexture2D( InParent, Name, Flags );
@@ -3915,7 +3915,7 @@ UTexture* UTextureFactory::ImportTexture(UClass* Class, UObject* InParent, FName
 				Warn->Logf( TEXT("Failed to read this PSD") );
 				Texture->Source.UnlockMip(0);
 				Texture->MarkPendingKill();
-				return NULL;
+				return nullptr;
 			}
 			Texture->Source.UnlockMip(0);
 		}
@@ -3931,7 +3931,7 @@ UTexture* UTextureFactory::ImportTexture(UClass* Class, UObject* InParent, FName
 		if(!IsImportResolutionValid(DDSLoadHelper.DDSHeader->dwWidth, DDSLoadHelper.DDSHeader->dwHeight, bAllowNonPowerOfTwo, Warn))
 		{
 			Warn->Logf(ELogVerbosity::Error, TEXT("DDS uses an unsupported format"));
-			return NULL;
+			return nullptr;
 		}
 
 		int32 NumMips = DDSLoadHelper.ComputeMipMapCount();
@@ -3939,7 +3939,7 @@ UTexture* UTextureFactory::ImportTexture(UClass* Class, UObject* InParent, FName
 		if (Format == TSF_Invalid)
 		{
 			Warn->Logf(ELogVerbosity::Error, TEXT("DDS file contains data in an unsupported format."));
-			return NULL;
+			return nullptr;
 		}
 
 		// create the cube texture
@@ -4000,7 +4000,7 @@ UTexture* UTextureFactory::ImportTexture(UClass* Class, UObject* InParent, FName
 		if(!IsImportResolutionValid(DDSLoadHelper.DDSHeader->dwWidth, DDSLoadHelper.DDSHeader->dwHeight, bAllowNonPowerOfTwo, Warn))
 		{
 			Warn->Logf(ELogVerbosity::Error, TEXT("DDS uses an unsupported format"));
-			return NULL;
+			return nullptr;
 		}
 		
 		ETextureSourceFormat SourceFormat = DDSLoadHelper.ComputeSourceFormat();
@@ -4103,7 +4103,7 @@ UTexture* UTextureFactory::ImportTexture(UClass* Class, UObject* InParent, FName
 		}
 	}
 
-	return NULL;
+	return nullptr;
 }
 
 bool UTextureFactory::DoesSupportClass(UClass* Class)
@@ -4188,8 +4188,8 @@ UObject* UTextureFactory::FactoryCreateBinary
 		case EAppReturnType::Cancel:
 		default:
 			{
-				FEditorDelegates::OnAssetPostImport.Broadcast( this, NULL );
-				return NULL;
+				FEditorDelegates::OnAssetPostImport.Broadcast( this, nullptr );
+				return nullptr;
 			}
 		}
 	}
@@ -4252,8 +4252,8 @@ UObject* UTextureFactory::FactoryCreateBinary
 		}
 
 		Warn->Logf(ELogVerbosity::Error, TEXT("Texture import failed") );
-		FEditorDelegates::OnAssetPostImport.Broadcast( this, NULL );
-		return NULL;
+		FEditorDelegates::OnAssetPostImport.Broadcast( this, nullptr );
+		return nullptr;
 	}
 
 	// Start with the value that the loader suggests.
@@ -4370,7 +4370,7 @@ UObject* UTextureFactory::FactoryCreateBinary
 		// Create the package for the material
 		const FString MaterialName = FString::Printf( TEXT("%s_Mat"), *Name.ToString() );
 		const FString MaterialPackageName = FPackageName::GetLongPackagePath(InParent->GetName()) + TEXT("/") + MaterialName;
-		UPackage* MaterialPackage = CreatePackage(NULL, *MaterialPackageName);
+		UPackage* MaterialPackage = CreatePackage(nullptr, *MaterialPackageName);
 
 		// Create the material
 		UMaterialFactoryNew* Factory = new UMaterialFactoryNew(FPostConstructInitializeProperties());
@@ -4861,7 +4861,7 @@ public:
 	*/
 	bool ExportHDR(UTextureRenderTarget2D* TexRT, FArchive& Ar)
 	{
-		check(TexRT != NULL);
+		check(TexRT != nullptr);
 		FRenderTarget* RenderTarget = TexRT->GameThread_GetRenderTargetResource();
 		Size = RenderTarget->GetSizeXY();
 		Format = TexRT->GetFormat();
@@ -4886,7 +4886,7 @@ public:
 	template <typename TCubeTextureType>
 	bool ExportHDR(TCubeTextureType* TexCube, FArchive& Ar)
 	{
-		check(TexCube != NULL);
+		check(TexCube != nullptr);
 
 		// Generate 2D image.
 		TArray<uint8> RawData;
@@ -4922,11 +4922,11 @@ bool UTextureExporterHDR::ExportBinary(UObject* Object, const TCHAR* Type, FArch
 	UTextureRenderTargetCube* TexRTCube = Cast<UTextureRenderTargetCube>(Object);
 
 	FHDRExportHelper Exporter;
-	if (TexRT2D != NULL)
+	if (TexRT2D != nullptr)
 	{
 		return Exporter.ExportHDR(TexRT2D, Ar);
 	}
-	else if (TexRTCube != NULL)
+	else if (TexRTCube != nullptr)
 	{
 		return Exporter.ExportHDR(TexRTCube, Ar);
 	}
@@ -4951,7 +4951,7 @@ bool UTextureCubeExporterHDR::ExportBinary(UObject* Object, const TCHAR* Type, F
 	UTextureCube* TexCube = Cast<UTextureCube>(Object);
 
 	FHDRExportHelper Exporter;
-	if (TexCube != NULL)
+	if (TexCube != nullptr)
 	{
 		return Exporter.ExportHDR(TexCube, Ar);
 	}
@@ -5171,7 +5171,7 @@ UObject* UFontFactory::FactoryCreateBinary
 	UTexture2D* Tex = CastChecked<UTexture2D>( UTextureFactory::FactoryCreateBinary( 
 		UTexture2D::StaticClass(), Font, NAME_None, RF_Public, Context, Type, Buffer, BufferEnd, Warn ) );
 
-	if( Tex != NULL )
+	if( Tex != nullptr )
 	{
 		Tex->LODGroup = TEXTUREGROUP_UI;  // set the LOD group otherwise this will be in the World Group
 
@@ -5252,8 +5252,8 @@ UObject* UFontFactory::FactoryCreateBinary
 	else 
 	{
 		Font->MarkPendingKill();
-		FEditorDelegates::OnAssetPostImport.Broadcast( this, NULL );
-		return NULL;
+		FEditorDelegates::OnAssetPostImport.Broadcast( this, nullptr );
+		return nullptr;
 	}
 }
 
@@ -5265,8 +5265,8 @@ FCustomizableTextObjectFactory implementation.
 void FCustomizableTextObjectFactory::ClearObjectNameUsage(UObject* InParent, FName InName)
 {
 	// Make sure this name is unique within the scope of InParent.
-	UObject* Found=NULL;
-	if( (InName != NAME_None) && (InParent != NULL) )
+	UObject* Found=nullptr;
+	if( (InName != NAME_None) && (InParent != nullptr) )
 	{
 		Found = FindObject<UObject>( InParent, *InName.ToString() );
 	}
@@ -5276,7 +5276,7 @@ void FCustomizableTextObjectFactory::ClearObjectNameUsage(UObject* InParent, FNa
 	{
 		check(Found->GetOuter() == InParent);
 
-		Found->Rename(NULL, NULL, REN_None);
+		Found->Rename(nullptr, nullptr, REN_None);
 	}
 }
 
@@ -5322,7 +5322,7 @@ void FCustomizableTextObjectFactory::ProcessBuffer(UObject* InParent, EObjectFla
 				FString ObjArchetypeName;
 				FParse::Value( Str, TEXT("ARCHETYPE="), ObjArchetypeName );
 				UObject* ObjArchetype;
-				ObjArchetype = LoadObject<UObject>(NULL, *ObjArchetypeName, NULL, LOAD_None, NULL);
+				ObjArchetype = LoadObject<UObject>(nullptr, *ObjArchetypeName, nullptr, LOAD_None, nullptr);
 
 				// Make sure this name is not used by anything else. Will rename other stuff if necessary
 				ClearObjectNameUsage(InParent, ObjName);
@@ -5543,7 +5543,7 @@ EReimportResult::Type UReimportTextureFactory::Reimport( UObject* Obj )
 	// Suppress the import overwrite dialog because we know that for explicitly re-importing we want to preserve existing settings
 	UTextureFactory::SuppressImportOverwriteDialog();
 
-	if (UFactory::StaticImportObject(pTex->GetClass(), pTex->GetOuter(), *pTex->GetName(), RF_Public|RF_Standalone, *ResolvedSourceFilePath, NULL, this))
+	if (UFactory::StaticImportObject(pTex->GetClass(), pTex->GetOuter(), *pTex->GetName(), RF_Public|RF_Standalone, *ResolvedSourceFilePath, nullptr, this))
 	{
 		UE_LOG(LogEditorFactories, Log, TEXT("-- imported successfully") );
 		// Try to find the outer package so we can dirty it up
@@ -6054,8 +6054,8 @@ UBlueprintGeneratedClassFactory::UBlueprintGeneratedClassFactory(const class FPo
 	FString ClassPath;
 	GConfig->GetString(TEXT("/Script/Engine.Engine"), TEXT("DefaultBlueprintBaseClassName"), /*out*/ ClassPath, GEngineIni);
 	UClass* DefaultParentClass = !ClassPath.IsEmpty() 
-		? LoadClass<UObject>(NULL, *ClassPath, NULL, LOAD_None, NULL) 
-		: NULL;
+		? LoadClass<UObject>(nullptr, *ClassPath, nullptr, LOAD_None, nullptr) 
+		: nullptr;
 	
 	if( !DefaultParentClass || !FKismetEditorUtilities::CanCreateBlueprintOfClass(DefaultParentClass) )
 	{
@@ -6071,7 +6071,7 @@ UBlueprintGeneratedClassFactory::UBlueprintGeneratedClassFactory(const class FPo
 bool UBlueprintGeneratedClassFactory::ConfigureProperties()
 {
 	// Null the parent class to ensure one is selected
-	ParentClass = NULL;
+	ParentClass = nullptr;
 
 	// Fill in options
 	FClassViewerInitializationOptions Options;
@@ -6098,7 +6098,7 @@ bool UBlueprintGeneratedClassFactory::ConfigureProperties()
 	Filter->DisallowedChildrenOfClasses.Add(UInterface::StaticClass());
 
 	const FText TitleText = LOCTEXT("CreateBlueprintOptions", "Pick Parent Class");
-	UClass* ChosenClass = NULL;
+	UClass* ChosenClass = nullptr;
 	const bool bPressedOk = SClassPickerDialog::PickClass(TitleText, Options, ChosenClass, UBlueprint::StaticClass());
 
 	if ( bPressedOk )
@@ -6116,11 +6116,11 @@ UObject* UBlueprintGeneratedClassFactory::FactoryCreateNew(UClass* Class, UObjec
 	// Make sure we are trying to factory a blueprint, then create and init one
 	check(Class->IsChildOf(UBlueprintGeneratedClass::StaticClass()));
 
-	if ((ParentClass == NULL) || !FKismetEditorUtilities::CanCreateBlueprintOfClass(ParentClass))
+	if ((ParentClass == nullptr) || !FKismetEditorUtilities::CanCreateBlueprintOfClass(ParentClass))
 	{
 		FMessageDialog::Open( EAppMsgType::Ok, FText::FromString(FString::Printf(TEXT("Cannot create a blueprint based on the class '%s'."),
-			(ParentClass != NULL) ? *(ParentClass->GetName()) : TEXT("(null)"))));
-		return NULL;
+			(ParentClass != nullptr) ? *(ParentClass->GetName()) : TEXT("(null)"))));
+		return nullptr;
 	}
 	else
 	{
@@ -6139,8 +6139,8 @@ UBlueprintFactory::UBlueprintFactory(const class FPostConstructInitializePropert
 	FString ClassPath;
 	GConfig->GetString(TEXT("/Script/Engine.Engine"), TEXT("DefaultBlueprintBaseClassName"), /*out*/ ClassPath, GEngineIni);
 	UClass* DefaultParentClass = !ClassPath.IsEmpty() 
-		? LoadClass<UObject>(NULL, *ClassPath, NULL, LOAD_None, NULL) 
-		: NULL;
+		? LoadClass<UObject>(nullptr, *ClassPath, nullptr, LOAD_None, nullptr) 
+		: nullptr;
 	
 	if( !DefaultParentClass || !FKismetEditorUtilities::CanCreateBlueprintOfClass(DefaultParentClass) )
 	{
@@ -6156,7 +6156,7 @@ UBlueprintFactory::UBlueprintFactory(const class FPostConstructInitializePropert
 bool UBlueprintFactory::ConfigureProperties()
 {
 	// Null the parent class to ensure one is selected
-	ParentClass = NULL;
+	ParentClass = nullptr;
 
 	// Fill in options
 	FClassViewerInitializationOptions Options;
@@ -6186,7 +6186,7 @@ bool UBlueprintFactory::ConfigureProperties()
 	Filter->DisallowedChildrenOfClasses.Add(UInterface::StaticClass());
 
 	const FText TitleText = LOCTEXT("CreateBlueprintOptions", "Pick Parent Class");
-	UClass* ChosenClass = NULL;
+	UClass* ChosenClass = nullptr;
 	const bool bPressedOk = SClassPickerDialog::PickClass(TitleText, Options, ChosenClass, UBlueprint::StaticClass());
 
 	if ( bPressedOk )
@@ -6204,12 +6204,12 @@ UObject* UBlueprintFactory::FactoryCreateNew(UClass* Class, UObject* InParent, F
 	// Make sure we are trying to factory a blueprint, then create and init one
 	check(Class->IsChildOf(UBlueprint::StaticClass()));
 
-	if ((ParentClass == NULL) || !FKismetEditorUtilities::CanCreateBlueprintOfClass(ParentClass))
+	if ((ParentClass == nullptr) || !FKismetEditorUtilities::CanCreateBlueprintOfClass(ParentClass))
 	{
 		FFormatNamedArguments Args;
-		Args.Add( TEXT("ClassName"), (ParentClass != NULL) ? FText::FromString( ParentClass->GetName() ) : LOCTEXT("Null", "(null)") );
+		Args.Add( TEXT("ClassName"), (ParentClass != nullptr) ? FText::FromString( ParentClass->GetName() ) : LOCTEXT("Null", "(null)") );
 		FMessageDialog::Open( EAppMsgType::Ok, FText::Format( LOCTEXT("CannotCreateBlueprintFromClass", "Cannot create a blueprint based on the class '{0}'."), Args ) );
-		return NULL;
+		return nullptr;
 	}
 	else
 	{
@@ -6256,12 +6256,12 @@ UObject* UBlueprintMacroFactory::FactoryCreateNew(UClass* Class, UObject* InPare
 	// Make sure we are trying to factory a blueprint, then create and init one
 	check(Class->IsChildOf(UBlueprint::StaticClass()));
 
-	if ((ParentClass == NULL) || !FKismetEditorUtilities::CanCreateBlueprintOfClass(ParentClass))
+	if ((ParentClass == nullptr) || !FKismetEditorUtilities::CanCreateBlueprintOfClass(ParentClass))
 	{
 		FFormatNamedArguments Args;
-		Args.Add( TEXT("ClassName"), (ParentClass != NULL) ? FText::FromString( ParentClass->GetName() ) : LOCTEXT("Null", "(null)") );
+		Args.Add( TEXT("ClassName"), (ParentClass != nullptr) ? FText::FromString( ParentClass->GetName() ) : LOCTEXT("Null", "(null)") );
 		FMessageDialog::Open( EAppMsgType::Ok, FText::Format( LOCTEXT("CannotCreateBlueprintFromClass", "Cannot create a blueprint based on the class '{0}'."), Args ) );
-		return NULL;
+		return nullptr;
 	}
 	else
 	{
@@ -6311,9 +6311,9 @@ UObject* UBlueprintFunctionLibraryFactory::FactoryCreateNew(UClass* Class, UObje
 	if (ParentClass != UBlueprintFunctionLibrary::StaticClass())
 	{
 		FFormatNamedArguments Args;
-		Args.Add(TEXT("ClassName"), (ParentClass != NULL) ? FText::FromString(ParentClass->GetName()) : LOCTEXT("Null", "(null)"));
+		Args.Add(TEXT("ClassName"), (ParentClass != nullptr) ? FText::FromString(ParentClass->GetName()) : LOCTEXT("Null", "(null)"));
 		FMessageDialog::Open(EAppMsgType::Ok, FText::Format(LOCTEXT("CannotCreateBlueprintFromClass", "Cannot create a blueprint based on the class '{0}'."), Args));
-		return NULL;
+		return nullptr;
 	}
 	else
 	{
@@ -6366,12 +6366,12 @@ UObject* UBlueprintInterfaceFactory::FactoryCreateNew(UClass* Class, UObject* In
 	// Force the parent class to be UInterface as per original code
 	UClass* ParentClass = UInterface::StaticClass();
 
-	if ((ParentClass == NULL) || !FKismetEditorUtilities::CanCreateBlueprintOfClass(ParentClass))
+	if ((ParentClass == nullptr) || !FKismetEditorUtilities::CanCreateBlueprintOfClass(ParentClass))
 	{
 		FFormatNamedArguments Args;
-		Args.Add( TEXT("ClassName"), (ParentClass != NULL) ? FText::FromString( ParentClass->GetName() ) : LOCTEXT("Null", "(null)") );
+		Args.Add( TEXT("ClassName"), (ParentClass != nullptr) ? FText::FromString( ParentClass->GetName() ) : LOCTEXT("Null", "(null)") );
 		FMessageDialog::Open( EAppMsgType::Ok, FText::Format( LOCTEXT("CannotCreateBlueprintFromClass", "Cannot create a blueprint based on the class '{0}'."), Args ) );
-		return NULL;
+		return nullptr;
 	}
 	else
 	{
@@ -6429,7 +6429,7 @@ public:
 bool UCurveFactory::ConfigureProperties()
 {
 	// Null the CurveClass so we can get a clean class
-	CurveClass = NULL;
+	CurveClass = nullptr;
 
 	// Load the classviewer module to display a class picker
 	FClassViewerModule& ClassViewerModule = FModuleManager::LoadModuleChecked<FClassViewerModule>("ClassViewer");
@@ -6445,7 +6445,7 @@ bool UCurveFactory::ConfigureProperties()
 	Filter->AllowedChildrenOfClasses.Add(UCurveBase::StaticClass());
 
 	const FText TitleText = LOCTEXT("CreateCurveOptions", "Pick Curve Class");
-	UClass* ChosenClass = NULL;
+	UClass* ChosenClass = nullptr;
 	const bool bPressedOk = SClassPickerDialog::PickClass(TitleText, Options, ChosenClass, UCurveBase::StaticClass());
 
 	if ( bPressedOk )
@@ -6458,8 +6458,8 @@ bool UCurveFactory::ConfigureProperties()
 
 UObject* UCurveFactory::FactoryCreateNew(UClass* Class, UObject* InParent, FName Name, EObjectFlags Flags, UObject* Context, FFeedbackContext* Warn)
 {
-	UCurveBase* NewCurve = NULL;
-	if(CurveClass != NULL)
+	UCurveBase* NewCurve = nullptr;
+	if(CurveClass != nullptr)
 	{
 		NewCurve = CastChecked<UCurveBase>(StaticConstructObject(CurveClass,InParent,Name,Flags));		
 	}
@@ -6495,30 +6495,30 @@ UObject* UCurveImportFactory::FactoryCreateText( UClass* InClass, UObject* InPar
 		TCHAR const* BufRead = Buffer;
 
 		// first line is faFile="", we can ignore
-		if ( !FParse::Line(&BufRead, Str) ) return NULL;
+		if ( !FParse::Line(&BufRead, Str) ) return nullptr;
 		FParse::Next(&BufRead);
 
 		// 2nd line is fps=X
 		float KeyFrameHz = 0.f;
-		if ( !FParse::Value(BufRead, TEXT("fps="), KeyFrameHz) ) return NULL;
-		if ( !FParse::Line(&BufRead, Str) ) return NULL;
+		if ( !FParse::Value(BufRead, TEXT("fps="), KeyFrameHz) ) return nullptr;
+		if ( !FParse::Line(&BufRead, Str) ) return nullptr;
 		FParse::Next(&BufRead);
 
 		// next line is scale=X, we can ignore?
-		if ( !FParse::Line(&BufRead, Str) ) return NULL;
+		if ( !FParse::Line(&BufRead, Str) ) return nullptr;
 		FParse::Next(&BufRead);
 		// next line is smoothing=X, we can ignore?
-		if ( !FParse::Line(&BufRead, Str) ) return NULL;
+		if ( !FParse::Line(&BufRead, Str) ) return nullptr;
 		FParse::Next(&BufRead);
 		// next line is dBValues=X, we can ignore?
-		if ( !FParse::Line(&BufRead, Str) ) return NULL;
+		if ( !FParse::Line(&BufRead, Str) ) return nullptr;
 		FParse::Next(&BufRead);
 		// next line is stereo=X, we can ignore?
-		if ( !FParse::Line(&BufRead, Str) ) return NULL;
+		if ( !FParse::Line(&BufRead, Str) ) return nullptr;
 		FParse::Next(&BufRead);
 
 		// next line is amplitude=[, then list of CSV floats
-		if ( !FParse::Value(BufRead, TEXT("amplitude=["), Str) ) return NULL;
+		if ( !FParse::Value(BufRead, TEXT("amplitude=["), Str) ) return nullptr;
 		BufRead += FCString::Strlen(TEXT("amplitude=["));
 
 		TArray<float> FloatKeys;
@@ -6563,8 +6563,8 @@ UObject* UCurveImportFactory::FactoryCreateText( UClass* InClass, UObject* InPar
 		}
 	}
 
-	FEditorDelegates::OnAssetPostImport.Broadcast( this, NULL );
-	return NULL;
+	FEditorDelegates::OnAssetPostImport.Broadcast( this, nullptr );
+	return nullptr;
 }
 
 
@@ -6599,8 +6599,8 @@ UDataAssetFactory::UDataAssetFactory(const class FPostConstructInitializePropert
 
 bool UDataAssetFactory::ConfigureProperties()
 {
-	// NULL the DataAssetClass so we can check for selection
-	DataAssetClass = NULL;
+	// nullptr the DataAssetClass so we can check for selection
+	DataAssetClass = nullptr;
 
 	// Load the classviewer module to display a class picker
 	FClassViewerModule& ClassViewerModule = FModuleManager::LoadModuleChecked<FClassViewerModule>("ClassViewer");
@@ -6616,7 +6616,7 @@ bool UDataAssetFactory::ConfigureProperties()
 	Filter->AllowedChildrenOfClasses.Add(UDataAsset::StaticClass());
 
 	const FText TitleText = LOCTEXT("CreateDataAssetOptions", "Pick Data Asset Class");
-	UClass* ChosenClass = NULL;
+	UClass* ChosenClass = nullptr;
 	const bool bPressedOk = SClassPickerDialog::PickClass(TitleText, Options, ChosenClass, UDataAsset::StaticClass());
 
 	if ( bPressedOk )
@@ -6629,7 +6629,7 @@ bool UDataAssetFactory::ConfigureProperties()
 
 UObject* UDataAssetFactory::FactoryCreateNew(UClass* Class, UObject* InParent, FName Name, EObjectFlags Flags, UObject* Context, FFeedbackContext* Warn)
 {
-	if (DataAssetClass != NULL)
+	if (DataAssetClass != nullptr)
 	{
 		return CastChecked<UDataAsset>(StaticConstructObject(DataAssetClass, InParent, Name, Flags | RF_Transactional));
 	}
@@ -6637,7 +6637,7 @@ UObject* UDataAssetFactory::FactoryCreateNew(UClass* Class, UObject* InParent, F
 	{
 		// if we have no data asset class, use the passed-in class instead
 		check(Class->IsChildOf(UDataAsset::StaticClass()));
-		return CastChecked<UDataAsset>(StaticConstructObject(Class,InParent,Name,Flags));
+		return CastChecked<UDataAsset>(StaticConstructObject(Class, InParent, Name, Flags));
 	}
 }
 
@@ -6682,15 +6682,15 @@ UObject* UDestructibleMeshFactory::FactoryCreateBinary
 	FEditorDelegates::OnAssetPreImport.Broadcast(this, Class, InParent, Name, FileType);
 
 	// The return value
-	UDestructibleMesh* DestructibleMesh = NULL;
+	UDestructibleMesh* DestructibleMesh = nullptr;
 
 	// Create an Apex NxDestructibleAsset from the binary blob
 	NxDestructibleAsset* ApexDestructibleAsset = CreateApexDestructibleAssetFromBuffer(Buffer, (int32)(BufferEnd-Buffer));
-	if( ApexDestructibleAsset != NULL )
+	if( ApexDestructibleAsset != nullptr )
 	{
 		// Succesfully created the NxDestructibleAsset, now create a UDestructibleMesh
-		DestructibleMesh = ImportDestructibleMeshFromApexDestructibleAsset(InParent, *ApexDestructibleAsset, Name, Flags, NULL);
-		if( DestructibleMesh != NULL )
+		DestructibleMesh = ImportDestructibleMeshFromApexDestructibleAsset(InParent, *ApexDestructibleAsset, Name, Flags, nullptr);
+		if( DestructibleMesh != nullptr )
 		{
 			FEditorDelegates::OnAssetPostImport.Broadcast(this, DestructibleMesh);
 
@@ -6764,7 +6764,7 @@ void UReimportDestructibleMeshFactory::SetReimportPaths( UObject* Obj, const TAr
 	UDestructibleMesh* DestructibleMesh = Cast<UDestructibleMesh>(Obj);
 	if(DestructibleMesh && ensure(NewReimportPaths.Num() == 1))
 	{
-		if ( DestructibleMesh->AssetImportData == NULL )
+		if ( DestructibleMesh->AssetImportData == nullptr )
 		{
 			// @todo AssetImportData make an apex destructible import data class
 			DestructibleMesh->AssetImportData = ConstructObject<UAssetImportData>(UAssetImportData::StaticClass(), DestructibleMesh);
@@ -6784,7 +6784,7 @@ EReimportResult::Type UReimportDestructibleMeshFactory::Reimport( UObject* Obj )
 
 	UDestructibleMesh* DestructibleMesh = Cast<UDestructibleMesh>( Obj );
 
-	if ( DestructibleMesh->AssetImportData == NULL )
+	if ( DestructibleMesh->AssetImportData == nullptr )
 	{
 		// @todo AssetImportData make an apex destructible import data class
 		DestructibleMesh->AssetImportData = ConstructObject<UAssetImportData>(UAssetImportData::StaticClass(), DestructibleMesh);
@@ -6813,12 +6813,12 @@ EReimportResult::Type UReimportDestructibleMeshFactory::Reimport( UObject* Obj )
 
 	// Create an Apex NxDestructibleAsset from the binary blob
 	NxDestructibleAsset* ApexDestructibleAsset = CreateApexDestructibleAssetFromFile(Filename);
-	if( ApexDestructibleAsset != NULL )
+	if( ApexDestructibleAsset != nullptr )
 	{
 		// Succesfully created the NxDestructibleAsset, now create a UDestructibleMesh
-		UDestructibleMesh* ReimportedDestructibleMesh = ImportDestructibleMeshFromApexDestructibleAsset(DestructibleMesh->GetOuter(), *ApexDestructibleAsset, DestructibleMesh->GetFName(), DestructibleMesh->GetFlags(), NULL,
+		UDestructibleMesh* ReimportedDestructibleMesh = ImportDestructibleMeshFromApexDestructibleAsset(DestructibleMesh->GetOuter(), *ApexDestructibleAsset, DestructibleMesh->GetFName(), DestructibleMesh->GetFlags(), nullptr,
 																										EDestructibleImportOptions::PreserveSettings);
-		if( ReimportedDestructibleMesh != NULL )
+		if( ReimportedDestructibleMesh != nullptr )
 		{
 			check( ReimportedDestructibleMesh == DestructibleMesh );
 
@@ -6865,7 +6865,7 @@ UBlendSpaceFactoryNew::UBlendSpaceFactoryNew(const class FPostConstructInitializ
 bool UBlendSpaceFactoryNew::ConfigureProperties()
 	{
 	// Null the parent class so we can check for selection later
-	TargetSkeleton = NULL;
+	TargetSkeleton = nullptr;
 
 		// Load the content browser module to display an asset picker
 	FContentBrowserModule& ContentBrowserModule = FModuleManager::LoadModuleChecked<FContentBrowserModule>("ContentBrowser");
@@ -6901,7 +6901,7 @@ bool UBlendSpaceFactoryNew::ConfigureProperties()
 	GEditor->EditorAddModalWindow(PickerWindow.ToSharedRef());
 	PickerWindow.Reset();
 
-	return TargetSkeleton != NULL;
+	return TargetSkeleton != nullptr;
 }
 
 UObject* UBlendSpaceFactoryNew::FactoryCreateNew(UClass* Class,UObject* InParent,FName Name,EObjectFlags Flags,UObject* Context,FFeedbackContext* Warn)
@@ -6915,7 +6915,7 @@ UObject* UBlendSpaceFactoryNew::FactoryCreateNew(UClass* Class,UObject* InParent
 		return BlendSpace;
 	}
 
-	return NULL;
+	return nullptr;
 }
 
 void UBlendSpaceFactoryNew::OnTargetSkeletonSelected(const FAssetData& SelectedAsset)
@@ -6938,7 +6938,7 @@ UBlendSpaceFactory1D::UBlendSpaceFactory1D(const class FPostConstructInitializeP
 bool UBlendSpaceFactory1D::ConfigureProperties()
 {
 	// Null the parent class so we can check for selection later
-	TargetSkeleton = NULL;
+	TargetSkeleton = nullptr;
 
 		// Load the content browser module to display an asset picker
 	FContentBrowserModule& ContentBrowserModule = FModuleManager::LoadModuleChecked<FContentBrowserModule>("ContentBrowser");
@@ -6974,7 +6974,7 @@ bool UBlendSpaceFactory1D::ConfigureProperties()
 	GEditor->EditorAddModalWindow(PickerWindow.ToSharedRef());
 	PickerWindow.Reset();
 
-	return TargetSkeleton != NULL;
+	return TargetSkeleton != nullptr;
 }
 
 UObject* UBlendSpaceFactory1D::FactoryCreateNew(UClass* Class,UObject* InParent,FName Name,EObjectFlags Flags,UObject* Context,FFeedbackContext* Warn)
@@ -6988,7 +6988,7 @@ UObject* UBlendSpaceFactory1D::FactoryCreateNew(UClass* Class,UObject* InParent,
 		return BlendSpace;
 	}
 
-	return NULL;
+	return nullptr;
 }
 
 void UBlendSpaceFactory1D::OnTargetSkeletonSelected(const FAssetData& SelectedAsset)
@@ -7018,7 +7018,7 @@ UObject* UAimOffsetBlendSpaceFactoryNew::FactoryCreateNew(UClass* Class,UObject*
 		return BlendSpace;
 	}
 
-	return NULL;
+	return nullptr;
 }
 
 /*------------------------------------------------------------------------------
@@ -7042,7 +7042,7 @@ UObject* UAimOffsetBlendSpaceFactory1D::FactoryCreateNew(UClass* Class,UObject*
 		return BlendSpace;
 	}
 
-	return NULL;
+	return nullptr;
 }
 
 /*------------------------------------------------------------------------------
@@ -7068,7 +7068,7 @@ UObject* UEnumFactory::FactoryCreateNew(UClass* Class, UObject* InParent, FName
 			Warn->Log( Message );
 		}
 		FMessageDialog::Open( EAppMsgType::Ok, Message);
-		return NULL;
+		return nullptr;
 	}
 
 	return FEnumEditorUtils::CreateUserDefinedEnum(InParent, Name, Flags);
