Added compression and sampler types for distance field fonts. Distance field fonts were previously being incorrectly re-linearized in  ProcessMaterialColorTextureLookup (MaterialTemplate.usf)

#ttp 315746 - UE4: RENDERING: UE4 - Distance field fonts do not work on mobile / ES2

diff --git a/Engine/Source/Developer/TargetPlatform/Public/Common/TargetPlatformBase.h b/Engine/Source/Developer/TargetPlatform/Public/Common/TargetPlatformBase.h
--- a/Engine/Source/Developer/TargetPlatform/Public/Common/TargetPlatformBase.h
+++ b/Engine/Source/Developer/TargetPlatform/Public/Common/TargetPlatformBase.h
@@ -122,10 +122,14 @@ public:
 		{
 			TextureFormatName = NameG8;
 		}
-		else if( Texture->CompressionSettings == TC_Alpha)
+		else if ( Texture->CompressionSettings == TC_Alpha)
 		{
 			TextureFormatName = NameBC4;
 		}
+		else if (Texture->CompressionSettings == TC_DistanceFieldFont)
+		{
+			TextureFormatName = NameG8;
+		}
 		else if (Texture->CompressionNoAlpha)
 		{
 			TextureFormatName = NameDXT1;
diff --git a/Engine/Source/Editor/UnrealEd/Private/Factories/TTFontImport.cpp b/Engine/Source/Editor/UnrealEd/Private/Factories/TTFontImport.cpp
--- a/Engine/Source/Editor/UnrealEd/Private/Factories/TTFontImport.cpp
+++ b/Engine/Source/Editor/UnrealEd/Private/Factories/TTFontImport.cpp
@@ -937,7 +937,7 @@ UTexture2D* UTrueTypeFontFactory::CreateTextureFromDC( UFont* Font, HDC dc, int3
 		FMemory::Memcpy(Texture->Source.LockMip(0),DistanceFieldTex.GetResultTexture(),DistanceFieldTex.GetResultTextureSize());		
 		Texture->Source.UnlockMip(0);		
 		// use PF_G8 for all distance field textures for better precision than DXT5
-		Texture->CompressionSettings = TC_Displacementmap;
+		Texture->CompressionSettings = TC_DistanceFieldFont;
 		// disable gamma correction since storing alpha in linear color for PF_G8
 		Texture->SRGB = false;
 	}
@@ -947,8 +947,8 @@ UTexture2D* UTrueTypeFontFactory::CreateTextureFromDC( UFont* Font, HDC dc, int3
 		if (ImportOptions->Data.bAlphaOnly &&
 			!ImportOptions->Data.bEnableDropShadow)
 		{
-			// use PF_G8 for all distance field textures for better precision than DXT5
-			Texture->CompressionSettings = TC_Displacementmap;
+			// Not a distance field texture, but we use the same compression settings for better precision than DXT5
+			Texture->CompressionSettings = TC_DistanceFieldFont;
 			// disable gamma correction since storing alpha in linear color for PF_G8
 			Texture->SRGB = false;
 		}
@@ -1797,7 +1797,7 @@ UTexture2D* UTrueTypeFontFactory::CreateTextureFromBitmap( UFont* Font, uint8* B
 		FMemory::Memcpy(Texture->Source.LockMip(0),DistanceFieldTex.GetResultTexture(),DistanceFieldTex.GetResultTextureSize());		
 		Texture->Source.UnlockMip(0);		
 		// use PF_G8 for all distance field textures for better precision than DXT5
-		Texture->CompressionSettings = TC_Displacementmap;
+		Texture->CompressionSettings = TC_DistanceFieldFont;
 		// disable gamma correction since storing alpha in linear color for PF_G8
 		Texture->SRGB = false;
 	}
@@ -1807,8 +1807,8 @@ UTexture2D* UTrueTypeFontFactory::CreateTextureFromBitmap( UFont* Font, uint8* B
 		if (ImportOptions->Data.bAlphaOnly &&
 			!ImportOptions->Data.bEnableDropShadow)
 		{
-			// use PF_G8 for all distance field textures for better precision than DXT5
-			Texture->CompressionSettings = TC_Displacementmap;
+			// Not a distance field texture, but we use the same compression settings for better precision than DXT5
+			Texture->CompressionSettings = TC_DistanceFieldFont;
 			// disable gamma correction since storing alpha in linear color for PF_G8
 			Texture->SRGB = false;
 		}
diff --git a/Engine/Source/Runtime/Engine/Classes/Engine/EngineTypes.h b/Engine/Source/Runtime/Engine/Classes/Engine/EngineTypes.h
--- a/Engine/Source/Runtime/Engine/Classes/Engine/EngineTypes.h
+++ b/Engine/Source/Runtime/Engine/Classes/Engine/EngineTypes.h
@@ -142,6 +142,7 @@ enum EMaterialSamplerType
 	SAMPLERTYPE_Alpha UMETA(DisplayName="Alpha"),
 	SAMPLERTYPE_Normal UMETA(DisplayName="Normal"),
 	SAMPLERTYPE_Masks UMETA(DisplayName="Masks"),
+	SAMPLERTYPE_DistanceFieldFont UMETA(DisplayName="Distance Field Font"),
 	SAMPLERTYPE_LinearColor UMETA(DisplayName = "Linear Color"),
 	SAMPLERTYPE_LinearGrayscale UMETA(DisplayName = "Linear Grayscale"),
 	SAMPLERTYPE_MAX,
diff --git a/Engine/Source/Runtime/Engine/Classes/Engine/Texture.h b/Engine/Source/Runtime/Engine/Classes/Engine/Texture.h
--- a/Engine/Source/Runtime/Engine/Classes/Engine/Texture.h
+++ b/Engine/Source/Runtime/Engine/Classes/Engine/Texture.h
@@ -21,6 +21,7 @@ enum TextureCompressionSettings
 	TC_HDR,
 	TC_EditorIcon,
 	TC_Alpha,
+	TC_DistanceFieldFont,
 	TC_MAX,
 };
 
diff --git a/Engine/Source/Runtime/Engine/Private/Font.cpp b/Engine/Source/Runtime/Engine/Private/Font.cpp
--- a/Engine/Source/Runtime/Engine/Private/Font.cpp
+++ b/Engine/Source/Runtime/Engine/Private/Font.cpp
@@ -39,7 +39,15 @@ void UFont::PostLoad()
 		if( Texture )
 		{	
 			Texture->SetFlags(RF_Public);
-			Texture->LODGroup = TEXTUREGROUP_UI;	
+			Texture->LODGroup = TEXTUREGROUP_UI;
+
+			// Fix up compression type for distance field fonts.
+			if (Texture->CompressionSettings == TC_Displacementmap && !Texture->SRGB)
+			{
+				Texture->ConditionalPostLoad();
+				Texture->CompressionSettings = TC_DistanceFieldFont;
+				Texture->UpdateResource();
+			}
 		}
 	}
 }
diff --git a/Engine/Source/Runtime/Engine/Private/Materials/HLSLMaterialTranslator.h b/Engine/Source/Runtime/Engine/Private/Materials/HLSLMaterialTranslator.h
--- a/Engine/Source/Runtime/Engine/Private/Materials/HLSLMaterialTranslator.h
+++ b/Engine/Source/Runtime/Engine/Private/Materials/HLSLMaterialTranslator.h
@@ -2153,6 +2153,7 @@ protected:
 			break;
 
 			case SAMPLERTYPE_Alpha:
+			case SAMPLERTYPE_DistanceFieldFont:
 				// Sampling a single channel texture in D3D9 gives: (G,G,G)
 				// Sampling a single channel texture in D3D11 gives: (G,0,0)
 				// This replication reproduces the D3D9 behavior in all cases.
diff --git a/Engine/Source/Runtime/Engine/Private/Materials/MaterialExpressions.cpp b/Engine/Source/Runtime/Engine/Private/Materials/MaterialExpressions.cpp
--- a/Engine/Source/Runtime/Engine/Private/Materials/MaterialExpressions.cpp
+++ b/Engine/Source/Runtime/Engine/Private/Materials/MaterialExpressions.cpp
@@ -1087,6 +1087,8 @@ EMaterialSamplerType UMaterialExpressionTextureBase::GetSamplerTypeForTexture(co
 				return SAMPLERTYPE_Alpha;
 			case TC_Masks:
 				return SAMPLERTYPE_Masks;
+			case TC_DistanceFieldFont:
+				return SAMPLERTYPE_DistanceFieldFont;
 			default:
 				return Texture->SRGB ? SAMPLERTYPE_Color : SAMPLERTYPE_LinearColor;
 		}
@@ -5397,7 +5399,17 @@ int32 UMaterialExpressionFontSample::Compile(class FMaterialCompiler* Compiler,
 				Texture = Texture = GEngine->DefaultTexture;
 			}
 			check(Texture);
-			EMaterialSamplerType ExpectedSamplerType = Texture->SRGB ? SAMPLERTYPE_Color : SAMPLERTYPE_LinearColor;
+
+			EMaterialSamplerType ExpectedSamplerType;
+			if (Texture->CompressionSettings == TC_DistanceFieldFont)
+			{
+				ExpectedSamplerType = SAMPLERTYPE_DistanceFieldFont;
+			}
+			else
+			{
+				ExpectedSamplerType = Texture->SRGB ? SAMPLERTYPE_Color : SAMPLERTYPE_LinearColor;
+			}
+
 			if (!VerifySamplerType(Compiler, (Desc.Len() > 0 ? *Desc : TEXT("FontSample")), Texture, ExpectedSamplerType))
 			{
 				return INDEX_NONE;
@@ -5497,7 +5509,17 @@ int32 UMaterialExpressionFontSampleParameter::Compile(class FMaterialCompiler* C
 			Texture = Texture = GEngine->DefaultTexture;
 		}
 		check(Texture);
-		EMaterialSamplerType ExpectedSamplerType = Texture->SRGB ? SAMPLERTYPE_Color : SAMPLERTYPE_LinearColor;
+
+		EMaterialSamplerType ExpectedSamplerType;
+		if (Texture->CompressionSettings == TC_DistanceFieldFont)
+		{
+			ExpectedSamplerType = SAMPLERTYPE_DistanceFieldFont;
+		}
+		else
+		{
+			ExpectedSamplerType = Texture->SRGB ? SAMPLERTYPE_Color : SAMPLERTYPE_LinearColor;
+		}
+
 		if (!VerifySamplerType(Compiler, (Desc.Len() > 0 ? *Desc : TEXT("FontSampleParameter")), Texture, ExpectedSamplerType))
 		{
 			return INDEX_NONE;
diff --git a/Engine/Source/Runtime/Engine/Private/TextureDerivedData.cpp b/Engine/Source/Runtime/Engine/Private/TextureDerivedData.cpp
--- a/Engine/Source/Runtime/Engine/Private/TextureDerivedData.cpp
+++ b/Engine/Source/Runtime/Engine/Private/TextureDerivedData.cpp
@@ -324,7 +324,7 @@ static void GetTextureBuildSettings(
 		OutBuildSettings.bLongLatSource = false;
 	}
 
-	if (Texture.CompressionSettings == TC_Displacementmap)
+	if (Texture.CompressionSettings == TC_Displacementmap || Texture.CompressionSettings == TC_DistanceFieldFont)
 	{
 		OutBuildSettings.bReplicateAlpha = true;
 	}
