Scalability - Changing the text when auto scalabiltiy kicks in, instead of otimum, reduce quality.  Refreshing all viewports after scalability settings are swapped.

diff --git a/Engine/Source/Editor/UnrealEd/Private/PerformanceMonitor.cpp b/Engine/Source/Editor/UnrealEd/Private/PerformanceMonitor.cpp
--- a/Engine/Source/Editor/UnrealEd/Private/PerformanceMonitor.cpp
+++ b/Engine/Source/Editor/UnrealEd/Private/PerformanceMonitor.cpp
@@ -161,6 +161,7 @@ void FPerformanceMonitor::AutoApplyScalability()
 
 	Scalability::SetQualityLevels(NewLevels);
 	Scalability::SaveState(GEditorGameAgnosticIni);
+	GEditor->RedrawAllViewports();
 
 	const bool bAutoApplied = true;
 	Scalability::RecordQualityLevelsAnalytics(bAutoApplied);
@@ -299,14 +300,14 @@ void FPerformanceMonitor::Tick(float DeltaTime)
 		// Choose an appropriate message
 		enum MessagesEnum { Seconds, SecondsPercent, Minute, MinutePercent, Minutes, MinutesPercent };
 		const FText Messages[] = {
-			LOCTEXT("PerformanceWarningInProgress_Seconds",			"Your framerate has been under {Framerate} FPS for the past {SampleTime} seconds.\n\nApplying optimum engine settings in {TimeRemaining}s."),
-			LOCTEXT("PerformanceWarningInProgress_Seconds_Percent", "Your framerate has been under {Framerate} FPS for {Percentage}% of the past {SampleTime} seconds.\n\nApplying optimum engine settings in {TimeRemaining}s."),
+			LOCTEXT("PerformanceWarningInProgress_Seconds",			"Your framerate has been under {Framerate} FPS for the past {SampleTime} seconds.\n\nApplying reduced quality settings in {TimeRemaining}s."),
+			LOCTEXT("PerformanceWarningInProgress_Seconds_Percent", "Your framerate has been under {Framerate} FPS for {Percentage}% of the past {SampleTime} seconds.\n\nApplying reduced quality settings in {TimeRemaining}s."),
 
-			LOCTEXT("PerformanceWarningInProgress_Minute",			"Your framerate has been under {Framerate} FPS for the past minute.\n\nApplying optimum engine settings in {TimeRemaining}s."),
-			LOCTEXT("PerformanceWarningInProgress_Minute_Percent",	"Your framerate has been under {Framerate} FPS for {Percentage}% of the last minute.\n\nApplying optimum engine settings in {TimeRemaining}s."),
+			LOCTEXT("PerformanceWarningInProgress_Minute",			"Your framerate has been under {Framerate} FPS for the past minute.\n\nApplying reduced quality settings in {TimeRemaining}s."),
+			LOCTEXT("PerformanceWarningInProgress_Minute_Percent",	"Your framerate has been under {Framerate} FPS for {Percentage}% of the last minute.\n\nApplying reduced quality settings in {TimeRemaining}s."),
 
-			LOCTEXT("PerformanceWarningInProgress_Minutes",			"Your framerate has been below {Framerate} FPS for the past {SampleTime} minutes.\n\nApplying optimum engine settings in {TimeRemaining}s."),
-			LOCTEXT("PerformanceWarningInProgress_Minutes_Percent", "Your framerate has been below {Framerate} FPS for {Percentage}% of the past {SampleTime} minutes.\n\nApplying optimum engine settings in {TimeRemaining}s."),
+			LOCTEXT("PerformanceWarningInProgress_Minutes",			"Your framerate has been below {Framerate} FPS for the past {SampleTime} minutes.\n\nApplying reduced quality settings in {TimeRemaining}s."),
+			LOCTEXT("PerformanceWarningInProgress_Minutes_Percent", "Your framerate has been below {Framerate} FPS for {Percentage}% of the past {SampleTime} minutes.\n\nApplying reduced quality settings in {TimeRemaining}s."),
 		};
 
 		int32 Message;
