Editor - Fix a bug where resetting Selection Color (and probably other properties) to default would not update the UI until the editor was restarted.
+ Raise change events for objects who's configs get reloaded.  This already worked correctly for their derived objects and instances.
+ Update editor setting change handling to handle the 'None' property change which happens when settings are reset.

TTP# 340790

diff --git a/Engine/Source/Editor/UnrealEd/Private/Editor.cpp b/Engine/Source/Editor/UnrealEd/Private/Editor.cpp
--- a/Engine/Source/Editor/UnrealEd/Private/Editor.cpp
+++ b/Engine/Source/Editor/UnrealEd/Private/Editor.cpp
@@ -432,14 +432,15 @@ void UEditorEngine::InitEditor(IEngineLoop* InEngineLoop)
 
 void UEditorEngine::HandleSettingChanged( FName Name )
 {
-	if (Name == FName(TEXT("ColorVisionDeficiencyPreviewType")))
+	// When settings are reset to default, the property name will be "None" so make sure that case is handled.
+	if (Name == FName(TEXT("ColorVisionDeficiencyPreviewType")) || Name == NAME_None)
 	{
 		uint32 DeficiencyType = (uint32)GetDefault<UEditorStyleSettings>()->ColorVisionDeficiencyPreviewType.GetValue();
 		FSlateApplication::Get().GetRenderer()->SetColorVisionDeficiencyType(DeficiencyType);
 
 		GEngine->Exec(NULL, TEXT("RecompileShaders SlateElementPixelShader"));
 	}
-	else if (Name == FName("SelectionColor"))
+	if (Name == FName("SelectionColor") || Name == NAME_None)
 	{
 		// Selection outline color and material color use the same color but sometimes the selected material color can be overidden so these need to be set independently
 		GEngine->SetSelectedMaterialColor(GetDefault<UEditorStyleSettings>()->SelectionColor);
diff --git a/Engine/Source/Runtime/CoreUObject/Private/UObject/Obj.cpp b/Engine/Source/Runtime/CoreUObject/Private/UObject/Obj.cpp
--- a/Engine/Source/Runtime/CoreUObject/Private/UObject/Obj.cpp
+++ b/Engine/Source/Runtime/CoreUObject/Private/UObject/Obj.cpp
@@ -1182,7 +1182,19 @@ void UObject::TagSubobjects(EObjectFlags NewFlags)
 
 void UObject::ReloadConfig( UClass* ConfigClass/*=NULL*/, const TCHAR* InFilename/*=NULL*/, uint32 PropagationFlags/*=LCPF_None*/, UProperty* PropertyToLoad/*=NULL*/ )
 {
-	LoadConfig(ConfigClass, InFilename, PropagationFlags|UE4::LCPF_ReloadingConfigData|UE4::LCPF_ReadParentSections, PropertyToLoad);
+		if (!GIsEditor)
+		{
+			LoadConfig(ConfigClass, InFilename, PropagationFlags | UE4::LCPF_ReloadingConfigData | UE4::LCPF_ReadParentSections, PropertyToLoad);
+		}
+#if WITH_EDITOR
+		else
+		{
+			// When in the editor, raise change events so that the UI will update correctly when object configs are reloaded.
+			PreEditChange(NULL);
+			LoadConfig(ConfigClass, InFilename, PropagationFlags | UE4::LCPF_ReloadingConfigData | UE4::LCPF_ReadParentSections, PropertyToLoad);
+			PostEditChange();
+		}
+#endif // WITH_EDITOR
 }
 
 /** Checks if a section specified as a long package name can be found as short name in ini. */
