Fixed typo in deployment toast.

diff --git a/Engine/Source/Editor/UnrealEd/Private/PlayLevel.cpp b/Engine/Source/Editor/UnrealEd/Private/PlayLevel.cpp
--- a/Engine/Source/Editor/UnrealEd/Private/PlayLevel.cpp
+++ b/Engine/Source/Editor/UnrealEd/Private/PlayLevel.cpp
@@ -1202,7 +1202,7 @@ void UEditorEngine::HandleLaunchCompleted(bool Succeeded, double TotalTime, bool
 		const FString DummyDeviceName(FString::Printf(TEXT("All_iOS_On_%s"), FPlatformProcess::ComputerName()));
 		if (PlayUsingLauncherDeviceId.Left(PlayUsingLauncherDeviceId.Find(TEXT("@"))) == TEXT("IOS") && PlayUsingLauncherDeviceName.Contains(DummyDeviceName))
 		{
-			CompletionMsg = LOCTEXT("LauncherTaskCompleted", "Deployment complete!! Open the app on your device to launch.");
+			CompletionMsg = LOCTEXT("LauncherTaskCompleted", "Deployment complete! Open the app on your device to launch.");
 			TSharedPtr<SNotificationItem> NotificationItem = NotificationItemPtr.Pin();
 //			NotificationItem->SetExpireDuration(30.0f);
 		}
