const correct

diff --git a/Engine/Source/Editor/UnrealEd/Private/SourceCodeNavigation.cpp b/Engine/Source/Editor/UnrealEd/Private/SourceCodeNavigation.cpp
--- a/Engine/Source/Editor/UnrealEd/Private/SourceCodeNavigation.cpp
+++ b/Engine/Source/Editor/UnrealEd/Private/SourceCodeNavigation.cpp
@@ -1424,7 +1424,7 @@ bool FSourceCodeNavigation::FindModulePath( const FString& ModuleName, FString &
 	return false;
 }
 
-bool FSourceCodeNavigation::FindClassHeaderPath( UField *Field, FString &OutClassHeaderPath )
+bool FSourceCodeNavigation::FindClassHeaderPath( const UField *Field, FString &OutClassHeaderPath )
 {
 	// Get the class package, and skip past the "/Script/" portion to get the module name
 	UPackage *ModulePackage = Field->GetTypedOuter<UPackage>();
diff --git a/Engine/Source/Editor/UnrealEd/Public/SourceCodeNavigation.h b/Engine/Source/Editor/UnrealEd/Public/SourceCodeNavigation.h
--- a/Engine/Source/Editor/UnrealEd/Public/SourceCodeNavigation.h
+++ b/Engine/Source/Editor/UnrealEd/Public/SourceCodeNavigation.h
@@ -116,7 +116,7 @@ public:
 	UNREALED_API static bool FindModulePath( const FString& ModuleName, FString &OutModulePath );
 
 	/** Finds the path to a given class. Does not rely on symbols; finds matching .build.cs files. */
-	UNREALED_API static bool FindClassHeaderPath( UField *Field, FString &OutClassHeaderPath );
+	UNREALED_API static bool FindClassHeaderPath( const UField *Field, FString &OutClassHeaderPath );
 
 	/** Opens a single source file */
 	UNREALED_API static bool OpenSourceFile( const FString& AbsoluteSourcePath, int32 LineNumber = 0, int32 ColumnNumber = 0 );
