Removed macro libraries from the variable type dropdown in the Blueprint editor.
#ttp 340697 - BLUEPRINTS: Macro libraries should not be used as BP variable types.
#branch UE4
#proj Editor.BlueprintGraph, Editor.UnrealEd
#add Created new method FKismetEditorUtilities::IsClassABlueprintMacroLibrary()
#add Made call to FKismetEditorUtilities::IsClassABlueprintMacroLibrary() in UEdGraphSchema_K2::IsAllowableBlueprintVariableType(), to filter macro library classes out of the list.

diff --git a/Engine/Source/Editor/BlueprintGraph/Private/EdGraphSchema_K2.cpp b/Engine/Source/Editor/BlueprintGraph/Private/EdGraphSchema_K2.cpp
--- a/Engine/Source/Editor/BlueprintGraph/Private/EdGraphSchema_K2.cpp
+++ b/Engine/Source/Editor/BlueprintGraph/Private/EdGraphSchema_K2.cpp
@@ -632,6 +632,12 @@ bool UEdGraphSchema_K2::IsAllowableBlueprintVariableType(const class UClass* InC
 			return false;
 		}
 
+		// No Blueprint Macro Libraries
+		if (FKismetEditorUtilities::IsClassABlueprintMacroLibrary(InClass))
+		{
+			return false;
+		}
+
 		// UObject is an exception, and is always a blueprint-able type
 		if(InClass == UObject::StaticClass())
 		{
diff --git a/Engine/Source/Editor/UnrealEd/Private/Kismet2/Kismet2.cpp b/Engine/Source/Editor/UnrealEd/Private/Kismet2/Kismet2.cpp
--- a/Engine/Source/Editor/UnrealEd/Private/Kismet2/Kismet2.cpp
+++ b/Engine/Source/Editor/UnrealEd/Private/Kismet2/Kismet2.cpp
@@ -1320,6 +1320,13 @@ bool FKismetEditorUtilities::IsClassABlueprintSkeleton(const UClass* Class)
 	return Class->HasAnyFlags(RF_Transient) && Class->HasAnyClassFlags(CLASS_CompiledFromBlueprint);
 }
 
+bool FKismetEditorUtilities::IsClassABlueprintMacroLibrary(const UClass* Class)
+{
+	// Find generating blueprint for a class
+	UBlueprint* GeneratingBP = Cast<UBlueprint>(Class->ClassGeneratedBy);
+	return (GeneratingBP && GeneratingBP->BlueprintType == BPTYPE_MacroLibrary);
+}
+
 /** Run over the components references, and then NULL any that fall outside this blueprint's scope (e.g. components brought over after reparenting from another class, which are now in the transient package) */
 void FKismetEditorUtilities::StripExternalComponents(class UBlueprint* Blueprint)
 {
diff --git a/Engine/Source/Editor/UnrealEd/Public/Kismet2/KismetEditorUtilities.h b/Engine/Source/Editor/UnrealEd/Public/Kismet2/KismetEditorUtilities.h
--- a/Engine/Source/Editor/UnrealEd/Public/Kismet2/KismetEditorUtilities.h
+++ b/Engine/Source/Editor/UnrealEd/Public/Kismet2/KismetEditorUtilities.h
@@ -128,6 +128,9 @@ public:
 	/** Check to see if a given class is blueprint skeleton class. */
 	static bool IsClassABlueprintSkeleton (const UClass* Class);
 
+	/** Check to see if a given class is a blueprint macro library */
+	static bool IsClassABlueprintMacroLibrary(const UClass* Class);
+
 	/** Run over the components in the blueprint, and then remove any that fall outside this blueprint's scope (e.g. components brought over after reparenting from another class) */
 	static void StripExternalComponents(class UBlueprint* Blueprint);
 
