Removed the spurious 'set brush' transactions when placing geometry.
#ttp 339798 - EDITOR: It takes multiple undos to fully remove a volume or geometry
#branch UE4
#proj Editor.UnrealEd
#remove Removed the "Set Brush" scoped transaction in UEditorBrushBuilder::EndBrush, and added them to any methods which require this transaction which call this method.

diff --git a/Engine/Source/Editor/DetailCustomizations/Private/BrushDetails.cpp b/Engine/Source/Editor/DetailCustomizations/Private/BrushDetails.cpp
--- a/Engine/Source/Editor/DetailCustomizations/Private/BrushDetails.cpp
+++ b/Engine/Source/Editor/DetailCustomizations/Private/BrushDetails.cpp
@@ -120,24 +120,28 @@ void FBrushDetails::CustomizeDetails( IDetailLayoutBuilder& DetailLayout )
 
 			TArray<FNewBrushBuilder> NewBuilders;
 			TArray<FString> NewObjectPaths;
-			for(UObject* OuterObject : OuterObjects)
+
 			{
-				UBrushBuilder* NewObject = ConstructObject<UBrushBuilder>(InChosenClass, OuterObject);
+				const FScopedTransaction Transaction(NSLOCTEXT("UnrealEd", "BrushSet", "Brush Set"));
+				for (UObject* OuterObject : OuterObjects)
+				{
+					UBrushBuilder* NewObject = ConstructObject<UBrushBuilder>(InChosenClass, OuterObject, NAME_None, RF_Transactional);
 
-				FNewBrushBuilder NewBuilder;
-				NewBuilder.Builder = NewObject;
-				NewBuilder.Brush = CastChecked<ABrush>( OuterObject );
+					FNewBrushBuilder NewBuilder;
+					NewBuilder.Builder = NewObject;
+					NewBuilder.Brush = CastChecked<ABrush>(OuterObject);
 
-				NewBuilders.Add(NewBuilder);
-				NewObjectPaths.Add(NewObject->GetPathName());
-			}
-			
-			BrushBuilderHandle->SetPerObjectValues(NewObjectPaths);
+					NewBuilders.Add(NewBuilder);
+					NewObjectPaths.Add(NewObject->GetPathName());
+				}
 
-			// make sure the brushes are rebuilt
-			for(FNewBrushBuilder& NewObject : NewBuilders)
-			{
-				NewObject.Builder->Build(NewObject.Brush->GetWorld(), NewObject.Brush);
+				BrushBuilderHandle->SetPerObjectValues(NewObjectPaths);
+
+				// make sure the brushes are rebuilt
+				for (FNewBrushBuilder& NewObject : NewBuilders)
+				{
+					NewObject.Builder->Build(NewObject.Brush->GetWorld(), NewObject.Brush);
+				}
 			}
 
 			FSlateApplication::Get().DismissAllMenus();
diff --git a/Engine/Source/Editor/UnrealEd/Private/EditorBrushBuilder.cpp b/Engine/Source/Editor/UnrealEd/Private/EditorBrushBuilder.cpp
--- a/Engine/Source/Editor/UnrealEd/Private/EditorBrushBuilder.cpp
+++ b/Engine/Source/Editor/UnrealEd/Private/EditorBrushBuilder.cpp
@@ -33,8 +33,8 @@ void UEditorBrushBuilder::BeginBrush(bool InMergeCoplanars, FName InLayer)
 bool UEditorBrushBuilder::EndBrush( UWorld* InWorld, ABrush* InBrush )
 {
 	//!!validate
-	check( InWorld );
-	ABrush* BuilderBrush = (InBrush != NULL) ? InBrush : InWorld->GetDefaultBrush();
+	check( InWorld != nullptr );
+	ABrush* BuilderBrush = (InBrush != nullptr) ? InBrush : InWorld->GetDefaultBrush();
 
 	// Ensure the builder brush is unhidden.
 	BuilderBrush->bHidden = false;
@@ -42,7 +42,7 @@ bool UEditorBrushBuilder::EndBrush( UWorld* InWorld, ABrush* InBrush )
 
 	AActor* Actor = GEditor->GetSelectedActors()->GetTop<AActor>();
 	FVector Location;
-	if ( InBrush == NULL )
+	if ( InBrush == nullptr )
 	{
 		Location = Actor ? Actor->GetActorLocation() : BuilderBrush->GetActorLocation();
 	}
@@ -52,108 +52,112 @@ bool UEditorBrushBuilder::EndBrush( UWorld* InWorld, ABrush* InBrush )
 	}
 
 	UModel* Brush = BuilderBrush->Brush;
-	if( Brush )
+	if (Brush == nullptr)
 	{
-		{
-			const FScopedTransaction Transaction( NSLOCTEXT("UnrealEd", "BrushSet", "Brush Set") );
-			Brush->Modify();
-			BuilderBrush->Modify();
+		return true;
+	}
 
-			BuilderBrush->Layers.AddUnique( Layer );
+	Brush->Modify();
+	BuilderBrush->Modify();
 
-			FRotator Temp(0.0f,0.0f,0.0f);
-			FSnappingUtils::SnapToBSPVertex( Location, FVector::ZeroVector, Temp );
-			BuilderBrush->SetActorLocation(Location, false);
-			BuilderBrush->SetPrePivot( FVector::ZeroVector );
-			{
-				// Try and maintain the materials assigned to the surfaces. 
-				TArray<FPoly> CachedPolys;
-				UMaterialInterface* CachedMaterial = NULL;
-				if( Brush->Polys->Element.Num() == Polys.Num() )
-				{
-					// If the number of polygons match we assume its the same shape.
-					CachedPolys.Append(Brush->Polys->Element);
-				}
-				else if( Brush->Polys->Element.Num() > 0 )
-				{
-					// If the polygons have changed check if we only had one material before. 
-					CachedMaterial = Brush->Polys->Element[0].Material;
-					if (CachedMaterial != NULL)
-					{
-						for( auto Poly : Brush->Polys->Element )
-						{
-							if( CachedMaterial != Poly.Material )
-							{
-								CachedMaterial = NULL;
-								break;
-							}
-						}
-					}
-				}
+	BuilderBrush->Layers.AddUnique( Layer );
 
-				// Clear existing polys.
-				Brush->Polys->Element.Empty();
+	FRotator Temp(0.0f,0.0f,0.0f);
+	FSnappingUtils::SnapToBSPVertex( Location, FVector::ZeroVector, Temp );
+	BuilderBrush->SetActorLocation(Location, false);
+	BuilderBrush->SetPrePivot( FVector::ZeroVector );
 
-				const bool bUseCachedPolysMaterial = CachedPolys.Num() > 0;
-				int32 CachedPolyIdx = 0;
-				for( TArray<FBuilderPoly>::TIterator It(Polys); It; ++It )
+	// Try and maintain the materials assigned to the surfaces. 
+	TArray<FPoly> CachedPolys;
+	UMaterialInterface* CachedMaterial = nullptr;
+	if( Brush->Polys->Element.Num() == Polys.Num() )
+	{
+		// If the number of polygons match we assume its the same shape.
+		CachedPolys.Append(Brush->Polys->Element);
+	}
+	else if( Brush->Polys->Element.Num() > 0 )
+	{
+		// If the polygons have changed check if we only had one material before. 
+		CachedMaterial = Brush->Polys->Element[0].Material;
+		if (CachedMaterial != NULL)
+		{
+			for( auto Poly : Brush->Polys->Element )
+			{
+				if( CachedMaterial != Poly.Material )
 				{
-					if( It->Direction<0 )
-					{
-						for( int32 i=0; i<It->VertexIndices.Num()/2; i++ )
-						{
-							Exchange( It->VertexIndices[i], It->VertexIndices.Last(i) );
-						}
-					}
-
-					FPoly Poly;
-					Poly.Init();
-					Poly.ItemName = It->ItemName;
-					Poly.Base = Vertices[It->VertexIndices[0]];
-					Poly.PolyFlags = It->PolyFlags;
-
-					// Try and maintain the polygons material where possible
-					Poly.Material = ( bUseCachedPolysMaterial ) ? CachedPolys[CachedPolyIdx++].Material : CachedMaterial;
-
-					for( int32 j=0; j<It->VertexIndices.Num(); j++ )
-					{
-						new(Poly.Vertices) FVector(Vertices[It->VertexIndices[j]]);
-					}
-					if( Poly.Finalize( BuilderBrush, 1 ) == 0 )
-					{
-						new(Brush->Polys->Element)FPoly(Poly);
-					}
+					CachedMaterial = NULL;
+					break;
 				}
 			}
-			if( MergeCoplanars )
+		}
+	}
+
+	// Clear existing polys.
+	Brush->Polys->Element.Empty();
+
+	const bool bUseCachedPolysMaterial = CachedPolys.Num() > 0;
+	int32 CachedPolyIdx = 0;
+	for( TArray<FBuilderPoly>::TIterator It(Polys); It; ++It )
+	{
+		if( It->Direction<0 )
+		{
+			for( int32 i=0; i<It->VertexIndices.Num()/2; i++ )
 			{
-				GEditor->bspMergeCoplanars( Brush, 0, 1 );
-				FBSPOps::bspValidateBrush( Brush, 1, 1 );
+				Exchange( It->VertexIndices[i], It->VertexIndices.Last(i) );
 			}
-			Brush->Linked = 1;
-			FBSPOps::bspValidateBrush( Brush, 0, 1 );
-			Brush->BuildBound();
 		}
 
-		GEditor->RedrawLevelEditingViewports();
-		GEditor->SetPivot( BuilderBrush->GetActorLocation(), false, true );
+		FPoly Poly;
+		Poly.Init();
+		Poly.ItemName = It->ItemName;
+		Poly.Base = Vertices[It->VertexIndices[0]];
+		Poly.PolyFlags = It->PolyFlags;
+
+		// Try and maintain the polygons material where possible
+		Poly.Material = ( bUseCachedPolysMaterial ) ? CachedPolys[CachedPolyIdx++].Material : CachedMaterial;
+
+		for( int32 j=0; j<It->VertexIndices.Num(); j++ )
+		{
+			new(Poly.Vertices) FVector(Vertices[It->VertexIndices[j]]);
+		}
+		if( Poly.Finalize( BuilderBrush, 1 ) == 0 )
+		{
+			new(Brush->Polys->Element)FPoly(Poly);
+		}
+	}
 
-		BuilderBrush->ReregisterAllComponents();
+	if( MergeCoplanars )
+	{
+		GEditor->bspMergeCoplanars( Brush, 0, 1 );
+		FBSPOps::bspValidateBrush( Brush, 1, 1 );
 	}
+	Brush->Linked = 1;
+	FBSPOps::bspValidateBrush( Brush, 0, 1 );
+	Brush->BuildBound();
+
+	GEditor->RedrawLevelEditingViewports();
+	GEditor->SetPivot( BuilderBrush->GetActorLocation(), false, true );
+
+	BuilderBrush->ReregisterAllComponents();
+
 	return true;
 }
+
 int32 UEditorBrushBuilder::GetVertexCount()
 {
 	return Vertices.Num();
 }
+
 FVector UEditorBrushBuilder::GetVertex(int32 i)
 {
 	return Vertices.IsValidIndex(i) ? Vertices[i] : FVector::ZeroVector;
 }
+
 int32 UEditorBrushBuilder::GetPolyCount()
 {
 	return Polys.Num();
 }
+
 bool UEditorBrushBuilder::BadParameters(const FText& Msg)
 {
 	if ( NotifyBadParams )
@@ -169,6 +173,7 @@ bool UEditorBrushBuilder::BadParameters(const FText& Msg)
 	}
 	return 0;
 }
+
 int32 UEditorBrushBuilder::Vertexv(FVector V)
 {
 	int32 Result = Vertices.Num();
@@ -176,12 +181,14 @@ int32 UEditorBrushBuilder::Vertexv(FVector V)
 
 	return Result;
 }
+
 int32 UEditorBrushBuilder::Vertex3f(float X, float Y, float Z)
 {
 	int32 Result = Vertices.Num();
 	new(Vertices)FVector(X,Y,Z);
 	return Result;
 }
+
 void UEditorBrushBuilder::Poly3i(int32 Direction, int32 i, int32 j, int32 k, FName ItemName, bool bIsTwoSidedNonSolid )
 {
 	new(Polys)FBuilderPoly;
@@ -192,6 +199,7 @@ void UEditorBrushBuilder::Poly3i(int32 Direction, int32 i, int32 j, int32 k, FNa
 	new(Polys.Last().VertexIndices)int32(k);
 	Polys.Last().PolyFlags = PF_DefaultFlags | (bIsTwoSidedNonSolid ? (PF_TwoSided|PF_NotSolid) : 0);
 }
+
 void UEditorBrushBuilder::Poly4i(int32 Direction, int32 i, int32 j, int32 k, int32 l, FName ItemName, bool bIsTwoSidedNonSolid )
 {
 	new(Polys)FBuilderPoly;
@@ -203,6 +211,7 @@ void UEditorBrushBuilder::Poly4i(int32 Direction, int32 i, int32 j, int32 k, int
 	new(Polys.Last().VertexIndices)int32(l);
 	Polys.Last().PolyFlags = PF_DefaultFlags | (bIsTwoSidedNonSolid ? (PF_TwoSided|PF_NotSolid) : 0);
 }
+
 void UEditorBrushBuilder::PolyBegin(int32 Direction, FName ItemName)
 {
 	new(Polys)FBuilderPoly;
@@ -210,10 +219,12 @@ void UEditorBrushBuilder::PolyBegin(int32 Direction, FName ItemName)
 	Polys.Last().Direction = Direction;
 	Polys.Last().PolyFlags = PF_DefaultFlags;
 }
+
 void UEditorBrushBuilder::Polyi(int32 i)
 {
 	new(Polys.Last().VertexIndices)int32(i);
 }
+
 void UEditorBrushBuilder::PolyEnd()
 {
 }
diff --git a/Engine/Source/Editor/UnrealEd/Private/UnrealEdSrv.cpp b/Engine/Source/Editor/UnrealEd/Private/UnrealEdSrv.cpp
--- a/Engine/Source/Editor/UnrealEd/Private/UnrealEdSrv.cpp
+++ b/Engine/Source/Editor/UnrealEd/Private/UnrealEdSrv.cpp
@@ -1842,41 +1842,45 @@ TArray<FPoly*> GetSelectedPolygons()
 void CreateBoundingBoxBuilderBrush( UWorld* InWorld, const TArray<FPoly*> SelectedPolys, bool bSnapVertsToGrid )
 {
 	int x;
-	FPoly* poly;
-	FBox bbox(0);
+	FPoly* Poly;
+	FBox BBox(0);
 	FVector Vertex;
 
 	for( x = 0 ; x < SelectedPolys.Num() ; ++x )
 	{
-		poly = SelectedPolys[x];
+		Poly = SelectedPolys[x];
 
-		for( int v = 0 ; v < poly->Vertices.Num() ; ++v )
+		for( int v = 0 ; v < Poly->Vertices.Num() ; ++v )
 		{
 			if( bSnapVertsToGrid )
 			{
-				Vertex = poly->Vertices[v].GridSnap(GEditor->GetGridSize());
+				Vertex = Poly->Vertices[v].GridSnap(GEditor->GetGridSize());
 			}
 			else
 			{
-				Vertex = poly->Vertices[v];
+				Vertex = Poly->Vertices[v];
 			}
 
-			bbox += Vertex;
+			BBox += Vertex;
 		}
 	}
 
 	// Change the builder brush to match the bounding box so that it exactly envelops the selected meshes
+	{
+		const FScopedTransaction Transaction(NSLOCTEXT("UnrealEd", "BrushSet", "Brush Set"));
 
-	FVector extent = bbox.GetExtent();
-	UCubeBuilder* CubeBuilder = ConstructObject<UCubeBuilder>( UCubeBuilder::StaticClass() );
-	CubeBuilder->X = extent.X * 2;
-	CubeBuilder->Y = extent.Y * 2;
-	CubeBuilder->Z = extent.Z * 2;
-	CubeBuilder->Build( InWorld );
-
-	InWorld->GetDefaultBrush()->SetActorLocation(bbox.GetCenter(), false);
+		UCubeBuilder* CubeBuilder = ConstructObject<UCubeBuilder>(UCubeBuilder::StaticClass(), GetTransientPackage(), NAME_None, RF_Transactional);
+		FVector Extent = BBox.GetExtent();
+		CubeBuilder->X = Extent.X * 2;
+		CubeBuilder->Y = Extent.Y * 2;
+		CubeBuilder->Z = Extent.Z * 2;
+		CubeBuilder->Build(InWorld);
 
-	InWorld->GetDefaultBrush()->ReregisterAllComponents();
+		ABrush* DefaultBrush = InWorld->GetDefaultBrush();
+		check(DefaultBrush != nullptr);
+		DefaultBrush->SetActorLocation(BBox.GetCenter(), false);
+		DefaultBrush->ReregisterAllComponents();
+	}
 }
 
 /**
