Moving actors into folder maintains their parent/child heirarchy where appropriate

#ttp 340456 - UE4: Moving selected actors into a folder in the Scene Outliner will unparent existing actors

#branch UE4

#change Selected child actors, whose parent is also being moved, are skipped - this allows them to move with the parent, instead of disassociating themselves when moving.

#change Unselected child actors, whose parent is being moved, are detached - this allows their parent to move freely without forcing the child to accompany them.

diff --git a/Engine/Source/Editor/SceneOutliner/Private/SSceneOutliner.cpp b/Engine/Source/Editor/SceneOutliner/Private/SSceneOutliner.cpp
--- a/Engine/Source/Editor/SceneOutliner/Private/SSceneOutliner.cpp
+++ b/Engine/Source/Editor/SceneOutliner/Private/SSceneOutliner.cpp
@@ -2052,12 +2052,7 @@ namespace SceneOutliner
 			TSharedRef<TOutlinerFolderTreeItem> NewFolderItem = MakeShareable(new TOutlinerFolderTreeItem(NewPath));
 			AddedItemsList.Add(NewFolderItem);
 
-			// Move the currently selected actors into the new folder
-			for( FSelectionIterator SelectionIt( *GEditor->GetSelectedActors() ); SelectionIt; ++SelectionIt )
-			{
-				AActor* Actor = CastChecked<AActor>(*SelectionIt);
-				Actor->SetFolderPath(NewPath);
-			}
+			FActorFolders::Get().SetSelectedFolderPath(NewPath);
 
 			Refresh();
 		}
diff --git a/Engine/Source/Editor/UnrealEd/Private/EditorActorFolders.cpp b/Engine/Source/Editor/UnrealEd/Private/EditorActorFolders.cpp
--- a/Engine/Source/Editor/UnrealEd/Private/EditorActorFolders.cpp
+++ b/Engine/Source/Editor/UnrealEd/Private/EditorActorFolders.cpp
@@ -309,11 +309,35 @@ void FActorFolders::CreateFolderContainingSelection(UWorld& InWorld, FName Path)
 {
 	const FScopedTransaction Transaction(LOCTEXT("UndoAction_CreateFolder", "Create Folder"));
 	CreateFolder(InWorld, Path);
+	SetSelectedFolderPath(Path);
+}
 
+void FActorFolders::SetSelectedFolderPath(FName Path) const
+{
 	// Move the currently selected actors into the new folder
-	for( FSelectionIterator SelectionIt( *GEditor->GetSelectedActors() ); SelectionIt; ++SelectionIt )
+	USelection* SelectedActors = GEditor->GetSelectedActors();
+	for (FSelectionIterator SelectionIt(*SelectedActors); SelectionIt; ++SelectionIt)
 	{
 		AActor* Actor = CastChecked<AActor>(*SelectionIt);
+
+		// If this actor is parented to another, which is also in the selection, skip it so that it moves when its parent does (otherwise it's orphaned)
+		const AActor* ParentActor = Actor->GetAttachParentActor();
+		if (ParentActor && SelectedActors->IsSelected(ParentActor))
+		{
+			continue;
+		}
+
+		// If any child actor isn't being moved with the parent, then orphan them (otherwise they'll end up being forced to move)
+		TArray<AActor*> ChildActors;
+		Actor->GetAttachedActors(ChildActors);
+		for (const auto& ChildActor : ChildActors)
+		{
+			if (ChildActor && !SelectedActors->IsSelected(ChildActor))
+			{
+				ChildActor->DetachRootComponentFromParent(true);
+			}
+		}
+
 		Actor->SetFolderPath(Path);
 	}
 }
diff --git a/Engine/Source/Editor/UnrealEd/Public/EditorActorFolders.h b/Engine/Source/Editor/UnrealEd/Public/EditorActorFolders.h
--- a/Engine/Source/Editor/UnrealEd/Public/EditorActorFolders.h
+++ b/Engine/Source/Editor/UnrealEd/Public/EditorActorFolders.h
@@ -82,6 +82,9 @@ struct UNREALED_API FActorFolders : public FGCObject
 	/** Same as CreateFolder, but moves the current actor selection into the new folder as well */
 	void CreateFolderContainingSelection(UWorld& InWorld, FName Path);
 
+	/** Sets the folder path for all the selected actors */
+	void SetSelectedFolderPath(FName Path) const;
+
 	/** Delete the specified folder in the world */
 	void DeleteFolder(UWorld& InWorld, FName FolderToDelete);
 
