The cancel button now appears on the feedback dialog while autosaving.

#ttp 337654 - EDITOR: Pressing Escape does not stop autosave

#branch UE4

#note The underlying issue appears to be that Escape should no longer cancels slow tasksà At around CL#1892733 a dedicated (optional) æCancelÆ button was added which is meant to be shown when the user is allowed to cancel the task. Unfortunately, autosave occurs in a common place where itÆs not always appropriate to cancel, so enabling Cancel to occur on the slow task by default wasnÆt a viable option. To get around this I modified the slow task so that Cancel can be toggled on or off at any point during its existence, depending on the context of the task at hand.

#add Added new function to FFeedbackContext: EnableUserCancel, so the user can manually show/hide the cancel button at any point during a slow task û this is handy when dealing with multiple jobs not all of which support the ability to cancel.

#change Tidied up SSlowTaskWidget by removing duplicate code and modifying how OnCancelClickedDelegate is used.

#remove Removed æEscapeÆ message now that cancel button is present.

diff --git a/Engine/Source/Editor/UnrealEd/Private/FeedbackContextEditor.cpp b/Engine/Source/Editor/UnrealEd/Private/FeedbackContextEditor.cpp
--- a/Engine/Source/Editor/UnrealEd/Private/FeedbackContextEditor.cpp
+++ b/Engine/Source/Editor/UnrealEd/Private/FeedbackContextEditor.cpp
@@ -17,34 +17,31 @@ class SSlowTaskWidget : public SBorder
 {
 public:
 	SLATE_BEGIN_ARGS( SSlowTaskWidget )	{ }
-		/** Called to when an asset is clicked */
-		SLATE_EVENT( FOnCancelClickedDelegate, OnCancelClickedDelegate )
 	SLATE_END_ARGS()
 
 	void Construct( const FArguments& InArgs )
 	{
-		OnCancelClickedDelegate = InArgs._OnCancelClickedDelegate;
-
-		TSharedRef<SVerticalBox> VerticalBox = SNew(SVerticalBox)
-			+SVerticalBox::Slot()
-			.AutoHeight()
-			.Padding( 14.0f, 4.0f, 14.0f, 10.0f )
+		SBorder::Construct(SBorder::FArguments()
+			.BorderImage(FEditorStyle::GetBrush("Menu.Background"))
+			.VAlign(VAlign_Center)
 			[
-				SNew( STextBlock )
-					.Text( this, &SSlowTaskWidget::OnGetProgressText )
-					.ShadowOffset( FVector2D( 1.0f, 1.0f ) )
-			];
-
-		if ( OnCancelClickedDelegate.IsBound() )
-		{
-			VerticalBox->AddSlot()
+				SNew(SVerticalBox)
+				+SVerticalBox::Slot()
+				.AutoHeight()
+				.Padding( 14.0f, 4.0f, 14.0f, 10.0f )
+				[
+					SNew( STextBlock )
+						.Text( this, &SSlowTaskWidget::OnGetProgressText )
+						.ShadowOffset( FVector2D( 1.0f, 1.0f ) )
+				]
+				+SVerticalBox::Slot()
 				.AutoHeight()
 				.Padding(10.0f, 7.0f)
 				[
 					SNew(SHorizontalBox)
 					+SHorizontalBox::Slot()
 					.Padding(5,0,0,0)
-					.FillWidth(0.8f)
+					.FillWidth(1.0f)
 					[
 						SNew(SProgressBar)
 						.Percent( this, &SSlowTaskWidget::GetProgressFraction )
@@ -57,31 +54,9 @@ public:
 						.Text( NSLOCTEXT("FeedbackContextProgress", "Cancel", "Cancel") )
 						.HAlign(EHorizontalAlignment::HAlign_Center)
 						.OnClicked(this, &SSlowTaskWidget::OnCancel)
+						.Visibility(this, &SSlowTaskWidget::GetCancelButtonVisibility)
 					]
-				];
-		}
-		else
-		{
-			VerticalBox->AddSlot()
-				.AutoHeight()
-				.Padding(10.0f, 7.0f)
-				[
-					SNew(SHorizontalBox)
-					+SHorizontalBox::Slot()
-					.Padding(5,0,0,0)
-					.FillWidth(1.0f)
-					[
-						SNew(SProgressBar)
-						.Percent( this, &SSlowTaskWidget::GetProgressFraction )
-					]
-				];
-		}
-
-		SBorder::Construct( SBorder::FArguments()
-			.BorderImage(FEditorStyle::GetBrush("Menu.Background"))
-			.VAlign(VAlign_Center)
-			[
-				VerticalBox
+				]
 			]
 		);
 	}
@@ -106,12 +81,28 @@ public:
 		UpdateProgressText();
 	}
 	
+	void BindOnCancelDelegate(const FOnCancelClickedDelegate& InCancelClickedDelegate)
+	{
+		OnCancelClickedDelegate = InCancelClickedDelegate;
+	}
+
+	void UnbindOnCancelDelegate()
+	{
+		OnCancelClickedDelegate.Unbind();
+	}
+
+private:
+
 	FReply OnCancel()
 	{
 		OnCancelClickedDelegate.ExecuteIfBound();
 		return FReply::Handled();
 	}
-private:
+
+	EVisibility GetCancelButtonVisibility() const
+	{
+		return OnCancelClickedDelegate.IsBound() ? EVisibility::Visible : EVisibility::Collapsed;
+	}
 
 	TOptional<float> GetProgressFraction() const
 	{
@@ -232,13 +223,6 @@ void FFeedbackContextEditor::BeginSlowTask( const FText& Task, bool bShowProgres
 
 		DialogRequestStack.Push( ( bProgressWindowShown ? 0 : bShowProgressDialog) );
 
-		FOnCancelClickedDelegate OnCancelClicked;
-		if ( bShowCancelButton )
-		{
-			// The cancel button is only displayed if a delegate is bound to it.
-			OnCancelClicked = FOnCancelClickedDelegate::CreateRaw(this, &FFeedbackContextEditor::OnUserCancel);
-		}
-
 		if(!bProgressWindowShown && bShowProgressDialog && ++DialogRequestCount == 1 )
 		{
 			if( !bProgressWindowShown && FSlateApplication::Get().CanDisplayWindows() )
@@ -248,10 +232,11 @@ void FFeedbackContextEditor::BeginSlowTask( const FText& Task, bool bShowProgres
 					.IsPopupWindow(true)
 					.ActivateWhenFirstShown(true);
 
-				SlowTaskWidget = SNew(SSlowTaskWidget)
-					.OnCancelClickedDelegate( OnCancelClicked );
+				SlowTaskWidget = SNew(SSlowTaskWidget);
 				SlowTaskWindowRef->SetContent( SlowTaskWidget.ToSharedRef() );
 
+				EnableUserCancel(bShowCancelButton);
+
 				SlowTaskWidget->SetProgressText( StatusMessage.StatusText );
 
 				SlowTaskWindow = SlowTaskWindowRef;
@@ -275,6 +260,21 @@ bool FFeedbackContextEditor::ReceivedUserCancel( void )
 	return res;
 }
 
+void FFeedbackContextEditor::EnableUserCancel( bool bUserCancel )
+{
+	if (SlowTaskWidget.IsValid())
+	{
+		if (!bUserCancel)
+		{
+			SlowTaskWidget->UnbindOnCancelDelegate();
+		}
+		else
+		{
+			SlowTaskWidget->BindOnCancelDelegate(FOnCancelClickedDelegate::CreateRaw(this, &FFeedbackContextEditor::OnUserCancel));
+		}
+	}
+}
+
 void FFeedbackContextEditor::OnUserCancel()
 {
 	HasTaskBeenCancelled = true;
diff --git a/Engine/Source/Editor/UnrealEd/Public/FeedbackContextEditor.h b/Engine/Source/Editor/UnrealEd/Public/FeedbackContextEditor.h
--- a/Engine/Source/Editor/UnrealEd/Public/FeedbackContextEditor.h
+++ b/Engine/Source/Editor/UnrealEd/Public/FeedbackContextEditor.h
@@ -102,6 +102,9 @@ public:
 	/** Whether or not the user has canceled out of this dialog */
 	bool ReceivedUserCancel();
 
+	/** Enable/Disable the ability to cancel the current task */
+	void EnableUserCancel(bool bUserCancel) override;
+
 	void OnUserCancel();
 
 	virtual bool YesNof( const FText& Question ) override
diff --git a/Engine/Source/Runtime/Core/Public/Misc/FeedbackContext.h b/Engine/Source/Runtime/Core/Public/Misc/FeedbackContext.h
--- a/Engine/Source/Runtime/Core/Public/Misc/FeedbackContext.h
+++ b/Engine/Source/Runtime/Core/Public/Misc/FeedbackContext.h
@@ -16,6 +16,8 @@ public:
 	virtual void BeginSlowTask( const FText& Task, bool ShowProgressDialog, bool bShowCancelButton=false )=0;
 	virtual void EndSlowTask()=0;
 	virtual bool ReceivedUserCancel() { return false; };
+	virtual void EnableUserCancel(bool bUserCancel){}
+
 
 	/** Slow task update status
 	 * @param Numerator		New progress numerator
diff --git a/Engine/Source/Runtime/CoreUObject/Private/UObject/SavePackage.cpp b/Engine/Source/Runtime/CoreUObject/Private/UObject/SavePackage.cpp
--- a/Engine/Source/Runtime/CoreUObject/Private/UObject/SavePackage.cpp
+++ b/Engine/Source/Runtime/CoreUObject/Private/UObject/SavePackage.cpp
@@ -2474,12 +2474,10 @@ bool UPackage::SavePackage( UPackage* InOuter, UObject* Base, EObjectFlags TopLe
 		Args.Add( TEXT("CleanFilename"), FText::FromString( CleanFilename ) );
 
 		FText StatusMessage = FText::Format( NSLOCTEXT("Core", "SavingFile", "Saving file: {CleanFilename}..."), Args );
-		if ((SaveFlags & SAVE_FromAutosave))
-		{
-			Args.Add( TEXT("SavingFile"), StatusMessage );
-			StatusMessage = FText::Format( NSLOCTEXT("Core", "AutoSaving", "{SavingFile} (Press Escape to stop the AutoSave)"), Args );
-		}
-		
+
+		// Only allow the user to cancel only if we're autosaving
+		GWarn->EnableUserCancel((SaveFlags & SAVE_FromAutosave ? true : false));
+
 		//Force status update based on current file name and progress percent
 		GWarn->StatusForceUpdate( CurSaveStep, TotalSaveSteps, StatusMessage );
 		bool Success = true;
