File missing from CL:2217535

diff --git a/Engine/Source/Editor/UnrealEd/Private/Factories/EditorFactories.cpp b/Engine/Source/Editor/UnrealEd/Private/Factories/EditorFactories.cpp
--- a/Engine/Source/Editor/UnrealEd/Private/Factories/EditorFactories.cpp
+++ b/Engine/Source/Editor/UnrealEd/Private/Factories/EditorFactories.cpp
@@ -6000,6 +6000,7 @@ EReimportResult::Type UReimportFbxAnimSequenceFactory::Reimport( UObject* Obj )
 		// update the data in case the file source has changed
 		ImportData->SourceFilePath = FReimportManager::SanitizeImportFilename(UFactory::CurrentFilename, AnimSequence);
 		ImportData->SourceFileTimestamp = IFileManager::Get().GetTimeStamp(*UFactory::CurrentFilename).ToString();
+		ImportData->bDirty = false;
 
 		// Try to find the outer package so we can dirty it up
 		if (AnimSequence->GetOuter())
