Adding an option to the DumpBlueprintsInfo commandlet so you can switch to the new experimental menu building systems (to diff against the old).

diff --git a/Engine/Source/Editor/BlueprintGraph/Public/BlueprintActionMenuBuilder.h b/Engine/Source/Editor/BlueprintGraph/Public/BlueprintActionMenuBuilder.h
--- a/Engine/Source/Editor/BlueprintGraph/Public/BlueprintActionMenuBuilder.h
+++ b/Engine/Source/Editor/BlueprintGraph/Public/BlueprintActionMenuBuilder.h
@@ -7,11 +7,11 @@
 // Forward declarations
 namespace FBlueprintActionMenuBuilderImpl
 {
-	struct FMenuSectionDefinition; // internal type, forward declared to hide implementation details
+	// internal type, forward declared to hide implementation details
+	struct FMenuSectionDefinition; 
 };
 class  FBlueprintActionFilter;
 struct FBlueprintActionContext;
-class  UActorComponent;
 
 /**
  * Responsible for constructing a list of viable blueprint actions. Runs the 
diff --git a/Engine/Source/Editor/UnrealEd/Private/Commandlets/DumpBlueprintsInfoCommandlet.cpp b/Engine/Source/Editor/UnrealEd/Private/Commandlets/DumpBlueprintsInfoCommandlet.cpp
--- a/Engine/Source/Editor/UnrealEd/Private/Commandlets/DumpBlueprintsInfoCommandlet.cpp
+++ b/Engine/Source/Editor/UnrealEd/Private/Commandlets/DumpBlueprintsInfoCommandlet.cpp
@@ -17,6 +17,13 @@
 #include "K2Node_CustomEvent.h"
 #include "ScopedTimers.h"
 
+#include "BlueprintActionMenuBuilder.h"
+#include "BlueprintActionFilter.h"
+#include "BlueprintActionDatabase.h"
+#include "K2Node_CallFunction.h"
+#include "K2Node_Event.h"
+#include "K2Node_Variable.h"
+
 DEFINE_LOG_CATEGORY_STATIC(LogBlueprintInfoDump, Log, All);
 
 /*******************************************************************************
@@ -86,6 +93,13 @@ DumpBlueprintsInfo commandlet params: \n\
                         you wish to use your own tool then specify it here. Use\n\
                         '{1}' and '{2}' as placeholders for filenames, like so:\n\
                         -diffcmd=\"AraxisP4Diff.exe {2} {1}\".                 \n\
+\n\
+    -experimental       Uses an new way of constructing Blueprint action menus \n\
+                        (that will replace the current system).                \n\
+\n\
+    -name=<Filename>    Overrides the default filename. Leave off the extention\n\
+                        (this will add .json to the end). When -multifile is   \n\
+                        supplied, the class name will be postfixed to the name.\n\
 \n\
     -help, -h, -?       Display this message and then exit.                    \n\
 \n");
@@ -104,6 +118,7 @@ DumpBlueprintsInfo commandlet params: \n\
 		BPDUMP_DoNotDumpActionInfo  = (1<<7),
 		BPDUMP_RecordTiming			= (1<<8),
 		BPDUMP_SelectAllObjTypes    = (1<<9), 
+		BPDUMP_UseLegacyMenuBuilder = (1<<10),
 
 		BPDUMP_PaletteMask = (BPDUMP_UnfilteredPalette | BPDUMP_FilteredPalette),
 		BPDUMP_ContextMask = (BPDUMP_GraphContextActions | BPDUMP_PinContextActions),
@@ -140,6 +155,7 @@ DumpBlueprintsInfo commandlet params: \n\
 		UClass*    SelectedObjectType;
 		FString    DiffPath;
 		FString    DiffCommand;
+		FString    Filename;
 	};
 
 	/** Static instance of the command switches (so we don't have to pass one along the call stack) */
@@ -345,12 +361,12 @@ DumpBlueprintsInfo commandlet params: \n\
 //------------------------------------------------------------------------------
 DumpBlueprintInfoUtils::CommandletOptions::CommandletOptions(TArray<FString> const& Switches)
 	: BlueprintClass(nullptr)
-	, DumpFlags(DumpBlueprintInfoUtils::BPDUMP_UnfilteredPalette)
+	, DumpFlags(BPDUMP_UnfilteredPalette | BPDUMP_UseLegacyMenuBuilder)
 	, PaletteFilter(nullptr)
 	, GraphFilter(GT_MAX)
 	, SelectedObjectType(nullptr)
 {
-	uint32 NewDumpFlags = 0u;
+	uint32 NewDumpFlags = BPDUMP_UseLegacyMenuBuilder;
 	for (FString const& Switch : Switches)
 	{
 		if (Switch.StartsWith("class="))
@@ -510,6 +526,15 @@ DumpBlueprintInfoUtils::CommandletOptions::CommandletOptions(TArray<FString> con
 		{
 			NewDumpFlags |= BPDUMP_RecordTiming;
 		}
+		else if (!Switch.Compare("experimental", ESearchCase::IgnoreCase))
+		{
+			NewDumpFlags &= ~BPDUMP_UseLegacyMenuBuilder;
+		}
+		else if (Switch.StartsWith("name="))
+		{
+			FString NameSwitch;
+			Switch.Split(TEXT("="), &NameSwitch, &Filename);
+		}
 		else
 		{
 			UE_LOG(LogBlueprintInfoDump, Warning, TEXT("Unrecognized command switch '%s', use -help for a listing of all accepted params"), *Switch);
@@ -520,7 +545,7 @@ DumpBlueprintInfoUtils::CommandletOptions::CommandletOptions(TArray<FString> con
 	{
 		DumpFlags = NewDumpFlags;
 	}
-	if (!(DumpFlags & (BPDUMP_PaletteMask | BPDUMP_ContextMask)))
+	if ((DumpFlags & (BPDUMP_PaletteMask | BPDUMP_ContextMask)) == 0)
 	{
 		DumpFlags |= BPDUMP_UnfilteredPalette;
 	}
@@ -736,12 +761,7 @@ static UBlueprint* DumpBlueprintInfoUtils::MakeTempBlueprint(UClass* ParentClass
 //------------------------------------------------------------------------------
 static bool DumpBlueprintInfoUtils::AddComponentToBlueprint(UBlueprint* Blueprint, UClass* ComponentClass)
 {
-	// copied from FBlueprintEditor::CanAccessComponentsMode()
-	bool const bCanUserAddComponents = (Blueprint->SimpleConstructionScript != nullptr) // An SCS must be present (otherwise there is nothing valid to edit)
-		&& FBlueprintEditorUtils::IsActorBased(Blueprint)        // Must be parented to an AActor-derived class (some older BPs may have an SCS but may not be Actor-based)
-		&& (Blueprint->BlueprintType != BPTYPE_MacroLibrary)     // Must not be a macro-type Blueprint
-		&& (Blueprint->BlueprintType != BPTYPE_FunctionLibrary); // Must not be a function library
-
+	bool const bCanUserAddComponents  = FBlueprintEditorUtils::DoesSupportComponents(Blueprint);
 	bool const bClassIsActorComponent = ComponentClass->IsChildOf(UActorComponent::StaticClass());
 	bool const bCanBeAddedToBlueprint = !ComponentClass->HasAnyClassFlags(CLASS_Abstract) && ComponentClass->HasMetaData(FBlueprintMetadata::MD_BlueprintSpawnableComponent);
 	bool const bCanMakeComponent = (bCanUserAddComponents && bClassIsActorComponent && bCanBeAddedToBlueprint);
@@ -787,8 +807,6 @@ static NodeType* DumpBlueprintInfoUtils::AddNodeToGraph(UEdGraph* Graph)
 //------------------------------------------------------------------------------
 static FString DumpBlueprintInfoUtils::BuildDumpFilePath(UClass* BlueprintClass)
 {
-	
-	
 	FString Pathname = FString::Printf(TEXT("BlueprintsInfoDump_%s"), FPlatformTime::StrTimestamp());
 	Pathname = Pathname.Replace(TEXT(" "), TEXT("_"));
 	Pathname = Pathname.Replace(TEXT("/"), TEXT("-"));
@@ -810,9 +828,18 @@ static FString DumpBlueprintInfoUtils::BuildDumpFilePath(UClass* BlueprintClass)
 		}
 	}
 	
+	if (!CommandOptions.Filename.IsEmpty())
+	{
+		Pathname = CommandOptions.Filename;
+	}
+
 	if (bSplitBlueprintsByFile && (BlueprintClass != nullptr))
 	{
-		Pathname = ("BlueprintInfo_" + BlueprintClass->GetName() + ".json");
+		if (CommandOptions.Filename.IsEmpty())
+		{
+			Pathname = TEXT("BlueprintInfo");
+		}
+		Pathname += "_" + BlueprintClass->GetName() + ".json";
 	}
 	else
 	{
@@ -950,10 +977,27 @@ static double DumpBlueprintInfoUtils::GetPaletteMenuActions(FBlueprintPaletteLis
 	UEdGraphSchema_K2 const* K2Schema = GetDefault<UEdGraphSchema_K2>();
 
 	double MenuBuildDuration = 0.0;
+	
+	if (CommandOptions.DumpFlags & BPDUMP_UseLegacyMenuBuilder)
 	{
 		FScopedDurationTimer DurationTimer(MenuBuildDuration);
 		FK2ActionMenuBuilder(K2Schema).GetPaletteActions(PaletteBuilder, PaletteFilter);
 	}
+	else
+	{
+		FBlueprintActionFilter ActionFilter;
+		ActionFilter.Context.Blueprints.Add(const_cast<UBlueprint*>(PaletteBuilder.Blueprint));
+		ActionFilter.ExcludedNodeTypes.Add(UK2Node_Variable::StaticClass());
+
+		// prime the database so it's not recorded in our timing capture
+		FBlueprintActionDatabase::Prime(); 
+
+		FDurationTimer DurationTimer(MenuBuildDuration);
+		FBlueprintActionMenuBuilder MenuBuilder(ActionFilter, /*bAutoBuild =*/true);
+		DurationTimer.Stop();
+
+		PaletteBuilder.Append(MenuBuilder);
+	}
 
 	return MenuBuildDuration;
 }
@@ -1471,12 +1515,57 @@ static double DumpBlueprintInfoUtils::GetContextMenuActions(FBlueprintGraphActio
 	ActionBuilder.Empty();
 	check(ActionBuilder.CurrentGraph != nullptr);
 
-	UEdGraphSchema const* GraphSchema = GetDefault<UEdGraphSchema>(ActionBuilder.CurrentGraph->Schema);
-
 	double MenuBuildDuration = 0.0;
+
+	UEdGraphSchema const* GraphSchema = GetDefault<UEdGraphSchema>(ActionBuilder.CurrentGraph->Schema);
+	if (CommandOptions.DumpFlags & BPDUMP_UseLegacyMenuBuilder)
 	{
- 		FScopedDurationTimer DurationTimer(MenuBuildDuration);
- 		GraphSchema->GetGraphContextActions(ActionBuilder);
+		FScopedDurationTimer DurationTimer(MenuBuildDuration);
+		GraphSchema->GetGraphContextActions(ActionBuilder);
+	}
+	else
+	{
+		UBlueprint* Blueprint = const_cast<UBlueprint*>(ActionBuilder.Blueprint);
+		
+		FBlueprintActionFilter MenuFilter;
+		MenuFilter.Context.Blueprints.Add(Blueprint);
+		MenuFilter.Context.Graphs.Add(const_cast<UEdGraph*>(ActionBuilder.CurrentGraph));
+		
+		if (ActionBuilder.FromPin != nullptr)
+		{
+			MenuFilter.Context.Pins.Add(const_cast<UEdGraphPin*>(ActionBuilder.FromPin));
+		}
+		
+		FBlueprintActionMenuBuilder MenuBuilder;
+		for (UObject* SelectedObj : ActionBuilder.SelectedObjects)
+		{
+			if (UObjectProperty* SelectedProperty = Cast<UObjectProperty>(SelectedObj))
+			{
+				MenuFilter.OwnerClasses.Add(SelectedProperty->PropertyClass);
+				FText const PropertyName = FText::FromName(SelectedProperty->GetFName());
+				
+				MenuFilter.NodeTypes.Add(UK2Node_CallFunction::StaticClass());
+				FText const FuncSectionName = FText::Format(NSLOCTEXT("DumpBpInfo", "ComponentFuncSection", "Call Function on {0}"), PropertyName);
+				MenuBuilder.AddMenuSection(MenuFilter, FuncSectionName);
+				
+				MenuFilter.NodeTypes[0] = UK2Node_Event::StaticClass();
+				FText const EventSectionName = FText::Format(NSLOCTEXT("DumpBpInfo", "ComponentEventSection", "Add Event for {0}"), PropertyName);
+				MenuBuilder.AddMenuSection(MenuFilter, EventSectionName);
+				
+				MenuFilter.NodeTypes.Empty();
+				MenuFilter.OwnerClasses.Empty();
+			}
+		}
+		MenuFilter.OwnerClasses.Add(Blueprint->GeneratedClass);
+		MenuBuilder.AddMenuSection(MenuFilter);
+
+		// prime the database so it's not recorded in our timing capture
+		FBlueprintActionDatabase::Prime();
+		{
+			FScopedDurationTimer DurationTimer(MenuBuildDuration);
+			MenuBuilder.RebuildActionList();
+		}
+		ActionBuilder.Append(MenuBuilder);
 	}
 
 	return MenuBuildDuration;
