Plugability - ClassIconFinder now allows additional style sets to be registered that will be searched first before the default icons are searched so that plugins or other modules can inject new style sets for class icons.

diff --git a/Engine/Source/Editor/UnrealEd/Private/ClassIconFinder.cpp b/Engine/Source/Editor/UnrealEd/Private/ClassIconFinder.cpp
--- a/Engine/Source/Editor/UnrealEd/Private/ClassIconFinder.cpp
+++ b/Engine/Source/Editor/UnrealEd/Private/ClassIconFinder.cpp
@@ -5,6 +5,34 @@
 #include "Slate.h"
 #include "ClassIconFinder.h"
 
+TArray < const ISlateStyle* > FClassIconFinder::Styles;
+
+void FClassIconFinder::RegisterIconSource(const ISlateStyle* StyleSet)
+{
+	Styles.AddUnique( StyleSet );
+}
+
+void FClassIconFinder::UnregisterIconSource(const ISlateStyle* StyleSet)
+{
+	Styles.Remove( StyleSet );
+}
+
+const FSlateBrush* FClassIconFinder::LookupBrush(FName IconName)
+{
+	const FSlateBrush* IconBrush = NULL;
+
+	for ( const ISlateStyle* Style : Styles )
+	{
+		IconBrush = Style->GetOptionalBrush(IconName, nullptr, nullptr);
+		if ( IconBrush )
+		{
+			return IconBrush;
+		}
+	}
+
+	return FEditorStyle::GetOptionalBrush(IconName, nullptr, nullptr);
+}
+
 const FSlateBrush* FClassIconFinder::FindIconForActors(const TArray< TWeakObjectPtr<AActor> >& InActors, UClass*& CommonBaseClass)
 {
 	// Get the common base class of the selected actors
@@ -40,7 +68,7 @@ const FSlateBrush* FClassIconFinder::FindIconForActors(const TArray< TWeakObject
 
 const FSlateBrush* FClassIconFinder::FindIconForActor( const TWeakObjectPtr<AActor>& InActor )
 {
-	return FEditorStyle::GetBrush( FindIconNameForActor( InActor ) );
+	return FClassIconFinder::LookupBrush( FindIconNameForActor(InActor) );
 }
 
 FName FClassIconFinder::FindIconNameForActor( const TWeakObjectPtr<AActor>& InActor )
@@ -81,7 +109,7 @@ FName FClassIconFinder::FindIconNameForActor( const TWeakObjectPtr<AActor>& InAc
 
 const FSlateBrush* FClassIconFinder::FindIconForClass(const UClass* InClass, const FName& InDefaultName )
 {
-	return FEditorStyle::GetBrush( FindIconNameImpl( InClass, InDefaultName ) );
+	return FClassIconFinder::LookupBrush( FindIconNameImpl( InClass, InDefaultName ) );
 }
 
 FName FClassIconFinder::FindIconNameForClass(const UClass* InClass, const FName& InDefaultName )
@@ -91,7 +119,7 @@ FName FClassIconFinder::FindIconNameForClass(const UClass* InClass, const FName&
 
 const FSlateBrush* FClassIconFinder::FindThumbnailForClass(const UClass* InClass, const FName& InDefaultName )
 {
-	return FEditorStyle::GetBrush( FindIconNameImpl( InClass, InDefaultName, TEXT("ClassThumbnail") ) );
+	return FClassIconFinder::LookupBrush( FindIconNameImpl( InClass, InDefaultName, TEXT("ClassThumbnail") ) );
 }
 
 FName FClassIconFinder::FindThumbnailNameForClass(const UClass* InClass, const FName& InDefaultName )
@@ -111,7 +139,7 @@ FName FClassIconFinder::FindIconNameImpl(const UClass* InClass, const FName& InD
 		while( Brush == NULL && CurrentClass && (CurrentClass != AActor::StaticClass()) )
 		{
 			BrushName = *FString::Printf( TEXT( "%s.%s" ), StyleRoot, *CurrentClass->GetName() );
-			Brush = FEditorStyle::GetOptionalBrush( BrushName, nullptr, nullptr );
+			Brush = FClassIconFinder::LookupBrush( BrushName );
 			CurrentClass = CurrentClass->GetSuperClass();
 		}
 	}
@@ -130,4 +158,5 @@ FName FClassIconFinder::FindIconNameImpl(const UClass* InClass, const FName& InD
 	}
 
 	return BrushName;
-}
\ No newline at end of file
+}
+
diff --git a/Engine/Source/Editor/UnrealEd/Public/ClassIconFinder.h b/Engine/Source/Editor/UnrealEd/Public/ClassIconFinder.h
--- a/Engine/Source/Editor/UnrealEd/Public/ClassIconFinder.h
+++ b/Engine/Source/Editor/UnrealEd/Public/ClassIconFinder.h
@@ -9,6 +9,12 @@
 class FClassIconFinder
 {
 public:
+	/** Registers a new style set to use to try to find icons. */
+	UNREALED_API static void RegisterIconSource(const ISlateStyle* StyleSet);
+
+	/** Unregisters an style set. */
+	UNREALED_API static void UnregisterIconSource(const ISlateStyle* StyleSet);
+
 	/** Find the best fitting small icon to use for the supplied actor array */
 	UNREALED_API static const FSlateBrush* FindIconForActors(const TArray< TWeakObjectPtr<AActor> >& InActors, UClass*& CommonBaseClass);
 
@@ -33,6 +39,12 @@ public:
 private:
 	/** Find a thumbnail/icon name to use for the supplied class */
 	UNREALED_API static FName FindIconNameImpl(const UClass* InClass, const FName& InDefaultName = FName(), const TCHAR* StyleRoot = TEXT("ClassIcon") );
+
+private:
+	UNREALED_API static const FSlateBrush* LookupBrush(FName IconName);
+
+	/** The available style sets to source icons from */
+	static TArray < const ISlateStyle* > Styles;
 };
 
 #endif // __ActorIconFinder_h__
\ No newline at end of file
