Plugability - Properly decorating the UAutoReimportManager so that the class is exported from the DLL.

diff --git a/Engine/Source/Editor/UnrealEd/Classes/AutoReimport/AutoReimportManager.h b/Engine/Source/Editor/UnrealEd/Classes/AutoReimport/AutoReimportManager.h
--- a/Engine/Source/Editor/UnrealEd/Classes/AutoReimport/AutoReimportManager.h
+++ b/Engine/Source/Editor/UnrealEd/Classes/AutoReimport/AutoReimportManager.h
@@ -5,7 +5,7 @@
 
 /* Deals with auto reimporting of objects when the objects file on disk is modified*/
 UCLASS(config=Editor, transient)
-class UAutoReimportManager : public UObject, public FTickableEditorObject
+class UNREALED_API UAutoReimportManager : public UObject, public FTickableEditorObject
 {
 	GENERATED_UCLASS_BODY()
 
