Fixed: EDITOR: Change Streaming Method clears current level streaming volume
TTP #341214

diff --git a/Engine/Source/Editor/UnrealEd/Private/EditorLevelUtils.cpp b/Engine/Source/Editor/UnrealEd/Private/EditorLevelUtils.cpp
--- a/Engine/Source/Editor/UnrealEd/Private/EditorLevelUtils.cpp
+++ b/Engine/Source/Editor/UnrealEd/Private/EditorLevelUtils.cpp
@@ -252,15 +252,19 @@ namespace EditorLevelUtils
 		UWorld* World = Level->OwningWorld;
 		
 		World->StreamingLevels.Remove(InLevel);
-
+		
 		// re-add the level with the desired streaming class
 		AddLevelToWorld(World, *(CachedPackageName.ToString()), LevelStreamingClass);
 
-		// Set original level transform
+		// Transfer level streaming settings
 		ULevelStreaming* NewStreamingLevel = FLevelUtils::FindStreamingLevel( Level );
 		if ( NewStreamingLevel )
 		{
 			NewStreamingLevel->LevelTransform = InLevel->LevelTransform;
+			NewStreamingLevel->EditorStreamingVolumes = InLevel->EditorStreamingVolumes;
+			NewStreamingLevel->MinTimeBetweenVolumeUnloadRequests = InLevel->MinTimeBetweenVolumeUnloadRequests;
+			NewStreamingLevel->DrawColor = InLevel->DrawColor;
+			NewStreamingLevel->Keywords = InLevel->Keywords;
 		}
 
 		return NewStreamingLevel;
