build fix for unity file size increase
#ue4

diff --git a/Engine/Source/Editor/UnrealEd/Private/LevelEditorViewport.cpp b/Engine/Source/Editor/UnrealEd/Private/LevelEditorViewport.cpp
--- a/Engine/Source/Editor/UnrealEd/Private/LevelEditorViewport.cpp
+++ b/Engine/Source/Editor/UnrealEd/Private/LevelEditorViewport.cpp
@@ -2295,14 +2295,14 @@ void FLevelEditorViewportClient::UpdateViewForLockedActor()
 }
 
 
-namespace ViewportDeadZoneConstants
+/*namespace ViewportDeadZoneConstants
 {
 	enum
 	{
 		NO_DEAD_ZONE,
 		STANDARD_DEAD_ZONE
 	};
-};
+};*/
 
 
 
