AssetThumbnail - Get the class thumbnail from the ClassIconFinder so that it works with alternate styles.

I'm not sure how much this class is actually used in the UI, but it does a lot of redundant work every frame in several different callback methods.  If it's used a lot, it should probably get an optimization pass.

TTP# 334639

diff --git a/Engine/Source/Editor/UnrealEd/Private/AssetThumbnail.cpp b/Engine/Source/Editor/UnrealEd/Private/AssetThumbnail.cpp
--- a/Engine/Source/Editor/UnrealEd/Private/AssetThumbnail.cpp
+++ b/Engine/Source/Editor/UnrealEd/Private/AssetThumbnail.cpp
@@ -77,6 +77,7 @@ public:
 		if ( Class == UClass::StaticClass() )
 		{
 			Substyle = FName(".ClassBackground");
+			ClassAssetClass = FindObject<UClass>(ANY_PACKAGE, *AssetData.AssetName.ToString());
 		}
 		else if(AssetTypeActions.IsValid())
 		{
@@ -84,14 +85,7 @@ public:
 		}
 		const FName BackgroundBrushName( *(Style.ToString() + Substyle.ToString()) );
 
-		if ( InArgs._ClassThumbnailBrushOverride != NAME_None )
-		{
-			ClassThumbnailBrushName = InArgs._ClassThumbnailBrushOverride;
-		}
-		else
-		{
-			ClassThumbnailBrushName = MakeClassThumbnailName();
-		}
+		ClassThumbnailBrushOverride = InArgs._ClassThumbnailBrushOverride;
 
 		OverlayWidget->AddSlot()
 		[
@@ -246,8 +240,6 @@ private:
 		{
 			HintTextBlock->SetText( GetLabelText() );
 		}
-		
-		ClassThumbnailBrushName = MakeClassThumbnailName();
 
 		// Check if the asset has a thumbnail.
 		const FObjectThumbnail* ObjectThumbnail = NULL;
@@ -275,6 +267,11 @@ private:
 			AssetTypeActions = AssetToolsModule.Get().GetAssetTypeActionsForClass(Class);
 		}
 
+		if (Class == UClass::StaticClass())
+		{
+			ClassAssetClass = FindObject<UClass>(ANY_PACKAGE, *AssetData.AssetName.ToString());
+		}
+
 		AssetColor = FLinearColor(1.f, 1.f, 1.f, 1.f);
 		if ( AssetTypeActions.IsValid() )
 		{
@@ -336,12 +333,22 @@ private:
 
 	const FSlateBrush* GetClassThumbnailBrush() const
 	{
-		return FEditorStyle::GetOptionalBrush( ClassThumbnailBrushName );
+		if (ClassThumbnailBrushOverride.IsNone())
+		{
+			return FClassIconFinder::FindThumbnailForClass(ClassAssetClass.Get());
+		}
+		else
+		{
+			// Instead of getting the override thumbnail directly from the editor style here get it from the
+			// ClassIconFinder since it may have additional styles registered which can be searched by passing
+			// it as a default with no class to search for.
+			return FClassIconFinder::FindThumbnailForClass(nullptr, ClassThumbnailBrushOverride);
+		}
 	}
 
 	EVisibility GetClassThumbnailVisibility() const
 	{
-		const FSlateBrush* ClassThumbnailBrush = FEditorStyle::GetOptionalBrush( ClassThumbnailBrushName, nullptr, nullptr );
+		const FSlateBrush* ClassThumbnailBrush = GetClassThumbnailBrush();
 		if(!bHasRenderedThumbnail && ClassThumbnailBrush)
 		{
 			const FAssetData& AssetData = AssetThumbnail->GetAssetData();
@@ -517,12 +524,6 @@ private:
 		return AssetData.AssetName.ToString();
 	}
 
-	FName MakeClassThumbnailName() const
-	{
-		const FAssetData& AssetData = AssetThumbnail->GetAssetData();
-		return FName(*FString::Printf( TEXT( "ClassThumbnail.%s" ), *AssetData.AssetName.ToString() ));
-	}
-
 private:
 	TSharedPtr<STextBlock> LabelTextBlock;
 	TSharedPtr<STextBlock> HintTextBlock;
@@ -543,8 +544,10 @@ private:
 
 	bool ShowClassBackground;
 
-	/** Brush name for rendering the class thumbnail */
-	FName ClassThumbnailBrushName;
+	/** The name of the thumbnail which should be used instead of the class thumbnail. */
+	FName ClassThumbnailBrushOverride;
+	/** The class instance for assets which represent a class. */
+	TWeakObjectPtr<UClass> ClassAssetClass;
 };
 
 
