AssetThumbnail - Add a header file to fix the non-unity builds.

diff --git a/Engine/Source/Editor/UnrealEd/Private/AssetThumbnail.cpp b/Engine/Source/Editor/UnrealEd/Private/AssetThumbnail.cpp
--- a/Engine/Source/Editor/UnrealEd/Private/AssetThumbnail.cpp
+++ b/Engine/Source/Editor/UnrealEd/Private/AssetThumbnail.cpp
@@ -7,6 +7,7 @@
 
 #include "AssetToolsModule.h"
 #include "AssetRegistryModule.h"
+#include "ClassIconFinder.h"
 
 class SAssetThumbnail : public SCompoundWidget
 {
