Minor clang fix (|| and && in the same scope).

(from github PR 242).

diff --git a/Engine/Source/Editor/UnrealEd/Private/MaterialGraph.cpp b/Engine/Source/Editor/UnrealEd/Private/MaterialGraph.cpp
--- a/Engine/Source/Editor/UnrealEd/Private/MaterialGraph.cpp
+++ b/Engine/Source/Editor/UnrealEd/Private/MaterialGraph.cpp
@@ -419,13 +419,13 @@ int32 UMaterialGraph::GetValidOutputIndex(FExpressionInput* Input) const
 			{
 				const FExpressionOutput& Output = Outputs[OutputIndex];
 
-				if(	bOutputIndexIsValid && OutputIndex == Input->OutputIndex
-					|| !bOutputIndexIsValid
+				if((bOutputIndexIsValid && OutputIndex == Input->OutputIndex)
+					|| (!bOutputIndexIsValid
 					&& Output.Mask == Input->Mask
 					&& Output.MaskR == Input->MaskR
 					&& Output.MaskG == Input->MaskG
 					&& Output.MaskB == Input->MaskB
-					&& Output.MaskA == Input->MaskA )
+					&& Output.MaskA == Input->MaskA))
 				{
 					break;
 				}
