Fixed regression in material generation for a MeshProxy

diff --git a/Engine/Source/Editor/UnrealEd/Private/EditorExporters.cpp b/Engine/Source/Editor/UnrealEd/Private/EditorExporters.cpp
--- a/Engine/Source/Editor/UnrealEd/Private/EditorExporters.cpp
+++ b/Engine/Source/Editor/UnrealEd/Private/EditorExporters.cpp
@@ -2964,7 +2964,7 @@ namespace MaterialExportUtils
 			//Create BaseColor Expression of the diffuse texture
 			UMaterialExpressionTextureSample* BasecolorExpression = ConstructObject<UMaterialExpressionTextureSample>(UMaterialExpressionTextureSample::StaticClass(), Material);
 			BasecolorExpression->Texture = DiffuseTexture;
-			BasecolorExpression->SamplerType = EMaterialSamplerType::SAMPLERTYPE_Color;
+			BasecolorExpression->SamplerType = EMaterialSamplerType::SAMPLERTYPE_LinearColor;
 			BasecolorExpression->MaterialExpressionEditorX = -400;
 			BasecolorExpression->MaterialExpressionEditorY = -150;
 			Material->Expressions.Add(BasecolorExpression);
