Updated some places using a hardcoded search in FEditorStyle to use FClassIconFinder instead

diff --git a/Engine/Source/Editor/LevelEditor/Private/LevelEditorCreateActorMenu.cpp b/Engine/Source/Editor/LevelEditor/Private/LevelEditorCreateActorMenu.cpp
--- a/Engine/Source/Editor/LevelEditor/Private/LevelEditorCreateActorMenu.cpp
+++ b/Engine/Source/Editor/LevelEditor/Private/LevelEditorCreateActorMenu.cpp
@@ -93,8 +93,7 @@ class SAssetMenuEntry : public SCompoundWidget
 					ActorTypeDisplayName = MenuItem.FactoryToUse->GetDisplayName();
 				}
 
-				FName BrushName = *FString::Printf( TEXT("ClassIcon.%s"), *MenuItem.FactoryToUse->GetClass()->GetName() );
-				IconBrush = FEditorStyle::GetOptionalBrush(BrushName, nullptr, nullptr);
+				IconBrush = FClassIconFinder::FindIconForClass( MenuItem.FactoryToUse->GetClass() );
 			}
 
 			if ( DefaultActor != NULL && ( MenuItem.FactoryToUse != NULL || !IsClass ) )
diff --git a/Engine/Source/Editor/UnrealEd/Public/DragAndDrop/ClassDragDropOp.h b/Engine/Source/Editor/UnrealEd/Public/DragAndDrop/ClassDragDropOp.h
--- a/Engine/Source/Editor/UnrealEd/Public/DragAndDrop/ClassDragDropOp.h
+++ b/Engine/Source/Editor/UnrealEd/Public/DragAndDrop/ClassDragDropOp.h
@@ -1,6 +1,8 @@
 // Copyright 1998-2014 Epic Games, Inc. All Rights Reserved.
 #pragma once
 
+#include "ClassIconFinder.h"
+
 class FClassDragDropOp : public FDragDropOperation
 {
 public:
@@ -13,7 +15,7 @@ public:
 	virtual TSharedPtr<SWidget> GetDefaultDecorator() const override
 	{
 		// Just use the first class for the cursor decorator.
-		const FSlateBrush* ClassIcon = FEditorStyle::GetOptionalBrush(*FString::Printf( TEXT( "ClassIcon.%s" ), *ClassesToDrop[0]->GetName() ), nullptr, nullptr );
+		const FSlateBrush* ClassIcon = FClassIconFinder::FindIconForClass( ClassesToDrop[0].Get() );
 
 		// If the class icon is the default brush, do not put it in the cursor decoration window.
 		if(ClassIcon)
