FUVTestHelper inherits from TSharedFromThis<> so that is the test is stopped the AsyncPackageLoad calback dosn't call back into a garbage collected FUVTestHelper instance

diff --git a/Engine/Source/Editor/UnrealEd/Private/Tests/EditorAutomationTests.cpp b/Engine/Source/Editor/UnrealEd/Private/Tests/EditorAutomationTests.cpp
--- a/Engine/Source/Editor/UnrealEd/Private/Tests/EditorAutomationTests.cpp
+++ b/Engine/Source/Editor/UnrealEd/Private/Tests/EditorAutomationTests.cpp
@@ -2018,7 +2018,7 @@ namespace StaticMeshUVTest
 		};
 	}
 
-	class FUVTestHelper
+	class FUVTestHelper : public TSharedFromThis<FUVTestHelper>
 	{
 	public:
 
@@ -2088,7 +2088,7 @@ namespace StaticMeshUVTest
 			}
 			else
 			{
-				LoadPackageAsync(PackageName,FLoadPackageAsyncDelegate::CreateRaw(this, &FUVTestHelper::PackageLoadCallback));
+				LoadPackageAsync(PackageName,FLoadPackageAsyncDelegate::CreateSP(this, &FUVTestHelper::PackageLoadCallback));
 				CurrentState = EStaticMeshUVTestState::WaitingForPackage;
 			}
 		}
