Deleting unused variables is now an undoable action.

Added a transaction when the action occurs.

#ttp 340635 - BP: hitting undo after using  'delete unused variables'  restores 1 variable less than the variables deleted

diff --git a/Engine/Source/Editor/UnrealEd/Private/Kismet2/BlueprintEditorUtils.cpp b/Engine/Source/Editor/UnrealEd/Private/Kismet2/BlueprintEditorUtils.cpp
--- a/Engine/Source/Editor/UnrealEd/Private/Kismet2/BlueprintEditorUtils.cpp
+++ b/Engine/Source/Editor/UnrealEd/Private/Kismet2/BlueprintEditorUtils.cpp
@@ -3163,6 +3163,9 @@ void FBlueprintEditorUtils::RemoveMemberVariable(UBlueprint* Blueprint, const FN
 
 void FBlueprintEditorUtils::BulkRemoveMemberVariables(UBlueprint* Blueprint, const TArray<FName>& VarNames)
 {
+	const FScopedTransaction Transaction( LOCTEXT("DeleteUnusedVariables", "Delete Unused Variables") );
+	Blueprint->Modify();
+
 	bool bModified = false;
 	for (int32 i = 0; i < VarNames.Num(); ++i)
 	{
