Fix for thumbnail rendering producing black or incorrect thumbnails

diff --git a/Engine/Source/Editor/UnrealEd/Private/AssetThumbnail.cpp b/Engine/Source/Editor/UnrealEd/Private/AssetThumbnail.cpp
--- a/Engine/Source/Editor/UnrealEd/Private/AssetThumbnail.cpp
+++ b/Engine/Source/Editor/UnrealEd/Private/AssetThumbnail.cpp
@@ -139,7 +139,7 @@ public:
 			InArgs._ThumbnailPool->OnThumbnailRendered().AddSP(this, &SAssetThumbnail::OnThumbnailRendered);
 			InArgs._ThumbnailPool->OnThumbnailRenderFailed().AddSP(this, &SAssetThumbnail::OnThumbnailRenderFailed);
 
-			if ( ShouldRender() && (!InArgs._AllowFadeIn || !InArgs._ThumbnailPool->IsInRenderStack(AssetThumbnail)) )
+			if ( ShouldRender() && (!InArgs._AllowFadeIn || InArgs._ThumbnailPool->IsRendered(AssetThumbnail)) )
 			{
 				bHasRenderedThumbnail = true;
 				ViewportFadeAnimation.JumpToEnd();
@@ -945,6 +945,9 @@ void FAssetThumbnailPool::Tick( float DeltaTime )
 
 					if ( bLoadedThumbnail )
 					{
+						// Mark it as updated
+						InfoRef->LastUpdateTime = FPlatformTime::Seconds();
+
 						// Notify listeners that a thumbnail has been rendered
 						ThumbnailRenderedEvent.Broadcast(InfoRef->AssetData);
 					}
@@ -980,7 +983,7 @@ FSlateTexture2DRHIRef* FAssetThumbnailPool::AccessTexture( const FAssetData& Ass
 			// If a the max number of thumbnails allowed by the pool exists then reuse its rendering resource for the new thumbnail
 			if( FreeThumbnails.Num() == 0 && ThumbnailToTextureMap.Num() == NumInPool )
 			{
-				// Find the thumbnail which was rendered last and use it for the new thumbnail
+				// Find the thumbnail which was accessed last and use it for the new thumbnail
 				float LastAccessTime = FLT_MAX;
 				const FThumbId* AssetToRemove = NULL;
 				for( TMap< FThumbId, TSharedRef<FThumbnailInfo> >::TConstIterator It(ThumbnailToTextureMap); It; ++It )
@@ -1039,6 +1042,9 @@ FSlateTexture2DRHIRef* FAssetThumbnailPool::AccessTexture( const FAssetData& Ass
 			ThumbnailInfo->Width = Width;
 			ThumbnailInfo->Height = Height;
 		
+			// Mark this thumbnail as needing to be updated
+			ThumbnailInfo->LastUpdateTime = -1.f;
+
 			// Request that the thumbnail be rendered as soon as possible
 			ThumbnailsToRenderStack.Push( ThumbnailRef );
 		}
@@ -1128,6 +1134,25 @@ bool FAssetThumbnailPool::IsInRenderStack( const TSharedPtr<FAssetThumbnail>& Th
 	return false;
 }
 
+bool FAssetThumbnailPool::IsRendered(const TSharedPtr<FAssetThumbnail>& Thumbnail) const
+{
+	const FAssetData& AssetData = Thumbnail->GetAssetData();
+	const uint32 Width = Thumbnail->GetSize().X;
+	const uint32 Height = Thumbnail->GetSize().Y;
+
+	if (ensure(AssetData.ObjectPath != NAME_None) && ensure(Width > 0) && ensure(Height > 0))
+	{
+		FThumbId ThumbId(AssetData.ObjectPath, Width, Height);
+		const TSharedRef<FThumbnailInfo>* ThumbnailInfoPtr = ThumbnailToTextureMap.Find(ThumbId);
+		if (ThumbnailInfoPtr)
+		{
+			return (*ThumbnailInfoPtr).Get().LastUpdateTime >= 0.f;
+		}
+	}
+
+	return false;
+}
+
 void FAssetThumbnailPool::PrioritizeThumbnails( const TArray< TSharedPtr<FAssetThumbnail> >& ThumbnailsToPrioritize, uint32 Width, uint32 Height )
 {
 	if ( ensure(Width > 0) && ensure(Height > 0) )
diff --git a/Engine/Source/Editor/UnrealEd/Public/AssetThumbnail.h b/Engine/Source/Editor/UnrealEd/Public/AssetThumbnail.h
--- a/Engine/Source/Editor/UnrealEd/Public/AssetThumbnail.h
+++ b/Engine/Source/Editor/UnrealEd/Public/AssetThumbnail.h
@@ -158,6 +158,9 @@ public:
 	/** Returns true if the thumbnail for the specified asset in the specified size is in the stack of thumbnails to render */
 	bool IsInRenderStack( const TSharedPtr<FAssetThumbnail>& Thumbnail ) const;
 
+	/** Returns true if the thumbnail for the specified asset in the specified size has been rendered */
+	bool IsRendered( const TSharedPtr<FAssetThumbnail>& Thumbnail ) const;
+
 	/** Brings all items in ThumbnailsToPrioritize to the front of the render stack if they are actually in the stack */
 	UNREALED_API void PrioritizeThumbnails( const TArray< TSharedPtr<FAssetThumbnail> >& ThumbnailsToPrioritize, uint32 Width, uint32 Height );
 
@@ -206,6 +209,8 @@ private:
 		FSlateTextureRenderTarget2DResource* ThumbnailRenderTarget;
 		/** The time since last access */
 		float LastAccessTime;
+		/** The time since last update */
+		float LastUpdateTime;
 		/** Width of the thumbnail */
 		uint32 Width;
 		/** Height of the thumbnail */
