Automation:  Milestone1 fixes

Asset Creation and Duplication test
 -Enabled Behavior trees while the test is running
 -Skipped assets that require a skeleton if we didn't find one to use
 -Skipped creating a data asset if GameplayAbilitySet is not found

diff --git a/Engine/Source/Editor/UnrealEd/Private/Tests/EditorAssetAutomationTests.cpp b/Engine/Source/Editor/UnrealEd/Private/Tests/EditorAssetAutomationTests.cpp
--- a/Engine/Source/Editor/UnrealEd/Private/Tests/EditorAssetAutomationTests.cpp
+++ b/Engine/Source/Editor/UnrealEd/Private/Tests/EditorAssetAutomationTests.cpp
@@ -128,6 +128,9 @@ namespace CreateAssetHelper
 		//Total number of assets
 		uint32 NumTotalAssets;
 
+		//Number of assets skipped
+		uint32 NumSkippedAssets;
+
 		//Number of assets created
 		uint32 NumCreated;
 
@@ -148,6 +151,7 @@ namespace CreateAssetHelper
 		*/
 		FCreateAssetStats() :
 			NumTotalAssets(0),
+			NumSkippedAssets(0),
 			NumCreated(0),
 			NumSaved(0),
 			NumDuplicatesSaved(0),
@@ -426,6 +430,17 @@ bool FClearEditorReferencesCommand::Update()
 	return true;
 }
 
+/**
+* Latent command to disable the behavior tree editor
+*/
+DEFINE_LATENT_AUTOMATION_COMMAND(FDisableBehaviorTreeEditorCommand);
+bool FDisableBehaviorTreeEditorCommand::Update()
+{
+	bool bEnableBehaviorTrees = false;
+	GConfig->SetBool(TEXT("BehaviorTreesEd"), TEXT("BehaviorTreeNewAssetsEnabled"), bEnableBehaviorTrees, GEngineIni);
+	return true;
+}
+
 /**
 * Latent command log the asset creation stats
 */
@@ -434,6 +449,10 @@ bool FLogAssetCreationStatsCommand::Update()
 {
 	UE_LOG(LogEditorAssetAutomationTests, Display, TEXT(" "));
 	UE_LOG(LogEditorAssetAutomationTests, Display, TEXT("Test Summary:"));
+	if (BuildStats->NumSkippedAssets)
+	{
+		UE_LOG(LogEditorAssetAutomationTests, Display, TEXT("Skipped %i assets"), BuildStats->NumSkippedAssets);
+	}
 	UE_LOG(LogEditorAssetAutomationTests, Display, TEXT("%i of %i assets were created successfully"),			BuildStats->NumCreated,			BuildStats->NumTotalAssets);
 	UE_LOG(LogEditorAssetAutomationTests, Display, TEXT("%i of %i assets were saved successfully"),				BuildStats->NumSaved,			BuildStats->NumTotalAssets);
 	UE_LOG(LogEditorAssetAutomationTests, Display, TEXT("%i of %i assets were duplicated successfully"),		BuildStats->NumDuplicated,		BuildStats->NumTotalAssets);
@@ -523,6 +542,22 @@ bool FAssetEditorTest::RunTest(const FString& Parameters)
 		}
 	}
 
+	//Check to see if we need to enable Behavior trees.
+	bool bEnabledBehaviorTrees = false;
+	bool bBehaviorTreeNewAssetsEnabled = false;
+	GConfig->GetBool(TEXT("BehaviorTreesEd"), TEXT("BehaviorTreeNewAssetsEnabled"), bBehaviorTreeNewAssetsEnabled, GEngineIni);
+	if (!bBehaviorTreeNewAssetsEnabled)
+	{
+		bEnabledBehaviorTrees = true;
+		GConfig->SetBool(TEXT("BehaviorTreesEd"), TEXT("BehaviorTreeNewAssetsEnabled"), bEnabledBehaviorTrees, GEngineIni);
+
+		if (!FModuleManager::Get().IsModuleLoaded(TEXT("BehaviorTreeEditor")))
+		{
+			//NOTE: This module gets left in after the test completes otherwise the content browser would crash when it tries to access the created BehaviorTree.
+			FModuleManager::Get().LoadModule(TEXT("BehaviorTreeEditor"));
+		}
+	}
+
 	//Holds info on each asset we are creating
 	TArray< TSharedPtr<CreateAssetHelper::FCreateAssetInfo> > AssetInfos;
 	TSharedPtr<CreateAssetHelper::FCreateAssetStats> BuildStats = (MakeShareable(new CreateAssetHelper::FCreateAssetStats()));
@@ -531,13 +566,23 @@ bool FAssetEditorTest::RunTest(const FString& Parameters)
 	ASSET_TEST_CREATE(UBlueprint, UBlueprintFactory, BP, FactoryInst->ParentClass = AActor::StaticClass();)
 	ASSET_TEST_CREATE(UMaterial, UMaterialFactoryNew, MAT, )
 	ASSET_TEST_CREATE(UParticleSystem, UParticleSystemFactoryNew, PS, )
-	ASSET_TEST_CREATE(UAimOffsetBlendSpace, UAimOffsetBlendSpaceFactoryNew, AO, FactoryInst->TargetSkeleton = FirstSkeleton;)
-	ASSET_TEST_CREATE(UAimOffsetBlendSpace1D, UAimOffsetBlendSpaceFactory1D, AO1D, FactoryInst->TargetSkeleton = FirstSkeleton;)
-	ASSET_TEST_CREATE(UAnimBlueprint, UAnimBlueprintFactory, AB, FactoryInst->ParentClass = UAnimInstance::StaticClass(); FactoryInst->TargetSkeleton = FirstSkeleton;)
-	ASSET_TEST_CREATE(UAnimComposite, UAnimCompositeFactory, AC, FactoryInst->TargetSkeleton = FirstSkeleton;)
-	ASSET_TEST_CREATE(UAnimMontage, UAnimMontageFactory, AM, FactoryInst->TargetSkeleton = FirstSkeleton;)
-	ASSET_TEST_CREATE(UBlendSpace, UBlendSpaceFactoryNew, BS, FactoryInst->TargetSkeleton = FirstSkeleton;)
-	ASSET_TEST_CREATE(UBlendSpace1D, UBlendSpaceFactory1D, BS1D, FactoryInst->TargetSkeleton = FirstSkeleton;)
+	
+	if (FirstSkeleton)
+	{
+		ASSET_TEST_CREATE(UAimOffsetBlendSpace, UAimOffsetBlendSpaceFactoryNew, AO, FactoryInst->TargetSkeleton = FirstSkeleton;)
+		ASSET_TEST_CREATE(UAimOffsetBlendSpace1D, UAimOffsetBlendSpaceFactory1D, AO1D, FactoryInst->TargetSkeleton = FirstSkeleton;)
+		ASSET_TEST_CREATE(UAnimBlueprint, UAnimBlueprintFactory, AB, FactoryInst->ParentClass = UAnimInstance::StaticClass(); FactoryInst->TargetSkeleton = FirstSkeleton;)
+		ASSET_TEST_CREATE(UAnimComposite, UAnimCompositeFactory, AC, FactoryInst->TargetSkeleton = FirstSkeleton;)
+		ASSET_TEST_CREATE(UAnimMontage, UAnimMontageFactory, AM, FactoryInst->TargetSkeleton = FirstSkeleton;)
+		ASSET_TEST_CREATE(UBlendSpace, UBlendSpaceFactoryNew, BS, FactoryInst->TargetSkeleton = FirstSkeleton;)
+		ASSET_TEST_CREATE(UBlendSpace1D, UBlendSpaceFactory1D, BS1D, FactoryInst->TargetSkeleton = FirstSkeleton;)
+	}
+	else
+	{
+		BuildStats->NumSkippedAssets += 7;
+		UE_LOG(LogEditorAssetAutomationTests, Warning, TEXT("NO AVAILABLE SKELETON.  Skipping related assets."));
+	}
+
 	ASSET_TEST_CREATE(UTextureRenderTargetCube, UTextureRenderTargetCubeFactoryNew, CRT, )
 	ASSET_TEST_CREATE(UFont, UTrueTypeFontFactory, F, )
 	ASSET_TEST_CREATE(UMaterialFunction, UMaterialFunctionFactoryNew, MF, )
@@ -560,7 +605,16 @@ bool FAssetEditorTest::RunTest(const FString& Parameters)
 	ASSET_TEST_CREATE(UCameraAnim, UCameraAnimFactory, CA, )
 	ASSET_TEST_CREATE(UCurveBase, UCurveFactory, C, FactoryInst->CurveClass = UCurveFloat::StaticClass();)
 	UClass* GameplayAbilityClass = StaticLoadClass(UObject::StaticClass(), NULL, TEXT("GameplayAbilities.GameplayAbilitySet"), NULL, LOAD_None, NULL);
-	ASSET_TEST_CREATE(UDataAsset, UDataAssetFactory, DA, FactoryInst->DataAssetClass = GameplayAbilityClass;)
+	if (GameplayAbilityClass)
+	{
+		ASSET_TEST_CREATE(UDataAsset, UDataAssetFactory, DA, FactoryInst->DataAssetClass = GameplayAbilityClass;)
+	}
+	else
+	{
+		BuildStats->NumSkippedAssets += 1;
+		UE_LOG(LogEditorAssetAutomationTests, Warning, TEXT("COULD NOT LOAD GameplayAbilitySet.  Skipping DataAsset creation."));
+	}
+	
 	ASSET_TEST_CREATE(UUserDefinedEnum, UEnumFactory, Enum, )
 	ASSET_TEST_CREATE(UForceFeedbackEffect, UForceFeedbackEffectFactory, FFE, )
 	ASSET_TEST_CREATE(UInterpData, UInterpDataFactoryNew, MD, )
@@ -599,6 +653,12 @@ bool FAssetEditorTest::RunTest(const FString& Parameters)
 
 	ADD_LATENT_AUTOMATION_COMMAND(FLogAssetCreationStatsCommand(BuildStats));
 
+	//Disable the behavior trees if we enabled them earlier
+	if (bEnabledBehaviorTrees)
+	{
+		ADD_LATENT_AUTOMATION_COMMAND(FDisableBehaviorTreeEditorCommand());
+	}
+
 	return true;
 }
 
@@ -742,13 +802,13 @@ namespace ImportExportAssetHelper
 		//The Asset file name
 		FString AssetName;
 
-		//If the asset imported successfuly
+		//If the asset imported successfully
 		bool bImportSuccessful;
 
 		//If the export step was skipped
 		bool bSkippedExport;
 
-		//If the asset exported successfuly
+		//If the asset exported successfully
 		bool bExportSuccessful;
 
 		//The size of the exported file
@@ -1088,6 +1148,12 @@ namespace ImportExportAssetHelper
 		*/
 		void DeleteAsset()
 		{
+			// Deselect all
+			GEditor->SelectNone(false, true);
+
+			// Clear the transaction buffer so we aren't referencing the new objects
+			GUnrealEd->ResetTransaction(FText::FromString(TEXT("FAssetEditorTest")));
+
 			//Clear references to the object so we can delete it
 			AssetAutomationCommon::NullReferencesToObject(ImportedAsset);
 
