Added abililty to snap actors to surfaces in the world whilst dragging in the viewport.

#ttp 323078 - Actor drag preview does not respect actor factory spawn position offset
#ttp 309631 - TOOLS FEATURE: Asset placement: Align placed assets with surface placed on
#ttp 335001 - UE4: EDITOR:"Paste here" works inconsistently
#ttp 334066 - EDITOR: PAPER2D: Dragging a mesh into an ortho viewport does not provide a preview

There is a lot of work here:
    * I have moved and unified all the prop preview and drag alignment code that attempted to align objects to surfaces in the world to a new class, FActorPositioning. This fixes inconsistent behaviour between the drop previews, actual placement, and drag alignment.
    * I have tidied up the methods by which surface alignment is performed in UActorFactory. There is now a virtual method for factories to define *how* an actor should be aligned to the surface. This allows us to specify different alignment routines for things like static meshes, decals and lights.
    * Actor factories now create actors with an FTransform rather than specifying pointers to location and rotation vectors. This makes creation simpler and parameter passing less error prone. As far as I can see, the rotation was only used for surface snapping which is now handled by the factory anyway.
    * Surface snapping can be controlled by the editor viewport toggle button, or via a keyboard command set in the editor preferences
    * It is possible to override the default placement offset when aligning actors to surfaces. SnapOffsetExtent can be found in the editor viewport preferences.
    * Finding a valid snap position no longer uses a hitproxy as this was prone to error and does not work correctly when dragging actors that render their own hit proxies.
    * ClickLocation and ClickPlane are now more consistently set and managed. I have tried to remove excess assignment to these variables to clarify their use and purpose.

diff --git a/Engine/Source/Editor/EditorStyle/Private/SlateEditorStyle.cpp b/Engine/Source/Editor/EditorStyle/Private/SlateEditorStyle.cpp
--- a/Engine/Source/Editor/EditorStyle/Private/SlateEditorStyle.cpp
+++ b/Engine/Source/Editor/EditorStyle/Private/SlateEditorStyle.cpp
@@ -3393,6 +3393,7 @@ void FSlateEditorStyle::FStyle::SetupLevelEditorStyle()
 		Set( "EditorViewport.LocationGridSnap", new IMAGE_BRUSH( "Old/LevelEditor/LocationGridSnap", Icon14x14, IconColor) );
 		Set( "EditorViewport.RotationGridSnap", new IMAGE_BRUSH( "Old/LevelEditor/RotationGridSnap", Icon14x14, IconColor ) );
 		Set( "EditorViewport.ScaleGridSnap", new IMAGE_BRUSH( "Old/LevelEditor/ScaleGridSnap", Icon14x14, IconColor ) );
+		Set( "EditorViewport.ToggleSurfaceSnapping", new IMAGE_BRUSH( "Icons/icon_surface_snapping_14px", Icon14x14 ) );
 		Set( "EditorViewport.RelativeCoordinateSystem_Local", new IMAGE_BRUSH( "Icons/icon_axis_local_16px", Icon16x16, IconColor ) );
 		Set( "EditorViewport.RelativeCoordinateSystem_Local.Small", new IMAGE_BRUSH( "Icons/icon_axis_local_16px", Icon16x16, IconColor ) );
 		Set( "EditorViewport.RelativeCoordinateSystem_World", new IMAGE_BRUSH( "Icons/icon_axis_world_16px", Icon16x16, IconColor ) );
diff --git a/Engine/Source/Editor/LevelEditor/Private/LevelEditorActions.cpp b/Engine/Source/Editor/LevelEditor/Private/LevelEditorActions.cpp
--- a/Engine/Source/Editor/LevelEditor/Private/LevelEditorActions.cpp
+++ b/Engine/Source/Editor/LevelEditor/Private/LevelEditorActions.cpp
@@ -1137,9 +1137,9 @@ void FLevelEditorActionCallbacks::AddActor_Clicked( UActorFactory* ActorFactory,
 	}
 }
 
-AActor* FLevelEditorActionCallbacks::AddActor( UActorFactory* ActorFactory, const FAssetData& AssetData, const FVector* ActorLocation )
+AActor* FLevelEditorActionCallbacks::AddActor( UActorFactory* ActorFactory, const FAssetData& AssetData, const FTransform* ActorTransform )
 {
-	AActor* NewActor = GEditor->UseActorFactory( ActorFactory, AssetData, ActorLocation, ActorFactory->bUseSurfaceOrientation );
+	AActor* NewActor = GEditor->UseActorFactory( ActorFactory, AssetData, ActorTransform );
 
 	if ( NewActor != NULL && IPlacementModeModule::IsAvailable() )
 	{
@@ -1151,10 +1151,10 @@ AActor* FLevelEditorActionCallbacks::AddActor( UActorFactory* ActorFactory, cons
 
 void FLevelEditorActionCallbacks::AddActorFromClass_Clicked( UClass* ActorClass )
 {
-	FLevelEditorActionCallbacks::AddActorFromClass(ActorClass, NULL);
+	FLevelEditorActionCallbacks::AddActorFromClass(ActorClass);
 }
 
-AActor* FLevelEditorActionCallbacks::AddActorFromClass( UClass* ActorClass, const FVector* ActorLocation )
+AActor* FLevelEditorActionCallbacks::AddActorFromClass( UClass* ActorClass )
 {
 	AActor* NewActor = NULL;
 
@@ -1164,7 +1164,7 @@ AActor* FLevelEditorActionCallbacks::AddActorFromClass( UClass* ActorClass, cons
 		UActorFactory* ActorFactory = GEditor->FindActorFactoryForActorClass( ActorClass );
 		if( ActorFactory )
 		{
-			NewActor = GEditor->UseActorFactoryOnCurrentSelection( ActorFactory, ActorLocation, ActorFactory->bUseSurfaceOrientation );
+			NewActor = GEditor->UseActorFactoryOnCurrentSelection( ActorFactory, nullptr );
 
 			if ( NewActor != NULL && IPlacementModeModule::IsAvailable() )
 			{
@@ -2193,7 +2193,7 @@ void FLevelEditorActionCallbacks::OnAddMatinee()
 		// Spawn a matinee actor at the origin, and either move infront of the camera or focus camera on it (depending on the viewport) and open for edit
 		UActorFactory* ActorFactory = GEditor->FindActorFactoryForActorClass( AMatineeActor::StaticClass() );
 		check( ActorFactory );
-		AMatineeActor* MatineeActor = CastChecked<AMatineeActor>( FLevelEditorActionCallbacks::AddActor( ActorFactory, FAssetData(), &FVector::ZeroVector ) );
+		AMatineeActor* MatineeActor = CastChecked<AMatineeActor>( FLevelEditorActionCallbacks::AddActor( ActorFactory, FAssetData(), &FTransform::Identity ) );
 		if( GCurrentLevelEditingViewportClient->IsPerspective() )
 		{
 			GEditor->MoveActorInFrontOfCamera( *MatineeActor, GCurrentLevelEditingViewportClient->GetViewLocation(), GCurrentLevelEditingViewportClient->GetViewRotation().Vector() );
diff --git a/Engine/Source/Editor/LevelEditor/Public/LevelEditorActions.h b/Engine/Source/Editor/LevelEditor/Public/LevelEditorActions.h
--- a/Engine/Source/Editor/LevelEditor/Public/LevelEditorActions.h
+++ b/Engine/Source/Editor/LevelEditor/Public/LevelEditorActions.h
@@ -800,16 +800,15 @@ public:
 	 * @param ActorLocation		[opt] If NULL, positions the actor at the mouse location, otherwise the location specified. Default is true.
 	 */
 	static void AddActor_Clicked( UActorFactory* ActorFactory, FAssetData AssetData, bool bUsePlacement );
-	static AActor* AddActor( UActorFactory* ActorFactory, const FAssetData& AssetData, const FVector* ActorLocation );
+	static AActor* AddActor( UActorFactory* ActorFactory, const FAssetData& AssetData, const FTransform* ActorLocation );
 
 	/**
 	 * Called when the AddActor command is executed and a class is selected in the actor browser
 	 *
 	 * @param ActorClass		The class of the actor to add
-	 * @param ActorLocation		[opt] If NULL, positions the actor at the mouse location, otherwise the location specified. Default is true.
 	 */
 	static void AddActorFromClass_Clicked( UClass* ActorClass );
-	static AActor* AddActorFromClass( UClass* ActorClass, const FVector* ActorLocation );
+	static AActor* AddActorFromClass( UClass* ActorClass );
 
 	/**
 	 * Replaces currently selected actors with an actor from the given actor factory
diff --git a/Engine/Source/Editor/Matinee/Private/MatineeTools.cpp b/Engine/Source/Editor/Matinee/Private/MatineeTools.cpp
--- a/Engine/Source/Editor/Matinee/Private/MatineeTools.cpp
+++ b/Engine/Source/Editor/Matinee/Private/MatineeTools.cpp
@@ -167,8 +167,8 @@ void  FMatinee::UpdateCameraRecording (void)
 							ActorToUseForBase = CastChecked< AActor >( SelectedActors.GetSelectedObject( 0 ) );
 						}
 
-						ACameraActor* NewCam = Cast<ACameraActor>(GEditor->AddActor( LevelVC->GetWorld()->GetCurrentLevel(), ACameraActor::StaticClass(), LevelVC->GetViewLocation() ));
-						NewCam->SetActorRotation(LevelVC->GetViewRotation());
+						const FTransform Transform(LevelVC->GetViewRotation(), LevelVC->GetViewLocation());
+						ACameraActor* NewCam = Cast<ACameraActor>(GEditor->AddActor( LevelVC->GetWorld()->GetCurrentLevel(), ACameraActor::StaticClass(), Transform ));
 						NewCam->CameraComponent->bConstrainAspectRatio = LevelVC->IsAspectRatioConstrained();
 						NewCam->CameraComponent->AspectRatio = LevelVC->AspectRatio;
 						NewCam->CameraComponent->FieldOfView = LevelVC->ViewFOV;
diff --git a/Engine/Source/Editor/UnrealEd/Classes/ActorFactories/ActorFactory.h b/Engine/Source/Editor/UnrealEd/Classes/ActorFactories/ActorFactory.h
--- a/Engine/Source/Editor/UnrealEd/Classes/ActorFactories/ActorFactory.h
+++ b/Engine/Source/Editor/UnrealEd/Classes/ActorFactories/ActorFactory.h
@@ -38,8 +38,8 @@ class UNREALED_API UActorFactory : public UObject
 	UPROPERTY()
 	FVector SpawnPositionOffset;
 
-	/** Called to actual create an actor at the supplied location/rotation, using the properties in the ActorFactory */
-	AActor* CreateActor( UObject* Asset, ULevel* InLevel, const FVector& Location, const FRotator* const Rotation, EObjectFlags ObjectFlags = RF_Transactional, const FName Name = NAME_None );
+	/** Called to actual create an actor with the supplied transform (scale is ignored), using the properties in the ActorFactory */
+	AActor* CreateActor( UObject* Asset, ULevel* InLevel, FTransform Transform, EObjectFlags ObjectFlags = RF_Transactional, const FName Name = NAME_None );
 
 	/** Called to create a blueprint class that can be used to spawn an actor from this factory */
 	UBlueprint* CreateBlueprint( UObject* Instance, UObject* Outer, const FName Name, const FName CallingContext = NAME_None );
@@ -58,9 +58,12 @@ class UNREALED_API UActorFactory : public UObject
 	/** Given an instance of an actor pertaining to this factory, find the asset that should be used to create a new actor */
 	virtual UObject* GetAssetFromActorInstance(AActor* ActorInstance);
 
+	/** Return a rotator which aligns this actor type to the specified surface normal */
+	virtual FQuat AlignObjectToSurfaceNormal(const FVector& InSurfaceNormal, const FQuat& ActorRotation = FQuat::Identity) const;
+
 protected:
 
-	virtual bool PreSpawnActor( UObject* Asset, FVector& InOutLocation, FRotator& InOutRotation, bool bRotationWasSupplied);
+	virtual bool PreSpawnActor( UObject* Asset, FTransform& InOutLocation);
 	virtual AActor* SpawnActor( UObject* Asset, ULevel* InLevel, const FVector& Location, const FRotator& Rotation, EObjectFlags ObjectFlags, const FName& Name );
 
 	/** Subclasses may implement this to modify the actor after it has been spawned 
diff --git a/Engine/Source/Editor/UnrealEd/Classes/ActorFactories/ActorFactoryBlueprint.h b/Engine/Source/Editor/UnrealEd/Classes/ActorFactories/ActorFactoryBlueprint.h
--- a/Engine/Source/Editor/UnrealEd/Classes/ActorFactories/ActorFactoryBlueprint.h
+++ b/Engine/Source/Editor/UnrealEd/Classes/ActorFactories/ActorFactoryBlueprint.h
@@ -10,7 +10,7 @@ class UActorFactoryBlueprint : public UActorFactory
 	GENERATED_UCLASS_BODY()
 
 	// Begin UActorFactory Interface
-	virtual bool PreSpawnActor( UObject* Asset, FVector& InOutLocation, FRotator& InOutRotation, bool bRotationWasSupplied ) override;
+	virtual bool PreSpawnActor( UObject* Asset, FTransform& InOutLocation ) override;
 	virtual void PostSpawnActor( UObject* Asset, AActor* NewActor ) override;
 	virtual bool CanCreateActorFrom( const FAssetData& AssetData, FText& OutErrorMsg ) override;
 	virtual AActor* GetDefaultActor( const FAssetData& AssetData ) override;
diff --git a/Engine/Source/Editor/UnrealEd/Classes/ActorFactories/ActorFactoryClass.h b/Engine/Source/Editor/UnrealEd/Classes/ActorFactories/ActorFactoryClass.h
--- a/Engine/Source/Editor/UnrealEd/Classes/ActorFactories/ActorFactoryClass.h
+++ b/Engine/Source/Editor/UnrealEd/Classes/ActorFactories/ActorFactoryClass.h
@@ -10,7 +10,7 @@ class UActorFactoryClass : public UActorFactory
 	GENERATED_UCLASS_BODY()
 
 	// Begin UActorFactory Interface
-	virtual bool PreSpawnActor( UObject* Asset, FVector& InOutLocation, FRotator& InOutRotation, bool bRotationWasSupplied ) override;
+	virtual bool PreSpawnActor( UObject* Asset, FTransform& InOutLocation ) override;
 	virtual void PostSpawnActor( UObject* Asset, AActor* NewActor ) override;
 	virtual bool CanCreateActorFrom( const FAssetData& AssetData, FText& OutErrorMsg ) override;
 	virtual AActor* GetDefaultActor( const FAssetData& AssetData ) override;
diff --git a/Engine/Source/Editor/UnrealEd/Classes/ActorFactories/ActorFactoryDeferredDecal.h b/Engine/Source/Editor/UnrealEd/Classes/ActorFactories/ActorFactoryDeferredDecal.h
--- a/Engine/Source/Editor/UnrealEd/Classes/ActorFactories/ActorFactoryDeferredDecal.h
+++ b/Engine/Source/Editor/UnrealEd/Classes/ActorFactories/ActorFactoryDeferredDecal.h
@@ -12,7 +12,6 @@ public:
 
 protected:
 	// Begin UActorFactory Interface
-	virtual bool PreSpawnActor( UObject* Asset, FVector& InOutLocation, FRotator& InOutRotation, bool bRotationWasSupplied ) override;
 	virtual void PostSpawnActor( UObject* Asset, AActor* NewActor ) override;
 	virtual void PostCreateBlueprint( UObject* Asset, AActor* CDO ) override;
 	virtual bool CanCreateActorFrom( const FAssetData& AssetData, FText& OutErrorMsg ) override;
diff --git a/Engine/Source/Editor/UnrealEd/Classes/ActorFactories/ActorFactoryPhysicsAsset.h b/Engine/Source/Editor/UnrealEd/Classes/ActorFactories/ActorFactoryPhysicsAsset.h
--- a/Engine/Source/Editor/UnrealEd/Classes/ActorFactories/ActorFactoryPhysicsAsset.h
+++ b/Engine/Source/Editor/UnrealEd/Classes/ActorFactories/ActorFactoryPhysicsAsset.h
@@ -10,7 +10,7 @@ class UActorFactoryPhysicsAsset : public UActorFactory
 	GENERATED_UCLASS_BODY()
 
 	// Begin UActorFactory Interface
-	virtual bool PreSpawnActor( UObject* Asset, FVector& InOutLocation, FRotator& InOutRotation, bool bRotationWasSupplied ) override;
+	virtual bool PreSpawnActor( UObject* Asset, FTransform& InOutLocation ) override;
 	virtual void PostSpawnActor( UObject* Asset, AActor* NewActor ) override;
 	virtual void PostCreateBlueprint( UObject* Asset, AActor* CDO ) override;
 	virtual bool CanCreateActorFrom( const FAssetData& AssetData, FText& OutErrorMsg ) override;
diff --git a/Engine/Source/Editor/UnrealEd/Classes/ActorFactories/ActorFactorySkeletalMesh.h b/Engine/Source/Editor/UnrealEd/Classes/ActorFactories/ActorFactorySkeletalMesh.h
--- a/Engine/Source/Editor/UnrealEd/Classes/ActorFactories/ActorFactorySkeletalMesh.h
+++ b/Engine/Source/Editor/UnrealEd/Classes/ActorFactories/ActorFactorySkeletalMesh.h
@@ -14,6 +14,7 @@ protected:
 	virtual void PostSpawnActor( UObject* Asset, AActor* NewActor ) override;
 	virtual void PostCreateBlueprint( UObject* Asset, AActor* CDO ) override;
 	virtual bool CanCreateActorFrom( const FAssetData& AssetData, FText& OutErrorMsg ) override;
+	virtual FQuat AlignObjectToSurfaceNormal(const FVector& InSurfaceNormal, const FQuat& ActorRotation) const override;
 	// End UActorFactory Interface
 
 	virtual USkeletalMesh* GetSkeletalMeshFromAsset( UObject* Asset ) const;
diff --git a/Engine/Source/Editor/UnrealEd/Classes/ActorFactories/ActorFactoryStaticMesh.h b/Engine/Source/Editor/UnrealEd/Classes/ActorFactories/ActorFactoryStaticMesh.h
--- a/Engine/Source/Editor/UnrealEd/Classes/ActorFactories/ActorFactoryStaticMesh.h
+++ b/Engine/Source/Editor/UnrealEd/Classes/ActorFactories/ActorFactoryStaticMesh.h
@@ -14,6 +14,7 @@ class UActorFactoryStaticMesh : public UActorFactory
 	virtual void PostSpawnActor( UObject* Asset, AActor* NewActor) override;
 	virtual void PostCreateBlueprint( UObject* Asset, AActor* CDO ) override;
 	virtual UObject* GetAssetFromActorInstance(AActor* ActorInstance) override;
+	virtual FQuat AlignObjectToSurfaceNormal(const FVector& InSurfaceNormal, const FQuat& ActorRotation) const override;
 	// End UActorFactory Interface
 };
 
diff --git a/Engine/Source/Editor/UnrealEd/Classes/Editor/EditorEngine.h b/Engine/Source/Editor/UnrealEd/Classes/Editor/EditorEngine.h
--- a/Engine/Source/Editor/UnrealEd/Classes/Editor/EditorEngine.h
+++ b/Engine/Source/Editor/UnrealEd/Classes/Editor/EditorEngine.h
@@ -956,12 +956,12 @@ public:
 	 *
 	 * @param	InLevel			Level in which to add the actor
 	 * @param	Class			A non-abstract, non-transient, placeable class.  Must be non-NULL.
-	 * @param	Location		The world-space location to spawn the actor.
+	 * @param	Transform		The world-space transform to spawn the actor with.
 	 * @param	bSilent			If true, suppress logging (optional, defaults to false).
 	 * @param	ObjectFlags		The object flags to place on the spawned actor.
 	 * @result					A pointer to the newly added actor, or NULL if add failed.
 	 */
-	virtual AActor* AddActor(ULevel* InLevel, UClass* Class, const FVector& Location, bool bSilent = false, EObjectFlags ObjectFlags = RF_Transactional);
+	virtual AActor* AddActor(ULevel* InLevel, UClass* Class, const FTransform& Transform, bool bSilent = false, EObjectFlags ObjectFlags = RF_Transactional);
 
 	/**
 	 * Adds actors to the world at the specified location using export text.
@@ -1055,12 +1055,10 @@ public:
 	 * Uses the supplied factory to create an actor at the clicked location and adds to level.
 	 *
 	 * @param	Factory					The factory to create the actor from.  Must be non-NULL.
-	 * @param	ActorLocation			[opt] If null, positions the actor at the mouse location, otherwise specified. Default is null.
-	 * @param	bUseSurfaceOrientation	[opt] If true, align new actor's orientation to the underlying surface normal.  Default is false.
 	 * @param	ObjectFlags				[opt] The flags to apply to the actor when it is created
 	 * @return							A pointer to the new actor, or NULL on fail.
 	 */
-	AActor* UseActorFactoryOnCurrentSelection( UActorFactory* Factory, const FVector* ActorLocation=NULL, bool bUseSurfaceOrientation=false, EObjectFlags ObjectFlags = RF_Transactional );
+	AActor* UseActorFactoryOnCurrentSelection( UActorFactory* Factory, const FTransform* InActorTransform, EObjectFlags ObjectFlags = RF_Transactional );
 
 	/**
 	 * Uses the supplied factory to create an actor at the clicked location and adds to level.
@@ -1072,7 +1070,7 @@ public:
 	 * @param	ObjectFlags				[opt] The flags to apply to the actor when it is created
 	 * @return							A pointer to the new actor, or NULL on fail.
 	 */
-	AActor* UseActorFactory( UActorFactory* Factory, const FAssetData& AssetData, const FVector* ActorLocation=NULL, bool bUseSurfaceOrientation=false, EObjectFlags ObjectFlags = RF_Transactional );
+	AActor* UseActorFactory( UActorFactory* Factory, const FAssetData& AssetData, const FTransform* ActorLocation, EObjectFlags ObjectFlags = RF_Transactional );
 
 	/**
 	 * Replaces the selected Actors with the same number of a different kind of Actor
diff --git a/Engine/Source/Editor/UnrealEd/Classes/Settings/LevelEditorViewportSettings.h b/Engine/Source/Editor/UnrealEd/Classes/Settings/LevelEditorViewportSettings.h
--- a/Engine/Source/Editor/UnrealEd/Classes/Settings/LevelEditorViewportSettings.h
+++ b/Engine/Source/Editor/UnrealEd/Classes/Settings/LevelEditorViewportSettings.h
@@ -158,6 +158,31 @@ struct UNREALED_API FLevelEditorViewportInstanceSettingsKeyValuePair
 };
 
 
+/**
+ * Settings that control the behavior of the "snap to surface" feature
+ */
+USTRUCT()
+struct UNREALED_API FSnapToSurfaceSettings
+{
+	GENERATED_USTRUCT_BODY()
+
+	FSnapToSurfaceSettings() : bEnabled(false), SnapOffsetExtent(0.f), bSnapRotation(true)
+	{}
+
+	/** Whether snapping to surfaces in the world is enabled */
+	UPROPERTY(config)
+	bool bEnabled;
+
+	/** The amount of offset to apply when snapping to surfaces */
+	UPROPERTY(config)
+	float SnapOffsetExtent;
+
+	/** Whether objects should match the rotation of the surfaces they snap to */
+	UPROPERTY(config)
+	bool bSnapRotation;
+};
+
+
 /**
  * Implements the Level Editor's view port settings.
  */
@@ -279,9 +304,9 @@ public:
 	UPROPERTY(EditAnywhere, config, Category=GridSnapping, meta=(DisplayName = "Enable Scale Snapping"))
 	uint32 SnapScaleEnabled:1;
 
-	/** If enabled the when dragging new objects out of the content browser, it will snap the objects Z coordinate to the floor below it (if any) instead of the Z grid snapping location */
-	UPROPERTY(EditAnywhere, config, Category=GridSnapping)
-	bool bSnapNewObjectsToFloor;
+	/** If enabled, actors will snap to surfaces in the viewport when dragged around */
+	UPROPERTY(config)
+	FSnapToSurfaceSettings SnapToSurface;
 
 private:
 
diff --git a/Engine/Source/Editor/UnrealEd/Private/AssetSelection.cpp b/Engine/Source/Editor/UnrealEd/Private/AssetSelection.cpp
--- a/Engine/Source/Editor/UnrealEd/Private/AssetSelection.cpp
+++ b/Engine/Source/Editor/UnrealEd/Private/AssetSelection.cpp
@@ -21,9 +21,9 @@
 #include "Landscape/Landscape.h"
 
 #include "Particles/Emitter.h"
+#include "Editor/ActorPositioning.h"
 #include "Animation/SkeletalMeshActor.h"
 
-extern void OnPlaceStaticMeshActor( AActor* MeshActor, bool bUseSurfaceOrientation );
 
 namespace AssetSelectionUtils
 {
@@ -351,103 +351,65 @@ namespace AssetSelectionUtils
 *
 * Does nothing if ActorClass is NULL.
 */
-static AActor* PrivateAddActor( UObject* Asset, UActorFactory* Factory, const FVector* ActorLocation=NULL, bool bUseSurfaceOrientation=false, bool SelectActor = true, EObjectFlags ObjectFlags = RF_Transactional, const FName Name = NAME_None )
+static AActor* PrivateAddActor( UObject* Asset, UActorFactory* Factory, bool SelectActor = true, EObjectFlags ObjectFlags = RF_Transactional, const FName Name = NAME_None )
 {
-	AActor* Actor = NULL;
-	if ( Factory )
+	if (!Factory)
 	{
-		AActor* NewActorTemplate = Factory->GetDefaultActor( Asset );
-		if ( !NewActorTemplate )
-		{
-			return NULL;
-		}
-
-		// Position the actor relative to the mouse?
-		FVector Location = FVector::ZeroVector;
-		FVector SnapNormal = FVector::ZeroVector;
-		if ( ActorLocation )
-		{
-			Location = *ActorLocation;
-		}
-		else
-		{
-			FSnappingUtils::SnapPointToGrid( GEditor->ClickLocation, FVector(0, 0, 0) );
+		return nullptr;
+	}
 
-			Location = GEditor->ClickLocation;
-			if( !GCurrentLevelEditingViewportClient || !FSnappingUtils::SnapLocationToNearestVertex( Location, GCurrentLevelEditingViewportClient->GetDropPreviewLocation(), GCurrentLevelEditingViewportClient, SnapNormal ) )
-			{
-				FVector Collision = NewActorTemplate->GetPlacementExtent();
+	AActor* Actor = NULL;
+	AActor* NewActorTemplate = Factory->GetDefaultActor( Asset );
+	if ( !NewActorTemplate )
+	{
+		return NULL;
+	}
 
-				Location += GEditor->ClickPlane * (FVector::BoxPushOut(GEditor->ClickPlane, Collision) + 0.1f);
-			}
+	const FTransform ActorTransform = FActorPositioning::GetSnappedSurfaceAlignedTransform(GCurrentLevelEditingViewportClient, Factory, GEditor->ClickLocation, GEditor->ClickPlane, NewActorTemplate->GetPlacementExtent());
 
-			// Do not fade snapping indicators over time if the viewport is not realtime
-			bool bClearImmediately = !GCurrentLevelEditingViewportClient || !GCurrentLevelEditingViewportClient->IsRealtime();
-			FSnappingUtils::ClearSnappingHelpers( bClearImmediately );
+	// Do not fade snapping indicators over time if the viewport is not realtime
+	bool bClearImmediately = !GCurrentLevelEditingViewportClient || !GCurrentLevelEditingViewportClient->IsRealtime();
+	FSnappingUtils::ClearSnappingHelpers( bClearImmediately );
 
-			FSnappingUtils::SnapPointToGrid( Location, FVector(0, 0, 0) );
-		}		
+	ULevel* DesiredLevel = GWorld->GetCurrentLevel();
 
-		const FRotator* Rotation = NULL;
-		if( bUseSurfaceOrientation )
-		{
-			FRotator SurfaceOrientation;
-			if( SnapNormal != FVector::ZeroVector )
-			{
-				// Orient the new actor with the snapped surface normal?
-				SurfaceOrientation = SnapNormal.Rotation();
-			}
-			else
-			{
-				SurfaceOrientation =  GEditor->ClickPlane.Rotation();
-			}
-
-			Rotation = &SurfaceOrientation;
-		}
-
-		ULevel* DesiredLevel = GWorld->GetCurrentLevel();
+	// Don't spawn the actor if the current level is locked.
+	if ( FLevelUtils::IsLevelLocked(DesiredLevel) )
+	{
+		FNotificationInfo Info( NSLOCTEXT("UnrealEd", "Error_OperationDisallowedOnLockedLevel", "The requested operation could not be completed because the level is locked.") );
+		Info.ExpireDuration = 3.0f;
+		FSlateNotificationManager::Get().AddNotification(Info);
+		return NULL;
+	}
 
-		// Don't spawn the actor if the current level is locked.
-		if ( FLevelUtils::IsLevelLocked(DesiredLevel) )
+	{
+		FScopedTransaction Transaction( NSLOCTEXT("UnrealEd", "CreateActor", "Create Actor") );
+		if ( !(ObjectFlags & RF_Transactional) )
 		{
-			FNotificationInfo Info( NSLOCTEXT("UnrealEd", "Error_OperationDisallowedOnLockedLevel", "The requested operation could not be completed because the level is locked.") );
-			Info.ExpireDuration = 3.0f;
-			FSlateNotificationManager::Get().AddNotification(Info);
-			return NULL;
+			Transaction.Cancel();
 		}
 
+		// Create the actor.
+		Actor = Factory->CreateActor( Asset, DesiredLevel, ActorTransform, ObjectFlags, Name );
+		if(Actor)
 		{
-			FScopedTransaction Transaction( NSLOCTEXT("UnrealEd", "CreateActor", "Create Actor") );
-			if ( !(ObjectFlags & RF_Transactional) )
+			if ( SelectActor )
 			{
-				Transaction.Cancel();
+				GEditor->SelectNone( false, true );
+				GEditor->SelectActor( Actor, true, true );
 			}
 
-			// Create the actor.
-			Actor = Factory->CreateActor( Asset, DesiredLevel, Location, Rotation, ObjectFlags, Name ); 
-			if(Actor)
-			{
-				// Apply any static mesh tool settings if we placed a static mesh. 
-				OnPlaceStaticMeshActor( Actor, bUseSurfaceOrientation );
-
-				if ( SelectActor )
-				{
-					GEditor->SelectNone( false, true );
-					GEditor->SelectActor( Actor, true, true );
-				}
-
-				Actor->InvalidateLightingCache();
-				Actor->PostEditChange();
-			}
+			Actor->InvalidateLightingCache();
+			Actor->PostEditChange();
 		}
+	}
 
-		GEditor->RedrawLevelEditingViewports();
+	GEditor->RedrawLevelEditingViewports();
 
-		if ( Actor )
-		{
-			Actor->MarkPackageDirty();
-			ULevel::LevelDirtiedEvent.Broadcast();
-		}
+	if ( Actor )
+	{
+		Actor->MarkPackageDirty();
+		ULevel::LevelDirtiedEvent.Broadcast();
 	}
 
 	return Actor;
@@ -601,7 +563,7 @@ UActorFactory* FActorFactoryAssetProxy::GetFactoryForAssetObject( UObject* Asset
 	return Result;
 }
 
-AActor* FActorFactoryAssetProxy::AddActorForAsset( UObject* AssetObj, const FVector* ActorLocation, bool bUseSurfaceOrientation, bool SelectActor, EObjectFlags ObjectFlags, UActorFactory* FactoryToUse /*= NULL*/, const FName Name )
+AActor* FActorFactoryAssetProxy::AddActorForAsset( UObject* AssetObj, bool SelectActor, EObjectFlags ObjectFlags, UActorFactory* FactoryToUse /*= NULL*/, const FName Name )
 {
 	AActor* Result = NULL;
 
@@ -614,7 +576,7 @@ AActor* FActorFactoryAssetProxy::AddActorForAsset( UObject* AssetObj, const FVec
 		{
 			if ( FactoryToUse->CanCreateActorFrom( AssetData, UnusedErrorMessage ) )
 			{
-				Result = PrivateAddActor( AssetObj, FactoryToUse, ActorLocation, bUseSurfaceOrientation, SelectActor, ObjectFlags, Name );
+				Result = PrivateAddActor( AssetObj, FactoryToUse, SelectActor, ObjectFlags, Name );
 			}
 		}
 		// If no specific factory has been provided, find the highest priority one that is valid for the asset and use
@@ -629,7 +591,7 @@ AActor* FActorFactoryAssetProxy::AddActorForAsset( UObject* AssetObj, const FVec
 				// Check if the actor can be created using this factory, making sure to check for an asset to be assigned from the selector
 				if ( ActorFactory->CanCreateActorFrom( AssetData, UnusedErrorMessage ) )
 				{
-					Result = PrivateAddActor(AssetObj, ActorFactory, ActorLocation, bUseSurfaceOrientation, SelectActor, ObjectFlags, Name);
+					Result = PrivateAddActor(AssetObj, ActorFactory, SelectActor, ObjectFlags, Name);
 					if ( Result != NULL )
 					{
 						break;
@@ -643,7 +605,7 @@ AActor* FActorFactoryAssetProxy::AddActorForAsset( UObject* AssetObj, const FVec
 	return Result;
 }
 
-AActor* FActorFactoryAssetProxy::AddActorFromSelection( UClass* ActorClass, const FVector* ActorLocation, bool bUseSurfaceOrientation, bool SelectActor, EObjectFlags ObjectFlags, UActorFactory* ActorFactory, const FName Name )
+AActor* FActorFactoryAssetProxy::AddActorFromSelection( UClass* ActorClass, const FVector* ActorLocation, bool SelectActor, EObjectFlags ObjectFlags, UActorFactory* ActorFactory, const FName Name )
 {
 	check( ActorClass != NULL );
 
@@ -663,7 +625,7 @@ AActor* FActorFactoryAssetProxy::AddActorFromSelection( UClass* ActorClass, cons
 		if( ActorFactory->CanCreateActorFrom( FAssetData(TargetObject), ErrorMessage ) )
 		{
 			// Attempt to add the actor
-			Result = PrivateAddActor( TargetObject, ActorFactory, ActorLocation, bUseSurfaceOrientation, SelectActor, ObjectFlags );
+			Result = PrivateAddActor( TargetObject, ActorFactory, SelectActor, ObjectFlags );
 		}
 	}
 
diff --git a/Engine/Source/Editor/UnrealEd/Private/Commandlets/DumpBlueprintsInfoCommandlet.cpp b/Engine/Source/Editor/UnrealEd/Private/Commandlets/DumpBlueprintsInfoCommandlet.cpp
--- a/Engine/Source/Editor/UnrealEd/Private/Commandlets/DumpBlueprintsInfoCommandlet.cpp
+++ b/Engine/Source/Editor/UnrealEd/Private/Commandlets/DumpBlueprintsInfoCommandlet.cpp
@@ -598,8 +598,6 @@ static AActor* DumpBlueprintInfoUtils::SpawnLevelActor(UClass* ActorClass, bool
 
 		UBlueprint* ActorTemplate = MakeTempBlueprint(ActorClass);
 		SpawnedActor = FActorFactoryAssetProxy::AddActorForAsset(ActorTemplate,
-			/*ActorLocation =*/nullptr,
-			/*bUseSurfaceOrientation =*/false,
 			/*SelectActor =*/bSelect,
 			RF_Transient, NewFactory, NAME_None);
 	}
diff --git a/Engine/Source/Editor/UnrealEd/Private/Editor.cpp b/Engine/Source/Editor/UnrealEd/Private/Editor.cpp
--- a/Engine/Source/Editor/UnrealEd/Private/Editor.cpp
+++ b/Engine/Source/Editor/UnrealEd/Private/Editor.cpp
@@ -44,6 +44,8 @@
 #include "LevelEditor.h"
 #include "SCreateAssetFromActor.h"
 
+#include "Editor/ActorPositioning.h"
+
 #include "Developer/DirectoryWatcher/Public/DirectoryWatcherModule.h"
 
 #include "Runtime/Engine/Public/Slate/SceneViewport.h"
@@ -2160,16 +2162,6 @@ void UEditorEngine::ApplyDeltaToActor(AActor* InActor,
 					InActor->SetActorRotation( InDeltaRot );
 				}
 			}
-
-			if ( bDelta )
-			{
-				FVector NewActorLocation = InActor->GetActorLocation();
-				NewActorLocation -= GLevelEditorModeTools().PivotLocation;
-				NewActorLocation = FRotationMatrix( InDeltaRot ).TransformPosition( NewActorLocation );
-				NewActorLocation += GLevelEditorModeTools().PivotLocation;
-				NewActorLocation -= InActor->GetActorLocation();
-				InActor->EditorApplyTranslation( NewActorLocation, bAltDown, bShiftDown, bControlDown );
-			}
 		}
 	}
 
@@ -4644,46 +4636,14 @@ FString UEditorEngine::GetFriendlyName( const UProperty* Property, UStruct* Owne
 	return FoundText.ToString();
 }
 
-AActor* UEditorEngine::UseActorFactoryOnCurrentSelection( UActorFactory* Factory, const FVector* ActorLocation, bool bUseSurfaceOrientation, EObjectFlags ObjectFlags )
+AActor* UEditorEngine::UseActorFactoryOnCurrentSelection( UActorFactory* Factory, const FTransform* InActorTransform, EObjectFlags ObjectFlags )
 {
 	// ensure that all selected assets are loaded
 	FEditorDelegates::LoadSelectedAssetsIfNeeded.Broadcast();
-	return UseActorFactory(Factory, FAssetData( GetSelectedObjects()->GetTop<UObject>() ), ActorLocation, bUseSurfaceOrientation, ObjectFlags );
-}
-
-/**
- * Corrects the orientation if the mesh if placed on a mesh 
- * and applies static mesh tool settings if applicable. 
- *
- * @param	Actor					The static mesh or speed tree actor that was placed.
- * @param	bUseSurfaceOrientation	When true, changes the mesh orientation to reflect the surface.
- */
-void OnPlaceStaticMeshActor( AActor* Actor, bool bUseSurfaceOrientation )
-{
-	check( Actor );
-
-	if( Actor->IsA(AStaticMeshActor::StaticClass())  && bUseSurfaceOrientation )
-	{
-		AStaticMeshActor* MeshActor = CastChecked<AStaticMeshActor>(Actor);
-		if( bUseSurfaceOrientation )
-		{
-			MeshActor->SetActorRotation(GEditor->ClickPlane.SafeNormal().Rotation());
-
-			// This is necessary because static meshes are authored along the vertical axis rather than
-			// the X axis, as the surface orientation code expects.  To compensate for this, we add 90
-			// degrees to the static mesh's Pitch.
-
-			FRotator Rot = MeshActor->GetActorRotation();
-			Rot.Pitch -= 90.f;
-			MeshActor->SetActorRotation(Rot);
-		}
-
-		MeshActor->UpdateComponentTransforms();
-	}
+	return UseActorFactory(Factory, FAssetData( GetSelectedObjects()->GetTop<UObject>() ), InActorTransform, ObjectFlags );
 }
 
-
-AActor* UEditorEngine::UseActorFactory( UActorFactory* Factory, const FAssetData& AssetData, const FVector* ActorLocation, bool bUseSurfaceOrientation, EObjectFlags ObjectFlags )
+AActor* UEditorEngine::UseActorFactory( UActorFactory* Factory, const FAssetData& AssetData, const FTransform* InActorTransform, EObjectFlags ObjectFlags )
 {
 	check( Factory );
 
@@ -4711,47 +4671,7 @@ AActor* UEditorEngine::UseActorFactory( UActorFactory* Factory, const FAssetData
 			return NULL;
 		}
 
-		FVector Collision = NewActorTemplate->GetPlacementExtent();
-
-		// Get cursor origin and direction in world space.
-		FViewportCursorLocation CursorLocation = GCurrentLevelEditingViewportClient->GetCursorWorldLocationFromMousePos();
-
-		// Position the actor relative to the mouse?
-		FVector Location = FVector::ZeroVector;
-		HHitProxy* HitProxy = NULL;
-		if ( ActorLocation )
-		{
-			Location = *ActorLocation;
-		}
-		else
-		{
-			FSnappingUtils::SnapPointToGrid( ClickLocation, FVector(0, 0, 0) );
-
-			float DistanceMultiplier = 1.0f;
-			if( CursorLocation.GetViewportType() != LVT_Perspective )
-			{
-				// If an asset was dropped in an orthographic view,
-				// adjust the location so the asset will not be positioned at the camera origin (which is at the edge of the world for ortho cams)
-				DistanceMultiplier = HALF_WORLD_MAX;
-			}
-
-			const FIntPoint CursorPos = CursorLocation.GetCursorPos();
-			HitProxy = GCurrentLevelEditingViewportClient->Viewport->GetHitProxy( CursorPos.X, CursorPos.Y );
-			if( HitProxy == NULL )
-			{
-				// If the hit proxy is null, we clicked on the background so we need to calculate our own click location since the old one will be off at the worlds edge somewhere.
-				ClickLocation = CursorLocation.GetOrigin() + CursorLocation.GetDirection() * DistanceMultiplier;
-			}
-
-			// Calculate the new actors location from the mouse click location and collision information.
-			Location = ClickLocation + ClickPlane * (FVector::BoxPushOut(ClickPlane, Collision) + 0.1);
-
-			FSnappingUtils::SnapPointToGrid( Location, FVector(0, 0, 0) );
-		}
-
-		// Orient the new actor with the surface normal?
-		const FRotator SurfaceOrientation( ClickPlane.Rotation() );
-		const FRotator* const Rotation = bUseSurfaceOrientation ? &SurfaceOrientation : NULL;
+		const FTransform ActorTransform = InActorTransform ? *InActorTransform : FActorPositioning::GetCurrentViewportPlacementTransform(*NewActorTemplate);
 
 		ULevel* DesiredLevel = GWorld->GetCurrentLevel();
 
@@ -4765,18 +4685,9 @@ AActor* UEditorEngine::UseActorFactory( UActorFactory* Factory, const FAssetData
 				const FScopedTransaction Transaction( NSLOCTEXT("UnrealEd", "CreateActor", "Create Actor") );
 
 				// Create the actor.
-				Actor = Factory->CreateActor( Asset, DesiredLevel, Location, Rotation, ObjectFlags );
+				Actor = Factory->CreateActor( Asset, DesiredLevel, ActorTransform, ObjectFlags );
 				if(Actor != NULL)
 				{
-					// Apply any static mesh tool settings if we placed a static mesh. 
-					OnPlaceStaticMeshActor( Actor, bUseSurfaceOrientation );
-
-					if( !ActorLocation && CursorLocation.GetViewportType() == LVT_Perspective && HitProxy == NULL  )
-					{
-						// Move the actor in front of the camera if the cursor location is in the perspective viewport and we dont have a valid hit proxy (which would place the actor somewhere else)
-						MoveActorInFrontOfCamera( *Actor, CursorLocation.GetOrigin(), CursorLocation.GetDirection() );
-					}
-
 					SelectNone( false, true );
 					SelectActor( Actor, true, true );
 					Actor->InvalidateLightingCache();
@@ -5070,13 +4981,12 @@ void UEditorEngine::ReplaceSelectedActors(UActorFactory* Factory, const FAssetDa
 		ULevel* Level = OldActor->GetLevel();
 		AActor* NewActor = NULL;
 
-		const FVector OldLocation = OldActor->GetActorLocation();
-		const FRotator OldRotation = OldActor->GetActorRotation();
+		const FTransform OldTransform = OldActor->ActorToWorld();
 
 		// create the actor
 		if (Factory != NULL)
 		{
-			NewActor = Factory->CreateActor( Asset, Level, OldLocation, &OldRotation);
+			NewActor = Factory->CreateActor( Asset, Level, OldTransform);
 			// For blueprints, try to copy over properties
 			if (Factory->IsA(UActorFactoryBlueprint::StaticClass()))
 			{
@@ -5095,7 +5005,10 @@ void UEditorEngine::ReplaceSelectedActors(UActorFactory* Factory, const FAssetDa
 		{
 			FActorSpawnParameters SpawnInfo;
 			SpawnInfo.OverrideLevel = Level;
-			NewActor = World->SpawnActor( NewActorClass, &OldLocation, &OldRotation,	SpawnInfo );
+
+			const auto Rotation = OldTransform.GetRotation().Rotator();
+			const auto Translation = OldTransform.GetTranslation();
+			NewActor = World->SpawnActor( NewActorClass, &Translation, &Rotation, SpawnInfo );
 		}
 		if ( NewActor != NULL )
 		{
@@ -6038,12 +5951,13 @@ bool UEditorEngine::AreAllWindowsHidden() const
 	return bAllHidden;
 }
 
-AActor* UEditorEngine::AddActor(ULevel* InLevel, UClass* Class, const FVector& Location, bool bSilent, EObjectFlags ObjectFlags)
+AActor* UEditorEngine::AddActor(ULevel* InLevel, UClass* Class, const FTransform& Transform, bool bSilent, EObjectFlags ObjectFlags)
 {
 	check( Class );
 
 	if( !bSilent )
 	{
+		const auto Location = Transform.GetLocation();
 		UE_LOG(LogEditor, Log,
 			TEXT("Attempting to add actor of class '%s' to level at %0.2f,%0.2f,%0.2f"),
 			*Class->GetName(), Location.X, Location.Y, Location.Z );
@@ -6098,7 +6012,9 @@ AActor* UEditorEngine::AddActor(ULevel* InLevel, UClass* Class, const FVector& L
 		SpawnInfo.OverrideLevel = DesiredLevel;
 		SpawnInfo.bNoCollisionFail = true;
 		SpawnInfo.ObjectFlags = ObjectFlags;
-		Actor = World->SpawnActor( Class, &Location, NULL, SpawnInfo );
+		const auto Location = Transform.GetLocation();
+		const auto Rotation = Transform.GetRotation().Rotator();
+		Actor = World->SpawnActor( Class, &Location, &Rotation, SpawnInfo );
 
 		if( Actor )
 		{
diff --git a/Engine/Source/Editor/UnrealEd/Private/Editor/ActorPositioning.cpp b/Engine/Source/Editor/UnrealEd/Private/Editor/ActorPositioning.cpp
new file mode 100644
--- /dev/null
+++ b/Engine/Source/Editor/UnrealEd/Private/Editor/ActorPositioning.cpp
@@ -0,0 +1,215 @@
+// Copyright 1998-2014 Epic Games, Inc. All Rights Reserved.
+
+#pragma once
+
+#include "UnrealEd.h"
+#include "ActorPositioning.h"
+#include "SnappingUtils.h"
+#include "Landscape/LandscapeHeightfieldCollisionComponent.h"
+#include "Landscape/LandscapeComponent.h"
+
+FActorPositionTraceResult FActorPositioning::TraceWorldForPositionWithDefault(const FViewportCursorLocation& Cursor, const FSceneView& View, const TArray<AActor*>* IgnoreActors)
+{
+	FActorPositionTraceResult Results = TraceWorldForPosition(Cursor, View, IgnoreActors);
+	if (Results.State == FActorPositionTraceResult::Failed)
+	{
+		Results.State = FActorPositionTraceResult::Default;
+
+		// And put it in front of the camera
+		const float DistanceMultiplier = ( Cursor.GetViewportType() == LVT_Perspective ) ? GetDefault<ULevelEditorViewportSettings>()->BackgroundDropDistance : 0.0f;
+		Results.Location = Cursor.GetOrigin() + Cursor.GetDirection() * DistanceMultiplier;
+	}
+	return Results;
+}
+
+FActorPositionTraceResult FActorPositioning::TraceWorldForPosition(const FViewportCursorLocation& Cursor, const FSceneView& View, const TArray<AActor*>* IgnoreActors)
+{
+	const auto* ViewportClient = Cursor.GetViewportClient();
+
+	FActorPositionTraceResult Results;
+
+	const auto ViewportType = ViewportClient->GetViewportType();
+
+	// Start with a ray that encapsulates the entire world
+	FVector RayStart = Cursor.GetOrigin();
+	if (ViewportType == LVT_OrthoXY || ViewportType == LVT_OrthoXZ || ViewportType == LVT_OrthoYZ)
+	{
+		RayStart -= Cursor.GetDirection() * HALF_WORLD_MAX/2;
+	}
+
+	const FVector RayEnd = RayStart + Cursor.GetDirection() * HALF_WORLD_MAX;
+	return TraceWorldForPosition(*ViewportClient->GetWorld(), View, RayStart, RayEnd, IgnoreActors);
+}
+
+/** Check to see if the specified hit result should be ignored from actor positioning calculations for the specified scene view */
+bool IsHitIgnored(const FHitResult& InHit, const FSceneView& InSceneView)
+{
+	const auto* Actor = InHit.GetActor();
+	
+	// Try and find a primitive component for the hit
+	const UPrimitiveComponent* PrimitiveComponent = Actor ? Actor->GetRootPrimitiveComponent() : nullptr;
+
+	if (!PrimitiveComponent)
+	{
+		PrimitiveComponent = InHit.Component.Get();
+	}
+	if (PrimitiveComponent && PrimitiveComponent->IsA(ULandscapeHeightfieldCollisionComponent::StaticClass()))
+	{
+		PrimitiveComponent = CastChecked<ULandscapeHeightfieldCollisionComponent>(PrimitiveComponent)->RenderComponent.Get();
+	}
+
+	if (InHit.bStartPenetrating || !PrimitiveComponent)
+	{
+		return true;
+	}
+
+	// Ignore volumes and shapes
+	if (Actor && Actor->IsA(AVolume::StaticClass()))
+	{
+		return true;
+	}
+	else if (PrimitiveComponent->IsA(UShapeComponent::StaticClass()))
+	{
+		return true;
+	}
+
+	// Only use this component if it is visible in the specified scene views
+	bool bIsRenderedOnScreen = false;
+	{				
+		if (PrimitiveComponent && PrimitiveComponent->SceneProxy)
+		{
+			const FPrimitiveViewRelevance ViewRelevance = PrimitiveComponent->SceneProxy->GetViewRelevance(&InSceneView);
+			// BSP is a bit special in that its bDrawRelevance is false even when drawn as wireframe because InSceneView.Family->EngineShowFlags.BSPTriangles is off
+			bIsRenderedOnScreen = ViewRelevance.bDrawRelevance || (PrimitiveComponent->IsA(UModelComponent::StaticClass()) && InSceneView.Family->EngineShowFlags.BSP);
+		}
+	}
+
+	return !bIsRenderedOnScreen;
+}
+
+FActorPositionTraceResult FActorPositioning::TraceWorldForPosition(const UWorld& InWorld, const FSceneView& InSceneView, const FVector& RayStart, const FVector& RayEnd, const TArray<AActor*>* IgnoreActors)
+{
+	TArray<FHitResult> Hits;
+
+	FCollisionQueryParams Param(TEXT("DragDropTrace"), true);
+	if (IgnoreActors)
+	{
+		Param.AddIgnoredActors(*IgnoreActors);
+	}
+
+	FActorPositionTraceResult Results;
+	if ( InWorld.LineTraceMulti(Hits, RayStart, RayEnd, Param, FCollisionObjectQueryParams(FCollisionObjectQueryParams::InitType::AllObjects)) )
+	{
+		{
+			// Filter out anything that should be ignored
+			FSuspendRenderingThread SuspendRendering(false);
+			Hits.RemoveAll([&](const FHitResult& Hit){
+				return IsHitIgnored(Hit, InSceneView);
+			});
+		}
+
+		// Go through all hits and find closest
+		float ClosestHitDistance = TNumericLimits<float>::Max();
+
+		for (const auto& Hit : Hits)
+		{
+			const float DistanceToHit = (Hit.ImpactPoint - RayStart).Size();
+			if (DistanceToHit < ClosestHitDistance)
+			{
+				ClosestHitDistance = DistanceToHit;
+				Results.Location = Hit.Location;
+				Results.SurfaceNormal = Hit.Normal.SafeNormal();
+				Results.State = FActorPositionTraceResult::HitSuccess;
+			}
+		}
+	}
+
+	return Results;
+}
+
+FTransform FActorPositioning::GetCurrentViewportPlacementTransform(const AActor& Actor, bool bSnap)
+{
+	FVector Collision = Actor.GetPlacementExtent();
+	const UActorFactory* Factory = GEditor->FindActorFactoryForActorClass(Actor.GetClass());
+
+	// Get cursor origin and direction in world space.
+	FViewportCursorLocation CursorLocation = GCurrentLevelEditingViewportClient->GetCursorWorldLocationFromMousePos();
+	const auto CursorPos = CursorLocation.GetCursorPos();
+
+	FTransform ActorTransform = FTransform::Identity;
+
+	if (CursorLocation.GetViewportType() == LVT_Perspective && !GCurrentLevelEditingViewportClient->Viewport->GetHitProxy( CursorPos.X, CursorPos.Y ))
+	{
+		ActorTransform.SetTranslation(FActorPositioning::GetActorPositionInFrontOfCamera(Actor, CursorLocation.GetOrigin(), CursorLocation.GetDirection()));
+	}
+	else if (bSnap)
+	{
+		ActorTransform = FActorPositioning::GetSnappedSurfaceAlignedTransform(GCurrentLevelEditingViewportClient, Factory, GEditor->ClickLocation, GEditor->ClickPlane, Actor.GetPlacementExtent());
+	}
+	else
+	{
+		ActorTransform = FActorPositioning::GetSurfaceAlignedTransform(Factory, GEditor->ClickLocation, GEditor->ClickPlane, Actor.GetPlacementExtent());
+	}
+
+	return ActorTransform;
+}
+
+FVector FActorPositioning::GetActorPositionInFrontOfCamera(const AActor& InActor, const FVector& InCameraOrigin, const FVector& InCameraDirection)
+{
+	// Get the  radius of the actors bounding cylinder.  Height is not needed.
+	float CylRadius, CylHeight;
+	InActor.GetComponentsBoundingCylinder(CylRadius, CylHeight);
+
+	// a default cylinder radius if no bounding cylinder exists.  
+	const float	DefaultCylinderRadius = 50.0f;
+
+	if( CylRadius == 0.0f )
+	{
+		// If the actor does not have a bounding cylinder, use a default value.
+		CylRadius = DefaultCylinderRadius;
+	}
+
+	// The new location the cameras origin offset by the actors bounding cylinder radius down the direction of the cameras view. 
+	FVector NewLocation = InCameraOrigin + InCameraDirection * CylRadius + InCameraDirection * GetDefault<ULevelEditorViewportSettings>()->BackgroundDropDistance;
+
+	// Snap the new location if snapping is enabled
+	FSnappingUtils::SnapPointToGrid( NewLocation, FVector::ZeroVector );
+	return NewLocation;
+}
+
+FTransform FActorPositioning::GetSurfaceAlignedTransform(const UActorFactory* InFactory, const FVector& InSurfaceLocation, const FVector& SurfaceNormal, const FVector& InPlacementExtent, const FTransform& InStartTransform)
+{
+	// Sort out the rotation first, then do the location
+	FQuat RotatorQuat = InStartTransform.GetRotation();
+	if (InFactory && GetDefault<ULevelEditorViewportSettings>()->SnapToSurface.bSnapRotation)
+	{
+		RotatorQuat = InFactory->AlignObjectToSurfaceNormal(SurfaceNormal, RotatorQuat);
+	}
+
+	// Choose the largest location offset of the various options (global viewport settings, collision, factory offset)
+	const float SnapOffsetExtent = GetDefault<ULevelEditorViewportSettings>()->SnapToSurface.SnapOffsetExtent;
+	const float CollisionOffsetExtent = FVector::BoxPushOut(SurfaceNormal, InPlacementExtent);
+
+	FVector LocationOffset = SurfaceNormal * FMath::Max(SnapOffsetExtent, CollisionOffsetExtent);
+	if (InFactory && LocationOffset.SizeSquared() < InFactory->SpawnPositionOffset.SizeSquared())
+	{
+		// Rotate the Spawn Position Offset to match our rotation
+		LocationOffset = RotatorQuat.RotateVector(-InFactory->SpawnPositionOffset);
+	}
+
+	return FTransform(RotatorQuat, InSurfaceLocation + LocationOffset);
+}
+
+FTransform FActorPositioning::GetSnappedSurfaceAlignedTransform(FLevelEditorViewportClient* InViewportClient, const UActorFactory* InFactory, FVector InSurfaceLocation, const FVector& SurfaceNormal, const FVector& InPlacementExtent, const FTransform& InStartTransform)
+{
+	FSnappingUtils::SnapPointToGrid(InSurfaceLocation, FVector(0.f));
+
+	// Secondly, attempt vertex snapping
+	FVector AlignToNormal;
+	if (!InViewportClient || !FSnappingUtils::SnapLocationToNearestVertex( InSurfaceLocation, InViewportClient->GetDropPreviewLocation(), InViewportClient, AlignToNormal ))
+	{
+		AlignToNormal = SurfaceNormal;
+	}
+
+	return GetSurfaceAlignedTransform(InFactory, InSurfaceLocation, AlignToNormal, InPlacementExtent, InStartTransform);
+}
\ No newline at end of file
diff --git a/Engine/Source/Editor/UnrealEd/Private/Editor/ActorPositioning.h b/Engine/Source/Editor/UnrealEd/Private/Editor/ActorPositioning.h
new file mode 100644
--- /dev/null
+++ b/Engine/Source/Editor/UnrealEd/Private/Editor/ActorPositioning.h
@@ -0,0 +1,110 @@
+// Copyright 1998-2014 Epic Games, Inc. All Rights Reserved.
+
+#pragma once
+
+
+struct FActorPositionTraceResult
+{
+	/** Enum representing the state of this result */
+	enum ResultState
+	{
+		/** The trace found a valid hit target */
+		HitSuccess,
+		/** The trace found no valid targets, so chose a default position */
+		Default,
+		/** The trace failed entirely */
+		Failed,
+	};
+
+	/** Constructor */
+	FActorPositionTraceResult() : State(Failed), Location(0.f), SurfaceNormal(0.f) {}
+
+	/** The state of this result */
+	ResultState	State;
+
+	/** The location of the preferred trace hit */
+	FVector		Location;
+
+	/** The surface normal of the trace hit */
+	FVector		SurfaceNormal;
+};
+
+struct FActorPositioning
+{
+	/** Trace the specified world to find a position to snap actors to
+	 *
+	 *	@param	Cursor			The cursor position and direction to trace at
+	 *	@param	View			The scene view that we are tracing on
+	 *	@param	IgnoreActors	Optional array of actors to exclude from the trace
+	 *
+	 *	@return	Result structure containing the location and normal of a trace hit, or a default position in front of the camera on failure
+	*/
+	static UNREALED_API FActorPositionTraceResult TraceWorldForPositionWithDefault(const FViewportCursorLocation& Cursor, const FSceneView& View, const TArray<AActor*>* IgnoreActors = nullptr);
+	
+	/** Trace the specified world to find a position to snap actors to
+	 *
+	 *	@param	Cursor			The cursor position and direction to trace at
+	 *	@param	View			The scene view that we are tracing on
+	 *	@param	IgnoreActors	Optional array of actors to exclude from the trace
+	 *
+	 *	@return	Result structure containing the location and normal of a trace hit, or empty on failure
+	*/
+	static UNREALED_API FActorPositionTraceResult TraceWorldForPosition(const FViewportCursorLocation& Cursor, const FSceneView& View, const TArray<AActor*>* IgnoreActors = nullptr);
+
+	/** Trace the specified world to find a position to snap actors to
+	 *
+	 *	@param	World			The world to trace
+	 *	@param	InSceneView		The scene view that we are tracing on
+	 *	@param	RayStart		The start of the ray in world space
+	 *	@param	RayEnd			The end of the ray in world space
+	 *	@param	IgnoreActors	Optional array of actors to exclude from the trace
+	 *
+	 *	@return	Result structure containing the location and normal of a trace hit, or empty on failure
+	*/
+	static UNREALED_API FActorPositionTraceResult TraceWorldForPosition(const UWorld& InWorld, const FSceneView& InSceneView, const FVector& RayStart, const FVector& RayEnd, const TArray<AActor*>* IgnoreActors = nullptr);
+
+	/** Get a transform that should be used to spawn the specified actor using the global editor click location and plane
+	 *
+	 *	@param	Actor			The actor to spawn
+	 *	@param	bSnap			Whether to perform snapping on the placement transform or not (defaults to true)
+	 *
+	 *	@return	The expected transform for the specified actor
+	*/
+	static UNREALED_API FTransform GetCurrentViewportPlacementTransform(const AActor& Actor, bool bSnap = true);
+
+	/** Get a default actor position in front of the camera
+	 *
+	 *	@param	InActor				The actor
+	 *	@param	InCameraOrigin		The location of the camera
+	 *	@param	InCameraDirection	The view direction of the camera
+	 *
+	 *	@return	The expected transform for the specified actor
+	*/
+	static UNREALED_API FVector GetActorPositionInFrontOfCamera(const AActor& InActor, const FVector& InCameraOrigin, const FVector& InCameraDirection);
+	
+	/**
+	 *	Get the snapped position and rotation transform for an actor to snap it to the specified surface. This will potentially align the actor to the specified surface normal, and offset it based on factory settings
+	 *
+	 *	@param	InFactory				The factory used to spawn the actor. Factories define how actors are aligned and offset.
+	 *	@param	InSurfaceLocation		The location on the surface to transform to
+	 *	@param	InSurfaceNormal			The surface normal of the spawn position
+	 *	@param	InPlacementExtent		(optional) Optional actor placement extent
+	 *	@param	InStartTransform		(optional) Optional initial transform. This will result in a transformation from InStartTransform, to the surface transform
+	 *
+	 *	@return	The expected transform for the specified actor
+	 */
+	static UNREALED_API FTransform GetSnappedSurfaceAlignedTransform(FLevelEditorViewportClient* InViewportClient, const UActorFactory* InFactory, FVector InSurfaceLocation, const FVector& InSurfaceNormal, const FVector& InPlacementExtent = FVector(0.f), const FTransform& InStartTransform = FTransform::Identity);
+
+	/**
+	 *	Get the position and rotation transform for an actor to snap it to the specified surface. This will potentially align the actor to the specified surface normal, and offset it.
+	 *
+	 *	@param	InFactory				The factory used to spawn the actor. Factories define how actors are aligned and offset.
+	 *	@param	InSurfaceLocation		The location on the surface to transform to
+	 *	@param	InSurfaceNormal			The surface normal of the spawn position
+	 *	@param	InPlacementExtent		(optional) Optional actor placement extent
+	 *	@param	InStartTransform		(optional) Optional initial transform. This will result in a transformation from InStartTransform, to the surface transform
+	 *
+	 *	@return	The expected transform for the specified actor
+	 */
+	static UNREALED_API FTransform GetSurfaceAlignedTransform(const UActorFactory* InFactory, const FVector& InSurfaceLocation, const FVector& InSurfaceNormal, const FVector& InPlacementExtent = FVector(0.f), const FTransform& InStartTransform = FTransform::Identity);
+};
\ No newline at end of file
diff --git a/Engine/Source/Editor/UnrealEd/Private/EditorServer.cpp b/Engine/Source/Editor/UnrealEd/Private/EditorServer.cpp
--- a/Engine/Source/Editor/UnrealEd/Private/EditorServer.cpp
+++ b/Engine/Source/Editor/UnrealEd/Private/EditorServer.cpp
@@ -39,6 +39,8 @@
 #include "TargetPlatform.h"
 #include "IConsoleManager.h"
 
+#include "Editor/ActorPositioning.h"
+
 #include "Editor/StatsViewer/Public/StatsViewerModule.h"
 #include "ActorEditorUtils.h"
 #include "ContentBrowserModule.h"
@@ -2998,19 +3000,8 @@ void UEditorEngine::PasteSelectedActorsFromClipboard( UWorld* InWorld, const FTe
 		return;
 	}
 
-	FVector SaveClickLocation = GEditor->ClickLocation;
-	FSnappingUtils::SnapPointToGrid( SaveClickLocation, FVector(GEditor->GetGridSize(),GEditor->GetGridSize(),GEditor->GetGridSize()) );
-
-	FVector AnyActorLocation = FVector::ZeroVector;
-	for ( FSelectionIterator It( GetSelectedActorIterator() ) ; It ; ++It )
-	{
-		AActor* Actor = static_cast<AActor*>( *It );
-		checkSlow( Actor->IsA(AActor::StaticClass()) );
-
-		AnyActorLocation = Actor->GetActorLocation();
-		break;
-	}
-
+	FVector SaveClickLocation = FActorPositioning::GetSnappedSurfaceAlignedTransform(GCurrentLevelEditingViewportClient, nullptr, GEditor->ClickLocation, GEditor->ClickPlane).GetLocation();
+	
 	ULevel* DesiredLevel = InWorld->GetCurrentLevel();
 
 	// Don't allow pasting to levels that are locked
@@ -4400,24 +4391,7 @@ bool UEditorEngine::SnapActorTo( AActor* InActor, const bool InAlign, const bool
 
 void UEditorEngine::MoveActorInFrontOfCamera( AActor& InActor, const FVector& InCameraOrigin, const FVector& InCameraDirection )
 {
-	// Get the  radius of the actors bounding cylinder.  Height is not needed.
-	float CylRadius, CylHeight;
-	InActor.GetComponentsBoundingCylinder(CylRadius, CylHeight);
-
-	// a default cylinder radius if no bounding cylinder exists.  
-	const float	DefaultCylinderRadius = 50.0f;
-
-	if( CylRadius == 0.0f )
-	{
-		// If the actor does not have a bounding cylinder, use a default value.
-		CylRadius = DefaultCylinderRadius;
-	}
-
-	// The new location the cameras origin offset by the actors bounding cylinder radius down the direction of the cameras view. 
-	FVector NewLocation = InCameraOrigin + InCameraDirection * CylRadius + InCameraDirection * GetDefault<ULevelEditorViewportSettings>()->BackgroundDropDistance;
-
-	// Snap the new location if snapping is enabled
-	FSnappingUtils::SnapPointToGrid( NewLocation, FVector::ZeroVector );
+	const FVector NewLocation = FActorPositioning::GetActorPositionInFrontOfCamera(InActor, InCameraOrigin, InCameraDirection);
 
 	// Move the actor to its new location.  Not checking for collisions
 	InActor.TeleportTo( NewLocation, InActor.GetActorRotation(), false, true );
diff --git a/Engine/Source/Editor/UnrealEd/Private/EditorViewportCommands.cpp b/Engine/Source/Editor/UnrealEd/Private/EditorViewportCommands.cpp
--- a/Engine/Source/Editor/UnrealEd/Private/EditorViewportCommands.cpp
+++ b/Engine/Source/Editor/UnrealEd/Private/EditorViewportCommands.cpp
@@ -59,6 +59,7 @@ void FEditorViewportCommands::RegisterCommands()
 	UI_COMMAND( LocationGridSnap, "Grid Snap", "Enables or disables snapping to the grid when dragging objects around", EUserInterfaceActionType::ToggleButton, FInputGesture() );
 	UI_COMMAND( RotationGridSnap, "Rotation Snap", "Enables or disables snapping objects to a rotation grid", EUserInterfaceActionType::ToggleButton, FInputGesture() );
 	UI_COMMAND( ScaleGridSnap, "Scale Snap", "Enables or disables snapping objects to a scale grid", EUserInterfaceActionType::ToggleButton, FInputGesture() );
+	UI_COMMAND( SurfaceSnapping, "Surface Snapping", "If enabled, actors will snap to surfaces in the world when dragging", EUserInterfaceActionType::ToggleButton, FInputGesture() );
 
 	UI_COMMAND( ToggleAutoExposure, "Automatic (Default in-game)", "Enable automatic expose", EUserInterfaceActionType::RadioButton, FInputGesture() );
 	UI_COMMAND( FixedExposure4m, "Fixed Exposure: -4", "Set the fixed exposure to -4", EUserInterfaceActionType::RadioButton, FInputGesture() );
diff --git a/Engine/Source/Editor/UnrealEd/Private/Factories/ActorFactory.cpp b/Engine/Source/Editor/UnrealEd/Private/Factories/ActorFactory.cpp
--- a/Engine/Source/Editor/UnrealEd/Private/Factories/ActorFactory.cpp
+++ b/Engine/Source/Editor/UnrealEd/Private/Factories/ActorFactory.cpp
@@ -30,6 +30,63 @@ DEFINE_LOG_CATEGORY_STATIC(LogActorFactory, Log, All);
 
 #define LOCTEXT_NAMESPACE "ActorFactory"
 
+/**
+ * Find am alignment transform for the specified actor rotation, given a model-space axis to align, and a world space normal to align to.
+ * This function attempts to find a 'natural' looking rotation by rotating around a local pitch axis, and a world Z. Rotating in this way
+ * should retain the roll around the model space axis, removing rotation artifacts introduced by a simpler quaternion rotation.
+ */
+FQuat FindActorAlignmentRotation(const FQuat& InActorRotation, const FVector& InModelAxis, const FVector& InWorldNormal)
+{
+	FVector TransformedModelAxis = InActorRotation.RotateVector(InModelAxis);
+
+	const auto InverseActorRotation = InActorRotation.Inverse();
+	const auto DestNormalModelSpace = InverseActorRotation.RotateVector(InWorldNormal);
+
+	FQuat DeltaRotation = FQuat::Identity;
+
+	const float VectorDot = InWorldNormal | TransformedModelAxis;
+	if (1.f - FMath::Abs(VectorDot) <= KINDA_SMALL_NUMBER)
+	{
+		if (VectorDot < 0.f)
+		{
+			// Anti-parallel
+			return InActorRotation * FQuat::FindBetween(InModelAxis, DestNormalModelSpace);
+		}
+	}
+	else
+	{
+		const FVector Z(0.f, 0.f, 1.f);
+
+		// Find a reference axis to measure the relative pitch rotations between the source axis, and the destination axis.
+		FVector PitchReferenceAxis = InverseActorRotation.RotateVector(Z);
+		if (FMath::Abs(FVector::DotProduct(InModelAxis, PitchReferenceAxis)) > 0.7f)
+		{
+			PitchReferenceAxis = DestNormalModelSpace;
+		}
+		
+		// Find a local 'pitch' axis to rotate around
+		const FVector OrthoPitchAxis = FVector::CrossProduct(PitchReferenceAxis, InModelAxis);
+		const float Pitch = FMath::Acos(PitchReferenceAxis | DestNormalModelSpace) - FMath::Acos(PitchReferenceAxis | InModelAxis);//FMath::Asin(OrthoPitchAxis.Size());
+
+		DeltaRotation = FQuat(OrthoPitchAxis.SafeNormal(), Pitch);
+		DeltaRotation.Normalize();
+
+		// Transform the model axis with this new pitch rotation to see if there is any need for yaw
+		TransformedModelAxis = (InActorRotation * DeltaRotation).RotateVector(InModelAxis);
+
+		if (!FVector::Parallel(InWorldNormal, TransformedModelAxis))
+		{
+			const float Yaw = FMath::Atan2(InWorldNormal.X, InWorldNormal.Y) - FMath::Atan2(TransformedModelAxis.X, TransformedModelAxis.Y);
+
+			// Rotation axis for yaw is the Z axis in world space
+			const FVector WorldYawAxis = (InActorRotation * DeltaRotation).Inverse().RotateVector(Z);
+			DeltaRotation *= FQuat(WorldYawAxis, -Yaw);
+		}
+	}
+
+	return InActorRotation * DeltaRotation;
+}
+
 /*-----------------------------------------------------------------------------
 UActorFactory
 -----------------------------------------------------------------------------*/
@@ -78,16 +135,29 @@ UObject* UActorFactory::GetAssetFromActorInstance(AActor* ActorInstance)
 	return NULL;
 }
 
-AActor* UActorFactory::CreateActor( UObject* Asset, ULevel* InLevel, const FVector& Location, const FRotator* const Rotation, EObjectFlags ObjectFlags, const FName Name )
+/** Return a rotator which aligns this actor type to the specified surface normal */
+FQuat UActorFactory::AlignObjectToSurfaceNormal(const FVector& InSurfaceNormal, const FQuat& ActorRotation) const
 {
-	AActor* DefaultActor = GetDefaultActor( FAssetData( Asset ) );
+	if (bUseSurfaceOrientation)
+	{
+		// By default we align the X axis with the inverse of the surface normal (so things look at the surface)
+		return FindActorAlignmentRotation(ActorRotation, FVector(-1.f, 0.f, 0.f), InSurfaceNormal);
+	}
+	else
+	{
+		return FQuat::Identity;
+	}
+}
 
-	FVector SpawnLocation = Location + SpawnPositionOffset;
-	FRotator SpawnRotation = Rotation ? *Rotation : DefaultActor ? DefaultActor->GetActorRotation() : FRotator::ZeroRotator;
+AActor* UActorFactory::CreateActor( UObject* Asset, ULevel* InLevel, FTransform SpawnTransform, EObjectFlags ObjectFlags, const FName Name )
+{
 	AActor* NewActor = NULL;
-	if ( PreSpawnActor(Asset, SpawnLocation, SpawnRotation, Rotation != NULL) )
+
+	if ( PreSpawnActor(Asset, SpawnTransform) )
 	{
-		NewActor = SpawnActor(Asset, InLevel, SpawnLocation, SpawnRotation, ObjectFlags, Name);
+		const auto Location = SpawnTransform.GetLocation();
+		const auto Rotation = SpawnTransform.GetRotation().Rotator();
+		NewActor = SpawnActor(Asset, InLevel, Location, Rotation, ObjectFlags, Name);
 
 		if ( NewActor )
 		{
@@ -109,7 +179,7 @@ UBlueprint* UActorFactory::CreateBlueprint( UObject* Asset, UObject* Outer, cons
 	return NewBlueprint;
 }
 
-bool UActorFactory::PreSpawnActor( UObject* Asset, FVector& InOutLocation, FRotator& InOutRotation, bool bRotationWasSupplied)
+bool UActorFactory::PreSpawnActor( UObject* Asset, FTransform& InOutLocation)
 {
 	// Subclasses may implement this to set up a spawn or to adjust the spawn location or rotation.
 	return true;
@@ -149,6 +219,7 @@ UActorFactoryStaticMesh::UActorFactoryStaticMesh(const class FPostConstructIniti
 {
 	DisplayName = LOCTEXT("StaticMeshDisplayName", "Static Mesh");
 	NewActorClass = AStaticMeshActor::StaticClass();
+	bUseSurfaceOrientation = true;
 }
 
 bool UActorFactoryStaticMesh::CanCreateActorFrom( const FAssetData& AssetData, FText& OutErrorMsg )
@@ -205,6 +276,12 @@ void UActorFactoryStaticMesh::PostCreateBlueprint( UObject* Asset, AActor* CDO )
 	}
 }
 
+FQuat UActorFactoryStaticMesh::AlignObjectToSurfaceNormal(const FVector& InSurfaceNormal, const FQuat& ActorRotation) const
+{
+	// Meshes align the Z (up) axis with the surface normal
+	return FindActorAlignmentRotation(ActorRotation, FVector(0.f, 0.f, 1.f), InSurfaceNormal);
+}
+
 /*-----------------------------------------------------------------------------
 UActorFactoryDeferredDecal
 -----------------------------------------------------------------------------*/
@@ -276,17 +353,6 @@ bool UActorFactoryDeferredDecal::CanCreateActorFrom( const FAssetData& AssetData
 	return true;
 }
 
-bool UActorFactoryDeferredDecal::PreSpawnActor(UObject* Asset, FVector& InOutLocation, FRotator& InOutRotation, bool bRotationWasSupplied)
-{
-	if ( bRotationWasSupplied )
-	{
-		// Orient the decal in the opposite direction of the receiving surface's normal
-		InOutRotation = -InOutRotation;
-	}
-
-	return true;
-}
-
 void UActorFactoryDeferredDecal::PostSpawnActor(UObject* Asset, AActor* NewActor)
 {
 	UMaterialInterface* Material = GetMaterial( Asset );
@@ -487,7 +553,7 @@ bool UActorFactoryPhysicsAsset::CanCreateActorFrom( const FAssetData& AssetData,
 	return true;
 }
 
-bool UActorFactoryPhysicsAsset::PreSpawnActor(UObject* Asset, FVector& InOutLocation, FRotator& InOutRotation, bool bRotationWasSupplied)
+bool UActorFactoryPhysicsAsset::PreSpawnActor(UObject* Asset, FTransform& InOutLocation)
 {
 	UPhysicsAsset* PhysicsAsset = CastChecked<UPhysicsAsset>(Asset);
 	USkeletalMesh* UseSkelMesh = PhysicsAsset->PreviewSkeletalMesh.Get();
@@ -743,6 +809,7 @@ UActorFactorySkeletalMesh::UActorFactorySkeletalMesh(const class FPostConstructI
 { 
 	DisplayName = LOCTEXT("SkeletalMeshDisplayName", "Skeletal Mesh");
 	NewActorClass = ASkeletalMeshActor::StaticClass();
+	bUseSurfaceOrientation = true;
 }
 
 bool UActorFactorySkeletalMesh::CanCreateActorFrom( const FAssetData& AssetData, FText& OutErrorMsg )
@@ -902,6 +969,11 @@ void UActorFactorySkeletalMesh::PostCreateBlueprint( UObject* Asset, AActor* CDO
 	}
 }
 
+FQuat UActorFactorySkeletalMesh::AlignObjectToSurfaceNormal(const FVector& InSurfaceNormal, const FQuat& ActorRotation) const
+{
+	// Meshes align the Z (up) axis with the surface normal
+	return FindActorAlignmentRotation(ActorRotation, FVector(0.f, 0.f, 1.f), InSurfaceNormal);
+}
 
 /*-----------------------------------------------------------------------------
 UActorFactoryCameraActor
@@ -1015,7 +1087,7 @@ AActor* UActorFactoryClass::GetDefaultActor( const FAssetData& AssetData )
 	return NULL;
 }
 
-bool UActorFactoryClass::PreSpawnActor( UObject* Asset, FVector& InOutLocation, FRotator& InOutRotation, bool bRotationWasSupplied)
+bool UActorFactoryClass::PreSpawnActor( UObject* Asset, FTransform& InOutLocation)
 {
 	UClass* ActualClass = Cast<UClass>(Asset);
 
@@ -1129,7 +1201,7 @@ AActor* UActorFactoryBlueprint::GetDefaultActor( const FAssetData& AssetData )
 	return GeneratedClass->GetDefaultObject<AActor>();
 }
 
-bool UActorFactoryBlueprint::PreSpawnActor( UObject* Asset, FVector& InOutLocation, FRotator& InOutRotation, bool bRotationWasSupplied)
+bool UActorFactoryBlueprint::PreSpawnActor( UObject* Asset, FTransform& InOutLocation)
 {
 	UBlueprint* Blueprint = CastChecked<UBlueprint>(Asset);
 
@@ -1218,6 +1290,8 @@ UActorFactoryDirectionalLight::UActorFactoryDirectionalLight(const class FPostCo
 {
 	DisplayName = LOCTEXT("DirectionalLightDisplayName", "Directional Light");
 	NewActorClass = ADirectionalLight::StaticClass();
+	SpawnPositionOffset = FVector(200, 0, 0);
+	bUseSurfaceOrientation = true;
 }
 
 /*-----------------------------------------------------------------------------
@@ -1228,6 +1302,8 @@ UActorFactorySpotLight::UActorFactorySpotLight(const class FPostConstructInitial
 {
 	DisplayName = LOCTEXT("SpotLightDisplayName", "Spot Light");
 	NewActorClass = ASpotLight::StaticClass();
+	SpawnPositionOffset = FVector(200, 0, 0);
+	bUseSurfaceOrientation = true;
 }
 
 /*-----------------------------------------------------------------------------
@@ -1258,7 +1334,8 @@ UActorFactorySphereReflectionCapture::UActorFactorySphereReflectionCapture(const
 {
 	DisplayName = LOCTEXT("ReflectionCaptureSphereDisplayName", "Sphere Reflection Capture");
 	NewActorClass = ASphereReflectionCapture::StaticClass();
-	SpawnPositionOffset = FVector(0, 0, 200);
+	SpawnPositionOffset = FVector(200, 0, 0);
+	bUseSurfaceOrientation = true;
 }
 
 /*-----------------------------------------------------------------------------
@@ -1269,7 +1346,8 @@ UActorFactoryBoxReflectionCapture::UActorFactoryBoxReflectionCapture(const class
 {
 	DisplayName = LOCTEXT("ReflectionCaptureBoxDisplayName", "Box Reflection Capture");
 	NewActorClass = ABoxReflectionCapture::StaticClass();
-	SpawnPositionOffset = FVector(0, 0, 200);
+	SpawnPositionOffset = FVector(200, 0, 0);
+	bUseSurfaceOrientation = true;
 }
 
 /*-----------------------------------------------------------------------------
@@ -1280,7 +1358,8 @@ UActorFactoryPlaneReflectionCapture::UActorFactoryPlaneReflectionCapture(const c
 {
 	DisplayName = LOCTEXT("ReflectionCapturePlaneDisplayName", "Plane Reflection Capture");
 	NewActorClass = APlaneReflectionCapture::StaticClass();
-	SpawnPositionOffset = FVector(0, 0, 200);
+	SpawnPositionOffset = FVector(200, 0, 0);
+	bUseSurfaceOrientation = true;
 }
 
 /*-----------------------------------------------------------------------------
diff --git a/Engine/Source/Editor/UnrealEd/Private/Fbx/FbxMatineeImport.cpp b/Engine/Source/Editor/UnrealEd/Private/Fbx/FbxMatineeImport.cpp
--- a/Engine/Source/Editor/UnrealEd/Private/Fbx/FbxMatineeImport.cpp
+++ b/Engine/Source/Editor/UnrealEd/Private/Fbx/FbxMatineeImport.cpp
@@ -233,7 +233,7 @@ void FFbxImporter::ImportMatineeSequence(AMatineeActor* InMatineeActor)
 			CameraNode = FindCamera(Node);
 			if ( bCreateUnknownCameras && CameraNode != NULL )
 			{
-				Actor = GEditor->AddActor( InMatineeActor->GetWorld()->GetCurrentLevel(), ACameraActor::StaticClass(), FVector::ZeroVector );
+				Actor = GEditor->AddActor( InMatineeActor->GetWorld()->GetCurrentLevel(), ACameraActor::StaticClass(), FTransform::Identity );
 				Actor->SetActorLabel( ANSI_TO_TCHAR(CameraNode->GetName()) );
 			}
 			else
diff --git a/Engine/Source/Editor/UnrealEd/Private/LevelEditorViewport.cpp b/Engine/Source/Editor/UnrealEd/Private/LevelEditorViewport.cpp
--- a/Engine/Source/Editor/UnrealEd/Private/LevelEditorViewport.cpp
+++ b/Engine/Source/Editor/UnrealEd/Private/LevelEditorViewport.cpp
@@ -43,6 +43,8 @@
 #include "AssetRegistryModule.h"
 #include "Animation/VertexAnim/VertexAnimation.h"
 
+#include "Editor/ActorPositioning.h"
+
 DEFINE_LOG_CATEGORY(LogEditorViewport);
 
 #define LOCTEXT_NAMESPACE "LevelEditorViewportClient"
@@ -157,14 +159,13 @@ static FVector4 AttemptToSnapLocationToOriginPlane( const FViewportCursorLocatio
  * @param	InLevel			Level in which to drop actor
  * @param	ObjToUse		Asset to attempt to use for an actor to place
  * @param	CursorLocation	Location of the cursor while dropping
- * @param	bUsedHitProxy	Whether or not a hit proxy was used for spawning
  * @param	bSelectActors	If true, select the newly dropped actors (defaults: true)
  * @param	ObjectFlags		The flags to place on the actor when it is spawned
  * @param	FactoryToUse	The preferred actor factory to use (optional)
  *
  * @return	true if the object was successfully used to place an actor; false otherwise
  */
-static TArray<AActor*> AttemptDropObjAsActors( ULevel* InLevel, UObject* ObjToUse, const FViewportCursorLocation& CursorLocation, bool bUsedHitProxy, bool bSelectActors, EObjectFlags ObjectFlags, UActorFactory* FactoryToUse, const FName Name = NAME_None )
+static TArray<AActor*> AttemptDropObjAsActors( ULevel* InLevel, UObject* ObjToUse, const FViewportCursorLocation& CursorLocation, bool bSelectActors, EObjectFlags ObjectFlags, UActorFactory* FactoryToUse, const FName Name = NAME_None )
 {
 	TArray<AActor*> PlacedActors;
 
@@ -175,8 +176,6 @@ static TArray<AActor*> AttemptDropObjAsActors( ULevel* InLevel, UObject* ObjToUs
 		ObjectClass = ObjToUse->GetClass();
 	}
 
-	FVector PreSnapLocation = GEditor->ClickLocation;
-
 	AActor* PlacedActor = NULL;
 	if ( ObjectClass != NULL && ObjectClass->IsChildOf( AActor::StaticClass() ) )
 	{
@@ -189,25 +188,24 @@ static TArray<AActor*> AttemptDropObjAsActors( ULevel* InLevel, UObject* ObjToUs
 
 		if ( ActorFactory != NULL )
 		{
-			PlacedActor = FActorFactoryAssetProxy::AddActorFromSelection( ObjectClass, NULL, ActorFactory->bUseSurfaceOrientation, bSelectActors, ObjectFlags, ActorFactory, Name );
+			PlacedActor = FActorFactoryAssetProxy::AddActorFromSelection( ObjectClass, NULL, bSelectActors, ObjectFlags, ActorFactory, Name );
 		}
 
 		if ( PlacedActor == NULL && ActorFactory != NULL )
 		{
-			PlacedActor = FActorFactoryAssetProxy::AddActorForAsset( ObjToUse, NULL, ActorFactory->bUseSurfaceOrientation, bSelectActors, ObjectFlags, ActorFactory, Name );
+			PlacedActor = FActorFactoryAssetProxy::AddActorForAsset( ObjToUse, bSelectActors, ObjectFlags, ActorFactory, Name );
 		}
 		
 		if ( PlacedActor == NULL && !ObjectClass->HasAnyClassFlags(CLASS_NotPlaceable | CLASS_Abstract) )
 		{
 			// If no actor factory was found or failed, add the actor directly.
-			FVector Collision = ObjectClass->GetDefaultObject<AActor>()->GetPlacementExtent();
-			PlacedActor = GEditor->AddActor( InLevel, ObjectClass, GEditor->ClickLocation + GEditor->ClickPlane * ( FVector::BoxPushOut( GEditor->ClickPlane, Collision ) + 0.1 ), /*bSilent=*/false, ObjectFlags );
+			const FTransform ActorTransform = FActorPositioning::GetCurrentViewportPlacementTransform(*ObjectClass->GetDefaultObject<AActor>());
+			PlacedActor = GEditor->AddActor( InLevel, ObjectClass, ActorTransform, /*bSilent=*/false, ObjectFlags );
 		}
 
 		if ( PlacedActor != NULL )
 		{
 			FVector Collision = ObjectClass->GetDefaultObject<AActor>()->GetPlacementExtent();
-			PreSnapLocation += GEditor->ClickPlane * ( FVector::BoxPushOut( GEditor->ClickPlane, Collision ) + 0.1 );
 			PlacedActors.Add(PlacedActor);
 		}
 	}
@@ -237,8 +235,6 @@ static TArray<AActor*> AttemptDropObjAsActors( ULevel* InLevel, UObject* ObjToUs
 	}
 	else if (NULL == PlacedActor)
 	{
-		const bool UseSurfaceOrientation = FactoryToUse != NULL ? FactoryToUse->bUseSurfaceOrientation : false;
-
 		bool bPlace = true;
 		if (ObjectClass->IsChildOf(UBlueprint::StaticClass()))
 		{
@@ -256,7 +252,7 @@ static TArray<AActor*> AttemptDropObjAsActors( ULevel* InLevel, UObject* ObjToUs
 
 		if (bPlace)
 		{
-			PlacedActor = FActorFactoryAssetProxy::AddActorForAsset( ObjToUse, NULL, UseSurfaceOrientation, bSelectActors, ObjectFlags, FactoryToUse, Name );
+			PlacedActor = FActorFactoryAssetProxy::AddActorForAsset( ObjToUse, bSelectActors, ObjectFlags, FactoryToUse, Name );
 			if ( PlacedActor != NULL )
 			{
 				PlacedActors.Add(PlacedActor);
@@ -264,55 +260,6 @@ static TArray<AActor*> AttemptDropObjAsActors( ULevel* InLevel, UObject* ObjToUs
 		}
 	}
 
-	if ( PlacedActors.Num() > 0 && CursorLocation.GetViewportType() == LVT_Perspective )
-	{
-		if (bUsedHitProxy)
-		{
-			const bool bSnapNewObjectsToFloor = GetDefault<ULevelEditorViewportSettings>()->bSnapNewObjectsToFloor;
-			for (auto ActorIt = PlacedActors.CreateConstIterator(); ActorIt; ++ActorIt)
-			{
-				AActor* Actor = *ActorIt;
-
-				FVector Loc = Actor->GetActorLocation();
-				if( bSnapNewObjectsToFloor )
-				{
-					Loc.Z = PreSnapLocation.Z;
-				}
-
-				Actor->TeleportTo(Loc, Actor->GetActorRotation(), false, true);
-			}
-		}
-		else
-		{
-			// If the actor was successfully placed, but it was done on the background of a perspective viewport, forcibly move the actor in front
-			// of the camera
-
-			// First calculate the origin location for all the created actors
-			FVector Origin = FVector::ZeroVector;
-			for (auto ActorIt = PlacedActors.CreateConstIterator(); ActorIt; ++ActorIt)
-			{
-				Origin += (*ActorIt)->GetActorLocation();
-			}
-			Origin /= PlacedActors.Num();
-
-			// Now move the actor in front of the camera then offset it by the offset to the calculated origin
-			for (auto ActorIt = PlacedActors.CreateConstIterator(); ActorIt; ++ActorIt)
-			{
-				AActor* Actor = *ActorIt;
-				FVector OffsetToOrigin = Actor->GetActorLocation() - Origin;
-				GEditor->MoveActorInFrontOfCamera(*Actor, CursorLocation.GetOrigin(), CursorLocation.GetDirection());
-				FVector NewLoc = Actor->GetActorLocation() + OffsetToOrigin;
-
-				if (ObjToUse->IsA(UBrushBuilder::StaticClass()))
-				{
-					NewLoc = AttemptToSnapLocationToOriginPlane(CursorLocation, NewLoc);
-				}
-
-				Actor->TeleportTo(NewLoc, Actor->GetActorRotation());
-			}
-		}
-	}
-
 	return PlacedActors;
 }
 
@@ -844,15 +791,8 @@ bool FLevelEditorViewportClient::DropObjectsOnBackground( FViewportCursorLocatio
 		UObject* AssetObj = DroppedObjects[DroppedObjectsIdx];
 		ensure( AssetObj );
 
-		// Configure editor click information
-		const float DistanceMultiplier = ( Cursor.GetViewportType() == LVT_Perspective ) ? 1.0f : 0.0f;
-
-		GEditor->ClickLocation = Cursor.GetOrigin() + Cursor.GetDirection() * DistanceMultiplier;
-		GEditor->ClickPlane = FPlane( 0, 0, 0, 0 );
-
 		// Attempt to create actors from the dropped object
-		const bool bUsedHitProxy = false;
-		TArray<AActor*> NewActors = AttemptDropObjAsActors(GetWorld()->GetCurrentLevel(), AssetObj, Cursor, bUsedHitProxy, bSelectActors, ObjectFlags, FactoryToUse);
+		TArray<AActor*> NewActors = AttemptDropObjAsActors(GetWorld()->GetCurrentLevel(), AssetObj, Cursor, bSelectActors, ObjectFlags, FactoryToUse);
 
 		if ( NewActors.Num() > 0 )
 		{
@@ -867,159 +807,48 @@ bool FLevelEditorViewportClient::DropObjectsOnBackground( FViewportCursorLocatio
 	return bResult;
 }
 
-bool FLevelEditorViewportClient::DropObjectsOnActor(FViewportCursorLocation& Cursor, const TArray<UObject*>& DroppedObjects, AActor* DroppedUponActor, int32 DroppedUponSlot, FVector* DroppedLocation, EObjectFlags ObjectFlags, TArray<AActor*>& OutNewActors, bool bUsedHitProxy, bool bSelectActors, UActorFactory* FactoryToUse)
+bool FLevelEditorViewportClient::DropObjectsOnActor(FViewportCursorLocation& Cursor, const TArray<UObject*>& DroppedObjects, AActor* DroppedUponActor, int32 DroppedUponSlot, EObjectFlags ObjectFlags, TArray<AActor*>& OutNewActors, bool bSelectActors, UActorFactory* FactoryToUse)
 {
 	bool bResult = false;
 
-	if ( DroppedUponActor != NULL )
-	{
-		const FVector TargetLocation = DroppedLocation ? *DroppedLocation : DroppedUponActor->GetActorLocation();
-		if(DroppedObjects.Num() > 1)
-		{
-			bResult = true;
-
-			// Create a transaction if we have more than 1 item. This causes them all to get "un-created" together
-			FScopedTransaction Transaction( NSLOCTEXT("UnrealEd", "CreateActors", "Create Actors") );
-			for ( auto DroppedObject : DroppedObjects )
-			{
-				if (!DropSingleObjectOnActor(Cursor, DroppedObject, DroppedUponActor, DroppedUponSlot, TargetLocation, ObjectFlags, OutNewActors, bUsedHitProxy, bSelectActors, FactoryToUse))
-				{
-					bResult = false;
-				}
-			}
-		}
-		else if (DroppedObjects.Num() == 1)
-		{
-			bResult = DropSingleObjectOnActor(Cursor, DroppedObjects[0], DroppedUponActor, DroppedUponSlot, TargetLocation, ObjectFlags, OutNewActors, bUsedHitProxy, bSelectActors, FactoryToUse);
-		}
-	}
-
-	return bResult;
-}
-
-bool FLevelEditorViewportClient::GetVisibleActors(TArray<struct FHitResult>& OutHits,const FVector& Start,const FVector& End,const struct FCollisionQueryParams& Params) const
-{
-	FCollisionObjectQueryParams ObjectParam(FCollisionObjectQueryParams::InitType::AllObjects);
-	if ( GetWorld()->LineTraceMulti(OutHits, Start, End, Params, ObjectParam) )
+	if ( !DroppedUponActor || DroppedObjects.Num() == 0 )
 	{
-		// see if any volume is around
-		for (int32 HitIdx=0; HitIdx < OutHits.Num(); ++HitIdx)
-		{
-			if ( OutHits[HitIdx].Component!= NULL &&
-				// ignore volume or trigger boxes
-				(OutHits[HitIdx].Component->IsA(UShapeComponent::StaticClass())))
-			{
-				// remove
-				OutHits.RemoveAt(HitIdx);
-				--HitIdx;
-			}
-			else if(OutHits[HitIdx].Actor!= NULL &&
-			// ignore volume or trigger boxes
-			(OutHits[HitIdx].Actor->IsA(AVolume::StaticClass())))
-			{
-				// remove
-				OutHits.RemoveAt(HitIdx);
-				--HitIdx;
-			}
-		}
-
-		 // now print all list
-//		for (int32 HitIdx=0; HitIdx < OutHits.Num(); ++HitIdx)
-//		{
-//			const AActor * Actor = OutHits[HitIdx].Actor.Get();
-//			UE_LOG(LogLevelTools, Warning, TEXT("%d) Actor: %s, ClassType : %s, Component : %s "), HitIdx+1, *GetNameSafe(Actor), (Actor)? *GetNameSafe(Actor->GetClass()): TEXT("None"),  *GetNameSafe(OutHits[HitIdx].Component.Get()));
-//		}
-
-		return (OutHits.Num() > 0);
+		return false;
 	}
 
-	return false;
-}
+	bool bSuccess = false;
 
-bool FLevelEditorViewportClient::DropSingleObjectOnActor(FViewportCursorLocation& Cursor, UObject* DroppedObject, AActor* DroppedUponActor, int32 DroppedUponSlot, const FVector& DroppedLocation, EObjectFlags ObjectFlags, TArray<AActor*>& OutNewActors, bool bUsedHitProxy /*= true*/, bool bSelectActors /*= true*/, class UActorFactory* FactoryToUse /*= NULL*/ )
-{
-	if (ensure( DroppedObject ))
+	// Create a transaction if we have more than 1 item. This causes them all to get "un-created" together
+	FScopedTransaction Transaction( NSLOCTEXT("UnrealEd", "CreateActors", "Create Actors") );
+	for ( auto DroppedObject : DroppedObjects )
 	{
-		// Attempt to apply the dropped asset as a material to the actor
 		const bool bAppliedToActor = ( FactoryToUse == NULL ) ? AttemptApplyObjToActor( DroppedObject, DroppedUponActor, DroppedUponSlot ) : false;
-		if ( !bAppliedToActor )
-		{
-			// Actor
-			GEditor->ClickLocation = DroppedLocation;
-			GEditor->ClickPlane = FPlane(DroppedLocation, FVector(0.0f, 0.0f, 1.0f));
-			TArray<FHitResult> Hits;
-			FCollisionQueryParams Param(TEXT("DragDropTrace"), true);
-			// grab the bounds of the actor we're being dropped upon
-			FVector DUAOrigin, DUAExtent;
-			DroppedUponActor->GetActorBounds(true, DUAOrigin, DUAExtent);
-			// now calculate how far back we need to be to collide a ray with it - extend it slightly
-			float WorldDistanceMultiplier = 0.0f;
-			switch( Cursor.GetViewportClient()->GetViewportType() )
-			{
-			case LVT_OrthoXY:	WorldDistanceMultiplier = DUAExtent.Z * 1.1; break; // Top
-			case LVT_OrthoXZ:	WorldDistanceMultiplier = DUAExtent.Y * 1.1; break; // Front
-			case LVT_OrthoYZ:	WorldDistanceMultiplier = DUAExtent.X * 1.1; break; // Side
-			}
-
-			if ( GetVisibleActors(Hits, Cursor.GetOrigin() - Cursor.GetDirection() * WorldDistanceMultiplier, Cursor.GetOrigin() + Cursor.GetDirection() * HALF_WORLD_MAX, Param) )
-			{
-				bool FoundMatch = false;
-
-				// We only care about the collision with the hit proxy actor as line trace will report hits on hidden actors too.
-				for( int32 HitIdx = 0; HitIdx < Hits.Num(); HitIdx++ )
-				{
-					const FHitResult& Hit = Hits[HitIdx];
-					if ( Hit.GetActor() == DroppedUponActor )
-					{
-						GEditor->ClickLocation = Hit.Location;
-						GEditor->ClickPlane = FPlane( Hit.Location,Hit.Normal );
-						FoundMatch = true;
-						break;
-					}
-				}
-
-				// If unsuccessful, use the first visible actor/component* instead
-				// *Limit this just to models/brush for the time being as its only needed for TTP#307379
-				if ( !FoundMatch )
-				{
-					for( int32 HitIdx = 0; HitIdx < Hits.Num(); HitIdx++ )
-					{
-						const FHitResult& Hit = Hits[HitIdx];
-						if((Hit.GetActor() && !Hit.GetActor()->IsHiddenEd())
-							|| (Hit.Component.IsValid() && Hit.Component->IsVisibleInEditor() && Hit.Component->IsA(UModelComponent::StaticClass()) && DroppedUponActor->IsA(ABrush::StaticClass())))
-						{
-							GEditor->ClickLocation = Hit.Location;
-							GEditor->ClickPlane = FPlane( Hit.Location, Hit.Normal );
-							break;
-						}
-					}
-				}
-			}
-			else
-			{
-				// if the line check fails clear the hit proxy flag so that we place it on the background
-				if ( Cursor.GetViewportType() == LVT_Perspective )
-				{
-					bUsedHitProxy = false;
-				}
-			}
 
+		if (!bAppliedToActor)
+		{
 			// Attempt to create actors from the dropped object
-			TArray<AActor*> NewActors = AttemptDropObjAsActors(GetWorld()->GetCurrentLevel(), DroppedObject, Cursor, bUsedHitProxy, bSelectActors, ObjectFlags, FactoryToUse);
+			TArray<AActor*> NewActors = AttemptDropObjAsActors(GetWorld()->GetCurrentLevel(), DroppedObject, Cursor, bSelectActors, ObjectFlags, FactoryToUse);
 
 			if ( NewActors.Num() > 0 )
 			{
 				OutNewActors.Append( NewActors );
-				return true;
+				bSuccess = true;
 			}
 		}
+		else
+		{
+			bSuccess = true;
+		}
 	}
-	return false;
+
+	return bSuccess;
 }
 
 bool FLevelEditorViewportClient::DropObjectsOnBSPSurface( FSceneView* View, FViewportCursorLocation& Cursor, const TArray<UObject*>& DroppedObjects, HModel* TargetProxy, EObjectFlags ObjectFlags, TArray<AActor*>& OutNewActors, bool bSelectActors, UActorFactory* FactoryToUse )
 {
-	if ( TargetProxy == NULL || DroppedObjects.Num() == 0 )
+	bool bResult = false;
+
+	if (DroppedObjects.Num() == 0)
 	{
 		return false;
 	}
@@ -1030,68 +859,23 @@ bool FLevelEditorViewportClient::DropObjectsOnBSPSurface( FSceneView* View, FVie
 		return true;
 	}
 
-	uint32 SurfaceIndex = INDEX_NONE;
-	const int32 DropX = Cursor.GetCursorPos().X;
-	const int32 DropY = Cursor.GetCursorPos().Y;
-
-	bool bResult = false;
-
-	GEditor->ClickLocation = TargetProxy->GetModelComponent()->GetComponentLocation();
-	GEditor->ClickPlane = FPlane( TargetProxy->GetModelComponent()->GetComponentLocation(), FVector( 0.0f, 0.0f, 1.0f ) );
-
-	TArray<FHitResult> Hits;
+	bool bSuccess = false;
 
-	FCollisionQueryParams Param(TEXT("DragDropTrace"), true);
-	if ( GetVisibleActors(Hits, Cursor.GetOrigin(), Cursor.GetOrigin() + Cursor.GetDirection() * HALF_WORLD_MAX, Param) )
+	// Create a transaction if we have more than 1 item. This causes them all to get "un-created" together
+	FScopedTransaction Transaction(NSLOCTEXT("UnrealEd", "CreateActors", "Create Actors"));
+	for (auto DroppedObject : DroppedObjects)
 	{
-		bool FoundMatch = false;
-
-		// We only care about the collision with the hit proxy component as line trace will report hits on hidden components too.
-		for( int32 HitIdx = 0; HitIdx < Hits.Num(); HitIdx++ )
-		{
-			const FHitResult& Hit = Hits[HitIdx];
-			if ( Hit.Component == TargetProxy->GetModelComponent() )
-			{
-				GEditor->ClickLocation = Hit.Location;
-				GEditor->ClickPlane = FPlane( Hit.Location,Hit.Normal );
-				FoundMatch = true;
-				break;
-			}
-		}
-
-		// If unsuccessful, use the first visible component instead
-		if ( !FoundMatch )
-		{
-			for( int32 HitIdx = 0; HitIdx < Hits.Num(); HitIdx++ )
-			{
-				const FHitResult& Hit = Hits[HitIdx];
-				if ( Hit.Component.IsValid() && Hit.Component->IsVisibleInEditor() )
-				{
-					GEditor->ClickLocation = Hit.Location;
-					GEditor->ClickPlane = FPlane( Hit.Location, Hit.Normal );
-					break;
-				}
-			}
-		}
+		// Attempt to create actors from the dropped object
+		TArray<AActor*> NewActors = AttemptDropObjAsActors(GetWorld()->GetCurrentLevel(), DroppedObject, Cursor, bSelectActors, ObjectFlags, FactoryToUse);
 
-		for ( int32 DroppedObjectsIdx = 0; DroppedObjectsIdx < DroppedObjects.Num(); ++DroppedObjectsIdx )
+		if (NewActors.Num() > 0)
 		{
-			UObject* AssetObj = DroppedObjects[ DroppedObjectsIdx ];
-			ensure( AssetObj );
-
-			// Attempt to create an actor from the dropped object
-			const bool bUsedHitProxy = true;
-			TArray<AActor*> NewActors = AttemptDropObjAsActors(GetWorld()->GetCurrentLevel(), AssetObj, Cursor, bUsedHitProxy, bSelectActors, ObjectFlags, FactoryToUse);
-
-			if ( NewActors.Num() > 0 )
-			{
-				OutNewActors.Append(NewActors);
-				bResult = true;
-			}
+			OutNewActors.Append(NewActors);
+			bSuccess = true;
 		}
 	}
 
-	return bResult;
+	return bSuccess;
 }
 
 /**
@@ -1100,7 +884,6 @@ bool FLevelEditorViewportClient::DropObjectsOnBSPSurface( FSceneView* View, FVie
  * @param	View				The SceneView for the dropped-in viewport
  * @param	Cursor				Mouse cursor location
  * @param	DroppedObjects		Array of objects dropped into the viewport
- * @param	TargetProxy			Hit proxy representing the dropped upon manipulation widget
  *
  * @return	true if the drop operation was successfully handled; false otherwise
  */
@@ -1151,301 +934,73 @@ bool FLevelEditorViewportClient::HasDropPreviewActors() const
 bool FLevelEditorViewportClient::UpdateDropPreviewActors(int32 MouseX, int32 MouseY, const TArray<UObject*>& DroppedObjects, bool& out_bDroppedObjectsVisible, UActorFactory* FactoryToUse)
 {
 	out_bDroppedObjectsVisible = false;
-	if( HasDropPreviewActors() )
+	if( !HasDropPreviewActors() )
 	{
-		// While dragging actors, allow viewport updates
-		bNeedsRedraw = true;
-
-		// If the mouse did not move so there is no need to update anything
-		if ( MouseX == DropPreviewMouseX && MouseY == DropPreviewMouseY )
-		{
-			return false;
-		}
-
-		// Update the cached mouse position
-		DropPreviewMouseX = MouseX;
-		DropPreviewMouseY = MouseY;
-
-		// Get the center point between all the drop preview actors for use in calculations below
-		// Also, build a list of valid AActor* pointers
-		FVector Origin = FVector::ZeroVector;
-		TArray<AActor*> DraggingActors;
-		for ( auto ActorIt = DropPreviewActors.CreateConstIterator(); ActorIt; ++ActorIt )
-		{
-			AActor* DraggingActor = (*ActorIt).Get();
-
-			if ( DraggingActor )
-			{
-				DraggingActors.Add(DraggingActor);
-				Origin += DraggingActor->GetActorLocation();
-			}
-		}
-
-		// If there were not valid actors after all, there is nothing to update
-		if ( DraggingActors.Num() == 0 )
-		{
-			return false;
-		}
-
-		// Finish the calculation of the actors origin now that we know we are not dividing by zero
-		Origin /= DraggingActors.Num();
-
-		AActor* FirstDraggingActor = DraggingActors[0];
-		FSceneViewFamilyContext ViewFamily(FSceneViewFamily::ConstructionValues(
-			Viewport, 
-			GetScene(),
-			EngineShowFlags)
-			.SetRealtimeUpdate( IsRealtime() ));
-		FSceneView* View = CalcSceneView( &ViewFamily );
-		FViewportCursorLocation Cursor(View, this, MouseX, MouseY);
-		FVector4 ScreenSpacePos = View->PixelToScreen(MouseX, MouseY, 0.f);
-
-		FVector MouseLocation = View->InvViewProjectionMatrix.TransformFVector4(ScreenSpacePos);
-		FPlane MousePlane = FPlane(ForceInit);
-
-		// Modify the ShowFlags for the scene so we can render the hit proxies without any axis widgets. 
-		// Store original ShowFlags and assign them back when we're done
-		const bool bOldModeWidgets1 = EngineShowFlags.ModeWidgets;
-		const bool bOldModeWidgets2 = View->Family->EngineShowFlags.ModeWidgets;
-
-		EngineShowFlags.ModeWidgets = 0;
-		FSceneViewFamily* SceneViewFamily = const_cast< FSceneViewFamily* >( View->Family );
-		SceneViewFamily->EngineShowFlags.ModeWidgets = 0;
-
-		// Invalidate the hit proxy map so it will be rendered out again when GetHitProxy is called
-		HHitProxy* HitProxy = Viewport->GetHitProxy(MouseX, MouseY);
-
-		// We should never encounter a widget axis.  If we do, then something is wrong
-		// with our ShowFlags (or the widget drawing code - in which case, try refreshing the hit proxies first)
-		if( HitProxy && HitProxy->IsA( HWidgetAxis::StaticGetType() ) )
-		{
-			Viewport->InvalidateHitProxy();
-			HitProxy = Viewport->GetHitProxy(MouseX, MouseY);
-			check( !HitProxy || ( HitProxy && !HitProxy->IsA( HWidgetAxis::StaticGetType() ) ) );
-		}
-
-		// Restore the original flags
-		EngineShowFlags.ModeWidgets = bOldModeWidgets1;
-		SceneViewFamily->EngineShowFlags.ModeWidgets = bOldModeWidgets2;
-
-		out_bDroppedObjectsVisible = true;
-
-		if ( HitProxy == NULL )
-		{
-			// Background
-			const float DistanceMultiplier = ( Cursor.GetViewportType() == LVT_Perspective ) ? 1.0f : HALF_WORLD_MAX;
-
-			MouseLocation = Cursor.GetOrigin() + Cursor.GetDirection() * DistanceMultiplier;
-			MousePlane = FPlane( 0, 0, 0, 0 );
-		}
-		else if (HitProxy->IsA(HActor::StaticGetType()) || HitProxy->IsA(HBSPBrushVert::StaticGetType()))
-		{
-			AActor* TargetActor = NULL;
-			FVector TargetLocation = FVector::ZeroVector;
-			if (HitProxy->IsA(HActor::StaticGetType()))
-			{
-				HActor* TargetProxy = static_cast<HActor*>(HitProxy);
-				TargetActor = TargetProxy->Actor;
-				TargetLocation = TargetActor ? TargetActor->GetActorLocation() : FVector::ZeroVector;
-			}
-			else if (HitProxy->IsA(HBSPBrushVert::StaticGetType()))
-			{
-				HBSPBrushVert* TargetProxy = static_cast<HBSPBrushVert*>(HitProxy);
-				TargetActor = TargetProxy->Brush.Get();
-				TargetLocation = TargetActor ? (TargetProxy->Vertex ? TargetActor->ActorToWorld().TransformPosition(*TargetProxy->Vertex) : TargetActor->GetActorLocation()) : FVector::ZeroVector;
-			}
+		return false;
+	}
 
-			if (TargetActor != NULL)
-			{
-				if ( FactoryToUse == NULL )
-				{
-					out_bDroppedObjectsVisible = false;
-					// go through all objects that can be dropped and if any can be applied to actor, hide dropped preview objects
-					for ( int32 DroppedObjectsIdx = 0; DroppedObjectsIdx < DroppedObjects.Num(); ++DroppedObjectsIdx )
-					{
-						UObject* AssetObj = DroppedObjects[DroppedObjectsIdx];
-						ensure( AssetObj );
+	// While dragging actors, allow viewport updates
+	bNeedsRedraw = true;
 
-						// Attempt to apply the dropped asset as a material to the actor, just test if it is possible
-						if ( !AttemptApplyObjToActor( AssetObj, TargetActor, -1, true ) )
-						{
-							// hide all objects as we are applying them, not creating new object
-							out_bDroppedObjectsVisible = true;
-							break;
-						}
-					}
-				}
-
-				// Actor
-				MouseLocation = TargetLocation;
-				MousePlane = FPlane(TargetLocation, FVector(0.0f, 0.0f, 1.0f));
-				TArray<FHitResult> Hits;
-				FCollisionQueryParams Param(TEXT("DragDropTrace"), true);
-				Param.AddIgnoredActors(DraggingActors);
-
-				// grab the bounds of the actor we're being dropped upon
-				FVector DUAOrigin, DUAExtent;
-				TargetActor->GetActorBounds(true, DUAOrigin, DUAExtent);
-				// now calculate how far back we need to be to collide a ray with it - extend it slightly
-				float WorldDistanceMultiplier = 0.0f;
-				switch( Cursor.GetViewportClient()->GetViewportType() )
-				{
-				case LVT_OrthoXY:	WorldDistanceMultiplier = DUAExtent.Z * 1.1; break; // Top
-				case LVT_OrthoXZ:	WorldDistanceMultiplier = DUAExtent.Y * 1.1; break; // Front
-				case LVT_OrthoYZ:	WorldDistanceMultiplier = DUAExtent.X * 1.1; break; // Side
-				}
+	// If the mouse did not move, there is no need to update anything
+	if ( MouseX == DropPreviewMouseX && MouseY == DropPreviewMouseY )
+	{
+		return false;
+	}
 
-				if ( GetVisibleActors(Hits, Cursor.GetOrigin() - Cursor.GetDirection() * WorldDistanceMultiplier, Cursor.GetOrigin() + Cursor.GetDirection() * HALF_WORLD_MAX, Param) )
-				{
-					bool FoundMatch = false;
+	// Update the cached mouse position
+	DropPreviewMouseX = MouseX;
+	DropPreviewMouseY = MouseY;
 
-					// We only care about the collision with the hit proxy actor as line trace will report hits on hidden actors too.
-					for( int32 HitIdx = 0; HitIdx < Hits.Num(); HitIdx++ )
-					{
-						const FHitResult& Hit = Hits[HitIdx];
-						if ( Hit.GetActor() == TargetActor )
-						{
-							MouseLocation = Hit.Location;
-							MousePlane = FPlane( Hit.Location,Hit.Normal );
-							FoundMatch = true;
-							break;
-						}
-					}
+	// Get the center point between all the drop preview actors for use in calculations below
+	// Also, build a list of valid AActor* pointers
+	FVector ActorOrigin = FVector::ZeroVector;
+	TArray<AActor*> DraggingActors;
+	for ( auto ActorIt = DropPreviewActors.CreateConstIterator(); ActorIt; ++ActorIt )
+	{
+		AActor* DraggingActor = (*ActorIt).Get();
 
-					// If unsuccessful, use the first visible actor/component* instead
-					// *Limit this just to models/brush for the time being as its only needed for TTP#307379
-					if ( !FoundMatch )
-					{
-						for( int32 HitIdx = 0; HitIdx < Hits.Num(); HitIdx++ )
-						{
-							const FHitResult& Hit = Hits[HitIdx];
-							if ((Hit.GetActor() && !Hit.GetActor()->IsHiddenEd())
-							|| (Hit.Component.IsValid() && Hit.Component->IsVisibleInEditor() && Hit.Component->IsA(UModelComponent::StaticClass()) && TargetActor->IsA(ABrush::StaticClass())))
-							{
-								MouseLocation = Hit.Location;
-								MousePlane = FPlane( Hit.Location, Hit.Normal );
-								break;
-							}
-						}
-					}
-				}
-				else
-				{
-					// if the line check fails clear the hit proxy so that we place it on the background
-					if ( Cursor.GetViewportType() == LVT_Perspective )
-					{
-						HitProxy = nullptr;
-					}
-				}
-			}
-		}
-		else if (HitProxy->IsA(HModel::StaticGetType()))
+		if ( DraggingActor )
 		{
-			// BSP surface
-			HModel* TargetProxy = static_cast<HModel*>(HitProxy);
-			MouseLocation = TargetProxy->GetModelComponent()->GetComponentLocation();
-			MousePlane = FPlane( TargetProxy->GetModelComponent()->GetComponentLocation(), FVector( 0.0f, 0.0f, 1.0f ) );
-
-			TArray<FHitResult> Hits;
-			FCollisionQueryParams Param(TEXT("DragDropTrace"), true);
-			Param.AddIgnoredActors(DraggingActors);		
-			if ( GetVisibleActors(Hits, Cursor.GetOrigin(), Cursor.GetOrigin() + Cursor.GetDirection() * HALF_WORLD_MAX, Param) )
-			{
-				bool FoundMatch = false;
-
-				// We only care about the collision with the hit proxy component as line trace will report hits on hidden components too.
-				for( int32 HitIdx = 0; HitIdx < Hits.Num(); HitIdx++ )
-				{
-					const FHitResult& Hit = Hits[HitIdx];
-					if ( Hit.Component == TargetProxy->GetModelComponent() )
-					{
-						MouseLocation = Hit.Location;
-						MousePlane = FPlane( Hit.Location,Hit.Normal );
-						FoundMatch = true;
-						break;
-					}
-				}
-
-				// If unsuccessful, use the first visible component instead
-				if ( !FoundMatch )
-				{
-					for( int32 HitIdx = 0; HitIdx < Hits.Num(); HitIdx++ )
-					{
-						const FHitResult& Hit = Hits[HitIdx];
-						if ( Hit.Component.IsValid() && Hit.Component->IsVisibleInEditor() )
-						{
-							MouseLocation = Hit.Location;
-							MousePlane = FPlane( Hit.Location, Hit.Normal );
-							break;
-						}
-					}
-				}
-			}
-			else
-			{
-				// if the line check fails clear the hit proxy so that we place it on the background
-				if ( Cursor.GetViewportType() == LVT_Perspective )
-				{
-					HitProxy = nullptr;
-				}
-			}
+			DraggingActors.Add(DraggingActor);
+			ActorOrigin += DraggingActor->GetActorLocation();
 		}
+	}
 
-		// Update the click location. This is to allow actions that happen on mouse release to know where the preview actor is
-		GEditor->ClickLocation = MouseLocation;
-		GEditor->ClickPlane = MousePlane;
+	// If there were not valid actors after all, there is nothing to update
+	if ( DraggingActors.Num() == 0 )
+	{
+		return false;
+	}
 
-		// If the actor was successfully placed, but it was done on the background of a perspective viewport, forcibly move the actor in front
-		// of the camera
-		if ( !HitProxy && Cursor.GetViewportType() == LVT_Perspective )
-		{
-			for ( auto ActorIt = DraggingActors.CreateConstIterator(); ActorIt; ++ActorIt )
-			{
-				// Move the actor in front of the camera while preserving the relative offset to the other actors
-				AActor* DraggingActor = *ActorIt;
-				FVector OffsetToOrigin = DraggingActor->GetActorLocation() - Origin;
-				GEditor->MoveActorInFrontOfCamera( *DraggingActor, Cursor.GetOrigin(), Cursor.GetDirection() );
-				FVector NewLoc = DraggingActor->GetActorLocation() + OffsetToOrigin;
+	// Finish the calculation of the actors origin now that we know we are not dividing by zero
+	ActorOrigin /= DraggingActors.Num();
 
-				if ( DraggingActor->IsA( ABrush::StaticClass() ) )
-				{
-					NewLoc = AttemptToSnapLocationToOriginPlane( Cursor, NewLoc );
-				}
+	FSceneViewFamilyContext ViewFamily(FSceneViewFamily::ConstructionValues(
+		Viewport, 
+		GetScene(),
+		EngineShowFlags)
+		.SetRealtimeUpdate( IsRealtime() ));
+	FSceneView* View = CalcSceneView( &ViewFamily );
+	FViewportCursorLocation Cursor(View, this, MouseX, MouseY);
 
-				DraggingActor->TeleportTo( NewLoc, DraggingActor->GetActorRotation() );
-				DraggingActor->SetIsTemporarilyHiddenInEditor( !out_bDroppedObjectsVisible );
-			}
-		}
-		else
-		{
-			const bool bSnapNewObjectsToFloor = GetDefault<ULevelEditorViewportSettings>()->bSnapNewObjectsToFloor;
+	const FActorPositionTraceResult TraceResult = FActorPositioning::TraceWorldForPositionWithDefault(Cursor, *View, &DraggingActors);
 
-			// Move the actor to the target location while preserving the relative offset to the other actors
-			FVector Collision = FirstDraggingActor->GetPlacementExtent();
-			FVector PreSnapLocation = MouseLocation + MousePlane * (FVector::BoxPushOut(MousePlane, Collision) + 0.1f);
-			FVector NewLocation = PreSnapLocation;
-			FSnappingUtils::SnapPointToGrid( NewLocation, FVector(0, 0, 0) );
+	GEditor->ClickLocation = TraceResult.Location;
+	GEditor->ClickPlane = FPlane(TraceResult.Location, TraceResult.SurfaceNormal);
 
-			if( bSnapNewObjectsToFloor && HitProxy != NULL && Cursor.GetViewportType() == LVT_Perspective )
-			{
-				NewLocation.Z = PreSnapLocation.Z;
-			}
+	// Snap the new location if snapping is enabled
+	FSnappingUtils::SnapPointToGrid(GEditor->ClickLocation, FVector::ZeroVector);
 
-			// @todo Support vertex normal orientation
-			FVector UnusedVertexNormal;
-			FSnappingUtils::SnapLocationToNearestVertex( NewLocation, FVector2D(MouseX,MouseY), this, UnusedVertexNormal );
+	for (AActor* Actor : DraggingActors)
+	{
+		const UActorFactory* ActorFactory = FactoryToUse ? FactoryToUse : GEditor->FindActorFactoryForActorClass(Actor->GetClass());
+		const FTransform ActorTransform = FActorPositioning::GetSnappedSurfaceAlignedTransform(this, ActorFactory, TraceResult.Location, TraceResult.SurfaceNormal, Actor->GetPlacementExtent());
 
-			for ( auto ActorIt = DraggingActors.CreateConstIterator(); ActorIt; ++ActorIt )
-			{
-				AActor* DraggingActor = *ActorIt;
-				FVector OffsetToOrigin = DraggingActor->GetActorLocation() - Origin;
-				DraggingActor->TeleportTo(NewLocation + OffsetToOrigin, DraggingActor->GetActorRotation());
-				DraggingActor->SetIsTemporarilyHiddenInEditor( !out_bDroppedObjectsVisible );
-			}
-		}
-		return true;
+		Actor->SetActorTransform(ActorTransform);
 	}
-	return false;
+
+	return true;
 }
 
 void FLevelEditorViewportClient::DestroyDropPreviewActors()
@@ -1545,6 +1100,8 @@ bool FLevelEditorViewportClient::DropObjectsAtCoordinates(int32 MouseX, int32 Mo
 
 	if(DroppedObjects.Num() > 0)
 	{
+		Viewport->InvalidateHitProxy();
+
 		FSceneViewFamilyContext ViewFamily(FSceneViewFamily::ConstructionValues(
 			Viewport, 
 			GetScene(),
@@ -1553,38 +1110,36 @@ bool FLevelEditorViewportClient::DropObjectsAtCoordinates(int32 MouseX, int32 Mo
 		FSceneView* View = CalcSceneView( &ViewFamily );
 		FViewportCursorLocation Cursor(View, this, MouseX, MouseY);
 
-		FVector4 ScreenSpacePos = View->PixelToScreen(MouseX, MouseY, 0.f);
-		GEditor->ClickLocation = View->InvViewProjectionMatrix.TransformFVector4(ScreenSpacePos);
-		GEditor->ClickPlane = FPlane(ForceInit);
+		HHitProxy* HitProxy = Viewport->GetHitProxy(Cursor.GetCursorPos().X, Cursor.GetCursorPos().Y);
 
-		// Invalidate the hit proxy map so it will be rendered out again when GetHitProxy is called
-		Viewport->InvalidateHitProxy();
-		HHitProxy* HitProxy = Viewport->GetHitProxy(MouseX, MouseY);
+		const FActorPositionTraceResult TraceResult = FActorPositioning::TraceWorldForPositionWithDefault(Cursor, *View);
+		
+		GEditor->ClickLocation = TraceResult.Location;
+		GEditor->ClickPlane = FPlane(TraceResult.Location, TraceResult.SurfaceNormal);
+
+		// Snap the new location if snapping is enabled
+		FSnappingUtils::SnapPointToGrid(GEditor->ClickLocation, FVector::ZeroVector);
 
 		EObjectFlags ObjectFlags = bCreateDropPreview ? RF_Transient : RF_Transactional;
-		if ( HitProxy == NULL )
+		if ( HitProxy == nullptr )
 		{
 			bResult = DropObjectsOnBackground( Cursor, DroppedObjects, ObjectFlags, OutNewActors, SelectActors, FactoryToUse );
 		}
 		else if (HitProxy->IsA(HActor::StaticGetType()) || HitProxy->IsA(HBSPBrushVert::StaticGetType()))
 		{
-			const bool bUsedHitProxy = true;
 			AActor* TargetActor = NULL;
-			FVector TargetLocation = FVector::ZeroVector;
 			int32 TargetMaterialSlot = -1;
 
 			if (HitProxy->IsA(HActor::StaticGetType()))
 			{
 				HActor* TargetProxy = static_cast<HActor*>(HitProxy);
 				TargetActor = TargetProxy->Actor;
-				TargetLocation = TargetActor ? TargetActor->GetActorLocation() : FVector::ZeroVector;
 				TargetMaterialSlot = TargetProxy->MaterialIndex;
 			}
 			else if (HitProxy->IsA(HBSPBrushVert::StaticGetType()))
 			{
 				HBSPBrushVert* TargetProxy = static_cast<HBSPBrushVert*>(HitProxy);
 				TargetActor = TargetProxy->Brush.Get();
-				TargetLocation = TargetActor ? (TargetProxy->Vertex ? TargetActor->ActorToWorld().TransformPosition(*TargetProxy->Vertex) : TargetActor->GetActorLocation()) : FVector::ZeroVector;
 			}
 
 			// If shift is pressed set the material slot to -1, so that it's applied to every slot.
@@ -1606,7 +1161,7 @@ bool FLevelEditorViewportClient::DropObjectsAtCoordinates(int32 MouseX, int32 Mo
 					FactoryToUse != NULL ||
 					!AttemptApplyObjToActor(DroppedObjects[0], TargetActor, TargetMaterialSlot, true) )
 				{
-					bResult = DropObjectsOnActor( Cursor, DroppedObjects, TargetActor, TargetMaterialSlot, &TargetLocation, ObjectFlags, OutNewActors, bUsedHitProxy, SelectActors, FactoryToUse );
+					bResult = DropObjectsOnActor( Cursor, DroppedObjects, TargetActor, TargetMaterialSlot, ObjectFlags, OutNewActors, SelectActors, FactoryToUse );
 				}
 				else
 				{
@@ -1615,8 +1170,7 @@ bool FLevelEditorViewportClient::DropObjectsAtCoordinates(int32 MouseX, int32 Mo
 						TargetActor = static_cast<AActor*>(*It);
 						if( TargetActor )
 						{
-							TargetLocation = TargetActor->GetActorLocation();
-							DropObjectsOnActor( Cursor, DroppedObjects, TargetActor, TargetMaterialSlot, &TargetLocation, ObjectFlags, OutNewActors, bUsedHitProxy, SelectActors, FactoryToUse );
+							DropObjectsOnActor( Cursor, DroppedObjects, TargetActor, TargetMaterialSlot, ObjectFlags, OutNewActors, SelectActors, FactoryToUse );
 							bResult = true;
 						}
 					}
@@ -1669,7 +1223,7 @@ bool FLevelEditorViewportClient::DropObjectsAtCoordinates(int32 MouseX, int32 Mo
 
 				// Invalidate the hit proxy now so the drop preview will be accurate.
 				// We don't invalidate the hit proxy in the drop preview update itself because it is slow.
-				Viewport->InvalidateHitProxy();
+				//Viewport->InvalidateHitProxy();
 			}
 			// Dropping the actors rather than a preview? Probably want to select them all then. 
 			else if(!bCreateDropPreview && SelectActors && OutNewActors.Num() > 0)
@@ -2304,7 +1858,135 @@ void FLevelEditorViewportClient::UpdateViewForLockedActor()
 	};
 };*/
 
+/** Trim the specified line to the planes of the frustum */
+void TrimLineToFrustum(const FConvexVolume& Frustum, FVector& Start, FVector& End)
+{
+	FVector Intersection;
+	for (const auto& Plane : Frustum.Planes)
+	{
+		if (FMath::SegmentPlaneIntersection(Start, End, Plane, Intersection))
+		{
+			// Chop the line inside the frustum
+			if ((static_cast<const FVector&>(Plane) | (Intersection - End)) > 0.0f)
+			{
+				Start = Intersection;
+			}
+			else
+			{
+				End = Intersection;
+			}
+		}
+	}
+}
+
+void FLevelEditorViewportClient::ProjectActorsIntoWorld(const TArray<AActor*>& Actors, FViewport* Viewport, const FVector& Drag, const FRotator& Rot)
+{
+	// Compile an array of selected actors
+	FSceneViewFamilyContext ViewFamily(FSceneViewFamily::ConstructionValues(
+		Viewport, 
+		GetScene(),
+		EngineShowFlags)
+		.SetRealtimeUpdate( IsRealtime() ));
+	// SceneView is deleted with the ViewFamily
+	FSceneView* SceneView = CalcSceneView( &ViewFamily );
+
+	// Calculate the frustum so we can trim rays to it
+	const FConvexVolume Frustum;
+	GetViewFrustumBounds(const_cast<FConvexVolume&>(Frustum), SceneView->ViewProjectionMatrix, true);
+
+	const FMatrix InputCoordSystem = GetWidgetCoordSystem();
+	const EAxisList::Type CurrentAxis = GetCurrentWidgetAxis();
+	
+	const FVector DeltaTranslation = (GLevelEditorModeTools().PivotLocation - GLevelEditorModeTools().CachedLocation) + Drag;
+
+
+	// Loop over all the actors and attempt to snap them along the drag axis normal
+	for (auto* Actor : Actors)
+	{
+		// Use the Delta of the Mode tool with the actor pre drag location to avoid accumulating snapping offsets
+		FVector NewActorPosition;
+		if (const FTransform* PreDragTransform = PreDragActorTransforms.Find(Actor))
+		{
+			NewActorPosition = PreDragTransform->GetLocation() + DeltaTranslation;
+		}
+		else
+		{
+			const FTransform& ActorTransform = PreDragActorTransforms.Add(Actor, Actor->GetTransform());
+			NewActorPosition = ActorTransform.GetLocation() + DeltaTranslation;
+		}
+
+		FViewportCursorLocation Cursor(SceneView, this, 0, 0);
+
+		FActorPositionTraceResult TraceResult;
+		bool bSnapped = false;
+
+		bool bIsOnScreen = false;
+		{
+			// We only snap things that are on screen
+			FVector2D ScreenPos;
+			FIntPoint ViewportSize = Viewport->GetSizeXY();
+			if (SceneView->WorldToPixel(NewActorPosition, ScreenPos) && FMath::IsWithin<float>(ScreenPos.X, 0, ViewportSize.X) && FMath::IsWithin<float>(ScreenPos.Y, 0, ViewportSize.Y))
+			{
+				bIsOnScreen = true;
+				Cursor = FViewportCursorLocation(SceneView, this, ScreenPos.X, ScreenPos.Y);
+			}
+		}
+
+		if (bIsOnScreen)
+		{
+			// Determine how we're going to attempt to project the object onto the world
+			if (CurrentAxis == EAxisList::XY || CurrentAxis == EAxisList::XZ || CurrentAxis == EAxisList::YZ)
+			{
+				// Snap along the perpendicular axis
+				const FVector PlaneNormal = CurrentAxis == EAxisList::XY ? FVector(0, 0, 1) : CurrentAxis == EAxisList::XZ ? FVector(0, 1, 0) : FVector(1, 0, 0);
+				FVector TraceDirection = InputCoordSystem.TransformVector(PlaneNormal);
+
+				// Make sure the trace normal points along the view direction
+				if (FVector::DotProduct(SceneView->GetViewDirection(), TraceDirection) < 0.f)
+				{
+					TraceDirection = -TraceDirection;
+				}
+
+				FVector RayStart	= NewActorPosition - (TraceDirection * HALF_WORLD_MAX/2);
+				FVector RayEnd		= NewActorPosition + (TraceDirection * HALF_WORLD_MAX/2);
+
+				TrimLineToFrustum(Frustum, RayStart, RayEnd);
+
+				TraceResult = FActorPositioning::TraceWorldForPosition(*GetWorld(), *SceneView, RayStart, RayEnd, &Actors);
+			}
+			else
+			{
+				TraceResult = FActorPositioning::TraceWorldForPosition(Cursor, *SceneView, &Actors);
+			}
+		}
+				
+		if (TraceResult.State == FActorPositionTraceResult::HitSuccess)
+		{
+			// Move the actor to the position of the trace hit using the spawn offset rules
+			// We only do this if we found a valid hit (we don't want to move the actor in front of the camera by default)
+
+			const auto* Factory = GEditor->FindActorFactoryForActorClass(Actor->GetClass());
+			
+			const auto ActorLocation = Actor->GetActorLocation();
+			const auto ActorRotation = Actor->GetActorRotation().Quaternion();
+
+			FTransform* PreDragActorTransform = PreDragActorTransforms.Find(Actor);
+			check(PreDragActorTransform);
 
+			// Compute the surface aligned transform. Note we do not use the snapped version here as our DragDelta is already snapped
+			const FTransform ActorTransform = FActorPositioning::GetSurfaceAlignedTransform(Factory,
+				TraceResult.Location, TraceResult.SurfaceNormal, Actor->GetPlacementExtent(), *PreDragActorTransform);
+			
+			const FRotator DeltaRotation(ActorTransform.GetRotation() * ActorRotation.Inverse());
+			ApplyDeltaToActor(Actor, ActorTransform.GetLocation() - ActorLocation, DeltaRotation, FVector(0.f, 0.f, 0.f));
+		}
+		else
+		{
+			// Didn't find a valid surface snapping candidate, just apply the deltas directly
+			ApplyDeltaToActor(Actor, NewActorPosition - Actor->GetActorLocation(), Rot, FVector(0.f, 0.f, 0.f));
+		}
+	}
+}
 
 bool FLevelEditorViewportClient::InputWidgetDelta( FViewport* Viewport, EAxisList::Type CurrentAxis, FVector& Drag, FRotator& Rot, FVector& Scale )
 {
@@ -2361,15 +2043,32 @@ bool FLevelEditorViewportClient::InputWidgetDelta( FViewport* Viewport, EAxisLis
 				if(!MiddleMouseButtonDown)
 				{
 					bool bSnapped = FSnappingUtils::SnapActorsToNearestActor( Drag, this );
-					// Check to see if the actor needs snapping first
-					if ( !bSnapped )
-					{	
-						FSnappingUtils::SnapDraggedActorsToNearestVertex( Drag, this );
+					bSnapped = bSnapped || FSnappingUtils::SnapDraggedActorsToNearestVertex( Drag, this );
+
+					// If we are only changing position, project the actors onto the world
+					const bool bOnlyTranslation = !Drag.IsZero() && Rot.IsZero() && Scale.IsZero();
 
-						// Apply deltas to selected actors/components or viewport cameras
+					const EAxisList::Type CurrentAxis = GetCurrentWidgetAxis();
+					const bool bSingleAxisDrag = CurrentAxis == EAxisList::X || CurrentAxis == EAxisList::Y || CurrentAxis == EAxisList::Z;
+					if (!bSnapped && !bSingleAxisDrag && GetDefault<ULevelEditorViewportSettings>()->SnapToSurface.bEnabled && bOnlyTranslation)
+					{
+						TArray<AActor*> SelectedActors;
+						for (FSelectionIterator It(GEditor->GetSelectedActorIterator()); It; ++It)
+						{
+							if (AActor* Actor = Cast<AActor>(*It))
+							{
+								SelectedActors.Add(Actor);
+							}
+						}
+
+						ProjectActorsIntoWorld(SelectedActors, Viewport, Drag, Rot);
+					}
+					else
+					{
 						ApplyDeltaToActors( Drag, Rot, Scale );
-						ApplyDeltaToRotateWidget( Rot );
 					}
+
+					ApplyDeltaToRotateWidget( Rot );
 				}
 				else
 				{
@@ -2533,7 +2232,11 @@ bool FLevelEditorViewportClient::InputKey(FViewport* Viewport, int32 ControllerI
 	// Compute the click location.
 	if ( InputState.IsAnyMouseButtonDown() )
 	{
-		GEditor->ClickLocation = FVector((View->ViewMatrices.ViewMatrix * View->ViewMatrices.ProjMatrix).Inverse().TransformFVector4(FVector4((HitX - Viewport->GetSizeXY().X / 2.0f) / (Viewport->GetSizeXY().X / 2.0f),(HitY - Viewport->GetSizeXY().Y / 2.0f) / -(Viewport->GetSizeXY().Y / 2.0f),0.5f,1.0f)));
+		const FViewportCursorLocation Cursor(View, this, HitX, HitY);
+		GEditor->ClickLocation = FActorPositioning::TraceWorldForPositionWithDefault(Cursor, *View).Location;
+
+		// Snap the new location if snapping is enabled
+		FSnappingUtils::SnapPointToGrid(GEditor->ClickLocation, FVector::ZeroVector);
 	}
 
 	if (GUnrealEd->ComponentVisManager.HandleInputKey(this, Viewport, Key, Event))
@@ -2643,7 +2346,7 @@ void FLevelEditorViewportClient::TrackingStarted( const FInputEventState& InInpu
 		}
 	}
 
-
+	PreDragActorTransforms.Empty();
 	for (FSelectionIterator It(GEditor->GetSelectedActorIterator()); It && !bIsTrackingBrushModification; ++It)
 	{
 		AActor* Actor = static_cast<AActor*>(*It);
@@ -2809,6 +2512,8 @@ void FLevelEditorViewportClient::TrackingStopped()
 			Actor->PostEditMove( true );
 			GEditor->BroadcastEndObjectMovement( *Actor );
 		}
+
+		GUnrealEd->UpdatePivotLocationForSelection();
 	}
 
 	// End the transaction here if one was started in StartTransaction()
diff --git a/Engine/Source/Editor/UnrealEd/Private/LevelViewportClickHandlers.cpp b/Engine/Source/Editor/UnrealEd/Private/LevelViewportClickHandlers.cpp
--- a/Engine/Source/Editor/UnrealEd/Private/LevelViewportClickHandlers.cpp
+++ b/Engine/Source/Editor/UnrealEd/Private/LevelViewportClickHandlers.cpp
@@ -11,6 +11,7 @@
 #include "Editor/GeometryMode/Public/EditorGeometry.h"
 #include "MessageLog.h"
 #include "ActorEditorUtils.h"
+#include "Editor/ActorPositioning.h"
 
 #define LOCTEXT_NAMESPACE "ClickHandlers"
 
@@ -29,7 +30,7 @@ namespace ClickHandlers
 	 * falling back to "ACTOR ADD" exec and SpawnActor if no factory is found.
 	 * Does nothing if ActorClass is NULL.
 	 */
-	static AActor* PrivateAddActor(UClass* ActorClass, const FVector* ActorLocation=NULL, bool bUseSurfaceOrientation=false, bool bTransactional = true)
+	static AActor* PrivateAddActor(UClass* ActorClass)
 	{
 		if ( ActorClass )
 		{
@@ -37,51 +38,13 @@ namespace ClickHandlers
 			UActorFactory* ActorFactory = GEditor->FindActorFactoryForActorClass( ActorClass );
 			if( ActorFactory )
 			{
-				return GEditor->UseActorFactory( ActorFactory, FAssetData(), ActorLocation, bUseSurfaceOrientation || ActorFactory->bUseSurfaceOrientation );
+				return GEditor->UseActorFactory( ActorFactory, FAssetData(), nullptr );
 			}
 			// otherwise use AddActor so that we can return the newly created actor
 			else
 			{
-				// Get cursor origin and direction in world space.
-				const FViewportCursorLocation& CursorLocation = GCurrentLevelEditingViewportClient->GetCursorWorldLocationFromMousePos();
-
-				// Determine if the actor is being added onto the backdrop. If so, and it is being added from a perspective viewport,
-				// it will be moved in front of the camera.
-				FVector Location = FVector::ZeroVector;
-				bool bOnBackdrop = false;
-				if ( ActorLocation )
-				{
-					Location = *ActorLocation;
-				}
-				else
-				{
-					const FIntPoint& CursorPos = CursorLocation.GetCursorPos();
-					bOnBackdrop = ( GCurrentLevelEditingViewportClient->Viewport->GetHitProxy( CursorPos.X, CursorPos.Y ) == NULL );
-
-					AActor* Default = ActorClass->GetDefaultObject<AActor>();
-					FVector Collision;
-					UCapsuleComponent* CylComp = Cast<UCapsuleComponent>(Default->GetRootComponent());
-					if (CylComp)
-					{
-						Collision = FVector(CylComp->GetScaledCapsuleRadius(), CylComp->GetScaledCapsuleRadius(), CylComp->GetScaledCapsuleHalfHeight());
-					}
-					else
-					{
-						float CollisionRadius, CollisionHeight;
-						Default->GetComponentsBoundingCylinder(CollisionRadius, CollisionHeight);
-						Collision = FVector(CollisionRadius, CollisionRadius, CollisionHeight);
-					}
-
-					Location = GEditor->ClickLocation + GEditor->ClickPlane * ( FVector::BoxPushOut( GEditor->ClickPlane, Collision ) + 0.1 );
-				}					
-				AActor* CreatedActor = GEditor->AddActor( GCurrentLevelEditingViewportClient->GetWorld()->GetCurrentLevel(), ActorClass, Location );
-
-				// If the actor was added to the backdrop in a perspective viewport, move it in front of the camera 
-				if ( CreatedActor && !ActorLocation && GCurrentLevelEditingViewportClient->IsPerspective() && bOnBackdrop )
-				{
-					GEditor->MoveActorInFrontOfCamera( *CreatedActor, CursorLocation.GetOrigin(), CursorLocation.GetDirection() );
-				}
-				return CreatedActor;
+				const FTransform ActorTransform = FActorPositioning::GetCurrentViewportPlacementTransform(*ActorClass->GetDefaultObject<AActor>());
+				return GEditor->AddActor( GCurrentLevelEditingViewportClient->GetWorld()->GetCurrentLevel(), ActorClass, ActorTransform );
 			}
 		}
 		return NULL;
@@ -123,20 +86,9 @@ namespace ClickHandlers
 
 	bool ClickActor(FLevelEditorViewportClient* ViewportClient,AActor* Actor,const FViewportClick& Click,bool bAllowSelectionChange)
 	{
-		// Find the point on the actor component which was clicked on.
-		// Do an accurate trace to avoid legacy pull back by an arbitrary amount.
-		// TRACE_Accurate is needed for texel selection to work.
-		FHitResult Hit(ForceInit);
-		if ( GWorld->LineTraceSingle(Hit, Click.GetOrigin(), Click.GetOrigin() + Click.GetDirection() * HALF_WORLD_MAX, ECC_Pawn, FCollisionQueryParams(true)) )
-		{	
-			GEditor->ClickLocation = Hit.Location;
-			GEditor->ClickPlane = FPlane(Hit.Location,Hit.Normal);
-		}
-
 		// Pivot snapping
 		if( Click.GetKey() == EKeys::MiddleMouseButton && Click.IsAltDown() )
 		{
-			GEditor->ClickLocation = Click.GetOrigin() + Click.GetDirection() * HALF_WORLD_MAX;
 			GEditor->SetPivot( GEditor->ClickLocation, true, false, true );
 
 			return true;
@@ -226,7 +178,7 @@ namespace ClickHandlers
 			else if( Click.GetKey() == EKeys::LeftMouseButton && ViewportClient->Viewport->KeyState(EKeys::S) )
 			{
 				// Create a static mesh.
-				PrivateAddActor( AStaticMeshActor::StaticClass(), NULL, Click.IsAltDown() );
+				PrivateAddActor( AStaticMeshActor::StaticClass() );
 
 				return true;
 			}
@@ -313,7 +265,6 @@ namespace ClickHandlers
 		// Pivot snapping
 		if( Click.GetKey() == EKeys::MiddleMouseButton && Click.IsAltDown() )
 		{
-			GEditor->ClickLocation = Click.GetOrigin() + Click.GetDirection() * HALF_WORLD_MAX;
 			GEditor->SetPivot( GEditor->ClickLocation, true, false, true );
 		}
 		else if( Click.GetKey() == EKeys::RightMouseButton )
@@ -355,7 +306,6 @@ namespace ClickHandlers
 		// Pivot snapping
 		if( Click.GetKey() == EKeys::MiddleMouseButton && Click.IsAltDown() )
 		{
-			GEditor->ClickLocation = Click.GetOrigin() + Click.GetDirection() * HALF_WORLD_MAX;
 			GEditor->SetPivot( GEditor->ClickLocation, true, false, true );
 		}
 		else if( Click.GetKey() == EKeys::RightMouseButton )
@@ -388,7 +338,6 @@ namespace ClickHandlers
 		// Pivot snapping
 		if( Click.GetKey() == EKeys::MiddleMouseButton && Click.IsAltDown() )
 		{
-			GEditor->ClickLocation = Click.GetOrigin() + Click.GetDirection() * HALF_WORLD_MAX;
 			GEditor->SetPivot( GEditor->ClickLocation, true, false, true );
 
 			return true;
@@ -471,7 +420,6 @@ namespace ClickHandlers
 		// Pivot snapping
 		if( Click.GetKey() == EKeys::MiddleMouseButton && Click.IsAltDown() )
 		{
-			GEditor->ClickLocation = Click.GetOrigin() + Click.GetDirection() * HALF_WORLD_MAX;
 			GEditor->SetPivot( GEditor->ClickLocation, true, false, true );
 
 			return true;
@@ -597,7 +545,6 @@ namespace ClickHandlers
 		{
 			// Pivot snapping
 
-			GEditor->ClickLocation = Click.GetOrigin() + Click.GetDirection() * HALF_WORLD_MAX;
 			GEditor->SetPivot( GEditor->ClickLocation, true, false, true );
 
 			return true;
@@ -663,14 +610,10 @@ namespace ClickHandlers
 
 		// Remember hit location for actor-adding.
 		FBspSurf& Surf			= Model->Surfs[iSurf];
-		const FPlane Plane		= Surf.Plane;
-		GEditor->ClickLocation	= FMath::LinePlaneIntersection( Click.GetOrigin(), Click.GetOrigin() + Click.GetDirection(), Plane );
-		GEditor->ClickPlane		= Plane;
 
 		// Pivot snapping
 		if( Click.GetKey() == EKeys::MiddleMouseButton && Click.IsAltDown() )
 		{
-			GEditor->ClickLocation = Click.GetOrigin() + Click.GetDirection() * HALF_WORLD_MAX;
 			GEditor->SetPivot( GEditor->ClickLocation, true, false, true );
 		}
 		else if( Click.GetKey() == EKeys::LeftMouseButton && Click.IsShiftDown() && Click.IsControlDown() )
@@ -759,7 +702,7 @@ namespace ClickHandlers
 		else if( Click.GetKey() == EKeys::LeftMouseButton && ViewportClient->Viewport->KeyState(EKeys::S) )
 		{
 			// Create a static mesh.
-			PrivateAddActor( AStaticMeshActor::StaticClass(), NULL, Click.IsAltDown() );
+			PrivateAddActor( AStaticMeshActor::StaticClass() );
 		}
 		else if( Click.GetKey() == EKeys::LeftMouseButton && ViewportClient->Viewport->KeyState(EKeys::Period) )
 		{
@@ -912,13 +855,9 @@ namespace ClickHandlers
 
 	void ClickBackdrop(FLevelEditorViewportClient* ViewportClient,const FViewportClick& Click)
 	{
-		GEditor->ClickLocation = Click.GetOrigin() + Click.GetDirection() * HALF_WORLD_MAX;
-		GEditor->ClickPlane    = FPlane(0,0,0,0);
-
 		// Pivot snapping
 		if( Click.GetKey() == EKeys::MiddleMouseButton && Click.IsAltDown() )
 		{
-			GEditor->ClickLocation = Click.GetOrigin() + Click.GetDirection() * HALF_WORLD_MAX;
 			GEditor->SetPivot( GEditor->ClickLocation, true, false, true );
 		}
 		else if( Click.GetKey() == EKeys::LeftMouseButton && ViewportClient->Viewport->KeyState(EKeys::A) )
diff --git a/Engine/Source/Editor/UnrealEd/Private/SEditorViewport.cpp b/Engine/Source/Editor/UnrealEd/Private/SEditorViewport.cpp
--- a/Engine/Source/Editor/UnrealEd/Private/SEditorViewport.cpp
+++ b/Engine/Source/Editor/UnrealEd/Private/SEditorViewport.cpp
@@ -257,6 +257,12 @@ void SEditorViewport::BindCommands()
 		//FExecuteAction::CreateStatic( &FLevelEditorActionCallbacks::ExecuteExecCommand, FString( TEXT("CAMERA ALIGN ACTIVEVIEWPORTONLY") ) )
 		);
 
+	CommandListRef.MapAction(
+		Commands.SurfaceSnapping,
+		FExecuteAction::CreateStatic( &SEditorViewport::OnToggleSurfaceSnap ),
+		FCanExecuteAction(),
+		FIsActionChecked::CreateStatic( &SEditorViewport::OnIsSurfaceSnapEnabled ) 
+		);
 
 		// Simple macro for binding many Exposure Setting commands
 #define MAP_EXPOSURE_ACTION( Command, ID ) \
@@ -437,4 +443,15 @@ void SEditorViewport::OnCycleCoordinateSystem()
 	Client->SetWidgetCoordSystemSpace( (ECoordSystem)CoordSystemAsInt );
 }
 
-#undef LOCTEXT_NAMESPACE
\ No newline at end of file
+void SEditorViewport::OnToggleSurfaceSnap()
+{
+	auto* Settings = GetMutableDefault<ULevelEditorViewportSettings>();
+	Settings->SnapToSurface.bEnabled = !Settings->SnapToSurface.bEnabled;
+}
+
+bool SEditorViewport::OnIsSurfaceSnapEnabled()
+{
+	return GetDefault<ULevelEditorViewportSettings>()->SnapToSurface.bEnabled;
+}
+
+#undef LOCTEXT_NAMESPACE
diff --git a/Engine/Source/Editor/UnrealEd/Private/STransformViewportToolbar.cpp b/Engine/Source/Editor/UnrealEd/Private/STransformViewportToolbar.cpp
--- a/Engine/Source/Editor/UnrealEd/Private/STransformViewportToolbar.cpp
+++ b/Engine/Source/Editor/UnrealEd/Private/STransformViewportToolbar.cpp
@@ -27,7 +27,7 @@ namespace TransformViewportToolbarDefs
 class SGridSnapSetting : public SCompoundWidget 
 {
 public:
-	SLATE_BEGIN_ARGS( SGridSnapSetting ){}
+	SLATE_BEGIN_ARGS( SGridSnapSetting ) : _BlockLocation(EMultiBlockLocation::Start) {}
 	
 		/** We need to know about the toolbar we are in */
 		SLATE_ARGUMENT( TSharedPtr<class SViewportToolBar>, ParentToolBar );
@@ -56,6 +56,9 @@ public:
 		/** ToolTip shown on the toggle button */
 		SLATE_ATTRIBUTE( FText, ToggleButtonToolTip )
 
+		/** The button location */
+		SLATE_ARGUMENT( EMultiBlockLocation::Type, BlockLocation )
+
 	SLATE_END_ARGS( )
 
 	/**
@@ -77,7 +80,7 @@ public:
 			ToggleControl = SNew(SCheckBox)
 			.Cursor( EMouseCursor::Default )
 			.Padding(FMargin( 4.0f ))
-			.Style(FEditorStyle::Get(), EMultiBlockLocation::ToName(CheckboxStyle, EMultiBlockLocation::Start))
+			.Style(FEditorStyle::Get(), EMultiBlockLocation::ToName(CheckboxStyle, InArgs._BlockLocation))
 			.OnCheckStateChanged(InArgs._OnCheckStateChanged)
 			.ToolTipText(InArgs._ToggleButtonToolTip)
 			.IsChecked(InArgs._IsChecked)
@@ -381,6 +384,145 @@ void STransformViewportToolBar::Construct( const FArguments& InArgs )
 	SViewportToolBar::Construct(SViewportToolBar::FArguments());
 }
 
+TSharedRef< SWidget > STransformViewportToolBar::MakeSurfaceSnappingButton( FName ToolBarStyle )
+{
+	auto IsSnappingEnabled = []{
+		return GetDefault<ULevelEditorViewportSettings>()->SnapToSurface.bEnabled;
+	};
+
+	const bool bCloseAfterSelection = true;
+	FMenuBuilder MenuBuilder( bCloseAfterSelection, CommandList );
+
+	MenuBuilder.AddMenuEntry(FEditorViewportCommands::Get().SurfaceSnapping);
+
+	MenuBuilder.BeginSection("SurfaceSnappingSettings", LOCTEXT("SnapToSurfaceSettings", "Settings"));
+	{
+		MenuBuilder.AddMenuEntry(
+			LOCTEXT("SnapToSurfaceSettings_Rotation", "Rotate to Surface Normal"),
+			LOCTEXT("SnapToSurfaceSettings_RotationTip", "When checked, snapping an object to a surface will also rotate the object to align to the surface normal"),
+			FSlateIcon(),
+			FUIAction(
+				FExecuteAction::CreateStatic( []{
+					auto& Settings = GetMutableDefault<ULevelEditorViewportSettings>()->SnapToSurface;
+					Settings.bSnapRotation = !Settings.bSnapRotation;
+				}),
+				FCanExecuteAction::CreateStatic( IsSnappingEnabled ),
+				FIsActionChecked::CreateStatic( []{
+					const auto& Settings = GetDefault<ULevelEditorViewportSettings>()->SnapToSurface;
+					return Settings.bSnapRotation;
+				})
+			),
+			NAME_None,
+			EUserInterfaceActionType::RadioButton);
+
+		MenuBuilder.AddWidget(
+			SNew(SHorizontalBox)
+			.IsEnabled(TAttribute<bool>::Create(TAttribute<bool>::FGetter::CreateStatic(IsSnappingEnabled)))
+
+			+ SHorizontalBox::Slot()
+			.AutoWidth()
+			.Padding(FMargin(0.f, 0.f, 5.f, 0.f))
+			[
+				SNew( STextBlock )
+				.Text( LOCTEXT("SnapToSurfaceSettings_Offset", "Surface Offset") )
+			]
+			+SHorizontalBox::Slot()
+			.VAlign( VAlign_Bottom )
+			.FillWidth( 1.f )
+			[
+				SNew( SNumericEntryBox<float> )
+				.Value( 
+					TAttribute<TOptional<float>>::Create(TAttribute<TOptional<float>>::FGetter::CreateStatic([]{
+						const auto& Settings = GetDefault<ULevelEditorViewportSettings>()->SnapToSurface;
+						return TOptional<float>(Settings.SnapOffsetExtent);
+					}))
+				)
+				.OnValueChanged(
+					SNumericEntryBox<float>::FOnValueChanged::CreateStatic([](float Val){
+						GetMutableDefault<ULevelEditorViewportSettings>()->SnapToSurface.SnapOffsetExtent = Val;
+					})
+				)
+				.MinValue(0.f)
+				.MaxValue(HALF_WORLD_MAX)
+				.MaxSliderValue(1000.f) // 'Sensible' range for the slider (10m)
+				.AllowSpin(true)
+			],
+			FText::GetEmpty()
+		);
+
+	}
+	MenuBuilder.EndSection();
+
+	// Have to use a custom widget here to make the checkbox work with the subsequent widget :(
+	return 
+		SNew(SHorizontalBox)
+
+		+SHorizontalBox::Slot()
+		.AutoWidth()
+		[
+			SNew(SCheckBox)
+			.Cursor( EMouseCursor::Default )
+			.Style( FEditorStyle::Get(), EMultiBlockLocation::ToName(FEditorStyle::Join(ToolBarStyle, ".ToggleButton"), EMultiBlockLocation::Start) )
+			.Padding( 0 )
+			.ToolTipText( LOCTEXT("SurfaceSnappingCheckboxDescription", "Open editor surface snapping options") )
+			.IsChecked_Static( []{ return GetDefault<ULevelEditorViewportSettings>()->SnapToSurface.bEnabled ? ESlateCheckBoxState::Checked : ESlateCheckBoxState::Unchecked; })
+			.Content()
+			[
+				SNew( SComboButton )
+				.ButtonStyle( FEditorStyle::Get(), "HoverHintOnly" )
+				.HasDownArrow( false )
+				.ContentPadding( 0 )
+				.ButtonContent()
+				[
+					SNew( SVerticalBox )
+
+					+ SVerticalBox::Slot()
+					.AutoHeight()
+					.Padding(FMargin(5.f, 2.f, 3.f, 0.f))
+					[
+						SNew( SBox )
+						.WidthOverride( TransformViewportToolbarDefs::ToggleImageScale )
+						.HeightOverride( TransformViewportToolbarDefs::ToggleImageScale )
+						.HAlign( HAlign_Center )
+						.VAlign( VAlign_Center )
+						[
+							SNew( SImage )
+							.Image( FEditorStyle::GetBrush("EditorViewport.ToggleSurfaceSnapping") )
+						]
+					]
+
+					+ SVerticalBox::Slot()
+					.AutoHeight()
+					.HAlign( HAlign_Center )
+					.Padding(FMargin(0.f, 0.f, 0.f, 3.f))
+					[
+						SNew( SBox )
+						.WidthOverride( TransformViewportToolbarDefs::DownArrowSize )
+						.HeightOverride( TransformViewportToolbarDefs::DownArrowSize )
+						[
+							SNew(SImage)
+							.Image(FEditorStyle::GetBrush("ComboButton.Arrow"))
+							.ColorAndOpacity(FLinearColor::Black)
+						]
+					]
+				]
+				.MenuContent()
+				[
+					MenuBuilder.MakeWidget()
+				]
+			]
+		]
+
+		+SHorizontalBox::Slot()
+		.AutoWidth()
+		[
+			SNew(SBorder)
+			.Padding(FMargin(1.0f, 0.0f, 0.0f, 0.0f))
+			.BorderImage(FEditorStyle::GetDefaultBrush())
+			.BorderBackgroundColor(FLinearColor::Black)
+		];
+}
+
 TSharedRef< SWidget > STransformViewportToolBar::MakeTransformToolBar( const TSharedPtr< FExtender > InExtenders )
 {
 	FToolBarBuilder ToolbarBuilder( CommandList, FMultiBoxCustomization::None, InExtenders );
@@ -434,6 +576,9 @@ TSharedRef< SWidget > STransformViewportToolBar::MakeTransformToolBar( const TSh
 	
 	ToolbarBuilder.BeginSection("LocationGridSnap");
 	{
+		static FName SurfaceSnapName = FName(TEXT("SurfaceSnap"));
+		ToolbarBuilder.AddWidget( MakeSurfaceSnappingButton( ToolBarStyle ), SurfaceSnapName );
+
 		// Grab the existing UICommand 
 		FUICommandInfo* Command = FEditorViewportCommands::Get().LocationGridSnap.Get();
 
@@ -442,6 +587,7 @@ TSharedRef< SWidget > STransformViewportToolBar::MakeTransformToolBar( const TSh
 		// Setup a GridSnapSetting with the UICommand
 		ToolbarBuilder.AddWidget(	SNew(SGridSnapSetting)
 									.Style(ToolBarStyle)
+									.BlockLocation(EMultiBlockLocation::Middle)
 									.Cursor( EMouseCursor::Default )
 									.IsChecked(this, &STransformViewportToolBar::IsLocationGridSnapChecked)
 									.OnCheckStateChanged(this, &STransformViewportToolBar::HandleToggleLocationGridSnap)
@@ -885,4 +1031,4 @@ void STransformViewportToolBar::HandleToggleScaleGridSnap( ESlateCheckBoxState::
 	GUnrealEd->Exec( GEditor->GetEditorWorldContext().World(), *FString::Printf( TEXT("MODE SCALEGRID=%d"), !GetDefault<ULevelEditorViewportSettings>()->SnapScaleEnabled ? 1 : 0 ) );
 }
 
-#undef LOCTEXT_NAMESPACE 
\ No newline at end of file
+#undef LOCTEXT_NAMESPACE 
diff --git a/Engine/Source/Editor/UnrealEd/Private/Tests/EditorAutomationTests.cpp b/Engine/Source/Editor/UnrealEd/Private/Tests/EditorAutomationTests.cpp
--- a/Engine/Source/Editor/UnrealEd/Private/Tests/EditorAutomationTests.cpp
+++ b/Engine/Source/Editor/UnrealEd/Private/Tests/EditorAutomationTests.cpp
@@ -1144,7 +1144,8 @@ bool FBSPValidation::RunTest(const FString& Parameters)
 	GEditor->Exec( World, TEXT("BRUSH SUBTRACT"));
 
 	//Directional Light
-	ADirectionalLight* DirectionalLight = Cast<ADirectionalLight>(GEditor->AddActor(World->GetCurrentLevel(), ADirectionalLight::StaticClass(), FVector(-611.0f, 242.0f, 805.0f)));
+	const FTransform Transform(FVector(-611.0f, 242.0f, 805.0f));
+	ADirectionalLight* DirectionalLight = Cast<ADirectionalLight>(GEditor->AddActor(World->GetCurrentLevel(), ADirectionalLight::StaticClass(), Transform));
 	DirectionalLight->SetMobility(EComponentMobility::Movable);
 	DirectionalLight->SetActorRotation(FRotator(300, 250, -91));
 	DirectionalLight->SetBrightness(3.142f);
@@ -1214,11 +1215,12 @@ bool FStaticMeshValidation::RunTest(const FString& Parameters)
 
 	//Single Anim Skeletal Mesh
 	//UActorFactory* SingleAnimSkeletalActorFactory = GEditor->FindActorFactoryForActorClass( ASingleAnimSkeletalActor::StaticClass() );
-	//SkeletalMesh = FActorFactoryAssetProxy::AddActorForAsset( EditorSkeletalMesh, NULL, false, true, RF_Transactional, SingleAnimSkeletalActorFactory );
+	//SkeletalMesh = FActorFactoryAssetProxy::AddActorForAsset( EditorSkeletalMesh, false, true, RF_Transactional, SingleAnimSkeletalActorFactory );
 	//SkeletalMesh->TeleportTo(FVector(1152.0f, 256.0f, 256.0f), FRotator(0, 0, 0));
 
 	//Directional Light
-	ADirectionalLight* DirectionalLight = Cast<ADirectionalLight>(GEditor->AddActor(World->GetCurrentLevel(), ADirectionalLight::StaticClass(), FVector(-611.0f, 242.0f, 805.0f)));
+	const FTransform Transform(FVector(-611.0f, 242.0f, 805.0f));
+	ADirectionalLight* DirectionalLight = Cast<ADirectionalLight>(GEditor->AddActor(World->GetCurrentLevel(), ADirectionalLight::StaticClass(), Transform));
 	DirectionalLight->SetMobility(EComponentMobility::Movable);
 	DirectionalLight->SetActorRotation(FRotator(329, 346, -105));
 	DirectionalLight->SetBrightness(3.142f);
@@ -1372,6 +1374,14 @@ bool FConvertToValidation::RunTest(const FString& Parameters)
 	//Wait to give the screenshot capture some time to complete.
 	ADD_LATENT_AUTOMATION_COMMAND(FWaitLatentCommand(0.1f));
 
+	//Directional Light
+	const FTransform Transform(FVector(384, 0, 384));
+	ADirectionalLight* DirectionalLight = Cast<ADirectionalLight>(GEditor->AddActor(World->GetCurrentLevel(), ADirectionalLight::StaticClass(), Transform));
+	DirectionalLight->SetMobility(EComponentMobility::Movable);
+	DirectionalLight->SetActorRotation(FRotator(314, 339, 0));
+	DirectionalLight->SetBrightness(3.142f);
+	DirectionalLight->SetLightColor(FColor::White);
+
 	GLevelEditorModeTools().MapChangeNotify();
 
 	return true;
@@ -1556,7 +1566,8 @@ bool FLightPlacement::RunTest(const FString& Parameters)
 	}
 
 	//Create the point light and set it's mobility, brightness, and light color.
-	APointLight* PointLight = Cast<APointLight>(GEditor->AddActor(World->GetCurrentLevel(), APointLight::StaticClass(), FVector(0.0f, 0.0f, 400.0f)));
+	const FTransform Transform(FVector(0.0f, 0.0f, 400.0f));
+	APointLight* PointLight = Cast<APointLight>(GEditor->AddActor(World->GetCurrentLevel(), APointLight::StaticClass(), Transform));
 	LightParameters.PointLight = PointLight;
 	LightParameters.LightColor = FColor(255, 0, 0);
 	LightParameters.LightLocation = FVector(0.0f, 0.0f, 400.0f);
diff --git a/Engine/Source/Editor/UnrealEd/Private/UnrealEdSrv.cpp b/Engine/Source/Editor/UnrealEd/Private/UnrealEdSrv.cpp
--- a/Engine/Source/Editor/UnrealEd/Private/UnrealEdSrv.cpp
+++ b/Engine/Source/Editor/UnrealEd/Private/UnrealEdSrv.cpp
@@ -24,6 +24,7 @@
 #include "AssetSelection.h"
 #include "HighResScreenshot.h"
 #include "ActorEditorUtils.h"
+#include "Editor/ActorPositioning.h"
 #include "Matinee/InterpData.h"
 #include "Animation/SkeletalMeshActor.h"
 #include "Landscape/LandscapeInfo.h"
@@ -1585,9 +1586,6 @@ bool UUnrealEdEngine::Exec_Edit( UWorld* InWorld, const TCHAR* Str, FOutputDevic
 			}
 		}
 
-		FVector SaveClickLocation = GEditor->ClickLocation;
-		FSnappingUtils::SnapPointToGrid( SaveClickLocation, FVector(GEditor->GetGridSize(),GEditor->GetGridSize(),GEditor->GetGridSize()) );
-
 		// How should this paste be handled
 		EPasteTo PasteTo = PT_OriginalLocation;
 		FText TransDescription = NSLOCTEXT("UnrealEd", "Paste", "Paste");
@@ -1983,49 +1981,13 @@ bool UUnrealEdEngine::Exec_Actor( UWorld* InWorld, const TCHAR* Str, FOutputDevi
 		UClass* Class;
 		if( ParseObject<UClass>( Str, TEXT("CLASS="), Class, ANY_PACKAGE ) )
 		{
-			AActor* Default   = Class->GetDefaultObject<AActor>();
-			
-			FVector Collision;
-			UCapsuleComponent* CylComp = Cast<UCapsuleComponent>(Default->GetRootComponent());
-			if (CylComp)
-			{
-				Collision = FVector(CylComp->GetScaledCapsuleRadius(), CylComp->GetScaledCapsuleRadius(), CylComp->GetScaledCapsuleHalfHeight());
-			}
-			else
-			{
-				float CollisionRadius, CollisionHeight;
-				Default->GetComponentsBoundingCylinder(CollisionRadius, CollisionHeight);
-				Collision = FVector(CollisionRadius, CollisionRadius, CollisionHeight);
-			}
-			
 			int32 bSnap = 1;
 			FParse::Value(Str,TEXT("SNAP="),bSnap);
-			if( bSnap )
-			{
-				FSnappingUtils::SnapPointToGrid( ClickLocation, FVector(0, 0, 0) );
-			}
-			FVector Location = ClickLocation + ClickPlane * (FVector::BoxPushOut(ClickPlane,Collision) + 0.1);
-			if( bSnap )
-			{
-				FSnappingUtils::SnapPointToGrid( Location, FVector(0, 0, 0) );
-			}
-
-			// Determine if we clicked on the background.
-			FIntPoint CurrentMousePos;
-			GCurrentLevelEditingViewportClient->Viewport->GetMousePos(CurrentMousePos);
-
-			HHitProxy* HitProxy = GCurrentLevelEditingViewportClient->Viewport->GetHitProxy( CurrentMousePos.X, CurrentMousePos.Y );
-			// If the hit proxy is NULL we clicked on the background
-			bool bClickedOnBackground = (HitProxy == NULL);
-
-			AActor* NewActor = AddActor( InWorld->GetCurrentLevel(), Class, Location );
-
-			if( NewActor && bClickedOnBackground && GCurrentLevelEditingViewportClient->IsPerspective() )
-			{
-				// Only move the actor in front of the camera if we didn't click on something useful like bsp or another actor and if we are in the perspective view
-				MoveActorInFrontOfCamera( *NewActor, GCurrentLevelEditingViewportClient->GetViewLocation(), GCurrentLevelEditingViewportClient->GetViewRotation().Vector() );
-			}
+			
+			AActor* Default = Class->GetDefaultObject<AActor>();
+			const FTransform ActorTransform = FActorPositioning::GetCurrentViewportPlacementTransform(*Default, !!bSnap);
 
+			AddActor( InWorld->GetCurrentLevel(), Class, ActorTransform );
 			RedrawLevelEditingViewports();
 			return true;
 		}
diff --git a/Engine/Source/Editor/UnrealEd/Public/AssetSelection.h b/Engine/Source/Editor/UnrealEd/Public/AssetSelection.h
--- a/Engine/Source/Editor/UnrealEd/Public/AssetSelection.h
+++ b/Engine/Source/Editor/UnrealEd/Public/AssetSelection.h
@@ -252,10 +252,6 @@ public:
 	 * Places an actor instance using the factory appropriate for the type of asset
 	 *
 	 * @param	AssetObj						the asset that is contained in the d&d operation
-	 * @param	ActorLocation					specify null to position the actor at the mouse location, 
-												otherwise it will be placed at the origin.
-	 * @param	bUseSurfaceOrientation			specify true to indicate that the factory should align the actor
-	 *											instance with the target surface.
 	 * @param	ObjectFlags						The object flags to place on the actor when it is spawned
 	 * @param	FactoryToUse					optional actor factory to use to create the actor; if not specified,
 	 *											the highest priority factory that is valid will be used
@@ -263,7 +259,7 @@ public:
 	 * @return	the actor that was created by the factory, or NULL if there aren't any factories for this asset (or
 	 *			the actor couldn't be created for some other reason)
 	 */
-	static AActor* AddActorForAsset( UObject* AssetObj, const FVector* ActorLocation=NULL, bool bUseSurfaceOrientation=false, bool SelectActor = true, EObjectFlags ObjectFlags = RF_Transactional, UActorFactory* FactoryToUse = NULL, const FName Name = NAME_None );
+	static AActor* AddActorForAsset( UObject* AssetObj, bool SelectActor = true, EObjectFlags ObjectFlags = RF_Transactional, UActorFactory* FactoryToUse = NULL, const FName Name = NAME_None );
 
 	/**
 	 * Places an actor instance using the factory appropriate for the type of asset using the current object selection as the asset
@@ -271,8 +267,6 @@ public:
 	 * @param	ActorClass						The type of actor to create
 	 * @param	ActorLocation					specify null to position the actor at the mouse location, 
 												otherwise it will be placed at the origin.
-	 * @param	bUseSurfaceOrientation			specify true to indicate that the factory should align the actor
-	 *											instance with the target surface.
 	 * @param	ObjectFlags						The object flags to place on the actor when it is spawned
 	 * @param	FactoryToUse					optional actor factory to use to create the actor; if not specified,
 	 *											the highest priority factory that is valid will be used
@@ -280,7 +274,7 @@ public:
 	 * @return	the actor that was created by the factory, or NULL if there aren't any factories for this asset (or
 	 *			the actor couldn't be created for some other reason)
 	 */
-	static AActor* AddActorFromSelection( UClass* ActorClass, const FVector* ActorLocation=NULL, bool bUseSurfaceOrientation=false, bool SelectActor = true, EObjectFlags ObjectFlags = RF_Transactional, UActorFactory* ActorFactory = NULL, const FName Name = NAME_None );
+	static AActor* AddActorFromSelection( UClass* ActorClass, const FVector* ActorLocation=NULL, bool SelectActor = true, EObjectFlags ObjectFlags = RF_Transactional, UActorFactory* ActorFactory = NULL, const FName Name = NAME_None );
 
 	/**
 	 * Determines if the provided actor is capable of having a material applied to it.
diff --git a/Engine/Source/Editor/UnrealEd/Public/EditorViewportCommands.h b/Engine/Source/Editor/UnrealEd/Public/EditorViewportCommands.h
--- a/Engine/Source/Editor/UnrealEd/Public/EditorViewportCommands.h
+++ b/Engine/Source/Editor/UnrealEd/Public/EditorViewportCommands.h
@@ -156,6 +156,8 @@ public:
 	/** Enables or disables snapping to a scaling grid while scaling objects */
 	TSharedPtr< FUICommandInfo > ScaleGridSnap;
 
+	/** Enables or disables snapping to surfaces when dragging in the viewport */
+	TSharedPtr< FUICommandInfo > SurfaceSnapping;
 public:
 	/** Registers our commands with the binding system */
 	virtual void RegisterCommands() override;
diff --git a/Engine/Source/Editor/UnrealEd/Public/LevelEditorViewport.h b/Engine/Source/Editor/UnrealEd/Public/LevelEditorViewport.h
--- a/Engine/Source/Editor/UnrealEd/Public/LevelEditorViewport.h
+++ b/Engine/Source/Editor/UnrealEd/Public/LevelEditorViewport.h
@@ -571,9 +571,6 @@ private:
 		return EnabledStats.Num();
 	}
 
-	/** Get the list of visible actors for editor viewport op. It filters out certain actor or component types */
-	bool GetVisibleActors(TArray<struct FHitResult>& OutHits,const FVector& Start,const FVector& End,const struct FCollisionQueryParams& Params) const;
-
 protected:
 	/** 
 	 * Checks the viewport to see if the given blueprint asset can be dropped on the viewport.
@@ -649,34 +646,14 @@ private:
 	* @param	DroppedObjects		Array of objects dropped into the viewport
 	* @param	DroppedUponActor	The actor that we are dropping upon
 	* @param    DroppedUponSlot     The material slot/submesh that was identified as the drop location.  If unknown use -1.
-	* @param	DroppedLocation		The location that we're dropping the objects
-	* @param	ObjectFlags			The object flags to place on the actors that this function spawns.
-	* @param	OutNewActors		The list of actors created while dropping
-	* @param	bUsedHitProxy		Whether or not a hit proxy was used for spawning
-	* @param	bSelectActors		If true, select the newly dropped actors (defaults: true)
-	* @param	FactoryToUse		The preferred actor factory to use (optional)
-	*
-	* @return	true if the drop operation was successfully handled; false otherwise
-	*/
-	bool DropObjectsOnActor(struct FViewportCursorLocation& Cursor, const TArray<UObject*>& DroppedObjects, AActor* DroppedUponActor, int32 DroppedUponSlot, FVector* DroppedLocation, EObjectFlags ObjectFlags, TArray<AActor*>& OutNewActors, bool bUsedHitProxy = true, bool bSelectActors = true, class UActorFactory* FactoryToUse = NULL);
-
-	/**
-	* Called when an asset is dropped upon an existing actor.
-	*
-	* @param	Cursor				Mouse cursor location
-	* @param	DroppedObject		Object dropped into the viewport
-	* @param	DroppedUponActor	The actor that we are dropping upon
-	* @param    DroppedUponSlot		The material slot/submesh that was identified as the drop location.  If unknown use -1.
-	* @param	DroppedLocation		The location that we're dropping the objects
 	* @param	ObjectFlags			The object flags to place on the actors that this function spawns.
 	* @param	OutNewActors		The list of actors created while dropping
-	* @param	bUsedHitProxy		Whether or not a hit proxy was used for spawning
 	* @param	bSelectActors		If true, select the newly dropped actors (defaults: true)
 	* @param	FactoryToUse		The preferred actor factory to use (optional)
 	*
 	* @return	true if the drop operation was successfully handled; false otherwise
 	*/
-	bool DropSingleObjectOnActor(struct FViewportCursorLocation& Cursor, UObject* DroppedObject, AActor* DroppedUponActor, int32 DroppedUponSlot, const FVector& DroppedLocation, EObjectFlags ObjectFlags, TArray<AActor*>& OutNewActors, bool bUsedHitProxy = true, bool bSelectActors = true, class UActorFactory* FactoryToUse = NULL);
+	bool DropObjectsOnActor(struct FViewportCursorLocation& Cursor, const TArray<UObject*>& DroppedObjects, AActor* DroppedUponActor, int32 DroppedUponSlot, EObjectFlags ObjectFlags, TArray<AActor*>& OutNewActors, bool bSelectActors = true, class UActorFactory* FactoryToUse = NULL);
 
 	/**
 	 * Called when an asset is dropped upon a BSP surface.
@@ -692,7 +669,7 @@ private:
 	 *
 	 * @return	true if the drop operation was successfully handled; false otherwise
 	 */
-	bool DropObjectsOnBSPSurface( FSceneView* View, struct FViewportCursorLocation& Cursor, const TArray<UObject*>& DroppedObjects, class HModel* TargetProxy, EObjectFlags ObjectFlags, TArray<AActor*>& OutNewActors, bool bSelectActors, UActorFactory* FactoryToUse );
+	bool DropObjectsOnBSPSurface( FSceneView* View, struct FViewportCursorLocation& Cursor, const TArray<UObject*>& DroppedObjects, HModel* TargetProxy, EObjectFlags ObjectFlags, TArray<AActor*>& OutNewActors, bool bSelectActors, UActorFactory* FactoryToUse );
 
 	/**
 	 * Called when an asset is dropped upon a manipulation widget.
@@ -711,6 +688,9 @@ private:
 	void ModifyScale( AActor* InActor, FVector& ScaleDelta, bool bCheckSmallExtent = false ) const;
 	void ValidateScale( const FVector& CurrentScale, const FVector& BoxExtent, FVector& ScaleDelta, bool bCheckSmallExtent = false ) const;
 
+	/** Project the specified actors into the world according to the current drag parameters */
+	void ProjectActorsIntoWorld(const TArray<AActor*>& Actors, FViewport* Viewport, const FVector& Drag, const FRotator& Rot);
+
 	/** Delegate handler to see if a stat is enabled on this viewport */
 	void HandleViewportStatCheckEnabled(const TCHAR* InName, bool& bOutCurrentEnabled, bool& bOutOthersEnabled);
 
@@ -785,6 +765,9 @@ private:
 	/** The actors that are currently being placed in the viewport via dragging */
 	static TArray< TWeakObjectPtr< AActor > > DropPreviewActors;
 
+	/** A map of actor locations before a drag operation */
+	TMap<TWeakObjectPtr<AActor>, FTransform> PreDragActorTransforms;
+
 	/** Bit array representing the visibility of every sprite category in the current viewport */
 	TBitArray<>	SpriteCategoryVisibility;
 
diff --git a/Engine/Source/Editor/UnrealEd/Public/SEditorViewport.h b/Engine/Source/Editor/UnrealEd/Public/SEditorViewport.h
--- a/Engine/Source/Editor/UnrealEd/Public/SEditorViewport.h
+++ b/Engine/Source/Editor/UnrealEd/Public/SEditorViewport.h
@@ -121,6 +121,17 @@ protected:
 	 */
 	virtual void OnFocusViewportToSelection(){}
 
+
+	/**
+	 * Called when surface snapping has been enabled/disabled
+	 */
+	static void OnToggleSurfaceSnap();
+
+	/**
+	 * Called to test whether surface snapping is enabled or not
+	 */
+	static bool OnIsSurfaceSnapEnabled();
+
 protected:
 	TSharedPtr<SOverlay> ViewportOverlay;
 	/** Viewport that renders the scene provided by the viewport client */
@@ -131,4 +142,4 @@ protected:
 	TSharedPtr<FEditorViewportClient> Client;
 	/** Commandlist used in the viewport (Maps commands to viewport specific actions) */
 	TSharedPtr<FUICommandList> CommandList;
-};
\ No newline at end of file
+};
diff --git a/Engine/Source/Editor/UnrealEd/Public/STransformViewportToolbar.h b/Engine/Source/Editor/UnrealEd/Public/STransformViewportToolbar.h
--- a/Engine/Source/Editor/UnrealEd/Public/STransformViewportToolbar.h
+++ b/Engine/Source/Editor/UnrealEd/Public/STransformViewportToolbar.h
@@ -74,6 +74,9 @@ private:
 	TSharedRef<SWidget> BuildLocationGridCheckBoxList(FName InExtentionHook, const FText& InHeading, const TArray<float>& InGridSizes) const;
 	TSharedRef<SWidget> BuildRotationGridCheckBoxList(FName InExtentionHook, const FText& InHeading, const TArray<float>& InGridSizes, ERotationGridMode InGridMode) const;
 
+	/** Make the surface snapping toolbar checkbox button */
+	TSharedRef< SWidget > MakeSurfaceSnappingButton( FName ToolBarStyle );
+
 	/** Grid Snap checked state callbacks */
 	ESlateCheckBoxState::Type IsLocationGridSnapChecked() const;
 	ESlateCheckBoxState::Type IsRotationGridSnapChecked() const;
