Removed #pragma once directive from cpp

diff --git a/Engine/Source/Editor/UnrealEd/Private/Editor/ActorPositioning.cpp b/Engine/Source/Editor/UnrealEd/Private/Editor/ActorPositioning.cpp
--- a/Engine/Source/Editor/UnrealEd/Private/Editor/ActorPositioning.cpp
+++ b/Engine/Source/Editor/UnrealEd/Private/Editor/ActorPositioning.cpp
@@ -1,7 +1,5 @@
 // Copyright 1998-2014 Epic Games, Inc. All Rights Reserved.
 
-#pragma once
-
 #include "UnrealEd.h"
 #include "ActorPositioning.h"
 #include "SnappingUtils.h"
