New thumbnails for animations that show a posed mesh

diff --git a/Engine/Config/BaseEditor.ini b/Engine/Config/BaseEditor.ini
--- a/Engine/Config/BaseEditor.ini
+++ b/Engine/Config/BaseEditor.ini
@@ -28,6 +28,9 @@ MaxNumberOfLogLines=2000
 +RenderableThumbnailTypes=(ClassNeedingThumbnailName="/Script/Engine.ParticleSystem",RendererClassName="/Script/UnrealEd.ParticleSystemThumbnailRenderer")
 +RenderableThumbnailTypes=(ClassNeedingThumbnailName="/Script/Engine.StaticMesh",RendererClassName="/Script/UnrealEd.StaticMeshThumbnailRenderer")
 +RenderableThumbnailTypes=(ClassNeedingThumbnailName="/Script/Engine.SkeletalMesh",RendererClassName="/Script/UnrealEd.SkeletalMeshThumbnailRenderer")
++RenderableThumbnailTypes=(ClassNeedingThumbnailName="/Script/Engine.AnimSequenceBase",RendererClassName="/Script/UnrealEd.AnimSequenceThumbnailRenderer")
++RenderableThumbnailTypes=(ClassNeedingThumbnailName="/Script/Engine.BlendSpaceBase",RendererClassName="/Script/UnrealEd.BlendSpaceThumbnailRenderer")
++RenderableThumbnailTypes=(ClassNeedingThumbnailName="/Script/Engine.AnimBlueprint",RendererClassName="/Script/UnrealEd.AnimBlueprintThumbnailRenderer")
 +RenderableThumbnailTypes=(ClassNeedingThumbnailName="/Script/Engine.Level",RendererClassName="/Script/UnrealEd.LevelThumbnailRenderer")
 +RenderableThumbnailTypes=(ClassNeedingThumbnailName="/Script/Engine.World",RendererClassName="/Script/UnrealEd.WorldThumbnailRenderer")
 +RenderableThumbnailTypes=(ClassNeedingThumbnailName="/Script/Engine.SoundWave",RendererClassName="/Script/UnrealEd.SoundWaveThumbnailRenderer")
diff --git a/Engine/Source/Developer/AssetTools/Private/AssetTypeActions/AssetTypeActions_AnimBlueprint.cpp b/Engine/Source/Developer/AssetTools/Private/AssetTypeActions/AssetTypeActions_AnimBlueprint.cpp
--- a/Engine/Source/Developer/AssetTools/Private/AssetTypeActions/AssetTypeActions_AnimBlueprint.cpp
+++ b/Engine/Source/Developer/AssetTools/Private/AssetTypeActions/AssetTypeActions_AnimBlueprint.cpp
@@ -70,6 +70,19 @@ void FAssetTypeActions_AnimBlueprint::FillRetargetMenu( FMenuBuilder& MenuBuilde
 		);
 }
 
+UThumbnailInfo* FAssetTypeActions_AnimBlueprint::GetThumbnailInfo(UObject* Asset) const
+{
+	UAnimBlueprint* AnimBlueprint = CastChecked<UAnimBlueprint>(Asset);
+	UThumbnailInfo* ThumbnailInfo = AnimBlueprint->ThumbnailInfo;
+	if (ThumbnailInfo == NULL)
+	{
+		ThumbnailInfo = ConstructObject<USceneThumbnailInfo>(USceneThumbnailInfo::StaticClass(), AnimBlueprint);
+		AnimBlueprint->ThumbnailInfo = ThumbnailInfo;
+	}
+
+	return ThumbnailInfo;
+}
+
 void FAssetTypeActions_AnimBlueprint::OpenAssetEditor( const TArray<UObject*>& InObjects, TSharedPtr<IToolkitHost> EditWithinLevelEditor )
 {
 	EToolkitMode::Type Mode = EditWithinLevelEditor.IsValid() ? EToolkitMode::WorldCentric : EToolkitMode::Standalone;
diff --git a/Engine/Source/Developer/AssetTools/Private/AssetTypeActions/AssetTypeActions_AnimBlueprint.h b/Engine/Source/Developer/AssetTools/Private/AssetTypeActions/AssetTypeActions_AnimBlueprint.h
--- a/Engine/Source/Developer/AssetTools/Private/AssetTypeActions/AssetTypeActions_AnimBlueprint.h
+++ b/Engine/Source/Developer/AssetTools/Private/AssetTypeActions/AssetTypeActions_AnimBlueprint.h
@@ -12,6 +12,7 @@ public:
 	virtual void GetActions( const TArray<UObject*>& InObjects, FMenuBuilder& MenuBuilder ) override;
 	virtual void OpenAssetEditor( const TArray<UObject*>& InObjects, TSharedPtr<class IToolkitHost> EditWithinLevelEditor = TSharedPtr<IToolkitHost>() ) override;
 	virtual uint32 GetCategories() override { return EAssetTypeCategories::Animation; }
+	virtual class UThumbnailInfo* GetThumbnailInfo(UObject* Asset) const override;
 
 private:
 
diff --git a/Engine/Source/Developer/AssetTools/Private/AssetTypeActions/AssetTypeActions_AnimationAsset.cpp b/Engine/Source/Developer/AssetTools/Private/AssetTypeActions/AssetTypeActions_AnimationAsset.cpp
--- a/Engine/Source/Developer/AssetTools/Private/AssetTypeActions/AssetTypeActions_AnimationAsset.cpp
+++ b/Engine/Source/Developer/AssetTools/Private/AssetTypeActions/AssetTypeActions_AnimationAsset.cpp
@@ -68,6 +68,19 @@ void FAssetTypeActions_AnimationAsset::FillRetargetMenu( FMenuBuilder& MenuBuild
 		);
 }
 
+UThumbnailInfo* FAssetTypeActions_AnimationAsset::GetThumbnailInfo(UObject* Asset) const
+{
+	UAnimationAsset* Anim = CastChecked<UAnimationAsset>(Asset);
+	UThumbnailInfo* ThumbnailInfo = Anim->ThumbnailInfo;
+	if (ThumbnailInfo == NULL)
+	{
+		ThumbnailInfo = ConstructObject<USceneThumbnailInfo>(USceneThumbnailInfo::StaticClass(), Anim);
+		Anim->ThumbnailInfo = ThumbnailInfo;
+	}
+
+	return ThumbnailInfo;
+}
+
 void FAssetTypeActions_AnimationAsset::OpenAssetEditor( const TArray<UObject*>& InObjects, TSharedPtr<IToolkitHost> EditWithinLevelEditor )
 {
 	EToolkitMode::Type Mode = EditWithinLevelEditor.IsValid() ? EToolkitMode::WorldCentric : EToolkitMode::Standalone;
diff --git a/Engine/Source/Developer/AssetTools/Private/AssetTypeActions/AssetTypeActions_AnimationAsset.h b/Engine/Source/Developer/AssetTools/Private/AssetTypeActions/AssetTypeActions_AnimationAsset.h
--- a/Engine/Source/Developer/AssetTools/Private/AssetTypeActions/AssetTypeActions_AnimationAsset.h
+++ b/Engine/Source/Developer/AssetTools/Private/AssetTypeActions/AssetTypeActions_AnimationAsset.h
@@ -14,6 +14,7 @@ public:
 	virtual void OpenAssetEditor( const TArray<UObject*>& InObjects, TSharedPtr<class IToolkitHost> EditWithinLevelEditor = TSharedPtr<IToolkitHost>() ) override;
 	virtual bool CanFilter() override { return false; }
 	virtual uint32 GetCategories() override { return EAssetTypeCategories::Animation; }
+	virtual class UThumbnailInfo* GetThumbnailInfo(UObject* Asset) const override;
 
 private:
 	/** Handler to fill the retarget submenu */
diff --git a/Engine/Source/Editor/UnrealEd/Classes/ThumbnailRendering/AnimBlueprintThumbnailRenderer.h b/Engine/Source/Editor/UnrealEd/Classes/ThumbnailRendering/AnimBlueprintThumbnailRenderer.h
new file mode 100644
--- /dev/null
+++ b/Engine/Source/Editor/UnrealEd/Classes/ThumbnailRendering/AnimBlueprintThumbnailRenderer.h
@@ -0,0 +1,27 @@
+// Copyright 1998-2014 Epic Games, Inc. All Rights Reserved.
+
+/**
+ *
+ * This thumbnail renderer displays a given AnimBlueprint
+ */
+
+#pragma once
+#include "AnimBlueprintThumbnailRenderer.generated.h"
+
+UCLASS(config=Editor, MinimalAPI)
+class UAnimBlueprintThumbnailRenderer : public UDefaultSizedThumbnailRenderer
+{
+	GENERATED_UCLASS_BODY()
+
+
+	// Begin UThumbnailRenderer Object
+	UNREALED_API virtual void Draw(UObject* Object, int32 X, int32 Y, uint32 Width, uint32 Height, FRenderTarget*, FCanvas* Canvas) override;
+	// End UThumbnailRenderer Object
+
+	// UObject implementation
+	UNREALED_API virtual void BeginDestroy() override;
+
+private:
+	class FAnimBlueprintThumbnailScene* ThumbnailScene;
+};
+
diff --git a/Engine/Source/Editor/UnrealEd/Classes/ThumbnailRendering/AnimSequenceThumbnailRenderer.h b/Engine/Source/Editor/UnrealEd/Classes/ThumbnailRendering/AnimSequenceThumbnailRenderer.h
new file mode 100644
--- /dev/null
+++ b/Engine/Source/Editor/UnrealEd/Classes/ThumbnailRendering/AnimSequenceThumbnailRenderer.h
@@ -0,0 +1,27 @@
+// Copyright 1998-2014 Epic Games, Inc. All Rights Reserved.
+
+/**
+ *
+ * This thumbnail renderer displays a given anim sequence
+ */
+
+#pragma once
+#include "AnimSequenceThumbnailRenderer.generated.h"
+
+UCLASS(config=Editor, MinimalAPI)
+class UAnimSequenceThumbnailRenderer : public UDefaultSizedThumbnailRenderer
+{
+	GENERATED_UCLASS_BODY()
+
+
+	// Begin UThumbnailRenderer Object
+	UNREALED_API virtual void Draw(UObject* Object, int32 X, int32 Y, uint32 Width, uint32 Height, FRenderTarget*, FCanvas* Canvas) override;
+	// End UThumbnailRenderer Object
+
+	// UObject implementation
+	UNREALED_API virtual void BeginDestroy() override;
+
+private:
+	class FAnimationSequenceThumbnailScene* ThumbnailScene;
+};
+
diff --git a/Engine/Source/Editor/UnrealEd/Classes/ThumbnailRendering/BlendSpaceThumbnailRenderer.h b/Engine/Source/Editor/UnrealEd/Classes/ThumbnailRendering/BlendSpaceThumbnailRenderer.h
new file mode 100644
--- /dev/null
+++ b/Engine/Source/Editor/UnrealEd/Classes/ThumbnailRendering/BlendSpaceThumbnailRenderer.h
@@ -0,0 +1,27 @@
+// Copyright 1998-2014 Epic Games, Inc. All Rights Reserved.
+
+/**
+ *
+ * This thumbnail renderer displays a given blend space
+ */
+
+#pragma once
+#include "BlendSpaceThumbnailRenderer.generated.h"
+
+UCLASS(config=Editor, MinimalAPI)
+class UBlendSpaceThumbnailRenderer : public UDefaultSizedThumbnailRenderer
+{
+	GENERATED_UCLASS_BODY()
+
+
+	// Begin UThumbnailRenderer Object
+	UNREALED_API virtual void Draw(UObject* Object, int32 X, int32 Y, uint32 Width, uint32 Height, FRenderTarget*, FCanvas* Canvas) override;
+	// End UThumbnailRenderer Object
+
+	// UObject implementation
+	UNREALED_API virtual void BeginDestroy() override;
+
+private:
+	class FBlendSpaceThumbnailScene* ThumbnailScene;
+};
+
diff --git a/Engine/Source/Editor/UnrealEd/Private/ThumbnailHelpers.cpp b/Engine/Source/Editor/UnrealEd/Private/ThumbnailHelpers.cpp
--- a/Engine/Source/Editor/UnrealEd/Private/ThumbnailHelpers.cpp
+++ b/Engine/Source/Editor/UnrealEd/Private/ThumbnailHelpers.cpp
@@ -493,6 +493,364 @@ void FStaticMeshThumbnailScene::GetViewMatrixParameters(const float InFOVDegrees
 	OutOrbitZoom = TargetDistance + ThumbnailInfo->OrbitZoom;
 }
 
+/*
+***************************************************************
+FAnimationThumbnailScene
+***************************************************************
+*/
+
+FAnimationSequenceThumbnailScene::FAnimationSequenceThumbnailScene()
+	: FThumbnailPreviewScene()
+{
+	bForceAllUsedMipsResident = false;
+	// Create preview actor
+	// checked
+	FActorSpawnParameters SpawnInfo;
+	SpawnInfo.bNoCollisionFail = true;
+	SpawnInfo.bNoFail = true;
+	SpawnInfo.ObjectFlags = RF_Transient;
+	PreviewActor = GetWorld()->SpawnActor<ASkeletalMeshActor>(SpawnInfo);
+
+	PreviewActor->SetActorEnableCollision(false);
+}
+
+bool FAnimationSequenceThumbnailScene::SetAnimation(UAnimSequenceBase* InAnimation)
+{
+	PreviewActor->SkeletalMeshComponent->Materials.Empty();
+
+	bool bSetSucessfully = false;
+
+	PreviewAnimation = InAnimation;
+
+	if (InAnimation)
+	{
+		if (USkeleton* Skeleton = InAnimation->GetSkeleton())
+		{
+			USkeletalMesh* PreviewSkeletalMesh = Skeleton->GetPreviewMesh(true);
+
+			PreviewActor->SkeletalMeshComponent->SetSkeletalMesh(PreviewSkeletalMesh);
+
+			if (PreviewSkeletalMesh)
+			{
+				bSetSucessfully = true;
+
+				// Handle posing the mesh at the middle of the animation
+				const float AnimPosition = InAnimation->SequenceLength / 2.f;
+
+				PreviewActor->SkeletalMeshComponent->PlayAnimation(InAnimation,false);
+				PreviewActor->SkeletalMeshComponent->Stop();
+				PreviewActor->SkeletalMeshComponent->SetPosition(AnimPosition);
+
+				UAnimSingleNodeInstance* SingleNodeInstance = PreviewActor->SkeletalMeshComponent->GetSingleNodeInstance();
+				if (SingleNodeInstance)
+				{
+					SingleNodeInstance->UpdateMontageWeightForTimeSkip(AnimPosition);
+				}
+
+				PreviewActor->SkeletalMeshComponent->RefreshBoneTransforms(nullptr);
+
+				FTransform MeshTransform = FTransform::Identity;
+
+				PreviewActor->SetActorLocation(FVector(0, 0, 0), false);
+				PreviewActor->SkeletalMeshComponent->UpdateBounds();
+
+				// Center the mesh at the world origin then offset to put it on top of the plane
+				const float BoundsZOffset = GetBoundsZOffset(PreviewActor->SkeletalMeshComponent->Bounds);
+				PreviewActor->SetActorLocation(-PreviewActor->SkeletalMeshComponent->Bounds.Origin + FVector(0, 0, BoundsZOffset), false);
+				PreviewActor->SkeletalMeshComponent->RecreateRenderState_Concurrent();
+			}
+		}
+	}
+	
+	if(!bSetSucessfully)
+	{
+		CleanupComponentChildren(PreviewActor->SkeletalMeshComponent);
+		PreviewActor->SkeletalMeshComponent->SetSkeletalMesh(nullptr);
+	}
+
+	return bSetSucessfully;
+}
+
+void FAnimationSequenceThumbnailScene::CleanupComponentChildren(USceneComponent* Component)
+{
+	if (Component)
+	{
+		for (int32 I = 0; I < Component->AttachChildren.Num(); ++I)
+		{
+			CleanupComponentChildren(Component->AttachChildren[I]);
+			Component->AttachChildren[I]->DestroyComponent();
+		}
+
+		Component->AttachChildren.Empty();
+	}
+}
+
+void FAnimationSequenceThumbnailScene::GetViewMatrixParameters(const float InFOVDegrees, FVector& OutOrigin, float& OutOrbitPitch, float& OutOrbitYaw, float& OutOrbitZoom) const
+{
+	check(PreviewAnimation);
+	check(PreviewActor->SkeletalMeshComponent);
+	check(PreviewActor->SkeletalMeshComponent->SkeletalMesh);
+
+	const float HalfFOVRadians = FMath::DegreesToRadians<float>(InFOVDegrees) * 0.5f;
+	// No need to add extra size to view slightly outside of the sphere to compensate for perspective since skeletal meshes already buffer bounds.
+	const float HalfMeshSize = PreviewActor->SkeletalMeshComponent->Bounds.SphereRadius;
+	const float BoundsZOffset = GetBoundsZOffset(PreviewActor->SkeletalMeshComponent->Bounds);
+	const float TargetDistance = HalfMeshSize / FMath::Tan(HalfFOVRadians);
+
+	USceneThumbnailInfo* ThumbnailInfo = Cast<USceneThumbnailInfo>(PreviewAnimation->ThumbnailInfo);
+	if (ThumbnailInfo)
+	{
+		if (TargetDistance + ThumbnailInfo->OrbitZoom < 0)
+		{
+			ThumbnailInfo->OrbitZoom = -TargetDistance;
+		}
+	}
+	else
+	{
+		ThumbnailInfo = USceneThumbnailInfo::StaticClass()->GetDefaultObject<USceneThumbnailInfo>();
+	}
+
+	OutOrigin = FVector(0, 0, -BoundsZOffset);
+	OutOrbitPitch = ThumbnailInfo->OrbitPitch;
+	OutOrbitYaw = ThumbnailInfo->OrbitYaw;
+	OutOrbitZoom = TargetDistance + ThumbnailInfo->OrbitZoom;
+}
+
+/*
+***************************************************************
+FBlendSpaceThumbnailScene
+***************************************************************
+*/
+
+FBlendSpaceThumbnailScene::FBlendSpaceThumbnailScene()
+	: FThumbnailPreviewScene()
+{
+	bForceAllUsedMipsResident = false;
+	// Create preview actor
+	// checked
+	FActorSpawnParameters SpawnInfo;
+	SpawnInfo.bNoCollisionFail = true;
+	SpawnInfo.bNoFail = true;
+	SpawnInfo.ObjectFlags = RF_Transient;
+	PreviewActor = GetWorld()->SpawnActor<ASkeletalMeshActor>(SpawnInfo);
+
+	PreviewActor->SetActorEnableCollision(false);
+}
+
+bool FBlendSpaceThumbnailScene::SetBlendSpace(class UBlendSpaceBase* InBlendSpace)
+{
+	PreviewActor->SkeletalMeshComponent->Materials.Empty();
+
+	bool bSetSucessfully = false;
+
+	PreviewAnimation = InBlendSpace;
+
+	if (InBlendSpace)
+	{
+		if (USkeleton* Skeleton = InBlendSpace->GetSkeleton())
+		{
+			USkeletalMesh* PreviewSkeletalMesh = Skeleton->GetPreviewMesh(true);
+
+			PreviewActor->SkeletalMeshComponent->SetSkeletalMesh(PreviewSkeletalMesh);
+
+			if (PreviewSkeletalMesh)
+			{
+				bSetSucessfully = true;
+
+				// Handle posing the mesh at the middle of the animation
+				PreviewActor->SkeletalMeshComponent->PlayAnimation(InBlendSpace, false);
+				PreviewActor->SkeletalMeshComponent->Stop();
+
+				UAnimSingleNodeInstance* AnimInstance = PreviewActor->SkeletalMeshComponent->GetSingleNodeInstance();
+				if (AnimInstance)
+				{
+					FVector BlendInput(0.f);
+					for (int32 i = 0; i < InBlendSpace->NumOfDimension; ++i)
+					{
+						const FBlendParameter& Param = InBlendSpace->GetBlendParameter(i);
+						BlendInput[i] = (Param.GetRange() / 2.f) + Param.Min;
+					}
+					AnimInstance->UpdateBlendspaceSamples(BlendInput);
+				}
+				PreviewActor->SkeletalMeshComponent->RefreshBoneTransforms(nullptr);
+
+				FTransform MeshTransform = FTransform::Identity;
+
+				PreviewActor->SetActorLocation(FVector(0, 0, 0), false);
+				PreviewActor->SkeletalMeshComponent->UpdateBounds();
+
+				// Center the mesh at the world origin then offset to put it on top of the plane
+				const float BoundsZOffset = GetBoundsZOffset(PreviewActor->SkeletalMeshComponent->Bounds);
+				PreviewActor->SetActorLocation(-PreviewActor->SkeletalMeshComponent->Bounds.Origin + FVector(0, 0, BoundsZOffset), false);
+				PreviewActor->SkeletalMeshComponent->RecreateRenderState_Concurrent();
+			}
+		}
+	}
+
+	if (!bSetSucessfully)
+	{
+		CleanupComponentChildren(PreviewActor->SkeletalMeshComponent);
+		PreviewActor->SkeletalMeshComponent->SetSkeletalMesh(nullptr);
+	}
+
+	return bSetSucessfully;
+}
+
+void FBlendSpaceThumbnailScene::CleanupComponentChildren(USceneComponent* Component)
+{
+	if (Component)
+	{
+		for (int32 I = 0; I < Component->AttachChildren.Num(); ++I)
+		{
+			CleanupComponentChildren(Component->AttachChildren[I]);
+			Component->AttachChildren[I]->DestroyComponent();
+		}
+
+		Component->AttachChildren.Empty();
+	}
+}
+
+void FBlendSpaceThumbnailScene::GetViewMatrixParameters(const float InFOVDegrees, FVector& OutOrigin, float& OutOrbitPitch, float& OutOrbitYaw, float& OutOrbitZoom) const
+{
+	check(PreviewAnimation);
+	check(PreviewActor->SkeletalMeshComponent);
+	check(PreviewActor->SkeletalMeshComponent->SkeletalMesh);
+
+	const float HalfFOVRadians = FMath::DegreesToRadians<float>(InFOVDegrees) * 0.5f;
+	// No need to add extra size to view slightly outside of the sphere to compensate for perspective since skeletal meshes already buffer bounds.
+	const float HalfMeshSize = PreviewActor->SkeletalMeshComponent->Bounds.SphereRadius;
+	const float BoundsZOffset = GetBoundsZOffset(PreviewActor->SkeletalMeshComponent->Bounds);
+	const float TargetDistance = HalfMeshSize / FMath::Tan(HalfFOVRadians);
+
+	USceneThumbnailInfo* ThumbnailInfo = Cast<USceneThumbnailInfo>(PreviewAnimation->ThumbnailInfo);
+	if (ThumbnailInfo)
+	{
+		if (TargetDistance + ThumbnailInfo->OrbitZoom < 0)
+		{
+			ThumbnailInfo->OrbitZoom = -TargetDistance;
+		}
+	}
+	else
+	{
+		ThumbnailInfo = USceneThumbnailInfo::StaticClass()->GetDefaultObject<USceneThumbnailInfo>();
+	}
+
+	OutOrigin = FVector(0, 0, -BoundsZOffset);
+	OutOrbitPitch = ThumbnailInfo->OrbitPitch;
+	OutOrbitYaw = ThumbnailInfo->OrbitYaw;
+	OutOrbitZoom = TargetDistance + ThumbnailInfo->OrbitZoom;
+}
+
+/*
+***************************************************************
+FAnimBlueprintThumbnailScene
+***************************************************************
+*/
+
+FAnimBlueprintThumbnailScene::FAnimBlueprintThumbnailScene()
+	: FThumbnailPreviewScene()
+{
+	bForceAllUsedMipsResident = false;
+	
+	// Create preview actor
+	FActorSpawnParameters SpawnInfo;
+	SpawnInfo.bNoCollisionFail = true;
+	SpawnInfo.bNoFail = true;
+	SpawnInfo.ObjectFlags = RF_Transient;
+	PreviewActor = GetWorld()->SpawnActor<ASkeletalMeshActor>(SpawnInfo);
+
+	PreviewActor->SetActorEnableCollision(false);
+}
+
+bool FAnimBlueprintThumbnailScene::SetAnimBlueprint(class UAnimBlueprint* InBlueprint)
+{
+	PreviewActor->SkeletalMeshComponent->Materials.Empty();
+
+	bool bSetSucessfully = false;
+
+	PreviewBlueprint = InBlueprint;
+
+	if (InBlueprint)
+	{
+		if (USkeleton* Skeleton = InBlueprint->TargetSkeleton)
+		{
+			USkeletalMesh* PreviewSkeletalMesh = Skeleton->GetPreviewMesh(true);
+
+			PreviewActor->SkeletalMeshComponent->SetSkeletalMesh(PreviewSkeletalMesh);
+
+			if (PreviewSkeletalMesh)
+			{
+				bSetSucessfully = true;
+
+				PreviewActor->SkeletalMeshComponent->SetAnimInstanceClass(InBlueprint->GeneratedClass);
+
+				FTransform MeshTransform = FTransform::Identity;
+
+				PreviewActor->SetActorLocation(FVector(0, 0, 0), false);
+				PreviewActor->SkeletalMeshComponent->UpdateBounds();
+
+				// Center the mesh at the world origin then offset to put it on top of the plane
+				const float BoundsZOffset = GetBoundsZOffset(PreviewActor->SkeletalMeshComponent->Bounds);
+				PreviewActor->SetActorLocation(-PreviewActor->SkeletalMeshComponent->Bounds.Origin + FVector(0, 0, BoundsZOffset), false);
+				PreviewActor->SkeletalMeshComponent->RecreateRenderState_Concurrent();
+			}
+		}
+	}
+
+	if (!bSetSucessfully)
+	{
+		CleanupComponentChildren(PreviewActor->SkeletalMeshComponent);
+		PreviewActor->SkeletalMeshComponent->SetSkeletalMesh(nullptr);
+	}
+
+	return bSetSucessfully;
+}
+
+void FAnimBlueprintThumbnailScene::CleanupComponentChildren(USceneComponent* Component)
+{
+	if (Component)
+	{
+		for (int32 I = 0; I < Component->AttachChildren.Num(); ++I)
+		{
+			CleanupComponentChildren(Component->AttachChildren[I]);
+			Component->AttachChildren[I]->DestroyComponent();
+		}
+
+		Component->AttachChildren.Empty();
+	}
+}
+
+void FAnimBlueprintThumbnailScene::GetViewMatrixParameters(const float InFOVDegrees, FVector& OutOrigin, float& OutOrbitPitch, float& OutOrbitYaw, float& OutOrbitZoom) const
+{
+	check(PreviewBlueprint);
+	check(PreviewActor->SkeletalMeshComponent);
+	check(PreviewActor->SkeletalMeshComponent->SkeletalMesh);
+
+	const float HalfFOVRadians = FMath::DegreesToRadians<float>(InFOVDegrees) * 0.5f;
+	// No need to add extra size to view slightly outside of the sphere to compensate for perspective since skeletal meshes already buffer bounds.
+	const float HalfMeshSize = PreviewActor->SkeletalMeshComponent->Bounds.SphereRadius;
+	const float BoundsZOffset = GetBoundsZOffset(PreviewActor->SkeletalMeshComponent->Bounds);
+	const float TargetDistance = HalfMeshSize / FMath::Tan(HalfFOVRadians);
+
+	USceneThumbnailInfo* ThumbnailInfo = Cast<USceneThumbnailInfo>(PreviewBlueprint->ThumbnailInfo);
+	if (ThumbnailInfo)
+	{
+		if (TargetDistance + ThumbnailInfo->OrbitZoom < 0)
+		{
+			ThumbnailInfo->OrbitZoom = -TargetDistance;
+		}
+	}
+	else
+	{
+		ThumbnailInfo = USceneThumbnailInfo::StaticClass()->GetDefaultObject<USceneThumbnailInfo>();
+	}
+
+	OutOrigin = FVector(0, 0, -BoundsZOffset);
+	OutOrbitPitch = ThumbnailInfo->OrbitPitch;
+	OutOrbitYaw = ThumbnailInfo->OrbitYaw;
+	OutOrbitZoom = TargetDistance + ThumbnailInfo->OrbitZoom;
+}
+
 /*
 ***************************************************************
   FBlueprintThumbnailScene
diff --git a/Engine/Source/Editor/UnrealEd/Private/ThumbnailRendering/AnimBlueprintThumbnailRenderer.cpp b/Engine/Source/Editor/UnrealEd/Private/ThumbnailRendering/AnimBlueprintThumbnailRenderer.cpp
new file mode 100644
--- /dev/null
+++ b/Engine/Source/Editor/UnrealEd/Private/ThumbnailRendering/AnimBlueprintThumbnailRenderer.cpp
@@ -0,0 +1,50 @@
+// Copyright 1998-2014 Epic Games, Inc. All Rights Reserved.
+
+#include "UnrealEd.h"
+
+// FPreviewScene derived helpers for rendering
+#include "ThumbnailHelpers.h"
+#include "EngineModule.h"
+
+UAnimBlueprintThumbnailRenderer::UAnimBlueprintThumbnailRenderer(const class FPostConstructInitializeProperties& PCIP)
+	: Super(PCIP)
+{
+	ThumbnailScene = nullptr;
+}
+
+void UAnimBlueprintThumbnailRenderer::Draw(UObject* Object, int32 X, int32 Y, uint32 Width, uint32 Height, FRenderTarget* RenderTarget, FCanvas* Canvas)
+{
+	UAnimBlueprint* AnimBlueprint = Cast<UAnimBlueprint>(Object);
+	if (AnimBlueprint != nullptr)
+	{
+		if ( ThumbnailScene == nullptr )
+		{
+			ThumbnailScene = new FAnimBlueprintThumbnailScene();
+		}
+
+		if(ThumbnailScene->SetAnimBlueprint(AnimBlueprint))
+		{
+			FSceneViewFamilyContext ViewFamily(FSceneViewFamily::ConstructionValues(RenderTarget, ThumbnailScene->GetScene(), FEngineShowFlags(ESFIM_Game))
+				.SetWorldTimes(FApp::GetCurrentTime() - GStartTime, FApp::GetDeltaTime(), FApp::GetCurrentTime() - GStartTime));
+
+			ViewFamily.EngineShowFlags.DisableAdvancedFeatures();
+			ViewFamily.EngineShowFlags.MotionBlur = 0;
+			ViewFamily.EngineShowFlags.LOD = 0;
+
+			ThumbnailScene->GetView(&ViewFamily, X, Y, Width, Height);
+			GetRendererModule().BeginRenderingViewFamily(Canvas, &ViewFamily);
+			ThumbnailScene->SetAnimBlueprint(nullptr);
+		}
+	}
+}
+
+void UAnimBlueprintThumbnailRenderer::BeginDestroy()
+{
+	if ( ThumbnailScene != nullptr )
+	{
+		delete ThumbnailScene;
+		ThumbnailScene = nullptr;
+	}
+
+	Super::BeginDestroy();
+}
diff --git a/Engine/Source/Editor/UnrealEd/Private/ThumbnailRendering/AnimSequenceThumbnailRenderer.cpp b/Engine/Source/Editor/UnrealEd/Private/ThumbnailRendering/AnimSequenceThumbnailRenderer.cpp
new file mode 100644
--- /dev/null
+++ b/Engine/Source/Editor/UnrealEd/Private/ThumbnailRendering/AnimSequenceThumbnailRenderer.cpp
@@ -0,0 +1,50 @@
+// Copyright 1998-2014 Epic Games, Inc. All Rights Reserved.
+
+#include "UnrealEd.h"
+
+// FPreviewScene derived helpers for rendering
+#include "ThumbnailHelpers.h"
+#include "EngineModule.h"
+
+UAnimSequenceThumbnailRenderer::UAnimSequenceThumbnailRenderer(const class FPostConstructInitializeProperties& PCIP)
+	: Super(PCIP)
+{
+	ThumbnailScene = nullptr;
+}
+
+void UAnimSequenceThumbnailRenderer::Draw(UObject* Object, int32 X, int32 Y, uint32 Width, uint32 Height, FRenderTarget* RenderTarget, FCanvas* Canvas)
+{
+	UAnimSequenceBase* Anim = Cast<UAnimSequenceBase>(Object);
+	if (Anim != nullptr)
+	{
+		if ( ThumbnailScene == nullptr )
+		{
+			ThumbnailScene = new FAnimationSequenceThumbnailScene();
+		}
+
+		if(ThumbnailScene->SetAnimation(Anim))
+		{
+			FSceneViewFamilyContext ViewFamily(FSceneViewFamily::ConstructionValues(RenderTarget, ThumbnailScene->GetScene(), FEngineShowFlags(ESFIM_Game))
+				.SetWorldTimes(FApp::GetCurrentTime() - GStartTime, FApp::GetDeltaTime(), FApp::GetCurrentTime() - GStartTime));
+
+			ViewFamily.EngineShowFlags.DisableAdvancedFeatures();
+			ViewFamily.EngineShowFlags.MotionBlur = 0;
+			ViewFamily.EngineShowFlags.LOD = 0;
+
+			ThumbnailScene->GetView(&ViewFamily, X, Y, Width, Height);
+			GetRendererModule().BeginRenderingViewFamily(Canvas, &ViewFamily);
+			ThumbnailScene->SetAnimation(nullptr);
+		}
+	}
+}
+
+void UAnimSequenceThumbnailRenderer::BeginDestroy()
+{
+	if ( ThumbnailScene != nullptr )
+	{
+		delete ThumbnailScene;
+		ThumbnailScene = nullptr;
+	}
+
+	Super::BeginDestroy();
+}
diff --git a/Engine/Source/Editor/UnrealEd/Private/ThumbnailRendering/BlendSpaceThumbnailRenderer.cpp b/Engine/Source/Editor/UnrealEd/Private/ThumbnailRendering/BlendSpaceThumbnailRenderer.cpp
new file mode 100644
--- /dev/null
+++ b/Engine/Source/Editor/UnrealEd/Private/ThumbnailRendering/BlendSpaceThumbnailRenderer.cpp
@@ -0,0 +1,50 @@
+// Copyright 1998-2014 Epic Games, Inc. All Rights Reserved.
+
+#include "UnrealEd.h"
+
+// FPreviewScene derived helpers for rendering
+#include "ThumbnailHelpers.h"
+#include "EngineModule.h"
+
+UBlendSpaceThumbnailRenderer::UBlendSpaceThumbnailRenderer(const class FPostConstructInitializeProperties& PCIP)
+	: Super(PCIP)
+{
+	ThumbnailScene = nullptr;
+}
+
+void UBlendSpaceThumbnailRenderer::Draw(UObject* Object, int32 X, int32 Y, uint32 Width, uint32 Height, FRenderTarget* RenderTarget, FCanvas* Canvas)
+{
+	UBlendSpaceBase* BlendSpace = Cast<UBlendSpaceBase>(Object);
+	if (BlendSpace != nullptr)
+	{
+		if (ThumbnailScene == nullptr)
+		{
+			ThumbnailScene = new FBlendSpaceThumbnailScene();
+		}
+
+		if (ThumbnailScene->SetBlendSpace(BlendSpace))
+		{
+			FSceneViewFamilyContext ViewFamily(FSceneViewFamily::ConstructionValues(RenderTarget, ThumbnailScene->GetScene(), FEngineShowFlags(ESFIM_Game))
+				.SetWorldTimes(FApp::GetCurrentTime() - GStartTime, FApp::GetDeltaTime(), FApp::GetCurrentTime() - GStartTime));
+
+			ViewFamily.EngineShowFlags.DisableAdvancedFeatures();
+			ViewFamily.EngineShowFlags.MotionBlur = 0;
+			ViewFamily.EngineShowFlags.LOD = 0;
+
+			ThumbnailScene->GetView(&ViewFamily, X, Y, Width, Height);
+			GetRendererModule().BeginRenderingViewFamily(Canvas, &ViewFamily);
+			ThumbnailScene->SetBlendSpace(nullptr);
+		}
+	}
+}
+
+void UBlendSpaceThumbnailRenderer::BeginDestroy()
+{
+	if ( ThumbnailScene != nullptr )
+	{
+		delete ThumbnailScene;
+		ThumbnailScene = nullptr;
+	}
+
+	Super::BeginDestroy();
+}
diff --git a/Engine/Source/Editor/UnrealEd/Public/ThumbnailHelpers.h b/Engine/Source/Editor/UnrealEd/Public/ThumbnailHelpers.h
--- a/Engine/Source/Editor/UnrealEd/Public/ThumbnailHelpers.h
+++ b/Engine/Source/Editor/UnrealEd/Public/ThumbnailHelpers.h
@@ -107,6 +107,78 @@ private:
 	AStaticMeshActor* PreviewActor;
 };
 
+class UNREALED_API FAnimationSequenceThumbnailScene : public FThumbnailPreviewScene
+{
+public:
+	/** Constructor */
+	FAnimationSequenceThumbnailScene();
+
+	/** Sets the animation to use in the next GetView() */
+	bool SetAnimation(class UAnimSequenceBase* InAnimation);
+
+protected:
+	// FThumbnailPreviewScene implementation
+	virtual void GetViewMatrixParameters(const float InFOVDegrees, FVector& OutOrigin, float& OutOrbitPitch, float& OutOrbitYaw, float& OutOrbitZoom) const override;
+	
+	//Clean up the children of this component
+	void CleanupComponentChildren(USceneComponent* Component);
+
+private:
+	/** The skeletal mesh actor used to display all animation thumbnails */
+	class ASkeletalMeshActor* PreviewActor;
+
+	/** Animation we are generating the thumbnail for */
+	class UAnimSequenceBase* PreviewAnimation;
+};
+
+class UNREALED_API FBlendSpaceThumbnailScene : public FThumbnailPreviewScene
+{
+public:
+	/** Constructor */
+	FBlendSpaceThumbnailScene();
+
+	/** Sets the animation to use in the next GetView() */
+	bool SetBlendSpace(class UBlendSpaceBase* InBlendSpace);
+
+protected:
+	// FThumbnailPreviewScene implementation
+	virtual void GetViewMatrixParameters(const float InFOVDegrees, FVector& OutOrigin, float& OutOrbitPitch, float& OutOrbitYaw, float& OutOrbitZoom) const override;
+
+	//Clean up the children of this component
+	void CleanupComponentChildren(USceneComponent* Component);
+
+private:
+	/** The skeletal mesh actor used to display all animation thumbnails */
+	class ASkeletalMeshActor* PreviewActor;
+
+	/** Animation we are generating the thumbnail for */
+	class UBlendSpaceBase* PreviewAnimation;
+};
+
+class UNREALED_API FAnimBlueprintThumbnailScene : public FThumbnailPreviewScene
+{
+public:
+	/** Constructor */
+	FAnimBlueprintThumbnailScene();
+
+	/** Sets the animation blueprint to use in the next GetView() */
+	bool SetAnimBlueprint(class UAnimBlueprint* InBlueprint);
+
+protected:
+	// FThumbnailPreviewScene implementation
+	virtual void GetViewMatrixParameters(const float InFOVDegrees, FVector& OutOrigin, float& OutOrbitPitch, float& OutOrbitYaw, float& OutOrbitZoom) const override;
+
+	//Clean up the children of this component
+	void CleanupComponentChildren(USceneComponent* Component);
+
+private:
+	/** The skeletal mesh actor used to display all animation thumbnails */
+	class ASkeletalMeshActor* PreviewActor;
+
+	/** Animation Blueprint we are generating the thumbnail for */
+	class UAnimBlueprint* PreviewBlueprint;
+};
+
 class UActorComponent;
 
 class UNREALED_API FBlueprintThumbnailScene : public FThumbnailPreviewScene
diff --git a/Engine/Source/Runtime/Engine/Classes/Animation/AnimSingleNodeInstance.h b/Engine/Source/Runtime/Engine/Classes/Animation/AnimSingleNodeInstance.h
--- a/Engine/Source/Runtime/Engine/Classes/Animation/AnimSingleNodeInstance.h
+++ b/Engine/Source/Runtime/Engine/Classes/Animation/AnimSingleNodeInstance.h
@@ -100,6 +100,12 @@ public:
 	virtual void RestartMontage(UAnimMontage * Montage, FName FromSection = FName());
 	void SetMontageLoop(UAnimMontage* Montage, bool bIsLooping, FName StartingSection = FName());
 
+	/** Updates montage weights based on a jump in time (as this wont be handled by SetPosition) */
+	void UpdateMontageWeightForTimeSkip(float TimeDifference);
+
+	/** Updates the blendspace samples list in the case of our asset being a blendspace */
+	void UpdateBlendspaceSamples(FVector InBlendInput);
+
 #if WITH_EDITORONLY_DATA
 	float PreviewPoseCurrentTime;
 #endif
diff --git a/Engine/Source/Runtime/Engine/Classes/Animation/AnimationAsset.h b/Engine/Source/Runtime/Engine/Classes/Animation/AnimationAsset.h
--- a/Engine/Source/Runtime/Engine/Classes/Animation/AnimationAsset.h
+++ b/Engine/Source/Runtime/Engine/Classes/Animation/AnimationAsset.h
@@ -413,7 +413,12 @@ public:
 	 * @param ReplacementMap	Mapping of original asset to new asset
 	 **/
 	ENGINE_API virtual void ReplaceReferredAnimations(const TMap<UAnimSequence*, UAnimSequence*>& ReplacementMap);	
+#endif
 
+#if WITH_EDITORONLY_DATA
+	/** Information for thumbnail rendering */
+	UPROPERTY(VisibleAnywhere, EditInline, Category = Thumbnail)
+	class UThumbnailInfo* ThumbnailInfo;
 #endif
 
 public:
diff --git a/Engine/Source/Runtime/Engine/Private/Animation/AnimSingleNodeInstance.cpp b/Engine/Source/Runtime/Engine/Private/Animation/AnimSingleNodeInstance.cpp
--- a/Engine/Source/Runtime/Engine/Private/Animation/AnimSingleNodeInstance.cpp
+++ b/Engine/Source/Runtime/Engine/Private/Animation/AnimSingleNodeInstance.cpp
@@ -163,6 +163,20 @@ void UAnimSingleNodeInstance::SetMontageLoop(UAnimMontage* Montage, bool bIsLoop
 	}
 }
 
+void UAnimSingleNodeInstance::UpdateMontageWeightForTimeSkip(float TimeDifference)
+{
+	Montage_UpdateWeight(TimeDifference);
+}
+
+void UAnimSingleNodeInstance::UpdateBlendspaceSamples(FVector InBlendInput)
+{
+	if (UBlendSpaceBase* BlendSpace = Cast<UBlendSpaceBase>(CurrentAsset))
+	{
+		float CurrentTime = 0.f;
+		BlendSpaceAdvanceImmediate(BlendSpace, InBlendInput, BlendSampleData, BlendFilter, false, 1.f, 0.f, CurrentTime);
+	}
+}
+
 void UAnimSingleNodeInstance::RestartMontage(UAnimMontage * Montage, FName FromSection)
 {
 	if( Montage == CurrentAsset )
