Changed 'future' map messages to make more sense

#ttp 340833 - EDITOR: Opening project in older version warning message does not display issue

#branch UE4

#change It's not possible to display any other version info (like changelist or GitHub version) as this information isn't saved to the packages, so instead I reworded the existing messages to try to be meaningful to the user and make it clearer what the current version is.

diff --git a/Engine/Source/Editor/UnrealEd/Private/EditorServer.cpp b/Engine/Source/Editor/UnrealEd/Private/EditorServer.cpp
--- a/Engine/Source/Editor/UnrealEd/Private/EditorServer.cpp
+++ b/Engine/Source/Editor/UnrealEd/Private/EditorServer.cpp
@@ -1936,7 +1936,7 @@ bool UEditorEngine::PackageIsAMapFile( const TCHAR* PackageFilename, FText& OutN
 		{
 			FFormatNamedArguments Arguments;
 			Arguments.Add(TEXT("File"), FText::FromString( FString( PackageFilename ) ));
-			OutNotMapReason = FText::Format( LOCTEXT( "FileIsAnAsset", "{File} appears to be an asset file!" ), 
+			OutNotMapReason = FText::Format( LOCTEXT( "FileIsAnAsset", "{File} appears to be an asset file." ), 
 				Arguments );
 			return false;
 		}
@@ -1951,7 +1951,7 @@ bool UEditorEngine::PackageIsAMapFile( const TCHAR* PackageFilename, FText& OutN
 			Arguments.Add(TEXT("File"), FText::FromString( FString( PackageFilename ) ));
 			Arguments.Add(TEXT("Version"), UE3Version);
 			Arguments.Add(TEXT("First"), VER_MIN_ENGINE_UE3);
-			OutNotMapReason = FText::Format( LOCTEXT( "UE3FileIsOlder", "{File} is an outdated UE3 map file [Cur:{Version}], from an engine release no longer supported [Min:{First}]!" ), 
+			OutNotMapReason = FText::Format( LOCTEXT( "UE3FileIsOlder", "{File} is an UE3 map [File:v{Version}], from an engine release no longer supported [Min:v{First}]." ), 
 				Arguments);
 			return false;
 		}
@@ -1962,7 +1962,7 @@ bool UEditorEngine::PackageIsAMapFile( const TCHAR* PackageFilename, FText& OutN
 			Arguments.Add(TEXT("File"), FText::FromString( FString( PackageFilename ) ));
 			Arguments.Add(TEXT("Version"), UE4Version);
 			Arguments.Add(TEXT("First"), VER_UE4_OLDEST_LOADABLE_PACKAGE);
-			OutNotMapReason = FText::Format( LOCTEXT( "UE4FileIsOlder", "{File} is an outdated UE4 map file [Cur:{Version}], from an engine release no longer supported [Min:{First}]!" ), 
+			OutNotMapReason = FText::Format( LOCTEXT( "UE4FileIsOlder", "{File} is an UE4 map [File:v{Version}], from an engine release no longer supported [Min:v{First}]." ), 
 				Arguments);
 			return false;
 		}
@@ -1976,7 +1976,7 @@ bool UEditorEngine::PackageIsAMapFile( const TCHAR* PackageFilename, FText& OutN
 			Arguments.Add(TEXT("File"), FText::FromString( FString( PackageFilename ) ));
 			Arguments.Add(TEXT("Version"), UE3Version);
 			Arguments.Add(TEXT("Last"), VER_LAST_ENGINE_UE3);
-			OutNotMapReason = FText::Format( LOCTEXT( "UE3FileIsNewer", "{File} is a future UE3 map file [Cur:{Version}], from an engine release newer than this [Max:{Last}]!" ), 
+			OutNotMapReason = FText::Format( LOCTEXT( "UE3FileIsNewer", "{File} is a UE3 map [File:v{Version}], from an engine release newer than this [Cur:v{Last}]." ), 
 				Arguments);
 			return false;
 		}
@@ -1986,7 +1986,7 @@ bool UEditorEngine::PackageIsAMapFile( const TCHAR* PackageFilename, FText& OutN
 			Arguments.Add(TEXT("File"), FText::FromString( FString( PackageFilename ) ));
 			Arguments.Add(TEXT("Version"), UE4Version);
 			Arguments.Add(TEXT("Last"), GPackageFileUE4Version);
-			OutNotMapReason = FText::Format( LOCTEXT( "UE4FileIsNewer", "{File} is a future UE4 map file [Cur:{Version}], from an engine release newer than this [Max:{Last}]!" ), 
+			OutNotMapReason = FText::Format( LOCTEXT( "UE4FileIsNewer", "{File} is a UE4 map [File:v{Version}], from an engine release newer than this [Cur:v{Last}]." ), 
 				Arguments);
 			return false;
 		}
@@ -1996,7 +1996,7 @@ bool UEditorEngine::PackageIsAMapFile( const TCHAR* PackageFilename, FText& OutN
 			Arguments.Add(TEXT("File"), FText::FromString( FString( PackageFilename ) ));
 			Arguments.Add(TEXT("Version"), UE4LicenseeVersion);
 			Arguments.Add(TEXT("Last"), GPackageFileLicenseeUE4Version);
-			OutNotMapReason = FText::Format( LOCTEXT( "UE4FileIsNewer", "{File} is a future UE4 map file [Cur:{Version}], from an engine release newer than this [Max:{Last}]!" ), 
+			OutNotMapReason = FText::Format( LOCTEXT( "UE4FileIsNewer", "{File} is a UE4 map [File:v{Version}], from an engine release newer than this [Cur:v{Last}]." ), 
 				Arguments);
 			return false;
 		}
@@ -2040,8 +2040,10 @@ bool UEditorEngine::Map_Load(const TCHAR* Str, FOutputDevice& Ar)
 				FText NotMapReason;
 				if( !ExistingWorld && !PackageIsAMapFile( TempFname, NotMapReason ) )
 				{
-					// Map load failed 
-					FMessageDialog::Open( EAppMsgType::Ok, NotMapReason );
+					// Map load failed
+					FFormatNamedArguments Arguments;
+					Arguments.Add(TEXT("Reason"), NotMapReason);
+					FMessageDialog::Open(EAppMsgType::Ok, FText::Format(LOCTEXT("MapLoadFailed", "Failed to load map!\n{Reason}"), Arguments));
 					GIsEditorLoadingPackage = false;
 					return false;
 				}
