Fixed brushes not rendering previews while dragging

This change also fixes asset drag/drop decorators not respecting the correct visibility for different actor types. Fixes TTP341719 - CRITICAL: Regression: EDITOR: Geometry: Dragging from Modes panel shows icon instead of outline.

diff --git a/Engine/Source/Editor/UnrealEd/Private/Editor/ActorPositioning.cpp b/Engine/Source/Editor/UnrealEd/Private/Editor/ActorPositioning.cpp
--- a/Engine/Source/Editor/UnrealEd/Private/Editor/ActorPositioning.cpp
+++ b/Engine/Source/Editor/UnrealEd/Private/Editor/ActorPositioning.cpp
@@ -118,6 +118,7 @@ FActorPositionTraceResult FActorPositioning::TraceWorldForPosition(const UWorld&
 				Results.Location = Hit.Location;
 				Results.SurfaceNormal = Hit.Normal.SafeNormal();
 				Results.State = FActorPositionTraceResult::HitSuccess;
+				Results.HitActor = Hit.Actor;
 			}
 		}
 	}
diff --git a/Engine/Source/Editor/UnrealEd/Private/Editor/ActorPositioning.h b/Engine/Source/Editor/UnrealEd/Private/Editor/ActorPositioning.h
--- a/Engine/Source/Editor/UnrealEd/Private/Editor/ActorPositioning.h
+++ b/Engine/Source/Editor/UnrealEd/Private/Editor/ActorPositioning.h
@@ -27,6 +27,9 @@ struct FActorPositionTraceResult
 
 	/** The surface normal of the trace hit */
 	FVector		SurfaceNormal;
+
+	/** Pointer to the actor that was hit, if any. nullptr otherwise */
+	TWeakObjectPtr<AActor>	HitActor;
 };
 
 struct FActorPositioning
diff --git a/Engine/Source/Editor/UnrealEd/Private/LevelEditorViewport.cpp b/Engine/Source/Editor/UnrealEd/Private/LevelEditorViewport.cpp
--- a/Engine/Source/Editor/UnrealEd/Private/LevelEditorViewport.cpp
+++ b/Engine/Source/Editor/UnrealEd/Private/LevelEditorViewport.cpp
@@ -992,12 +992,28 @@ bool FLevelEditorViewportClient::UpdateDropPreviewActors(int32 MouseX, int32 Mou
 	// Snap the new location if snapping is enabled
 	FSnappingUtils::SnapPointToGrid(GEditor->ClickLocation, FVector::ZeroVector);
 
+	AActor* DroppedOnActor = TraceResult.HitActor.Get();
+	
+	if (DroppedOnActor)
+	{
+		// We indicate that the dropped objects are visible if *any* of them are not applicable to other actors
+		out_bDroppedObjectsVisible = DroppedObjects.ContainsByPredicate([&](UObject* AssetObj){
+			return !AttemptApplyObjToActor(AssetObj, DroppedOnActor, -1, true);
+		});
+	}
+	else
+	{
+		// All dropped objects are visible if we're not dropping on an actor
+		out_bDroppedObjectsVisible = true;
+	}
+
 	for (AActor* Actor : DraggingActors)
 	{
 		const UActorFactory* ActorFactory = FactoryToUse ? FactoryToUse : GEditor->FindActorFactoryForActorClass(Actor->GetClass());
 		const FTransform ActorTransform = FActorPositioning::GetSnappedSurfaceAlignedTransform(this, ActorFactory, TraceResult.Location, TraceResult.SurfaceNormal, Actor->GetPlacementExtent());
 
 		Actor->SetActorTransform(ActorTransform);
+		Actor->SetIsTemporarilyHiddenInEditor(false);
 	}
 
 	return true;
