Actors now default to being aligned to Z-up if they are placed on the background.

TTP 341756 - Regression: EDITOR: Dragging Deferred Decal into the world from the modes tab defaults in pointing down -Y instead of -Z.

diff --git a/Engine/Source/Editor/UnrealEd/Private/Editor/ActorPositioning.h b/Engine/Source/Editor/UnrealEd/Private/Editor/ActorPositioning.h
--- a/Engine/Source/Editor/UnrealEd/Private/Editor/ActorPositioning.h
+++ b/Engine/Source/Editor/UnrealEd/Private/Editor/ActorPositioning.h
@@ -17,7 +17,7 @@ struct FActorPositionTraceResult
 	};
 
 	/** Constructor */
-	FActorPositionTraceResult() : State(Failed), Location(0.f), SurfaceNormal(0.f) {}
+	FActorPositionTraceResult() : State(Failed), Location(0.f), SurfaceNormal(0.f, 0.f, 1.f) {}
 
 	/** The state of this result */
 	ResultState	State;
