Move CookerSettings.h to better location, so that making changes doesn't hurt so much

diff --git a/Engine/Source/Editor/UnrealEd/Classes/Settings/CookerSettings.h b/Engine/Source/Editor/UnrealEd/Classes/Settings/CookerSettings.h
deleted file mode 100644
diff --git a/Engine/Source/Editor/UnrealEd/Public/CookerSettings.h b/Engine/Source/Editor/UnrealEd/Public/CookerSettings.h
new file mode 100644
--- /dev/null
+++ b/Engine/Source/Editor/UnrealEd/Public/CookerSettings.h
@@ -0,0 +1,34 @@
+// Copyright 1998-2014 Epic Games, Inc. All Rights Reserved.
+
+/*=============================================================================
+	CookerSettings.h: Declares the UCookerSettings class.
+=============================================================================*/
+
+#pragma once
+
+#include "CookerSettings.generated.h"
+
+
+/**
+ * Implements per-project cooker settings exposed to the editor
+ */
+UCLASS(config=Engine)
+class UNREALED_API UCookerSettings
+	: public UObject
+{
+	GENERATED_UCLASS_BODY()
+
+public:
+
+	/** Quality of 0 means fastest, 4 means best */
+	UPROPERTY(config, EditAnywhere, Category=Textures)
+	int32 DefaultPVRTCQuality;
+
+	/** Whether or not to add support for Metal API (requires IOS8 and A7 processors) */
+	UPROPERTY(config, EditAnywhere, Category=IOS)
+	bool bSupportMetal;
+
+	/** Whether or not to add support for OpenGL ES2 (if this is false, then your game should specify minimum IOS8 version and A7, in the plist) */
+	UPROPERTY(config, EditAnywhere, Category = IOS)
+	bool bSupportOpenGLES2;
+};
