- Missed file in last checkin

diff --git a/Engine/Source/Editor/UnrealEd/Private/Commandlets/CookCommandlet.cpp b/Engine/Source/Editor/UnrealEd/Private/Commandlets/CookCommandlet.cpp
--- a/Engine/Source/Editor/UnrealEd/Private/Commandlets/CookCommandlet.cpp
+++ b/Engine/Source/Editor/UnrealEd/Private/Commandlets/CookCommandlet.cpp
@@ -21,6 +21,7 @@
 #include "UnrealEdMessages.h"
 #include "GameDelegates.h"
 #include "ChunkManifestGenerator.h"
+#include "CookerSettings.h"
 
 DEFINE_LOG_CATEGORY_STATIC(LogCookCommandlet, Log, All);
 
