Updated Mac OS X version check to warn users if their version is older than 10.9.4 plus fixed the check to not detect 10.10 as older than 10.9.1

diff --git a/Engine/Source/Editor/UnrealEd/Private/Editor.cpp b/Engine/Source/Editor/UnrealEd/Private/Editor.cpp
--- a/Engine/Source/Editor/UnrealEd/Private/Editor.cpp
+++ b/Engine/Source/Editor/UnrealEd/Private/Editor.cpp
@@ -6473,7 +6473,7 @@ void UEditorEngine::UpdateAutoLoadProject()
 			ComponentValues[i] = [Component integerValue];
 		}
 		
-		if(ComponentValues[0] < 10 || ComponentValues[1] < 9 || ComponentValues[2] < 2)
+		if(ComponentValues[0] < 10 || ComponentValues[1] < 9 || (ComponentValues[1] == 9 && ComponentValues[2] < 4))
 		{
 			FSuppressableWarningDialog::FSetupInfo Info( LOCTEXT("UpdateMacOSX_Body","Please update to the latest version of Mac OS X for best performance."), LOCTEXT("UpdateMacOSX_Title","Update Mac OS X"), TEXT("UpdateMacOSX"), GEditorGameAgnosticIni );
 			Info.ConfirmText = LOCTEXT( "OK", "OK");
