- Added Create Asset toolbar
- Consolidated code between asset skeleton and toolbar
- Added Asset Picker to suppoert Sync To Assets

diff --git a/Engine/Source/Developer/AssetTools/Private/AssetTypeActions/AssetTypeActions_Skeleton.cpp b/Engine/Source/Developer/AssetTools/Private/AssetTypeActions/AssetTypeActions_Skeleton.cpp
--- a/Engine/Source/Developer/AssetTools/Private/AssetTypeActions/AssetTypeActions_Skeleton.cpp
+++ b/Engine/Source/Developer/AssetTools/Private/AssetTypeActions/AssetTypeActions_Skeleton.cpp
@@ -9,6 +9,7 @@
 #include "AssetToolsModule.h"
 #include "ContentBrowserModule.h"
 #include "SSkeletonWidget.h"
+#include "AnimationEditorUtils.h"
 
 #include "Kismet2/KismetEditorUtilities.h"
 #include "Kismet2/BlueprintEditorUtils.h"
@@ -71,81 +72,9 @@ void FAssetTypeActions_Skeleton::FillCreateMenu(FMenuBuilder& MenuBuilder, TArra
 			)
 		);
 
-	MenuBuilder.AddSubMenu(
-		LOCTEXT("SkeletalMesh_NewAimOffset", "Create AimOffset"),
-		LOCTEXT("SkeletalMesh_NewAimOffsetTooltip", "Creates an aimoffset blendspace using the selected skeleton."),
-		FNewMenuDelegate::CreateSP(this, &FAssetTypeActions_Skeleton::FillAimOffsetBlendSpaceMenu, Skeletons));
-
-	MenuBuilder.AddSubMenu(
-		LOCTEXT("SkeletalMesh_NewBlendspace", "Create BlendSpace"),
-		LOCTEXT("SkeletalMesh_NewBlendspaceTooltip", "Creates a blendspace using the skeleton of the selected mesh."),
-		FNewMenuDelegate::CreateSP(this, &FAssetTypeActions_Skeleton::FillBlendSpaceMenu, Skeletons));
-
-	MenuBuilder.AddMenuEntry(
-		LOCTEXT("Skeleton_NewAnimComposite", "Create AnimComposite"),
-		LOCTEXT("Skeleton_NewAnimCompositeTooltip", "Creates an AnimComposite using the selected skeleton."),
-		FSlateIcon(),
-		FUIAction(
-			FExecuteAction::CreateSP(this, &FAssetTypeActions_Skeleton::ExecuteNewAnimAsset<UAnimCompositeFactory, UAnimComposite>, Skeletons, FString("_Composite")),
-			FCanExecuteAction()
-			)
-		);
-
-	MenuBuilder.AddMenuEntry(
-		LOCTEXT("Skeleton_NewAnimMontage", "Create AnimMontage"),
-		LOCTEXT("Skeleton_NewAnimMontageTooltip", "Creates an AnimMontage using the selected skeleton."),
-		FSlateIcon(),
-		FUIAction(
-			FExecuteAction::CreateSP(this, &FAssetTypeActions_Skeleton::ExecuteNewAnimAsset<UAnimMontageFactory, UAnimMontage>, Skeletons, FString("_Montage")),
-			FCanExecuteAction()
-			)
-		);
-}
-
-void FAssetTypeActions_Skeleton::FillBlendSpaceMenu( FMenuBuilder& MenuBuilder, TArray<TWeakObjectPtr<USkeleton>> Skeletons )
-{
-	MenuBuilder.AddMenuEntry(
-		LOCTEXT("SkeletalMesh_New1DBlendspace", "Create 1D BlendSpace"),
-		LOCTEXT("SkeletalMesh_New1DBlendspaceTooltip", "Creates a 1D blendspace using the skeleton of the selected mesh."),
-		FSlateIcon(),
-		FUIAction(
-			FExecuteAction::CreateSP( this, &FAssetTypeActions_Skeleton::ExecuteNewAnimAsset<UBlendSpaceFactory1D, UBlendSpace1D>, Skeletons, FString("_BlendSpace1D") ),
-			FCanExecuteAction()
-			)
-		);
-
-	MenuBuilder.AddMenuEntry(
-		LOCTEXT("SkeletalMesh_New2DBlendspace", "Create 2D BlendSpace"),
-		LOCTEXT("SkeletalMesh_New2DBlendspaceTooltip", "Creates a 2D blendspace using the skeleton of the selected mesh."),
-		FSlateIcon(),
-		FUIAction(
-			FExecuteAction::CreateSP( this, &FAssetTypeActions_Skeleton::ExecuteNewAnimAsset<UBlendSpaceFactoryNew, UBlendSpace>, Skeletons, FString("_BlendSpace2D") ),
-			FCanExecuteAction()
-			)
-		);
-}
-
-void FAssetTypeActions_Skeleton::FillAimOffsetBlendSpaceMenu( FMenuBuilder& MenuBuilder, TArray<TWeakObjectPtr<USkeleton>> Skeletons )
-{
-	MenuBuilder.AddMenuEntry(
-		LOCTEXT("SkeletalMesh_New1DAimOffset", "Create 1D AimOffset"),
-		LOCTEXT("SkeletalMesh_New1DAimOffsetTooltip", "Creates a 1D aimoffset blendspace using the selected skeleton."),
-		FSlateIcon(),
-		FUIAction(
-			FExecuteAction::CreateSP( this, &FAssetTypeActions_Skeleton::ExecuteNewAnimAsset<UAimOffsetBlendSpaceFactory1D, UAimOffsetBlendSpace1D>, Skeletons, FString("_AimOffset1D") ),
-			FCanExecuteAction()
-			)
-		);
-
-	MenuBuilder.AddMenuEntry(
-		LOCTEXT("SkeletalMesh_New2DAimOffset", "Create 2D AimOffset"),
-		LOCTEXT("SkeletalMesh_New2DAimOffsetTooltip", "Creates a 2D aimoffset blendspace using the selected skeleton."),
-		FSlateIcon(),
-		FUIAction(
-			FExecuteAction::CreateSP( this, &FAssetTypeActions_Skeleton::ExecuteNewAnimAsset<UAimOffsetBlendSpaceFactoryNew, UAimOffsetBlendSpace>, Skeletons, FString("_AimOffset2D") ),
-			FCanExecuteAction()
-			)
-		);
+	FAnimAssetCreated OnAssetCreated;
+	OnAssetCreated.BindSP(this, &FAssetTypeActions_Skeleton::OnAssetCreated);
+	AnimationEditorUtils::FillCreateAssetMenu(MenuBuilder, Skeletons, /*OnAssetCreated*/FAnimAssetCreated::CreateSP(this, &FAssetTypeActions_Skeleton::OnAssetCreated));
 }
 
 void FAssetTypeActions_Skeleton::OpenAssetEditor( const TArray<UObject*>& InObjects, TSharedPtr<IToolkitHost> EditWithinLevelEditor )
@@ -175,33 +104,6 @@ void FAssetTypeActions_Skeleton::ExecuteEdit(TArray<TWeakObjectPtr<USkeleton>> O
 	}
 }
 
-template <typename TFactory, typename T>
-void FAssetTypeActions_Skeleton::ExecuteNewAnimAsset(TArray<TWeakObjectPtr<USkeleton>> Objects, const FString InSuffix)
-{
-	if ( Objects.Num() == 1 )
-	{
-		auto Object = Objects[0].Get();
-
-		if ( Object )
-		{
-			// Determine an appropriate name for inline-rename
-			FString Name;
-			FString PackageName;
-			CreateUniqueAssetName(Object->GetOutermost()->GetName(), InSuffix, PackageName, Name);
-
-			TFactory* Factory = ConstructObject<TFactory>(TFactory::StaticClass());
-			Factory->TargetSkeleton = Object;
-
-			FContentBrowserModule& ContentBrowserModule = FModuleManager::LoadModuleChecked<FContentBrowserModule>("ContentBrowser");
-			ContentBrowserModule.Get().CreateNewAsset(Name, FPackageName::GetLongPackagePath(PackageName), T::StaticClass(), Factory);
-		}
-	}
-	else
-	{
-		CreateAnimationAssets(Objects, T::StaticClass(), InSuffix);
-	}
-}
-
 void FAssetTypeActions_Skeleton::ExecuteNewAnimBlueprint(TArray<TWeakObjectPtr<USkeleton>> Objects)
 {
 	const FString DefaultSuffix = TEXT("_AnimBlueprint");
@@ -260,35 +162,9 @@ void FAssetTypeActions_Skeleton::ExecuteNewAnimBlueprint(TArray<TWeakObjectPtr<U
 
 void FAssetTypeActions_Skeleton::CreateAnimationAssets(const TArray<TWeakObjectPtr<USkeleton>>& Skeletons, TSubclassOf<UAnimationAsset> AssetClass, const FString& InPrefix)
 {
-	TArray<UObject*> ObjectsToSync;
-	for (auto SkelIt = Skeletons.CreateConstIterator(); SkelIt; ++SkelIt)
-	{
-		USkeleton* Skeleton = (*SkelIt).Get();
-		if ( Skeleton )
-		{
-			// Determine an appropriate name
-			FString Name;
-			FString PackageName;
-			CreateUniqueAssetName(Skeleton->GetOutermost()->GetName(), InPrefix, PackageName, Name);
-
-			// Create the asset, and assign its skeleton
-			FAssetToolsModule& AssetToolsModule = FModuleManager::GetModuleChecked<FAssetToolsModule>("AssetTools");
-			UAnimationAsset* NewAsset = Cast<UAnimationAsset>(AssetToolsModule.Get().CreateAsset(Name, FPackageName::GetLongPackagePath(PackageName), AssetClass, NULL));
-
-			if ( NewAsset )
-			{
-				NewAsset->SetSkeleton(Skeleton);
-				NewAsset->MarkPackageDirty();
-
-				ObjectsToSync.Add(NewAsset);
-			}
-		}
-	}
-
-	if ( ObjectsToSync.Num() > 0 )
-	{
-		FAssetTools::Get().SyncBrowserToAssets(ObjectsToSync);
-	}
+	FAnimAssetCreated OnAssetCreated;
+	OnAssetCreated.BindSP(this, &FAssetTypeActions_Skeleton::OnAssetCreated);
+	AnimationEditorUtils::CreateAnimationAssets(Skeletons, AssetClass, InPrefix, /*OnAssetCreated*/FAnimAssetCreated::CreateSP(this, &FAssetTypeActions_Skeleton::OnAssetCreated));
 }
 
 void FAssetTypeActions_Skeleton::ExecuteRetargetSkeleton(TArray<TWeakObjectPtr<USkeleton>> Skeletons)
@@ -610,4 +486,13 @@ void FAssetTypeActions_Skeleton::RetargetSkeleton(TArray<FAssetToRemapSkeleton>&
 		}
 	}
 }
+
+void FAssetTypeActions_Skeleton::OnAssetCreated(TArray<UObject*> NewAssets) const
+{
+	if(NewAssets.Num() > 1)
+	{
+		FAssetTools::Get().SyncBrowserToAssets(NewAssets);
+	}
+
+}
 #undef LOCTEXT_NAMESPACE
diff --git a/Engine/Source/Developer/AssetTools/Private/AssetTypeActions/AssetTypeActions_Skeleton.h b/Engine/Source/Developer/AssetTools/Private/AssetTypeActions/AssetTypeActions_Skeleton.h
--- a/Engine/Source/Developer/AssetTools/Private/AssetTypeActions/AssetTypeActions_Skeleton.h
+++ b/Engine/Source/Developer/AssetTools/Private/AssetTypeActions/AssetTypeActions_Skeleton.h
@@ -37,20 +37,10 @@ public:
 	virtual void OpenAssetEditor( const TArray<UObject*>& InObjects, TSharedPtr<class IToolkitHost> EditWithinLevelEditor = TSharedPtr<IToolkitHost>() ) override;
 	virtual uint32 GetCategories() override { return EAssetTypeCategories::Animation; }
 
-	/** Handler for the blend space sub menu */
-	void FillBlendSpaceMenu( FMenuBuilder& MenuBuilder, TArray<TWeakObjectPtr<USkeleton>> Skeletons );
-
-	/** Handler for the blend space sub menu */
-	void FillAimOffsetBlendSpaceMenu( FMenuBuilder& MenuBuilder, TArray<TWeakObjectPtr<USkeleton>> Skeletons );
-
 private:
 	/** Handler for when Edit is selected */
 	void ExecuteEdit(TArray<TWeakObjectPtr<USkeleton>> Objects);
 
-	/** Creates a new object of type T from factory TFactory */
-	template <typename TFactory, typename T>
-	void ExecuteNewAnimAsset(TArray<TWeakObjectPtr<USkeleton>> Objects, const FString InSuffix);
-
 	/** Handler for when NewAnimBlueprint is selected */
 	void ExecuteNewAnimBlueprint(TArray<TWeakObjectPtr<USkeleton>> Objects);
 
@@ -79,4 +69,5 @@ private: // Helper functions
 	void SavePackages(const TArray<UPackage*> PackagesToSave) const;
 	void DetectReadOnlyPackages(TArray<FAssetToRemapSkeleton>& AssetsToRemap, TArray<UPackage*>& InOutPackagesToSave) const;
 	void FillCreateMenu(FMenuBuilder& MenuBuilder, TArray<TWeakObjectPtr<USkeleton>> Skeletons) const;
+	void OnAssetCreated(TArray<UObject*> NewAssets) const;
 };
\ No newline at end of file
diff --git a/Engine/Source/Editor/ContentBrowser/Private/SAssetPicker.cpp b/Engine/Source/Editor/ContentBrowser/Private/SAssetPicker.cpp
--- a/Engine/Source/Editor/ContentBrowser/Private/SAssetPicker.cpp
+++ b/Engine/Source/Editor/ContentBrowser/Private/SAssetPicker.cpp
@@ -25,6 +25,14 @@ void SAssetPicker::Construct( const FArguments& InArgs )
 		}
 	}
 
+	for(auto DelegateIt = InArgs._AssetPickerConfig.SyncToAssetsDelegates.CreateConstIterator(); DelegateIt; ++DelegateIt)
+	{
+		if((*DelegateIt) != NULL)
+		{
+			(**DelegateIt) = FSyncToAssetsDelegate::CreateSP(this, &SAssetPicker::SyncToAssets);
+		}
+	}
+
 	for (auto DelegateIt = InArgs._AssetPickerConfig.SetFilterDelegates.CreateConstIterator(); DelegateIt; ++DelegateIt)
 	{
 		if ((*DelegateIt) != NULL)
@@ -337,6 +345,11 @@ void SAssetPicker::HandleAssetsActivated(const TArray<FAssetData>& ActivatedAsse
 	OnAssetsActivated.ExecuteIfBound( ActivatedAssets, ActivationMethod );
 }
 
+void SAssetPicker::SyncToAssets(const TArray<FAssetData>& AssetDataList)
+{
+	AssetViewPtr->SyncToAssets(AssetDataList);
+}
+
 TArray< FAssetData > SAssetPicker::GetCurrentSelection()
 {
 	return AssetViewPtr->GetSelectedAssets();
diff --git a/Engine/Source/Editor/ContentBrowser/Private/SAssetPicker.h b/Engine/Source/Editor/ContentBrowser/Private/SAssetPicker.h
--- a/Engine/Source/Editor/ContentBrowser/Private/SAssetPicker.h
+++ b/Engine/Source/Editor/ContentBrowser/Private/SAssetPicker.h
@@ -45,6 +45,9 @@ private:
 	/** Handler for when the user double clicks, presses enter, or presses space on an asset */
 	void HandleAssetsActivated(const TArray<FAssetData>& ActivatedAssets, EAssetTypeActivationMethod::Type ActivationMethod);
 
+	/** Selects the paths containing the specified assets. */
+	void SyncToAssets(const TArray<FAssetData>& AssetDataList);
+
 	/** @return The currently selected asset */
 	TArray< FAssetData > GetCurrentSelection();
 
diff --git a/Engine/Source/Editor/ContentBrowser/Public/ContentBrowserDelegates.h b/Engine/Source/Editor/ContentBrowser/Public/ContentBrowserDelegates.h
--- a/Engine/Source/Editor/ContentBrowser/Public/ContentBrowserDelegates.h
+++ b/Engine/Source/Editor/ContentBrowser/Public/ContentBrowserDelegates.h
@@ -25,6 +25,9 @@ DECLARE_DELEGATE_OneParam( FOnThumbnailScaleChanged, const float /*NewScale*/);
 /** Called to retrieve an array of the currently selected asset data */
 DECLARE_DELEGATE_RetVal( TArray< FAssetData >, FGetCurrentSelectionDelegate );
 
+/** Called to retrieve an array of the currently selected asset data */
+DECLARE_DELEGATE_OneParam(FSyncToAssetsDelegate, const TArray< FAssetData >& /*AssetData*/);
+
 /** Called to set a new filter for an existing asset picker */
 DECLARE_DELEGATE_OneParam(FSetARFilterDelegate, const FARFilter& /*NewFilter*/);
 
diff --git a/Engine/Source/Editor/ContentBrowser/Public/IContentBrowserSingleton.h b/Engine/Source/Editor/ContentBrowser/Public/IContentBrowserSingleton.h
--- a/Engine/Source/Editor/ContentBrowser/Public/IContentBrowserSingleton.h
+++ b/Engine/Source/Editor/ContentBrowser/Public/IContentBrowserSingleton.h
@@ -35,6 +35,9 @@ struct FAssetPickerConfig
 	/** An array of pointers to existing delegates which the AssetView will register a function which returns the current selection */
 	TArray<FGetCurrentSelectionDelegate*> GetCurrentSelectionDelegates;
 
+	/** An array of pointers to existing delegates which the AssetView will register a function which sync the asset list*/
+	TArray<FSyncToAssetsDelegate*> SyncToAssetsDelegates;
+
 	/** A pointer to an existing delegate that, when executed, will set the filter an the asset picker after it is created. */
 	TArray<FSetARFilterDelegate*> SetFilterDelegates;
 
diff --git a/Engine/Source/Editor/EditorStyle/Private/SlateEditorStyle.cpp b/Engine/Source/Editor/EditorStyle/Private/SlateEditorStyle.cpp
--- a/Engine/Source/Editor/EditorStyle/Private/SlateEditorStyle.cpp
+++ b/Engine/Source/Editor/EditorStyle/Private/SlateEditorStyle.cpp
@@ -3946,10 +3946,12 @@ void FSlateEditorStyle::FStyle::SetupPersonaStyle()
 		Set("Persona.ImportAnimation.Small", new IMAGE_BRUSH(TEXT("Icons/icon_PhAT_Prismatic_40x"), Icon20x20));
 		Set("Persona.ReimportAnimation", new IMAGE_BRUSH(TEXT("Icons/icon_PhAT_PhysMat_40x"), Icon40x40));
 		Set("Persona.ReimportAnimation.Small", new IMAGE_BRUSH(TEXT("Icons/icon_PhAT_PhysMat_40x"), Icon20x20));
-		Set("Persona.ApplyCompression", new IMAGE_BRUSH(TEXT("Icons/icon_PhAT_NewBody_40x"), Icon40x40));
-		Set("Persona.ApplyCompression.Small", new IMAGE_BRUSH(TEXT("Icons/icon_PhAT_NewBody_40x"), Icon20x20));
+		Set("Persona.ApplyCompression", new IMAGE_BRUSH(TEXT("Icons/icon_Cascade_LowestLOD_40x"), Icon40x40));
+		Set("Persona.ApplyCompression.Small", new IMAGE_BRUSH(TEXT("Icons/icon_Cascade_LowestLOD_40x"), Icon20x20));
 		Set("Persona.ExportToFBX", new IMAGE_BRUSH(TEXT("Icons/icon_DestructibleMeshEd_ImportFBX_40x"), Icon40x40));
 		Set("Persona.ExportToFBX.Small", new IMAGE_BRUSH(TEXT("Icons/icon_DestructibleMeshEd_ImportFBX_40x"), Icon20x20));
+		Set("Persona.CreateAsset", new IMAGE_BRUSH(TEXT("Icons/icon_PhAT_NewBody_40x"), Icon40x40));
+		Set("Persona.CreateAsset.Small", new IMAGE_BRUSH(TEXT("Icons/icon_PhAT_NewBody_40x"), Icon20x20));
 
 		// persona extras
 		Set("Persona.ConvertAnimationGraph", new IMAGE_BRUSH("Old/Graph/ConvertIcon", Icon40x40));
diff --git a/Engine/Source/Editor/Persona/Private/Persona.cpp b/Engine/Source/Editor/Persona/Private/Persona.cpp
--- a/Engine/Source/Editor/Persona/Private/Persona.cpp
+++ b/Engine/Source/Editor/Persona/Private/Persona.cpp
@@ -61,6 +61,8 @@
 #include "SAnimationDlgs.h"
 #include "Developer/AssetTools/Public/AssetToolsModule.h"
 #include "FbxMeshUtils.h"
+#include "AnimationEditorUtils.h"
+#include "SAnimationSequenceBrowser.h"
 
 #define LOCTEXT_NAMESPACE "FPersona"
 
@@ -351,6 +353,12 @@ TSharedPtr<SDockTab> FPersona::OpenNewAnimationDocumentTab(UObject* InAnimAsset)
 				SharedAnimDocumentTab = OpenedTab;
 			}
 		}
+
+		if(SequenceBrowser.IsValid())
+		{
+			UAnimationAsset * NewAsset = CastChecked<UAnimationAsset>(InAnimAsset);
+			SequenceBrowser.Pin()->SelectAsset(NewAsset);
+		}
 	}
 
 	return OpenedTab;
@@ -385,6 +393,18 @@ void FPersona::SetViewport(TWeakPtr<class SAnimationEditorViewportTabBody> NewVi
 	OnViewportCreated.Broadcast(Viewport);
 }
 
+void FPersona::SetSequenceBrowser(class SAnimationSequenceBrowser* InSequenceBrowser)
+{
+	if (InSequenceBrowser)
+	{
+		SequenceBrowser = SharedThis(InSequenceBrowser);
+	}
+	else
+	{
+		SequenceBrowser = NULL;
+	}
+}
+
 void FPersona::RefreshViewport()
 {
 	if (Viewport.IsValid())
@@ -674,7 +694,33 @@ void FPersona::InitPersona(const EToolkitMode::Type Mode, const TSharedPtr< clas
 	// Register post import callback to catch animation imports when we have the asset open (we need to reinit)
 	FEditorDelegates::OnAssetPostImport.AddRaw(this, &FPersona::OnPostImport);
 }
+	
+TSharedRef< SWidget > FPersona::GenerateCreateAssetMenu( USkeleton* Skeleton ) const
+{
+	const bool bShouldCloseWindowAfterMenuSelection = true;
+	FMenuBuilder MenuBuilder(bShouldCloseWindowAfterMenuSelection, NULL);
 
+	TArray<TWeakObjectPtr<USkeleton>> Skeletons;
+	
+	Skeletons.Add(Skeleton);
+	
+	MenuBuilder.BeginSection(NAME_None, LOCTEXT("Persona_CreateAsset", "Create Asset"));
+	{
+		AnimationEditorUtils::FillCreateAssetMenu(MenuBuilder, Skeletons, FAnimAssetCreated::CreateSP(this, &FPersona::OnAssetCreated), false);
+	}
+	MenuBuilder.EndSection();
+
+	return MenuBuilder.MakeWidget();
+}
+
+void FPersona::OnAssetCreated(const TArray<UObject*> NewAssets) 
+{
+	if ( NewAssets.Num() > 0 )
+	{
+		FAssetRegistryModule::AssetCreated(NewAssets[0]);
+		OpenNewDocumentTab(CastChecked<UAnimationAsset>(NewAssets[0]));
+	}
+}
 
 void FPersona::ExtendDefaultPersonaToolbar()
 {
@@ -692,25 +738,44 @@ void FPersona::ExtendDefaultPersonaToolbar()
 	// extend extra menu/toolbars
 	struct Local
 	{
-		static void FillToolbar(FToolBarBuilder& ToolbarBuilder)
+		static void FillToolbar(FToolBarBuilder& ToolbarBuilder, USkeleton * Skeleton, FPersona* PersonaPtr)
 		{
-			ToolbarBuilder.BeginSection("Person");
+			ToolbarBuilder.BeginSection("Skeleton");
 			{
 				ToolbarBuilder.AddToolBarButton(FPersonaCommands::Get().ChangeSkeletonPreviewMesh, NAME_None, LOCTEXT("Toolbar_ChangePreviewMesh", "Set Preview"));
 				ToolbarBuilder.AddToolBarButton(FPersonaCommands::Get().AnimNotifyWindow);
 				ToolbarBuilder.AddToolBarButton(FPersonaCommands::Get().RetargetSourceMgr, NAME_None, LOCTEXT("Toolbar_RetargetSourceMgr", "Retarget Source"));
 				ToolbarBuilder.AddToolBarButton(FPersonaCommands::Get().ImportMesh);
-
 				ToolbarBuilder.AddToolBarButton(FPersonaCommands::Get().ReimportMesh);
-				//ToolbarBuilder.AddToolBarButton(FPersonaCommands::Get().ImportLODs);
-				//ToolbarBuilder.AddToolBarButton(FPersonaCommands::Get().AddBodyPart);
 
+				// animation import menu
 				ToolbarBuilder.AddToolBarButton(FPersonaCommands::Get().ImportAnimation, NAME_None, LOCTEXT("Toolbar_ImportAnimation", "Import"));
 				ToolbarBuilder.AddToolBarButton(FPersonaCommands::Get().ReimportAnimation, NAME_None, LOCTEXT("Toolbar_ReimportAnimation", "Reimport"));
-				ToolbarBuilder.AddToolBarButton(FPersonaCommands::Get().ApplyCompression, NAME_None, LOCTEXT("Toolbar_ApplyCompression", "Compression"));
 				ToolbarBuilder.AddToolBarButton(FPersonaCommands::Get().ExportToFBX, NAME_None, LOCTEXT("Toolbar_ExportToFBX", "Export"));
 			}
 			ToolbarBuilder.EndSection();
+
+			ToolbarBuilder.BeginSection("Animation");
+			{
+				// create button
+				{
+					ToolbarBuilder.AddComboButton(
+						FUIAction(
+							FExecuteAction(),
+							FCanExecuteAction(),
+							FIsActionChecked(),
+							FIsActionButtonVisible::CreateSP(PersonaPtr, &FPersona::IsInPersonaMode, FPersonaModes::AnimationEditMode)
+							),
+						FOnGetContent::CreateSP(PersonaPtr, &GenerateCreateAssetMenu, Skeleton),
+						LOCTEXT("OpenBlueprint_Label", "Create Asset"),
+						LOCTEXT("OpenBlueprint_ToolTip", "Create Assets for this skeleton."),
+						FSlateIcon(FEditorStyle::GetStyleSetName(), "Persona.CreateAsset")
+						);
+				}
+
+				ToolbarBuilder.AddToolBarButton(FPersonaCommands::Get().ApplyCompression, NAME_None, LOCTEXT("Toolbar_ApplyCompression", "Compression"));
+			}
+			ToolbarBuilder.EndSection();
 		}
 	};
 
@@ -719,7 +784,7 @@ void FPersona::ExtendDefaultPersonaToolbar()
 		"Asset",
 		EExtensionHook::After,
 		GetToolkitCommands(),
-		FToolBarExtensionDelegate::CreateStatic(&Local::FillToolbar)
+		FToolBarExtensionDelegate::CreateStatic(&Local::FillToolbar, TargetSkeleton, this)
 		);
 
 	AddToolbarExtender(ToolbarExtender);
diff --git a/Engine/Source/Editor/Persona/Private/Persona.h b/Engine/Source/Editor/Persona/Private/Persona.h
--- a/Engine/Source/Editor/Persona/Private/Persona.h
+++ b/Engine/Source/Editor/Persona/Private/Persona.h
@@ -67,6 +67,8 @@ public:
 
 	/** Set the current preview viewport for Persona */
 	void SetViewport(TWeakPtr<class SAnimationEditorViewportTabBody> NewViewport);
+	/** Set Sequence browser */
+	void SetSequenceBrowser(class SAnimationSequenceBrowser * SequenceBrowser);
 
 	/** Refresh viewport */
 	void RefreshViewport();
@@ -563,6 +565,8 @@ private:
 	void OnReimportMesh();
 	void OnImportAsset(enum EFBXImportType DefaultImportType);
 	void OnReimportAnimation();
+	void OnAssetCreated(const TArray<UObject*> NewAssets);
+	TSharedRef< SWidget > GenerateCreateAssetMenu( USkeleton* Skeleton ) const;
 
 	/** Extend menu and toolbar */
 	void ExtendMenu();
@@ -584,4 +588,7 @@ private:
 
 	/** Preview instance inspector widget */
 	TSharedPtr<class SKismetInspector> PreviewEditor;
+
+	/** Sequence Browser **/
+	TWeakPtr<class SAnimationSequenceBrowser> SequenceBrowser;
 };
diff --git a/Engine/Source/Editor/Persona/Private/PersonaCommands.h b/Engine/Source/Editor/Persona/Private/PersonaCommands.h
--- a/Engine/Source/Editor/Persona/Private/PersonaCommands.h
+++ b/Engine/Source/Editor/Persona/Private/PersonaCommands.h
@@ -47,4 +47,11 @@ public:
 	TSharedPtr<FUICommandInfo> ExportToFBX;
 	// Add looping interpolation
 	TSharedPtr<FUICommandInfo> AddLoopingInterpolation;
+	// Create Asset command
+	TSharedPtr<FUICommandInfo> CreateAimOffset1D;
+	TSharedPtr<FUICommandInfo> CreateAimOffset2D;
+	TSharedPtr<FUICommandInfo> CreateBlendSpace1D;
+	TSharedPtr<FUICommandInfo> CreateBlendSpace2D;
+	TSharedPtr<FUICommandInfo> CreateAnimComposite;
+	TSharedPtr<FUICommandInfo> CreateAanimMontage;
 };
diff --git a/Engine/Source/Editor/Persona/Private/SAnimationSequenceBrowser.cpp b/Engine/Source/Editor/Persona/Private/SAnimationSequenceBrowser.cpp
--- a/Engine/Source/Editor/Persona/Private/SAnimationSequenceBrowser.cpp
+++ b/Engine/Source/Editor/Persona/Private/SAnimationSequenceBrowser.cpp
@@ -8,10 +8,10 @@
 #include "Persona.h"
 #include "AssetRegistryModule.h"
 #include "SSkeletonWidget.h"
-#include "Editor/ContentBrowser/Public/ContentBrowserModule.h"
 #include "FeedbackContextEditor.h"
 #include "EditorAnimUtils.h"
 #include "Editor/ContentBrowser/Public/FrontendFilterBase.h"
+#include "Runtime/AssetRegistry/Public/AssetRegistryModule.h"
 
 #define LOCTEXT_NAMESPACE "SequenceBrowser"
 
@@ -38,6 +38,14 @@ public:
 
 const int32 SAnimationSequenceBrowser::MaxAssetsHistory = 10;
 
+SAnimationSequenceBrowser::~SAnimationSequenceBrowser()
+{
+	if (PersonaPtr.IsValid())
+	{
+		PersonaPtr.Pin()->SetSequenceBrowser(NULL);
+	}
+}
+
 void SAnimationSequenceBrowser::OnAnimSelected(const FAssetData& AssetData)
 {
 	CacheOriginalAnimAssetHistory();
@@ -293,8 +301,9 @@ void SAnimationSequenceBrowser::OnCreateCopy(TArray<FAssetData> Selected)
 
 			if(UAnimationAsset* AnimAsset = Cast<UAnimationAsset>(AssetToOpen))
 			{
+				FAssetRegistryModule::AssetCreated(AssetToOpen);
 				// once all success, attempt to open new persona module with new skeleton
-				EToolkitMode::Type Mode = EToolkitMode::Standalone;
+				EToolkitMode::Type Mode = EToolkitMode::WorldCentric;
 				FPersonaModule& PersonaModule = FModuleManager::LoadModuleChecked<FPersonaModule>( "Persona" );
 				PersonaModule.CreatePersona( Mode, TSharedPtr<IToolkitHost>(), NewSkeleton, NULL, AnimAsset, NULL );
 			}
@@ -310,6 +319,12 @@ bool SAnimationSequenceBrowser::CanShowColumnForAssetRegistryTag(FName AssetType
 void SAnimationSequenceBrowser::Construct(const FArguments& InArgs)
 {
 	PersonaPtr = InArgs._Persona;
+
+	if (PersonaPtr.IsValid())
+	{
+		PersonaPtr.Pin()->SetSequenceBrowser(this);
+	}
+
 	CurrentAssetHistoryIndex = INDEX_NONE;
 	bTriedToCacheOrginalAsset = false;
 
@@ -339,6 +354,8 @@ void SAnimationSequenceBrowser::Construct(const FArguments& InArgs)
 	Config.OnAssetDoubleClicked = FOnAssetDoubleClicked::CreateSP(this, &SAnimationSequenceBrowser::OnRequestOpenAsset, false);
 	Config.OnGetAssetContextMenu = FOnGetAssetContextMenu::CreateSP(this, &SAnimationSequenceBrowser::OnGetAssetContextMenu);
 	Config.OnAssetTagWantsToBeDisplayed = FOnShouldDisplayAssetTag::CreateSP(this, &SAnimationSequenceBrowser::CanShowColumnForAssetRegistryTag);
+	Config.SyncToAssetsDelegates.Add(&SyncToAssetsDelegate);
+	Config.GetCurrentSelectionDelegates.Add(&GetCurrentSelectionDelegate);
 	Config.bFocusSearchBoxWhenOpened = false;
 	Config.DefaultFilterMenuExpansion = EAssetTypeCategories::Animation;
 
@@ -600,4 +617,21 @@ void SAnimationSequenceBrowser::CacheOriginalAnimAssetHistory()
 	}
 }
 
+void SAnimationSequenceBrowser::SelectAsset(UAnimationAsset * AnimAsset)
+{
+	FAssetData AssetData(AnimAsset);
+
+	if (AssetData.IsValid())
+	{
+		TArray<FAssetData> CurrentSelection = GetCurrentSelectionDelegate.Execute();
+
+		if ( !CurrentSelection.Contains(AssetData) )
+		{
+			TArray<FAssetData> AssetsToSelect;
+			AssetsToSelect.Add(AssetData);
+
+			SyncToAssetsDelegate.Execute(AssetsToSelect);
+		}
+	}
+}
 #undef LOCTEXT_NAMESPACE
diff --git a/Engine/Source/Editor/Persona/Private/SAnimationSequenceBrowser.h b/Engine/Source/Editor/Persona/Private/SAnimationSequenceBrowser.h
--- a/Engine/Source/Editor/Persona/Private/SAnimationSequenceBrowser.h
+++ b/Engine/Source/Editor/Persona/Private/SAnimationSequenceBrowser.h
@@ -4,6 +4,7 @@
 #pragma once
 
 #include "AssetData.h"
+#include "Editor/ContentBrowser/Public/ContentBrowserModule.h"
 
 //////////////////////////////////////////////////////////////////////////
 // SAnimationSequenceBrowser
@@ -23,6 +24,8 @@ public:
 	void OnAnimSelected(const FAssetData& AssetData);
 	void OnRequestOpenAsset(const FAssetData& AssetData, bool bFromHistory);
 
+	virtual ~SAnimationSequenceBrowser();
+
 	/** Delegate that handles creation of context menu */
 	TSharedPtr<SWidget> OnGetAssetContextMenu(const TArray<FAssetData>& SelectedAssets);
 
@@ -51,6 +54,9 @@ public:
 	 */
 	void OnCreateCopy(TArray<FAssetData> Selected);
 
+	/** Refresh list */
+	void SelectAsset(UAnimationAsset * AnimAsset);
+
 protected:
 	bool CanShowColumnForAssetRegistryTag(FName AssetType, FName TagName) const;
 	
@@ -112,4 +118,8 @@ protected:
 
 	// Track if we have tried to cache the first asset we were playing
 	bool bTriedToCacheOrginalAsset;
+
+	// delegate to sync the asset picker to selected assets
+	FSyncToAssetsDelegate SyncToAssetsDelegate;
+	FGetCurrentSelectionDelegate GetCurrentSelectionDelegate;
 };
\ No newline at end of file
diff --git a/Engine/Source/Editor/UnrealEd/Private/AnimationEditorUtils.cpp b/Engine/Source/Editor/UnrealEd/Private/AnimationEditorUtils.cpp
new file mode 100644
--- /dev/null
+++ b/Engine/Source/Editor/UnrealEd/Private/AnimationEditorUtils.cpp
@@ -0,0 +1,180 @@
+// Copyright 1998-2014 Epic Games, Inc. All Rights Reserved.
+
+#include "UnrealEd.h"
+#include "AnimationEditorUtils.h"
+#include "AssetToolsModule.h"
+#include "ContentBrowserModule.h"
+
+#define LOCTEXT_NAMESPACE "AnimationEditorUtils"
+
+namespace AnimationEditorUtils
+{
+	/** Creates a unique package and asset name taking the form InBasePackageName+InSuffix */
+	void CreateUniqueAssetName(const FString& InBasePackageName, const FString& InSuffix, FString& OutPackageName, FString& OutAssetName) 
+	{
+		FAssetToolsModule& AssetToolsModule = FModuleManager::Get().LoadModuleChecked<FAssetToolsModule>("AssetTools");
+		AssetToolsModule.Get().CreateUniqueAssetName(InBasePackageName, InSuffix, OutPackageName, OutAssetName);
+	}
+
+	void CreateAnimationAssets(const TArray<TWeakObjectPtr<USkeleton>>& Skeletons, TSubclassOf<UAnimationAsset> AssetClass, const FString& InPrefix, FAnimAssetCreated AssetCreated )
+	{
+		TArray<UObject*> ObjectsToSync;
+		for(auto SkelIt = Skeletons.CreateConstIterator(); SkelIt; ++SkelIt)
+		{
+			USkeleton* Skeleton = (*SkelIt).Get();
+			if(Skeleton)
+			{
+				// Determine an appropriate name
+				FString Name;
+				FString PackageName;
+				CreateUniqueAssetName(Skeleton->GetOutermost()->GetName(), InPrefix, PackageName, Name);
+
+				// Create the asset, and assign its skeleton
+				FAssetToolsModule& AssetToolsModule = FModuleManager::GetModuleChecked<FAssetToolsModule>("AssetTools");
+				UAnimationAsset* NewAsset = Cast<UAnimationAsset>(AssetToolsModule.Get().CreateAsset(Name, FPackageName::GetLongPackagePath(PackageName), AssetClass, NULL));
+
+				if(NewAsset)
+				{
+					NewAsset->SetSkeleton(Skeleton);
+					NewAsset->MarkPackageDirty();
+
+					ObjectsToSync.Add(NewAsset);
+				}
+			}
+		}
+
+		if (AssetCreated.IsBound())
+		{
+			AssetCreated.Execute(ObjectsToSync);
+		}
+// 		if(ObjectsToSync.Num() > 0)
+// 		{
+// 			FContentBrowserModule& ContentBrowserModule = FModuleManager::Get().LoadModuleChecked<FContentBrowserModule>("ContentBrowser");
+// 			ContentBrowserModule.Get().SyncBrowserToAssets(ObjectsToSync, /*bAllowLockedBrowsers=*/true);
+// 		}
+	}
+
+	template <typename TFactory, typename T>
+	void ExecuteNewAnimAsset(TArray<TWeakObjectPtr<USkeleton>> Objects, const FString InSuffix, FAnimAssetCreated AssetCreated, bool bInContentBrowser )
+	{
+		if(bInContentBrowser && Objects.Num() == 1)
+		{
+			auto Object = Objects[0].Get();
+
+			if(Object)
+			{
+				// Determine an appropriate name for inline-rename
+				FString Name;
+				FString PackageName;
+				CreateUniqueAssetName(Object->GetOutermost()->GetName(), InSuffix, PackageName, Name);
+
+				TFactory* Factory = ConstructObject<TFactory>(TFactory::StaticClass());
+				Factory->TargetSkeleton = Object;
+
+				FContentBrowserModule& ContentBrowserModule = FModuleManager::LoadModuleChecked<FContentBrowserModule>("ContentBrowser");
+				ContentBrowserModule.Get().CreateNewAsset(Name, FPackageName::GetLongPackagePath(PackageName), T::StaticClass(), Factory);
+
+				if(AssetCreated.IsBound())
+				{
+					// @TODO: this doesn't work
+					//FString LongPackagePath = FPackageName::GetLongPackagePath(PackageName);
+					UObject * 	Parent = FindPackage(NULL, *PackageName);
+					UObject* NewAsset = FindObject<UObject>(Parent, *Name, false);
+					if (NewAsset)
+					{
+						TArray<UObject*> NewAssets;
+						NewAssets.Add(NewAsset);
+						AssetCreated.Execute(NewAssets);
+					}
+				}
+			}
+		}
+		else
+		{
+			CreateAnimationAssets(Objects, T::StaticClass(), InSuffix, AssetCreated);
+		}
+	}
+
+	/** Handler for the blend space sub menu */
+	void FillBlendSpaceMenu(FMenuBuilder& MenuBuilder, TArray<TWeakObjectPtr<USkeleton>> Skeletons, FAnimAssetCreated AssetCreated, bool bInContentBrowser)
+	{
+		MenuBuilder.AddMenuEntry(
+			LOCTEXT("SkeletalMesh_New1DBlendspace", "Create 1D BlendSpace"),
+			LOCTEXT("SkeletalMesh_New1DBlendspaceTooltip", "Creates a 1D blendspace using the skeleton of the selected mesh."),
+			FSlateIcon(),
+			FUIAction(
+				FExecuteAction::CreateStatic(&ExecuteNewAnimAsset<UBlendSpaceFactory1D, UBlendSpace1D>, Skeletons, FString("_BlendSpace1D"), AssetCreated, bInContentBrowser),
+				FCanExecuteAction()
+				)
+			);
+
+		MenuBuilder.AddMenuEntry(
+			LOCTEXT("SkeletalMesh_New2DBlendspace", "Create 2D BlendSpace"),
+			LOCTEXT("SkeletalMesh_New2DBlendspaceTooltip", "Creates a 2D blendspace using the skeleton of the selected mesh."),
+			FSlateIcon(),
+			FUIAction(
+				FExecuteAction::CreateStatic(&ExecuteNewAnimAsset<UBlendSpaceFactoryNew, UBlendSpace>, Skeletons, FString("_BlendSpace2D"), AssetCreated, bInContentBrowser),
+				FCanExecuteAction()
+				)
+			);
+	}
+
+	/** Handler for the blend space sub menu */
+	void FillAimOffsetBlendSpaceMenu(FMenuBuilder& MenuBuilder, TArray<TWeakObjectPtr<USkeleton>> Skeletons, FAnimAssetCreated AssetCreated, bool bInContentBrowser)
+	{
+		MenuBuilder.AddMenuEntry(
+			LOCTEXT("SkeletalMesh_New1DAimOffset", "Create 1D AimOffset"),
+			LOCTEXT("SkeletalMesh_New1DAimOffsetTooltip", "Creates a 1D aimoffset blendspace using the selected skeleton."),
+			FSlateIcon(),
+			FUIAction(
+				FExecuteAction::CreateStatic(&ExecuteNewAnimAsset<UAimOffsetBlendSpaceFactory1D, UAimOffsetBlendSpace1D>, Skeletons, FString("_AimOffset1D"), AssetCreated, bInContentBrowser),
+				FCanExecuteAction()
+				)
+			);
+
+		MenuBuilder.AddMenuEntry(
+			LOCTEXT("SkeletalMesh_New2DAimOffset", "Create 2D AimOffset"),
+			LOCTEXT("SkeletalMesh_New2DAimOffsetTooltip", "Creates a 2D aimoffset blendspace using the selected skeleton."),
+			FSlateIcon(),
+			FUIAction(
+				FExecuteAction::CreateStatic(&ExecuteNewAnimAsset<UAimOffsetBlendSpaceFactoryNew, UAimOffsetBlendSpace>, Skeletons, FString("_AimOffset2D"), AssetCreated, bInContentBrowser),
+				FCanExecuteAction()
+				)
+			);
+	}
+
+	void FillCreateAssetMenu(FMenuBuilder& MenuBuilder, TArray<TWeakObjectPtr<USkeleton>> Skeletons, FAnimAssetCreated AssetCreated, bool bInContentBrowser) 
+	{
+		MenuBuilder.AddSubMenu(
+			LOCTEXT("SkeletalMesh_NewAimOffset", "Create AimOffset"),
+			LOCTEXT("SkeletalMesh_NewAimOffsetTooltip", "Creates an aimoffset blendspace using the selected skeleton."),
+			FNewMenuDelegate::CreateStatic(&FillAimOffsetBlendSpaceMenu, Skeletons, AssetCreated, bInContentBrowser));
+
+		MenuBuilder.AddSubMenu(
+			LOCTEXT("SkeletalMesh_NewBlendspace", "Create BlendSpace"),
+			LOCTEXT("SkeletalMesh_NewBlendspaceTooltip", "Creates a blendspace using the skeleton of the selected mesh."),
+			FNewMenuDelegate::CreateStatic(&FillBlendSpaceMenu, Skeletons, AssetCreated, bInContentBrowser));
+
+		MenuBuilder.AddMenuEntry(
+			LOCTEXT("Skeleton_NewAnimComposite", "Create AnimComposite"),
+			LOCTEXT("Skeleton_NewAnimCompositeTooltip", "Creates an AnimComposite using the selected skeleton."),
+			FSlateIcon(),
+			FUIAction(
+				FExecuteAction::CreateStatic(&ExecuteNewAnimAsset<UAnimCompositeFactory, UAnimComposite>, Skeletons, FString("_Composite"), AssetCreated, bInContentBrowser),
+				FCanExecuteAction()
+				)
+			);
+
+		MenuBuilder.AddMenuEntry(
+			LOCTEXT("Skeleton_NewAnimMontage", "Create AnimMontage"),
+			LOCTEXT("Skeleton_NewAnimMontageTooltip", "Creates an AnimMontage using the selected skeleton."),
+			FSlateIcon(),
+			FUIAction(
+				FExecuteAction::CreateStatic(&ExecuteNewAnimAsset<UAnimMontageFactory, UAnimMontage>, Skeletons, FString("_Montage"), AssetCreated, bInContentBrowser),
+				FCanExecuteAction()
+				)
+			);
+	}
+}
+
+#undef LOCTEXT_NAMESPACE
\ No newline at end of file
diff --git a/Engine/Source/Editor/UnrealEd/Public/AnimationEditorUtils.h b/Engine/Source/Editor/UnrealEd/Public/AnimationEditorUtils.h
new file mode 100644
--- /dev/null
+++ b/Engine/Source/Editor/UnrealEd/Public/AnimationEditorUtils.h
@@ -0,0 +1,16 @@
+// Copyright 1998-2014 Epic Games, Inc. All Rights Reserved.
+
+#ifndef __AnimationEditorUtils_h__
+#define __AnimationEditorUtils_h__
+
+/** Defines FCanExecuteAction delegate interface.  Returns true when an action is able to execute. */
+DECLARE_DELEGATE_OneParam(FAnimAssetCreated, TArray<class UObject*>);
+
+//Animation editor utility functions
+namespace AnimationEditorUtils
+{
+	UNREALED_API void CreateAnimationAssets(const TArray<TWeakObjectPtr<USkeleton>>& Skeletons, TSubclassOf<UAnimationAsset> AssetClass, const FString& InPrefix, FAnimAssetCreated AssetCreated );
+	UNREALED_API void FillCreateAssetMenu(FMenuBuilder& MenuBuilder, TArray<TWeakObjectPtr<USkeleton>> Skeletons, FAnimAssetCreated AssetCreated, bool bInContentBrowser=true);
+} // namespace AnimationEditorUtils
+
+#endif //__AnimationEditorUtils_h__
