Added "Save Animation Assets" menu in Persona

diff --git a/Engine/Source/Editor/Persona/Private/Persona.cpp b/Engine/Source/Editor/Persona/Private/Persona.cpp
--- a/Engine/Source/Editor/Persona/Private/Persona.cpp
+++ b/Engine/Source/Editor/Persona/Private/Persona.cpp
@@ -458,6 +458,12 @@ void FPersona::ExtendMenu()
 	// Add additional persona editor menus
 	struct Local
 	{
+		static void AddPersonaSaveLoadMenu(FMenuBuilder& MenuBuilder)
+		{
+			MenuBuilder.AddMenuEntry(FPersonaCommands::Get().SaveAnimationAssets);
+			MenuBuilder.AddMenuSeparator();
+		}
+		
 		static void AddPersonaFileMenu(FMenuBuilder& MenuBuilder)
 		{
 			// View
@@ -512,6 +518,12 @@ void FPersona::ExtendMenu()
 			FMenuExtensionDelegate::CreateStatic(&Local::AddPersonaFileMenu));
 	}
 
+	MenuExtender->AddMenuExtension(
+		"FileLoadAndSave",
+		EExtensionHook::First,
+		GetToolkitCommands(),
+		FMenuExtensionDelegate::CreateStatic(&Local::AddPersonaSaveLoadMenu));
+
 	MenuExtender->AddMenuExtension(
 			"AssetEditorActions",
 			EExtensionHook::After,
@@ -1512,6 +1524,24 @@ void FPersona::SaveAsset_Execute()
 	}
 }
 
+void FPersona::SaveAnimationAssets_Execute()
+{
+	// only save animation assets or animation blueprints
+	TArray<UClass*> SaveClasses;
+	SaveClasses.Add(UAnimationAsset::StaticClass());
+	SaveClasses.Add(UAnimBlueprint::StaticClass());
+	SaveClasses.Add(USkeletalMesh::StaticClass());
+	SaveClasses.Add(USkeleton::StaticClass());
+
+	const bool bPromptUserToSave = false;
+	const bool bFastSave = false;
+	FEditorFileUtils::SaveDirtyContentPackages(SaveClasses, bPromptUserToSave, bFastSave);
+}
+
+bool FPersona::CanSaveAnimationAssets() const
+{
+	return true;
+}
 
 void FPersona::OnActiveTabChanged( TSharedPtr<SDockTab> PreviouslyActive, TSharedPtr<SDockTab> NewlyActivated )
 {
diff --git a/Engine/Source/Editor/Persona/Private/Persona.h b/Engine/Source/Editor/Persona/Private/Persona.h
--- a/Engine/Source/Editor/Persona/Private/Persona.h
+++ b/Engine/Source/Editor/Persona/Private/Persona.h
@@ -202,6 +202,10 @@ public:
 	virtual void SaveAsset_Execute() override;
 	// End of IToolkit interface
 
+	/** Saves all animation assets related to a skeleton */
+	void SaveAnimationAssets_Execute();
+	bool CanSaveAnimationAssets() const;
+
 	/** @return the documentation location for this editor */
 	virtual FString GetDocumentationLink() const override
 	{
diff --git a/Engine/Source/Editor/Persona/Private/PersonaCommands.cpp b/Engine/Source/Editor/Persona/Private/PersonaCommands.cpp
--- a/Engine/Source/Editor/Persona/Private/PersonaCommands.cpp
+++ b/Engine/Source/Editor/Persona/Private/PersonaCommands.cpp
@@ -6,6 +6,9 @@
 
 void FPersonaCommands::RegisterCommands()
 {
+	// save menu
+	UI_COMMAND(SaveAnimationAssets, "Save Animation Assets", "Only save all animation assets including skeletons, skeletal meshes and animation blueprints.", EUserInterfaceActionType::Button, FInputGesture());
+
 	// skeleton menu
 	UI_COMMAND( ChangeSkeletonPreviewMesh, "Set Preview Mesh as Default", "Changes the skeletons default preview mesh to the current open preview mesh. The skeleton will require saving after this action.", EUserInterfaceActionType::Button, FInputGesture() );
 	UI_COMMAND( RemoveUnusedBones, "Remove Unused Bones from Skeleton", "Removes any bones from the skeleton that are not used by any of its meshes. The skeleton and associated animations will require saving after this action.", EUserInterfaceActionType::Button, FInputGesture() );
diff --git a/Engine/Source/Editor/Persona/Private/PersonaCommands.h b/Engine/Source/Editor/Persona/Private/PersonaCommands.h
--- a/Engine/Source/Editor/Persona/Private/PersonaCommands.h
+++ b/Engine/Source/Editor/Persona/Private/PersonaCommands.h
@@ -14,6 +14,11 @@ public:
 	virtual void RegisterCommands() override;
 
 public:
+
+	// save and load menu options
+	// save all animation assets related to skeleton
+	TSharedPtr<FUICommandInfo> SaveAnimationAssets;
+
 	// skeleton menu options
 	// Command to allow users to set the skeletons preview mesh
 	TSharedPtr<FUICommandInfo> ChangeSkeletonPreviewMesh;
diff --git a/Engine/Source/Editor/UnrealEd/Private/FileHelpers.cpp b/Engine/Source/Editor/UnrealEd/Private/FileHelpers.cpp
--- a/Engine/Source/Editor/UnrealEd/Private/FileHelpers.cpp
+++ b/Engine/Source/Editor/UnrealEd/Private/FileHelpers.cpp
@@ -2428,6 +2428,145 @@ bool FEditorFileUtils::SaveDirtyPackages(const bool bPromptUserToSave, const boo
 	return bReturnCode;
 }
 
+bool FEditorFileUtils::SaveDirtyContentPackages(TArray<UClass*>& SaveContentClasses, const bool bPromptUserToSave, const bool bFastSave, const bool bNotifyNoPackagesSaved)
+{
+	bool bReturnCode = true;
+
+	// A list of all packages that need to be saved
+	TArray<UPackage*> PackagesToSave;
+
+	CollectGarbage(GARBAGE_COLLECTION_KEEPFLAGS);
+
+	// Make a list of all content packages that we should save
+	for (TObjectIterator<UPackage> It; It; ++It)
+	{
+		UPackage*	Package = *It;
+		bool		bShouldIgnorePackage = false;
+
+		// Only look at root packages.
+		bShouldIgnorePackage |= Package->GetOuter() != NULL;
+		// Don't try to save "Transient" package.
+		bShouldIgnorePackage |= Package == GetTransientPackage();
+		// Ignore PIE packages.
+		bShouldIgnorePackage |= (Package->PackageFlags & PKG_PlayInEditor) != 0;
+		// Ignore packages that haven't been modified.
+		bShouldIgnorePackage |= !Package->IsDirty();
+
+		// Ignore packages with long, invalid names. This culls out packages with paths in read-only roots such as /Temp.
+		bShouldIgnorePackage |= (!FPackageName::IsShortPackageName(Package->GetFName()) && !FPackageName::IsValidLongPackageName(Package->GetName(), /*bIncludeReadOnlyRoots=*/false));
+
+		if (!bShouldIgnorePackage)
+		{
+			TArray<UObject*> Objects;
+			GetObjectsWithOuter(Package, Objects);
+
+			for (auto Iter = Objects.CreateIterator(); Iter; ++Iter)
+			{
+				bool bNeedToSave = false;
+
+				for (const UClass* ClassType : SaveContentClasses)
+				{
+					if ((*Iter)->GetClass()->IsChildOf(ClassType))
+					{
+						bNeedToSave = true;
+						break;
+					}
+				}
+
+				if (bNeedToSave)
+				{
+					// add to asset 
+					PackagesToSave.Add(Package);
+					break;
+				}
+			}
+		}
+	}
+
+	if (PackagesToSave.Num() > 0)
+	{
+		if (!bFastSave)
+		{
+			const FEditorFileUtils::EPromptReturnCode Return = FEditorFileUtils::PromptForCheckoutAndSave(PackagesToSave, true, bPromptUserToSave);
+			if (Return == FEditorFileUtils::EPromptReturnCode::PR_Cancelled)
+			{
+				// Only cancel should return false and stop whatever we were doing before.(like closing the editor)
+				// If failure is returned, the user was given ample times to retry saving the package and didn't want to
+				// So we should continue with whatever we were doing.  
+				bReturnCode = false;
+			}
+		}
+		else
+		{
+			FSaveErrorOutputDevice SaveErrors;
+			GWarn->BeginSlowTask(NSLOCTEXT("UnrealEd", "SavingPackagesE", "Saving packages..."), true);
+
+			// Packages that failed to save
+			TArray< UPackage* > FailedPackages;
+
+			for (TArray<UPackage*>::TConstIterator PkgIter(PackagesToSave); PkgIter; ++PkgIter)
+			{
+				UPackage* CurPackage = *PkgIter;
+
+				// Check if a file exists for this package
+				FString Filename;
+				bool bFoundFile = FPackageName::DoesPackageExist(CurPackage->GetName(), NULL, &Filename);
+				if (bFoundFile)
+				{
+					// determine if the package file is read only
+					const bool bPkgReadOnly = IFileManager::Get().IsReadOnly(*Filename);
+
+					// Only save writable files in fast mode
+					if (!bPkgReadOnly)
+					{
+						if (!CurPackage->IsFullyLoaded())
+						{
+							// Packages must be fully loaded to save
+							CurPackage->FullyLoad();
+						}
+
+						const UWorld* const AssociatedWorld = UWorld::FindWorldInPackage(CurPackage);
+						const bool bIsMapPackage = AssociatedWorld != nullptr;
+
+						const FText SavingPackageText = (bIsMapPackage)
+							? FText::Format(NSLOCTEXT("UnrealEd", "SavingMapf", "Saving map {0}"), FText::FromString(CurPackage->GetName()))
+							: FText::Format(NSLOCTEXT("UnrealEd", "SavingAssetf", "Saving asset {0}"), FText::FromString(CurPackage->GetName()));
+
+						GWarn->StatusForceUpdate(PkgIter.GetIndex(), PackagesToSave.Num(), SavingPackageText);
+
+						// Save the package
+						bool bPackageLocallyWritable;
+						const int32 SaveStatus = InternalSavePackage(CurPackage, bPackageLocallyWritable, SaveErrors);
+
+						if (SaveStatus == EAppReturnType::No)
+						{
+							// The package could not be saved so add it to the failed array 
+							FailedPackages.Add(CurPackage);
+
+						}
+					}
+				}
+
+			}
+			GWarn->EndSlowTask();
+			SaveErrors.Flush();
+
+			// Warn the user about any packages which failed to save.
+			WarnUserAboutFailedSave(FailedPackages);
+		}
+	}
+	else
+	{
+		FNotificationInfo NotificationInfo(LOCTEXT("NoAssetsToSave", "No new changes to save!"));
+		NotificationInfo.Image = FEditorStyle::GetBrush(FTokenizedMessage::GetSeverityIconName(EMessageSeverity::Info));
+		NotificationInfo.bFireAndForget = true;
+		NotificationInfo.ExpireDuration = 4.0f; // Need this message to last a little longer than normal since the user may have expected there to be modified files.
+		NotificationInfo.bUseThrobber = true;
+		FSlateNotificationManager::Get().AddNotification(NotificationInfo);
+	}
+
+	return bReturnCode;
+}
 
 /**
  * Saves the active level, prompting the use for checkout if necessary.
diff --git a/Engine/Source/Editor/UnrealEd/Public/FileHelpers.h b/Engine/Source/Editor/UnrealEd/Public/FileHelpers.h
--- a/Engine/Source/Editor/UnrealEd/Public/FileHelpers.h
+++ b/Engine/Source/Editor/UnrealEd/Public/FileHelpers.h
@@ -145,6 +145,17 @@ public:
 	 */
 	UNREALED_API static bool SaveDirtyPackages(const bool bPromptUserToSave, const bool bSaveMapPackages, const bool bSaveContentPackages, const bool bFastSave = false, const bool bNotifyNoPackagesSaved = false, bool* bOutPackagesNeededSaving = NULL);
 
+	/**
+	* Looks at all currently loaded packages and saves them if their "bDirty" flag is set and they include specified clasees, optionally prompting the user to select which packages to save)
+	*
+	* @param	SaveContentClasses			save only specified classes or children classes
+	* @param	bPromptUserToSave			true if we should prompt the user to save dirty packages we found. false to assume all dirty packages should be saved.  Regardless of this setting the user will be prompted for checkout(if needed) unless bFastSave is set
+	* @param	bFastSave					true if we should do a fast save. (I.E dont prompt the user to save, dont prompt for checkout, and only save packages that are currently writable).  Note: Still prompts for SaveAs if a package needs a filename
+	* @param	bNotifyNoPackagesSaved		true if a notification should be displayed when no packages need to be saved.
+	* @return								true on success, false on fail.
+	*/
+	UNREALED_API static bool SaveDirtyContentPackages(TArray<UClass*>& SaveContentClasses, const bool bPromptUserToSave, const bool bFastSave = false, const bool bNotifyNoPackagesSaved = false);
+
 	/**
 	 * Saves the active level, prompting the use for checkout if necessary.
 	 *
