Consider batched meshes when clearing line batchers

diff --git a/Engine/Source/Editor/UnrealEd/Private/EditorViewportClient.cpp b/Engine/Source/Editor/UnrealEd/Private/EditorViewportClient.cpp
--- a/Engine/Source/Editor/UnrealEd/Private/EditorViewportClient.cpp
+++ b/Engine/Source/Editor/UnrealEd/Private/EditorViewportClient.cpp
@@ -2519,12 +2519,12 @@ void FEditorViewportClient::Draw(FViewport* InViewport, FCanvas* Canvas)
 
 	// Remove temporary debug lines.
 	// Possibly a hack. Lines may get added without the scene being rendered etc.
-	if (World->LineBatcher != NULL && (World->LineBatcher->BatchedLines.Num() || World->LineBatcher->BatchedPoints.Num()))
+	if (World->LineBatcher != NULL && (World->LineBatcher->BatchedLines.Num() || World->LineBatcher->BatchedPoints.Num() || World->LineBatcher->BatchedMeshes.Num() ) )
 	{
 		World->LineBatcher->Flush();
 	}
 
-	if (World->ForegroundLineBatcher != NULL && (World->ForegroundLineBatcher->BatchedLines.Num() || World->ForegroundLineBatcher->BatchedPoints.Num()))
+	if (World->ForegroundLineBatcher != NULL && (World->ForegroundLineBatcher->BatchedLines.Num() || World->ForegroundLineBatcher->BatchedPoints.Num() || World->ForegroundLineBatcher->BatchedMeshes.Num() ) )
 	{
 		World->ForegroundLineBatcher->Flush();
 	}
