UserDefinedStruct removed with "force delete" is replaced with a special FallbackStruct

CRASH: EDITOR: UDS: Crash when using force delete option on struct and blueprint that refer to each other ttp#330163

diff --git a/Engine/Source/Editor/UnrealEd/Private/ObjectTools.cpp b/Engine/Source/Editor/UnrealEd/Private/ObjectTools.cpp
--- a/Engine/Source/Editor/UnrealEd/Private/ObjectTools.cpp
+++ b/Engine/Source/Editor/UnrealEd/Private/ObjectTools.cpp
@@ -1843,23 +1843,51 @@ namespace ObjectTools
 			// Replacing references inside already loaded objects could cause rendering issues, so globally detach all components from their scenes for now
 			FGlobalComponentReregisterContext ReregisterContext;
 
-			FForceReplaceInfo ReplaceInfo;
+			int32 ReplaceableObjectsNum = 0;
+			{
+				TArray<UObject*> ObjectsToReplace = ObjectsToDelete;
+				for (TArray<UObject*>::TIterator ObjectItr(ObjectsToReplace); ObjectItr; ++ObjectItr)
+				{
+					UObject* CurObject = *ObjectItr;
+					// If we're a blueprint add our generated class as well
+					UBlueprint *BlueprintObject = Cast<UBlueprint>(CurObject);
+					if (BlueprintObject && BlueprintObject->GeneratedClass)
+					{
+						ObjectsToReplace.AddUnique(BlueprintObject->GeneratedClass);
+					}
+				}
 
-			TArray<UObject*> ObjectsToReplace = ObjectsToDelete;
+				// UserDefinedStructs (probably all SctiptStructs) should be replaced with the FallbackStruct
+				{
+					TArray<UObject*> UDStructToReplace;
+					for (int32 Iter = 0; Iter < ObjectsToReplace.Num(); )
+					{
+						if (auto UDStruct = Cast<UUserDefinedStruct>(ObjectsToReplace[Iter]))
+						{
+							ObjectsToReplace.RemoveAtSwap(Iter);
+							UDStructToReplace.Add(UDStruct);
+						}
+						else
+						{
+							Iter++;
+						}
+					}
+
+					if (UDStructToReplace.Num())
+					{
+						FForceReplaceInfo ReplaceInfo;
+						ForceReplaceReferences(GetFallbackStruct(), UDStructToReplace, ReplaceInfo, false);
+						ReplaceableObjectsNum += ReplaceInfo.ReplaceableObjects.Num();
+					}
+				}
 
-			for( TArray<UObject*>::TIterator ObjectItr(ObjectsToReplace); ObjectItr; ++ObjectItr )
-			{
-				UObject* CurObject = *ObjectItr; 
-				// If we're a blueprint add our generated class as well
-				UBlueprint *BlueprintObject = Cast<UBlueprint>(CurObject);
-				if (BlueprintObject && BlueprintObject->GeneratedClass)
 				{
-					ObjectsToReplace.AddUnique(BlueprintObject->GeneratedClass);
+					FForceReplaceInfo ReplaceInfo;
+					ForceReplaceReferences(NULL, ObjectsToReplace, ReplaceInfo, false);
+					ReplaceableObjectsNum += ReplaceInfo.ReplaceableObjects.Num();
 				}
 			}
 
-			ForceReplaceReferences( NULL, ObjectsToReplace, ReplaceInfo, false );
-
 			// Load the asset tools module to get access to the browser type maps
 			FAssetToolsModule& AssetToolsModule = FModuleManager::LoadModuleChecked<FAssetToolsModule>(TEXT("AssetTools"));
 
@@ -1878,7 +1906,7 @@ namespace ObjectTools
 					++NumDeletedObjects;
 				}
 
-				GWarn->StatusUpdate( ObjectItr.GetIndex(), ReplaceInfo.ReplaceableObjects.Num(), NSLOCTEXT("UnrealEd", "ConsolidateAssetsUpdate_DeletingObjects", "Deleting Assets...") );
+				GWarn->StatusUpdate(ObjectItr.GetIndex(), ReplaceableObjectsNum, NSLOCTEXT("UnrealEd", "ConsolidateAssetsUpdate_DeletingObjects", "Deleting Assets..."));
 
 			}
 		}
