Ensure that package metadata is loaded when trying to use it

diff --git a/Engine/Source/Editor/UnrealEd/Private/Kismet2/BlueprintEditorUtils.cpp b/Engine/Source/Editor/UnrealEd/Private/Kismet2/BlueprintEditorUtils.cpp
--- a/Engine/Source/Editor/UnrealEd/Private/Kismet2/BlueprintEditorUtils.cpp
+++ b/Engine/Source/Editor/UnrealEd/Private/Kismet2/BlueprintEditorUtils.cpp
@@ -1026,16 +1026,6 @@ UClass* FBlueprintEditorUtils::RegenerateBlueprintClass(UBlueprint* Blueprint, U
 			}
 		}
 
-		// Ensure that the package metadata is preloaded and valid
-		if( Package )
-		{
-			UMetaData* MetaData = Package->GetMetaData();
-			if( MetaData->HasAnyFlags(RF_NeedLoad) )
-			{
-				MetaData->GetLinker()->Preload(MetaData);
-			}
-		}
-
 		// Make sure the simple construction script is loaded, since the outer hierarchy isn't compatible with PreloadMembers past the root node
 		FBlueprintEditorUtils::PreloadConstructionScript(Blueprint);
 
diff --git a/Engine/Source/Runtime/CoreUObject/Private/UObject/Package.cpp b/Engine/Source/Runtime/CoreUObject/Private/UObject/Package.cpp
--- a/Engine/Source/Runtime/CoreUObject/Private/UObject/Package.cpp
+++ b/Engine/Source/Runtime/CoreUObject/Private/UObject/Package.cpp
@@ -119,6 +119,11 @@ UMetaData* UPackage::GetMetaData()
 
 	check(MetaData);
 
+	if (MetaData->HasAnyFlags(RF_NeedLoad))
+	{
+		MetaData->GetLinker()->Preload(MetaData);
+	}
+
 	return MetaData;
 }
 
