[339682] Blueprints - Fix a crash that could occur when starting PIE with an open, uncompiled Blueprint.

- Prevent Blueprint actors with a native RootComponent from being incorrectly attached to the old (stale) Actor instance after compiling. This led to an assertion during PIE startup.
- Also cleaned up reinstancing attachment code so that it no longer records old objects into the transaction buffer, which was preventing those objects from being garbage collected.

diff --git a/Engine/Source/Editor/UnrealEd/Private/Kismet2/KismetReinstanceUtilities.cpp b/Engine/Source/Editor/UnrealEd/Private/Kismet2/KismetReinstanceUtilities.cpp
--- a/Engine/Source/Editor/UnrealEd/Private/Kismet2/KismetReinstanceUtilities.cpp
+++ b/Engine/Source/Editor/UnrealEd/Private/Kismet2/KismetReinstanceUtilities.cpp
@@ -259,37 +259,64 @@ void FBlueprintCompileReinstancer::ReplaceInstancesOfClass(UClass* OldClass, UCl
 				FVector Location = FVector::ZeroVector;
 				FRotator* RotationPtr = NULL;
 				FRotator Rotation = FRotator::ZeroRotator;
-				AActor* AttachParentActor = NULL;
-				TArray<AActor*> OldAttachChildren;
+				AActor* AttachParent = NULL;
+				USceneComponent* AttachParentComp = NULL;
 				FName   AttachSocketName;
 
+				// Struct to store info about attached actors
+				struct FAttachedActorInfo
+				{
+					AActor* AttachedActor;
+					FName AttachedToSocket;
+				};
+
+				// Save info about attached actors
+				TArray<FAttachedActorInfo> AttachedActorInfos;
+
 				// Create cache to store component data across rerunning construction scripts
 				FComponentInstanceDataCache InstanceDataCache(OldActor);
 
+				// If there are attached objects detach them and store the socket names
+				TArray<AActor*> AttachedActors;
+				OldActor->GetAttachedActors(AttachedActors);
+				for (AActor* AttachedActor : AttachedActors)
+				{
+					USceneComponent* AttachedActorRoot = AttachedActor->GetRootComponent();
+					if (AttachedActorRoot && AttachedActorRoot->AttachParent)
+					{
+						// Save info about actor to reattach
+						FAttachedActorInfo Info;
+						Info.AttachedActor = AttachedActor;
+						Info.AttachedToSocket = AttachedActorRoot->AttachSocketName;
+						AttachedActorInfos.Add(Info);
+
+						// Now detach it
+						AttachedActorRoot->DetachFromParent(true);
+					}
+				}
+
 				if (USceneComponent* OldRootComponent = OldActor->GetRootComponent())
 				{
 					if(OldRootComponent->AttachParent != NULL)
 					{
-						AttachParentActor = OldRootComponent->AttachParent->GetOwner();
-						check(AttachParentActor != OldActor); // Root component should never be attached to another component in the same actor!
+						AttachParent = OldRootComponent->AttachParent->GetOwner();
+						// Root component should never be attached to another component in the same actor!
+						if (AttachParent == OldActor)
+						{
+							UE_LOG(LogBlueprint, Warning, TEXT("ReplaceInstancesOfClass: RootComponent (%s) attached to another component in this Actor (%s)."), *OldRootComponent->GetPathName(), *AttachParent->GetPathName());
+							AttachParent = NULL;
+						}
 
+						AttachParentComp = OldRootComponent->AttachParent;
 						AttachSocketName = OldRootComponent->AttachSocketName;
 
 						//detach it to remove any scaling
 						OldRootComponent->DetachFromParent(true);
 					}
 
-					for (auto AttachIt = OldRootComponent->AttachChildren.CreateConstIterator(); AttachIt; ++AttachIt)
-					{
-						USceneComponent* Child = *AttachIt;
-						if (Child != nullptr)
-						{
-							OldAttachChildren.Add(Child->GetOwner());
-						}
-					}
-
 					// Save off transform
-					WorldTransform = OldActor->GetRootComponent()->ComponentToWorld;
+					WorldTransform = OldRootComponent->ComponentToWorld;
+					WorldTransform.SetTranslation(OldRootComponent->GetComponentLocation()); // take into account any custom location
 					Location = OldActor->GetActorLocation();
 					LocationPtr = &Location;
 					Rotation = OldActor->GetActorRotation();
@@ -345,16 +372,36 @@ void FBlueprintCompileReinstancer::ReplaceInstancesOfClass(UClass* OldClass, UCl
 				// Run the construction script, which will use the properties we just copied over
 				NewActor->ExecuteConstruction(WorldTransform, &InstanceDataCache);
 
-				//Attach the new instance to original parent
-				if(AttachParentActor)
+				if (USceneComponent* NewRootComponent = NewActor->GetRootComponent())
 				{
-					GEditor->ParentActors(AttachParentActor, NewActor, AttachSocketName);
-				}
+					//Attach the new instance to original parent
+					if(AttachParent != nullptr)
+					{
+						if(AttachParentComp == nullptr)
+						{
+							AttachParentComp = AttachParent->GetRootComponent();
+						}
 
-				for (auto AttachIt = OldAttachChildren.CreateIterator(); AttachIt; ++AttachIt)
-				{
-					// Cannot attach to a socket on a blueprint
-					GEditor->ParentActors(NewActor, *AttachIt, FName());
+						if(AttachParentComp != nullptr)
+						{
+							NewRootComponent->AttachTo(AttachParentComp, AttachSocketName, EAttachLocation::KeepWorldPosition);
+						}
+					}
+
+					// If we had attached children reattach them now - unless they are already attached
+					for(FAttachedActorInfo& Info : AttachedActorInfos)
+					{
+						// If this actor is no longer attached to anything, reattach
+						if (!Info.AttachedActor->IsPendingKill() && Info.AttachedActor->GetAttachParentActor() == nullptr)
+						{
+							USceneComponent* ChildRoot = Info.AttachedActor->GetRootComponent();
+							if (ChildRoot && ChildRoot->AttachParent != NewRootComponent)
+							{
+								ChildRoot->AttachTo(NewRootComponent, Info.AttachedToSocket, EAttachLocation::KeepWorldPosition);
+								ChildRoot->UpdateComponentToWorld();
+							}
+						}
+					}
 				}
 
 				// Determine if the actor used to be selected.
