addition of messaging when iOS cannot package due to missing iTunes, missing signing certificate, or missing mobile provision
TTP166302
part of TTP337760
TTP341402
#ue4
#ios

diff --git a/Engine/Source/Developer/TargetPlatform/Public/Common/TargetPlatformBase.h b/Engine/Source/Developer/TargetPlatform/Public/Common/TargetPlatformBase.h
--- a/Engine/Source/Developer/TargetPlatform/Public/Common/TargetPlatformBase.h
+++ b/Engine/Source/Developer/TargetPlatform/Public/Common/TargetPlatformBase.h
@@ -170,6 +170,11 @@ public:
 		return true;
 	}
 
+	virtual int IsReadyToBuild(const FString& ProjectPath, bool bProjectHasCode, FString& OutDocumentationPath) const override
+	{
+		return ETargetPlatformReadyStatus::Ready;
+	}
+
 	// End ITargetPlatform interface
 
 protected:
diff --git a/Engine/Source/Developer/TargetPlatform/Public/Interfaces/ITargetPlatform.h b/Engine/Source/Developer/TargetPlatform/Public/Interfaces/ITargetPlatform.h
--- a/Engine/Source/Developer/TargetPlatform/Public/Interfaces/ITargetPlatform.h
+++ b/Engine/Source/Developer/TargetPlatform/Public/Interfaces/ITargetPlatform.h
@@ -86,6 +86,37 @@ namespace ETargetPlatformFeatures
 };
 
 
+namespace ETargetPlatformReadyStatus
+{
+	/**
+	 * Flags specifiying what is needed to be able to complete and deploy a build
+	 */
+	enum Type
+	{
+		/** Ready */
+		Ready = 0,
+
+		/** SDK Not Found*/
+		SDKNotFound = 1,
+
+		/** Code Build Not Supported */
+		CodeUnsupported = 2,
+
+		/** Plugins Not Supported */
+		PluginsUnsupported = 4,
+
+		/** Signing Key Not Found */
+		SigningKeyNotFound = 8,
+
+		/** Provision Not Found */
+		ProvisionNotFound = 16,
+
+		/** Manifest Not Found */
+		ManifestNotFound = 32,
+	};
+};
+
+
 /**
  * ITargetPlatform, abstraction for cooking platforms and enumerating actual target devices
 **/
@@ -260,6 +291,18 @@ public:
 	 */
 	virtual bool IsSdkInstalled(bool bProjectHasCode, FString& OutDocumentationPath) const = 0;
 
+
+	/**
+	 * Checks whether the platform's build requirements are met so that we can do things like
+	 * package for the platform
+	 *
+	 * @param ProjectPath Path to the project
+	 * @param bProjectHasCode true if the project has code, and therefore any compilation based SDK requirements should be checked
+	 * @param OutDocumentationPath Let's the platform tell the editor a path to show some information about how to fix any problem
+	 * @return Readiness status
+	 */
+	virtual int IsReadyToBuild(const FString& ProjectPath, bool bProjectHasCode, FString& OutDocumentationPath) const = 0;
+
 #if WITH_ENGINE
 	/**
 	 * Gets the format to use for a particular body setup.
diff --git a/Engine/Source/Developer/iOS/IOSTargetPlatform/IOSTargetPlatform.Build.cs b/Engine/Source/Developer/iOS/IOSTargetPlatform/IOSTargetPlatform.Build.cs
--- a/Engine/Source/Developer/iOS/IOSTargetPlatform/IOSTargetPlatform.Build.cs
+++ b/Engine/Source/Developer/iOS/IOSTargetPlatform/IOSTargetPlatform.Build.cs
@@ -16,6 +16,7 @@ public class IOSTargetPlatform : ModuleRules
 				"TargetPlatform",
 				"DesktopPlatform",
 				"LaunchDaemonMessages",
+				"Projects"
 			}
 		);
 
diff --git a/Engine/Source/Developer/iOS/IOSTargetPlatform/Private/IOSTargetPlatform.cpp b/Engine/Source/Developer/iOS/IOSTargetPlatform/Private/IOSTargetPlatform.cpp
--- a/Engine/Source/Developer/iOS/IOSTargetPlatform/Private/IOSTargetPlatform.cpp
+++ b/Engine/Source/Developer/iOS/IOSTargetPlatform/Private/IOSTargetPlatform.cpp
@@ -5,7 +5,7 @@
 =============================================================================*/
 
 #include "IOSTargetPlatformPrivatePCH.h"
-
+#include "IProjectManager.h"
 
 /* FIOSTargetPlatform structors
  *****************************************************************************/
@@ -82,10 +82,79 @@ bool FIOSTargetPlatform::IsSdkInstalled(bool bProjectHasCode, FString& OutDocume
 #if PLATFORM_MAC
 	OutDocumentationPath = FString("Shared/Tutorials/InstallingXCodeTutorial");
 	biOSSDKInstalled = IFileManager::Get().DirectoryExists(TEXT("/Applications/Xcode.app/Contents/Developer/Platforms/iPhoneOS.platform"));
+#else
+	{
+		HKEY hKey;
+		LRESULT lRes = RegOpenKeyEx(HKEY_LOCAL_MACHINE, TEXT("SOFTWARE\\Wow6432Node\\Apple Inc.\\Apple Mobile Device Support\\Shared"), 0, KEY_READ, &hKey);
+		TCHAR dllPath[256];
+		unsigned long pathSize = 256;
+		if (lRes != ERROR_SUCCESS || RegQueryValueEx(hKey, TEXT("iTunesMobileDeviceDLL"), 0, NULL, (BYTE*)dllPath, &pathSize) != ERROR_SUCCESS || IFileManager::Get().FileSize(*FString(dllPath)) == INDEX_NONE)
+		{
+			OutDocumentationPath = FString("Shared/Tutorials/InstallingiTunesTutorial");
+			biOSSDKInstalled = false;
+		}
+	}
+
 #endif
 	return biOSSDKInstalled;
 }
 
+int FIOSTargetPlatform::IsReadyToBuild(const FString& ProjectPath, bool bProjectHasCode, FString& OutDocumentationPath) const
+{
+	int bReadyToBuild = ETargetPlatformReadyStatus::Ready; // @todo How do we check that the iOS SDK is installed when building from Windows? Is that even possible?
+	if (!IsSdkInstalled(bProjectHasCode, OutDocumentationPath))
+	{
+		bReadyToBuild |= ETargetPlatformReadyStatus::SDKNotFound;
+	}
+#if PLATFORM_MAC
+	OutDocumentationPath = FString("Shared/Tutorials/InstallingXCodeTutorial");
+#else
+	if (bProjectHasCode && FRocketSupport::IsRocket())
+	{
+		OutDocumentationPath = FString("Shared/Tutorials/iOSonPCRestrictions");
+		bReadyToBuild |= ETargetPlatformReadyStatus::CodeUnsupported;
+	}
+	if (FRocketSupport::IsRocket() && IProjectManager::Get().IsNonDefaultPluginEnabled())
+	{
+		OutDocumentationPath = FString("Shared/Tutorials/iOSonPCValidPlugins");
+		bReadyToBuild |= ETargetPlatformReadyStatus::PluginsUnsupported;
+	}
+
+	// shell to IPP and get the status of the provision and cert
+	FString CmdExe = FPaths::ConvertRelativePathToFull(FPaths::EngineDir() / TEXT("Binaries/DotNet/IOS/IPhonePackager.exe"));
+	FString CommandLine = FString::Printf(TEXT("Validate Engine -project %s"), *ProjectPath);
+	TSharedPtr<FMonitoredProcess> IPPProcess = MakeShareable(new FMonitoredProcess(CmdExe, CommandLine, false));
+	IPPProcess->Launch();
+	while(IPPProcess->IsRunning())
+	{
+		FPlatformProcess::Sleep(0.01f);
+	}
+	int RetCode = IPPProcess->GetReturnCode();
+	if (RetCode == 14)
+	{
+		OutDocumentationPath = FString("Shared/Tutorials/CreatingInfoPlist");
+		bReadyToBuild |= ETargetPlatformReadyStatus::ManifestNotFound;
+	}
+	else if (RetCode == 13)
+	{
+		OutDocumentationPath = FString("Shared/Tutorials/CreatingSigningCertAndProvision");
+		bReadyToBuild |= ETargetPlatformReadyStatus::SigningKeyNotFound;
+		bReadyToBuild |= ETargetPlatformReadyStatus::ProvisionNotFound;
+	}
+	else if (RetCode == 12)
+	{
+		OutDocumentationPath = FString("Shared/Tutorials/CreatingSigningCertAndProvision");
+		bReadyToBuild |= ETargetPlatformReadyStatus::SigningKeyNotFound;
+	}
+	else if (RetCode == 11)
+	{
+		OutDocumentationPath = FString("Shared/Tutorials/CreatingSigningCertAndProvision");
+		bReadyToBuild |= ETargetPlatformReadyStatus::ProvisionNotFound;
+	}
+#endif
+	return bReadyToBuild;
+}
+
 
 /* FIOSTargetPlatform implementation
  *****************************************************************************/
diff --git a/Engine/Source/Developer/iOS/IOSTargetPlatform/Private/IOSTargetPlatform.h b/Engine/Source/Developer/iOS/IOSTargetPlatform/Private/IOSTargetPlatform.h
--- a/Engine/Source/Developer/iOS/IOSTargetPlatform/Private/IOSTargetPlatform.h
+++ b/Engine/Source/Developer/iOS/IOSTargetPlatform/Private/IOSTargetPlatform.h
@@ -84,6 +84,7 @@ public:
 	}
 
 	virtual bool IsSdkInstalled(bool bProjectHasCode, FString& OutDocumentationPath) const override;
+	virtual int IsReadyToBuild(const FString& ProjectPath, bool bProjectHasCode, FString& OutDocumentationPath) const override;
 
 
 #if WITH_ENGINE
diff --git a/Engine/Source/Editor/MainFrame/Private/Frame/MainFrameActions.cpp b/Engine/Source/Editor/MainFrame/Private/Frame/MainFrameActions.cpp
--- a/Engine/Source/Editor/MainFrame/Private/Frame/MainFrameActions.cpp
+++ b/Engine/Source/Editor/MainFrame/Private/Frame/MainFrameActions.cpp
@@ -479,28 +479,36 @@ void FMainFrameActionCallbacks::PackageProject( const FName InPlatformInfoName )
 		if (Platform)
 		{
 			FString NotInstalledDocLink;
-			if (!Platform->IsSdkInstalled(bProjectHasCode, NotInstalledDocLink))
+			FString ProjectPath = FPaths::IsProjectFilePathSet() ? FPaths::ConvertRelativePathToFull(FPaths::GetProjectFilePath()) : FPaths::RootDir() / FApp::GetGameName() / FApp::GetGameName() + TEXT(".uproject");
+			switch (Platform->IsReadyToBuild(ProjectPath, bProjectHasCode, NotInstalledDocLink))
 			{
 				// broadcast this, and assume someone will pick it up
-				IMainFrameModule& MainFrameModule = FModuleManager::GetModuleChecked<IMainFrameModule>(TEXT("MainFrame"));
-				MainFrameModule.BroadcastMainFrameSDKNotInstalled(PlatformInfo->TargetPlatformName.ToString(), NotInstalledDocLink);
+			case ETargetPlatformReadyStatus::SDKNotFound:
+			case ETargetPlatformReadyStatus::ProvisionNotFound:
+			case ETargetPlatformReadyStatus::SigningKeyNotFound:
+			case (ETargetPlatformReadyStatus::ProvisionNotFound | ETargetPlatformReadyStatus::SigningKeyNotFound):
+				{
+					IMainFrameModule& MainFrameModule = FModuleManager::GetModuleChecked<IMainFrameModule>(TEXT("MainFrame"));
+					MainFrameModule.BroadcastMainFrameSDKNotInstalled(PlatformInfo->TargetPlatformName.ToString(), NotInstalledDocLink);
+				}
 				return;
-			}
-		}
-	}
 
 #if PLATFORM_WINDOWS
-	if (bProjectHasCode && FRocketSupport::IsRocket() && PlatformInfo->TargetPlatformName == FName("IOS"))
-	{
-		FMessageDialog::Open(EAppMsgType::Ok, LOCTEXT("PackagingNotWorkingMessage", "Sorry, packaging is currently not supported for code-based iOS projects. This feature will be available in a future release.") );
-		return;
-	}
-	if (FRocketSupport::IsRocket() && PlatformInfo->TargetPlatformName == FName("IOS") && IProjectManager::Get().IsNonDefaultPluginEnabled())
-	{
-		FMessageDialog::Open(EAppMsgType::Ok, LOCTEXT("PackagingNotWorkingMessage", "Sorry, packaging is currently not supported for content based projects with third-party plugins. This feature will be available in a future release.") );
-		return;
-	}
+			case ETargetPlatformReadyStatus::CodeUnsupported:
+				// show the message
+				FMessageDialog::Open(EAppMsgType::Ok, LOCTEXT("IOSNotSupported", "Sorry, launching on a device is currently not supported for code-based iOS projects. This feature will be available in a future release.") );
+				return;
+
+			case ETargetPlatformReadyStatus::PluginsUnsupported:
+				// show the message
+				FMessageDialog::Open(EAppMsgType::Ok, LOCTEXT("IOSNotSupported", "Sorry, launching on a device is currently not supported for content based projects with third-party plugins. This feature will be available in a future release.") );
+				return;
 #endif
+			default:
+				break;
+			}
+		}
+	}
 
 	if (!FModuleManager::LoadModuleChecked<IProjectTargetPlatformEditorModule>("ProjectTargetPlatformEditor").ShowUnsupportedTargetWarning(PlatformInfo->VanillaPlatformName))
 	{
diff --git a/Engine/Source/Editor/UnrealEd/Private/Kismet2/DebuggerCommands.cpp b/Engine/Source/Editor/UnrealEd/Private/Kismet2/DebuggerCommands.cpp
--- a/Engine/Source/Editor/UnrealEd/Private/Kismet2/DebuggerCommands.cpp
+++ b/Engine/Source/Editor/UnrealEd/Private/Kismet2/DebuggerCommands.cpp
@@ -123,6 +123,9 @@ protected:
 
 	/** Get the player start location to use when starting PIE */
 	static EPlayModeLocations GetPlayModeLocation();
+
+	/** checks to see if we have everything needed to launch a build to device */
+	static bool IsReadyToLaunchOnDevice(FString DeviceId);
 };
 
 
@@ -1224,7 +1227,10 @@ void FInternalPlayWorldCommandCallbacks::RepeatLastLaunch_Clicked()
 	switch (PlaySettings->LastExecutedLaunchModeType)
 	{
 	case LaunchMode_OnDevice:
-		LaunchOnDevice(PlaySettings->LastExecutedLaunchDevice, PlaySettings->LastExecutedLaunchName);
+		if (IsReadyToLaunchOnDevice(PlaySettings->LastExecutedLaunchDevice))
+		{
+			LaunchOnDevice(PlaySettings->LastExecutedLaunchDevice, PlaySettings->LastExecutedLaunchName);
+		}
 		break;
 
 	default:
@@ -1287,8 +1293,7 @@ FSlateIcon FInternalPlayWorldCommandCallbacks::GetRepeatLastLaunchIcon()
 	return FSlateIcon(FEditorStyle::GetStyleSetName(), "PlayWorld.RepeatLastLaunch");
 }
 
-
-void FInternalPlayWorldCommandCallbacks::HandleLaunchOnDeviceActionExecute( FString DeviceId, FString DeviceName )
+bool FInternalPlayWorldCommandCallbacks::IsReadyToLaunchOnDevice(FString DeviceId)
 {
 	int32 Index = 0;
 	DeviceId.FindChar(TEXT('@'), Index);
@@ -1302,23 +1307,53 @@ void FInternalPlayWorldCommandCallbacks::HandleLaunchOnDeviceActionExecute( FStr
 	if (Platform)
 	{
 		FString NotInstalledDocLink;
-		if (!Platform->IsSdkInstalled(bProjectHasCode, NotInstalledDocLink))
+		FString ProjectPath = FPaths::IsProjectFilePathSet() ? FPaths::ConvertRelativePathToFull(FPaths::GetProjectFilePath()) : FPaths::RootDir() / FApp::GetGameName() / FApp::GetGameName() + TEXT(".uproject");
+		switch (Platform->IsReadyToBuild(ProjectPath, bProjectHasCode, NotInstalledDocLink))
 		{
 			// broadcast this, and assume someone will pick it up
-			IMainFrameModule& MainFrameModule = FModuleManager::GetModuleChecked<IMainFrameModule>(TEXT("MainFrame"));
-			MainFrameModule.BroadcastMainFrameSDKNotInstalled(PlatformName, NotInstalledDocLink);
-			return;
+		case ETargetPlatformReadyStatus::SDKNotFound:
+		case ETargetPlatformReadyStatus::ProvisionNotFound:
+		case ETargetPlatformReadyStatus::SigningKeyNotFound:
+		case (ETargetPlatformReadyStatus::ProvisionNotFound | ETargetPlatformReadyStatus::SigningKeyNotFound):
+			{
+				IMainFrameModule& MainFrameModule = FModuleManager::GetModuleChecked<IMainFrameModule>(TEXT("MainFrame"));
+				MainFrameModule.BroadcastMainFrameSDKNotInstalled(PlatformName, NotInstalledDocLink);
+			}
+			return false;
+
+#if PLATFORM_WINDOWS
+		case ETargetPlatformReadyStatus::CodeUnsupported:
+			// show the message
+			FMessageDialog::Open(EAppMsgType::Ok, LOCTEXT("IOSNotSupported", "Sorry, launching on a device is currently not supported for code-based iOS projects. This feature will be available in a future release.") );
+			return false;
+
+		case ETargetPlatformReadyStatus::PluginsUnsupported:
+			// show the message
+			FMessageDialog::Open(EAppMsgType::Ok, LOCTEXT("IOSNotSupported", "Sorry, launching on a device is currently not supported for content based projects with third-party plugins. This feature will be available in a future release.") );
+			return false;
+#endif
+		default:
+			break;
 		}
 	}
 
-	ULevelEditorPlaySettings* PlaySettings = GetMutableDefault<ULevelEditorPlaySettings>();
+	return true;
+}
 
-	PlaySettings->LastExecutedLaunchModeType = LaunchMode_OnDevice;
-	PlaySettings->LastExecutedLaunchDevice = DeviceId;
-    PlaySettings->LastExecutedLaunchName = DeviceName;
-	PlaySettings->SaveConfig();
+void FInternalPlayWorldCommandCallbacks::HandleLaunchOnDeviceActionExecute( FString DeviceId, FString DeviceName )
+{
+	if (IsReadyToLaunchOnDevice(DeviceId))
+	{
+		ULevelEditorPlaySettings* PlaySettings = GetMutableDefault<ULevelEditorPlaySettings>();
+
+		PlaySettings->LastExecutedLaunchModeType = LaunchMode_OnDevice;
+		PlaySettings->LastExecutedLaunchDevice = DeviceId;
+		PlaySettings->LastExecutedLaunchName = DeviceName;
+		PlaySettings->SaveConfig();
+
+		LaunchOnDevice(DeviceId, DeviceName);
+	}
 
-	LaunchOnDevice(DeviceId, DeviceName);
 }
 
 
diff --git a/Engine/Source/Editor/UnrealEd/Private/PlayLevel.cpp b/Engine/Source/Editor/UnrealEd/Private/PlayLevel.cpp
--- a/Engine/Source/Editor/UnrealEd/Private/PlayLevel.cpp
+++ b/Engine/Source/Editor/UnrealEd/Private/PlayLevel.cpp
@@ -1285,21 +1285,6 @@ void UEditorEngine::PlayUsingLauncher()
 		FGameProjectGenerationModule& GameProjectModule = FModuleManager::LoadModuleChecked<FGameProjectGenerationModule>(TEXT("GameProjectGeneration"));
 		bool bHasCode = GameProjectModule.Get().GetProjectCodeFileCount() > 0;
 
-#if PLATFORM_WINDOWS
-		if (bHasCode && FRocketSupport::IsRocket() && PlayUsingLauncherDeviceId.Left(PlayUsingLauncherDeviceId.Find(TEXT("@"))) == TEXT("IOS"))
-		{
-			FMessageDialog::Open(EAppMsgType::Ok, LOCTEXT("LaunchOnDeviceNotWorkingMessage", "Sorry, launching on a device is currently not supported for code-based iOS projects. This feature will be available in a future release.") );
-			bPlayUsingLauncher = false;
-			return;
-		}
-		if (FRocketSupport::IsRocket() && PlayUsingLauncherDeviceId.Left(PlayUsingLauncherDeviceId.Find(TEXT("@"))) == TEXT("IOS") && IProjectManager::Get().IsNonDefaultPluginEnabled())
-		{
-			FMessageDialog::Open(EAppMsgType::Ok, LOCTEXT("LaunchOnDeviceNotWorkingMessage", "Sorry, launching on a device is currently not supported for content based projects with third-party plugins. This feature will be available in a future release.") );
-			bPlayUsingLauncher = false;
-			return;
-		}
-#endif
-
 		ILauncherProfileRef LauncherProfile = LauncherServicesModule.CreateProfile(TEXT("Play On Device"));
 		LauncherProfile->SetBuildGame(bHasCode && FSourceCodeNavigation::IsCompilerAvailable());
 		LauncherProfile->SetCookMode(ELauncherProfileCookModes::ByTheBook);
diff --git a/Engine/Source/Programs/IOS/iPhonePackager/CodeSigning.cs b/Engine/Source/Programs/IOS/iPhonePackager/CodeSigning.cs
--- a/Engine/Source/Programs/IOS/iPhonePackager/CodeSigning.cs
+++ b/Engine/Source/Programs/IOS/iPhonePackager/CodeSigning.cs
@@ -70,27 +70,81 @@ namespace iPhonePackager
 		/// <summary>
 		/// Makes sure the required files for code signing exist and can be found
 		/// </summary>
-		public static void FindRequiredFiles(out MobileProvision Provision, out X509Certificate2 Cert, out bool bHasOverridesFile)
+		public static bool FindRequiredFiles(out MobileProvision Provision, out X509Certificate2 Cert, out bool bHasOverridesFile)
 		{
-			// Check for a mobile provision
 			Provision = null;
+			Cert = null;
+			bHasOverridesFile = File.Exists(Config.GetPlistOverrideFilename());
+
+			// Load Info.plist, which guides nearly everything else
+			string plistFile = Path.GetDirectoryName(Config.ProjectFile) + "/Build/IOS/Info.plist";
+			if (!File.Exists(plistFile))
+			{
+				plistFile = Path.GetDirectoryName(Config.ProjectFile) + "/Build/IOS/" + Path.GetFileNameWithoutExtension(Config.ProjectFile) + "-Info.plist";
+
+				if (!File.Exists(plistFile))
+				{
+					plistFile = Config.EngineBuildDirectory + "/UE4Game-Info.plist";
+				}
+			}
+			Utilities.PListHelper Info = null;
+			try
+			{
+				string RawInfoPList = File.ReadAllText(plistFile, Encoding.UTF8);
+				Info = new Utilities.PListHelper(RawInfoPList);;
+			}
+			catch (Exception ex)
+			{
+				Console.WriteLine(ex.Message);
+			}
+
+			if (Info == null)
+			{
+				return false;
+			}
+
+			// Get the name of the bundle
+			string CFBundleIdentifier = null;
+			Info.GetString("CFBundleIdentifier", out CFBundleIdentifier);
+			if (CFBundleIdentifier == null)
+			{
+				return false;
+			}
+			else
+			{
+				CFBundleIdentifier = CFBundleIdentifier.Replace("${BUNDLE_IDENTIFIER}", Path.GetFileNameWithoutExtension(Config.ProjectFile));
+			}
+
+			// Check for a mobile provision
 			try
 			{
-				string ExpectedProvisionFile = FileOperations.FindPrefixedFile(Config.ProvisionDirectory, Program.GameName + ".mobileprovision");
-				Provision = MobileProvisionParser.ParseFile(ExpectedProvisionFile);
+				string MobileProvisionFilename = MobileProvision.FindCompatibleProvision(CFBundleIdentifier);
+				Provision = MobileProvisionParser.ParseFile(MobileProvisionFilename);
 			}
 			catch (Exception)
 			{
 			}
 
+			// if we have a null provision see if we can find a compatible provision without checking for a certificate
+			if (Provision == null)
+			{
+				try
+				{
+					string MobileProvisionFilename = MobileProvision.FindCompatibleProvision(CFBundleIdentifier, false);
+					Provision = MobileProvisionParser.ParseFile(MobileProvisionFilename);
+				}
+				catch (Exception)
+				{
+				}
+			}
+
 			// Check for a suitable signature to match the mobile provision
-			Cert = null;
 			if (Provision != null)
 			{
 				Cert = CodeSignatureBuilder.FindCertificate(Provision);
 			}
 
-			bHasOverridesFile = File.Exists(Config.GetPlistOverrideFilename());
+			return true;
 		}
 
 		public static bool DoRequiredFilesExist()
diff --git a/Engine/Source/Programs/IOS/iPhonePackager/iPhonePackager.cs b/Engine/Source/Programs/IOS/iPhonePackager/iPhonePackager.cs
--- a/Engine/Source/Programs/IOS/iPhonePackager/iPhonePackager.cs
+++ b/Engine/Source/Programs/IOS/iPhonePackager/iPhonePackager.cs
@@ -14,6 +14,7 @@ using System.Threading;
 using System.Windows.Forms;
 using Ionic.Zlib;
 using System.ComponentModel;
+using System.Security.Cryptography.X509Certificates;
 
 namespace iPhonePackager
 {
@@ -51,11 +52,11 @@ namespace iPhonePackager
 			Log(String.Format(Line, Args));
 		}
 
-		static public void Error( string Line )
+		static public void Error( string Line, int Code = 1 )
 		{
 			if (Program.ReturnCode == 0)
 			{
-				Program.ReturnCode = 1;
+				Program.ReturnCode = Code;
 			}
 			
 			Console.ForegroundColor = ConsoleColor.Red;
@@ -471,6 +472,7 @@ namespace iPhonePackager
 					Log(" ... Deploy PathToIPA");
 					Log(" ... RepackageFromStage GameName");
 					Log(" ... Devices");
+					Log(" ... Validate");
 					Log("");
 					Log("Configuration switches:");
 					Log("	 -stagedir <path>		  sets the directory to copy staged files from (defaults to none)");
@@ -497,6 +499,8 @@ namespace iPhonePackager
 
 				Log("Executing iPhonePackager " + String.Join(" ", args));
 				Log("CWD: " + Directory.GetCurrentDirectory());
+				Log("Initial Dir: " + InitialCurrentDirectory);
+				Log("Env CWD: " + Environment.CurrentDirectory);
 
 				// Ensure shipping configuration for final distributions
 				if (Config.bForDistribution && (GameConfiguration != "Shipping"))
@@ -514,11 +518,44 @@ namespace iPhonePackager
 				// setup configuration
 				if (!Config.Initialize(InitialCurrentDirectory, GamePath))
 				{
-					return 1;
+					return -1;
 				}
 
 				switch (MainCommand.ToLowerInvariant())
 				{
+					case "validate":
+						// check to see if iTunes is installed
+						string dllPath = Microsoft.Win32.Registry.GetValue("HKEY_LOCAL_MACHINE\\SOFTWARE\\Wow6432Node\\Apple Inc.\\Apple Mobile Device Support\\Shared", "iTunesMobileDeviceDLL", null) as string; 
+						if (String.IsNullOrEmpty(dllPath) || !File.Exists(dllPath))
+						{
+							Error("iTunes Not Found!!", 10);
+						}
+						else
+						{
+							// validate there is a useable provision and cert
+							MobileProvision Provision;
+							X509Certificate2 Cert;
+							bool bHasOverrides;
+							bool foundPlist = CodeSignatureBuilder.FindRequiredFiles(out Provision, out Cert, out bHasOverrides);
+							if (!foundPlist)
+							{
+								Error("Could not find a valid plist file!!", 14);
+							}
+							else if (Provision == null && Cert == null)
+							{
+								Error("No Provision or cert found!!", 13);
+							}
+							else if (Provision == null)
+							{
+								Error("No Provision found!!", 11);
+							}
+							else if (Cert == null)
+							{
+								Error("No Signing Certificate found!!", 12);
+							}
+						}
+						break;
+
 					case "packageapp":
 						if (CheckArguments())
 						{
@@ -607,7 +644,8 @@ namespace iPhonePackager
 				}
 			}
 
-			return ( ReturnCode );
+			Environment.ExitCode = ReturnCode;
+			return (ReturnCode);
 		}
 	}
 }
diff --git a/Engine/Source/Programs/IOS/iPhonePackager/iPhonePackager.csproj b/Engine/Source/Programs/IOS/iPhonePackager/iPhonePackager.csproj
--- a/Engine/Source/Programs/IOS/iPhonePackager/iPhonePackager.csproj
+++ b/Engine/Source/Programs/IOS/iPhonePackager/iPhonePackager.csproj
@@ -21,7 +21,8 @@
     <SccProvider>
     </SccProvider>
     <IsWebBootstrapper>false</IsWebBootstrapper>
-    <StartupObject>iPhonePackager.Program</StartupObject>
+    <StartupObject>
+    </StartupObject>
     <ApplicationIcon>IPhonePackagerIcon.ico</ApplicationIcon>
     <FileUpgradeFlags>
     </FileUpgradeFlags>
diff --git a/Engine/Source/Runtime/Core/Private/Misc/MonitoredProcess.cpp b/Engine/Source/Runtime/Core/Private/Misc/MonitoredProcess.cpp
--- a/Engine/Source/Runtime/Core/Private/Misc/MonitoredProcess.cpp
+++ b/Engine/Source/Runtime/Core/Private/Misc/MonitoredProcess.cpp
@@ -17,6 +17,7 @@ FMonitoredProcess::FMonitoredProcess( const FString& InURL, const FString& InPar
 	, Thread(NULL)
 	, URL(InURL)
 	, WritePipe(NULL)
+	, ReturnCode(0)
 { }
 
 
@@ -118,8 +119,6 @@ uint32 FMonitoredProcess::Run( )
 	ReadPipe = WritePipe = NULL;
 
 	// get completion status
-	int32 ReturnCode;
-	
 	if (!FPlatformProcess::GetProcReturnCode(ProcessHandle, &ReturnCode))
 	{
 		ReturnCode = -1;
diff --git a/Engine/Source/Runtime/Core/Public/Misc/MonitoredProcess.h b/Engine/Source/Runtime/Core/Public/Misc/MonitoredProcess.h
--- a/Engine/Source/Runtime/Core/Public/Misc/MonitoredProcess.h
+++ b/Engine/Source/Runtime/Core/Public/Misc/MonitoredProcess.h
@@ -107,6 +107,16 @@ public:
 		return OutputDelegate;
 	}
 
+	/**
+	 * Returns the return code from the exited process
+	 *
+	 * @return Process return code
+	 */
+	int GetReturnCode() const
+	{
+		return ReturnCode;
+	}
+
 public:
 
 	// FRunnable interface
@@ -169,6 +179,9 @@ private:
 	// Holds the write pipe.
 	void* WritePipe;
 
+	// Holds the return code
+	int ReturnCode;
+
 private:
 
 	// Holds a delegate that is executed when the process has been canceled.
