- Made CookerSettings "defaultconfig"

diff --git a/Engine/Source/Editor/UnrealEd/Public/CookerSettings.h b/Engine/Source/Editor/UnrealEd/Public/CookerSettings.h
--- a/Engine/Source/Editor/UnrealEd/Public/CookerSettings.h
+++ b/Engine/Source/Editor/UnrealEd/Public/CookerSettings.h
@@ -12,7 +12,7 @@
 /**
  * Implements per-project cooker settings exposed to the editor
  */
-UCLASS(config=Engine)
+UCLASS(config=Engine, defaultconfig)
 class UNREALED_API UCookerSettings
 	: public UObject
 {
