Editor: Add DLL export to RemoveRedundantTriangles

diff --git a/Engine/Source/Editor/UnrealEd/Public/GeomTools.h b/Engine/Source/Editor/UnrealEd/Public/GeomTools.h
--- a/Engine/Source/Editor/UnrealEd/Public/GeomTools.h
+++ b/Engine/Source/Editor/UnrealEd/Public/GeomTools.h
@@ -142,7 +142,7 @@ FClipSMPolygon Transform2DPolygonToSMPolygon(const FUtilPoly2D& InTri, const FMa
 void GeneratePolyUVs(FUtilPoly2D& Polygon);
 
 /** Given a set of triangles, remove those which share an edge and could be collapsed into one triangle. */
-void RemoveRedundantTriangles(TArray<FClipSMTriangle>& Tris);
+UNREALED_API void RemoveRedundantTriangles(TArray<FClipSMTriangle>& Tris);
 
 /** Split 2D polygons with a 3D plane. */
 void Split2DPolysWithPlane(FUtilPoly2DSet& PolySet, const FPlane& Plane, const FColor& ExteriorVertColor, const FColor& InteriorVertColor);
