- unified duplicated codes in saving content packages
- added a missing part to connect between "save animation assets" menu and an actual execute method

diff --git a/Engine/Source/Editor/Persona/Private/Persona.cpp b/Engine/Source/Editor/Persona/Private/Persona.cpp
--- a/Engine/Source/Editor/Persona/Private/Persona.cpp
+++ b/Engine/Source/Editor/Persona/Private/Persona.cpp
@@ -943,6 +943,14 @@ void FPersona::CreateDefaultCommands()
 	// now add default commands
 	FPersonaCommands::Register();
 
+	// save all animation assets
+	ToolkitCommands->MapAction(FPersonaCommands::Get().SaveAnimationAssets,
+		FExecuteAction::CreateSP(this, &FPersona::SaveAnimationAssets_Execute),
+		FCanExecuteAction::CreateSP(this, &FPersona::CanSaveAnimationAssets),
+		FIsActionChecked(),
+		FIsActionButtonVisible::CreateSP(this, &FPersona::IsRecordAvailable)
+		);
+
 	// record animation
 	ToolkitCommands->MapAction( FPersonaCommands::Get().RecordAnimation,
 		FExecuteAction::CreateSP( this, &FPersona::RecordAnimation ),
diff --git a/Engine/Source/Editor/UnrealEd/Private/FileHelpers.cpp b/Engine/Source/Editor/UnrealEd/Private/FileHelpers.cpp
--- a/Engine/Source/Editor/UnrealEd/Private/FileHelpers.cpp
+++ b/Engine/Source/Editor/UnrealEd/Private/FileHelpers.cpp
@@ -2260,10 +2260,103 @@
 		OpenMsgDlgInt( EAppMsgType::Ok, Message, NSLOCTEXT("FileHelper", "FailedSavePrompt_Title", "Packages Failed To Save") );
 	}
 }
-bool FEditorFileUtils::SaveDirtyPackages(const bool bPromptUserToSave, const bool bSaveMapPackages, const bool bSaveContentPackages, const bool bFastSave, const bool bNotifyNoPackagesSaved,  bool* bOutPackagesNeededSaving )
+
+static bool InternalSavePackages(TArray<UPackage*>& PackagesToSave, int32 NumPackagesNotIgnored, bool bPromptUserToSave, bool bFastSave, bool bNotifyNoPackagesSaved, bool* bOutPackagesNeededSaving)
 {
 	bool bReturnCode = true;
 
+	if (PackagesToSave.Num() > 0 && (NumPackagesNotIgnored > 0 || bPromptUserToSave))
+	{
+		// The caller asked us 
+		if (bOutPackagesNeededSaving != NULL)
+		{
+			*bOutPackagesNeededSaving = true;
+		}
+
+		if (!bFastSave)
+		{
+			const FEditorFileUtils::EPromptReturnCode Return = FEditorFileUtils::PromptForCheckoutAndSave(PackagesToSave, true, bPromptUserToSave);
+			if (Return == FEditorFileUtils::EPromptReturnCode::PR_Cancelled)
+			{
+				// Only cancel should return false and stop whatever we were doing before.(like closing the editor)
+				// If failure is returned, the user was given ample times to retry saving the package and didn't want to
+				// So we should continue with whatever we were doing.  
+				bReturnCode = false;
+			}
+		}
+		else
+		{
+			FSaveErrorOutputDevice SaveErrors;
+			GWarn->BeginSlowTask(NSLOCTEXT("UnrealEd", "SavingPackagesE", "Saving packages..."), true);
+
+			// Packages that failed to save
+			TArray< UPackage* > FailedPackages;
+
+			for (TArray<UPackage*>::TConstIterator PkgIter(PackagesToSave); PkgIter; ++PkgIter)
+			{
+				UPackage* CurPackage = *PkgIter;
+
+				// Check if a file exists for this package
+				FString Filename;
+				bool bFoundFile = FPackageName::DoesPackageExist(CurPackage->GetName(), NULL, &Filename);
+				if (bFoundFile)
+				{
+					// determine if the package file is read only
+					const bool bPkgReadOnly = IFileManager::Get().IsReadOnly(*Filename);
+
+					// Only save writable files in fast mode
+					if (!bPkgReadOnly)
+					{
+						if (!CurPackage->IsFullyLoaded())
+						{
+							// Packages must be fully loaded to save
+							CurPackage->FullyLoad();
+						}
+
+						const UWorld* const AssociatedWorld = UWorld::FindWorldInPackage(CurPackage);
+						const bool bIsMapPackage = AssociatedWorld != nullptr;
+
+						const FText SavingPackageText = (bIsMapPackage)
+							? FText::Format(NSLOCTEXT("UnrealEd", "SavingMapf", "Saving map {0}"), FText::FromString(CurPackage->GetName()))
+							: FText::Format(NSLOCTEXT("UnrealEd", "SavingAssetf", "Saving asset {0}"), FText::FromString(CurPackage->GetName()));
+
+						GWarn->StatusForceUpdate(PkgIter.GetIndex(), PackagesToSave.Num(), SavingPackageText);
+
+						// Save the package
+						bool bPackageLocallyWritable;
+						const int32 SaveStatus = InternalSavePackage(CurPackage, bPackageLocallyWritable, SaveErrors);
+
+						if (SaveStatus == EAppReturnType::No)
+						{
+							// The package could not be saved so add it to the failed array 
+							FailedPackages.Add(CurPackage);
+
+						}
+					}
+				}
+
+			}
+			GWarn->EndSlowTask();
+			SaveErrors.Flush();
+
+			// Warn the user about any packages which failed to save.
+			WarnUserAboutFailedSave(FailedPackages);
+		}
+	}
+	else if (bNotifyNoPackagesSaved)
+	{
+		FNotificationInfo NotificationInfo(LOCTEXT("NoAssetsToSave", "No new changes to save!"));
+		NotificationInfo.Image = FEditorStyle::GetBrush(FTokenizedMessage::GetSeverityIconName(EMessageSeverity::Info));
+		NotificationInfo.bFireAndForget = true;
+		NotificationInfo.ExpireDuration = 4.0f; // Need this message to last a little longer than normal since the user may have expected there to be modified files.
+		NotificationInfo.bUseThrobber = true;
+		FSlateNotificationManager::Get().AddNotification(NotificationInfo);
+	}
+	return bReturnCode;
+}
+
+bool FEditorFileUtils::SaveDirtyPackages(const bool bPromptUserToSave, const bool bSaveMapPackages, const bool bSaveContentPackages, const bool bFastSave, const bool bNotifyNoPackagesSaved,  bool* bOutPackagesNeededSaving )
+{
 	if (bOutPackagesNeededSaving != NULL)
 	{
 		*bOutPackagesNeededSaving = false;
@@ -2338,94 +2431,7 @@
 		}
 	}
 
-	if( PackagesToSave.Num() > 0 && (NumPackagesNotIgnored > 0 || bPromptUserToSave) ) 
-	{
-		// The caller asked us 
-		if (bOutPackagesNeededSaving != NULL)
-		{
-			*bOutPackagesNeededSaving  = true;
-		}
-
-		if( !bFastSave )
-		{
-			const FEditorFileUtils::EPromptReturnCode Return = FEditorFileUtils::PromptForCheckoutAndSave(PackagesToSave, true, bPromptUserToSave);
-			if( Return == FEditorFileUtils::EPromptReturnCode::PR_Cancelled )
-			{
-				// Only cancel should return false and stop whatever we were doing before.(like closing the editor)
-				// If failure is returned, the user was given ample times to retry saving the package and didn't want to
-				// So we should continue with whatever we were doing.  
-				bReturnCode = false;
-			}
-		}
-		else
-		{
-			FSaveErrorOutputDevice SaveErrors;
-			GWarn->BeginSlowTask( NSLOCTEXT("UnrealEd", "SavingPackagesE", "Saving packages..."), true );
-
-			// Packages that failed to save
-			TArray< UPackage* > FailedPackages;
-
-			for ( TArray<UPackage*>::TConstIterator PkgIter( PackagesToSave ); PkgIter; ++PkgIter )
-			{
-				UPackage* CurPackage = *PkgIter;
-
-				// Check if a file exists for this package
-				FString Filename;
-				bool bFoundFile = FPackageName::DoesPackageExist( CurPackage->GetName(), NULL, &Filename );
-				if( bFoundFile )
-				{
-					// determine if the package file is read only
-					const bool bPkgReadOnly = IFileManager::Get().IsReadOnly( *Filename );
-
-					// Only save writable files in fast mode
-					if ( !bPkgReadOnly )
-					{	
-						if( !CurPackage->IsFullyLoaded() )
-						{
-							// Packages must be fully loaded to save
-							CurPackage->FullyLoad();
-						}
-
-						const UWorld* const AssociatedWorld = UWorld::FindWorldInPackage(CurPackage);
-						const bool bIsMapPackage = AssociatedWorld != nullptr;
-
-						const FText SavingPackageText = (bIsMapPackage) 
-							? FText::Format(NSLOCTEXT("UnrealEd", "SavingMapf", "Saving map {0}"), FText::FromString(CurPackage->GetName()))
-							: FText::Format(NSLOCTEXT("UnrealEd", "SavingAssetf", "Saving asset {0}"), FText::FromString(CurPackage->GetName()));
-
-						GWarn->StatusForceUpdate( PkgIter.GetIndex(), PackagesToSave.Num(), SavingPackageText );
-
-						// Save the package
-						bool bPackageLocallyWritable;
-						const int32 SaveStatus = InternalSavePackage( CurPackage, bPackageLocallyWritable, SaveErrors );
-
-						if( SaveStatus == EAppReturnType::No )
-						{
-							// The package could not be saved so add it to the failed array 
-							FailedPackages.Add( CurPackage );
-
-						}
-					}
-				}
-				
-			}
-			GWarn->EndSlowTask();
-			SaveErrors.Flush();
-
-			// Warn the user about any packages which failed to save.
-			WarnUserAboutFailedSave( FailedPackages );
-		}
-	}
-	else if(bNotifyNoPackagesSaved)
-	{
-		FNotificationInfo NotificationInfo( LOCTEXT("NoAssetsToSave", "No new changes to save!") );
-		NotificationInfo.Image = FEditorStyle::GetBrush(FTokenizedMessage::GetSeverityIconName(EMessageSeverity::Info));
-		NotificationInfo.bFireAndForget = true;
-		NotificationInfo.ExpireDuration = 4.0f; // Need this message to last a little longer than normal since the user may have expected there to be modified files.
-		NotificationInfo.bUseThrobber = true;
-		FSlateNotificationManager::Get().AddNotification(NotificationInfo);
-	}
-	return bReturnCode;
+	return InternalSavePackages(PackagesToSave, NumPackagesNotIgnored, bPromptUserToSave, bFastSave, bNotifyNoPackagesSaved, bOutPackagesNeededSaving);
 }
 
 bool FEditorFileUtils::SaveDirtyContentPackages(TArray<UClass*>& SaveContentClasses, const bool bPromptUserToSave, const bool bFastSave, const bool bNotifyNoPackagesSaved)
@@ -2483,89 +2489,7 @@
 		}
 	}
 
-	if (PackagesToSave.Num() > 0)
-	{
-		if (!bFastSave)
-		{
-			const FEditorFileUtils::EPromptReturnCode Return = FEditorFileUtils::PromptForCheckoutAndSave(PackagesToSave, true, bPromptUserToSave);
-			if (Return == FEditorFileUtils::EPromptReturnCode::PR_Cancelled)
-			{
-				// Only cancel should return false and stop whatever we were doing before.(like closing the editor)
-				// If failure is returned, the user was given ample times to retry saving the package and didn't want to
-				// So we should continue with whatever we were doing.  
-				bReturnCode = false;
-			}
-		}
-		else
-		{
-			FSaveErrorOutputDevice SaveErrors;
-			GWarn->BeginSlowTask(NSLOCTEXT("UnrealEd", "SavingPackagesE", "Saving packages..."), true);
-
-			// Packages that failed to save
-			TArray< UPackage* > FailedPackages;
-
-			for (TArray<UPackage*>::TConstIterator PkgIter(PackagesToSave); PkgIter; ++PkgIter)
-			{
-				UPackage* CurPackage = *PkgIter;
-
-				// Check if a file exists for this package
-				FString Filename;
-				bool bFoundFile = FPackageName::DoesPackageExist(CurPackage->GetName(), NULL, &Filename);
-				if (bFoundFile)
-				{
-					// determine if the package file is read only
-					const bool bPkgReadOnly = IFileManager::Get().IsReadOnly(*Filename);
-
-					// Only save writable files in fast mode
-					if (!bPkgReadOnly)
-					{
-						if (!CurPackage->IsFullyLoaded())
-						{
-							// Packages must be fully loaded to save
-							CurPackage->FullyLoad();
-						}
-
-						const UWorld* const AssociatedWorld = UWorld::FindWorldInPackage(CurPackage);
-						const bool bIsMapPackage = AssociatedWorld != nullptr;
-
-						const FText SavingPackageText = (bIsMapPackage)
-							? FText::Format(NSLOCTEXT("UnrealEd", "SavingMapf", "Saving map {0}"), FText::FromString(CurPackage->GetName()))
-							: FText::Format(NSLOCTEXT("UnrealEd", "SavingAssetf", "Saving asset {0}"), FText::FromString(CurPackage->GetName()));
-
-						GWarn->StatusForceUpdate(PkgIter.GetIndex(), PackagesToSave.Num(), SavingPackageText);
-
-						// Save the package
-						bool bPackageLocallyWritable;
-						const int32 SaveStatus = InternalSavePackage(CurPackage, bPackageLocallyWritable, SaveErrors);
-
-						if (SaveStatus == EAppReturnType::No)
-						{
-							// The package could not be saved so add it to the failed array 
-							FailedPackages.Add(CurPackage);
-
-						}
-					}
-				}
-
-			}
-			GWarn->EndSlowTask();
-			SaveErrors.Flush();
-
-			// Warn the user about any packages which failed to save.
-			WarnUserAboutFailedSave(FailedPackages);
-		}
-	}
-	else
-	{
-		FNotificationInfo NotificationInfo(LOCTEXT("NoAssetsToSave", "No new changes to save!"));
-		NotificationInfo.Image = FEditorStyle::GetBrush(FTokenizedMessage::GetSeverityIconName(EMessageSeverity::Info));
-		NotificationInfo.bFireAndForget = true;
-		NotificationInfo.ExpireDuration = 4.0f; // Need this message to last a little longer than normal since the user may have expected there to be modified files.
-		NotificationInfo.bUseThrobber = true;
-		FSlateNotificationManager::Get().AddNotification(NotificationInfo);
-	}
-
-	return bReturnCode;
+	return InternalSavePackages(PackagesToSave, PackagesToSave.Num(), bPromptUserToSave, bFastSave, bNotifyNoPackagesSaved, NULL);
 }
 
 /**
