Header cleanup: World.h, ParticleEmitter.h

diff --git a/Engine/Source/Editor/UnrealEd/Private/Commandlets/ParticleSystemAuditCommandlet.cpp b/Engine/Source/Editor/UnrealEd/Private/Commandlets/ParticleSystemAuditCommandlet.cpp
--- a/Engine/Source/Editor/UnrealEd/Private/Commandlets/ParticleSystemAuditCommandlet.cpp
+++ b/Engine/Source/Editor/UnrealEd/Private/Commandlets/ParticleSystemAuditCommandlet.cpp
@@ -8,6 +8,7 @@
 #include "Runtime/Engine/Classes/Particles/Collision/ParticleModuleCollision.h"
 #include "Runtime/Engine/Classes/Particles/Spawn/ParticleModuleSpawn.h"
 #include "Runtime/AssetRegistry/Public/AssetRegistryModule.h"
+#include "Particles/ParticleLODLevel.h"
 
 DEFINE_LOG_CATEGORY_STATIC(LogParticleSystemAuditCommandlet, Log, All);
 
diff --git a/Engine/Source/Runtime/Engine/Classes/Engine/World.h b/Engine/Source/Runtime/Engine/Classes/Engine/World.h
--- a/Engine/Source/Runtime/Engine/Classes/Engine/World.h
+++ b/Engine/Source/Runtime/Engine/Classes/Engine/World.h
@@ -6,7 +6,6 @@
 	World.h: UWorld definition.
 =============================================================================*/
 
-#include "CollisionQueryParams.h"
 #include "WorldCollision.h"
 #include "PendingNetGame.h"
 #include "EngineDefines.h"
diff --git a/Engine/Source/Runtime/Engine/Classes/Particles/ParticleEmitter.h b/Engine/Source/Runtime/Engine/Classes/Particles/ParticleEmitter.h
--- a/Engine/Source/Runtime/Engine/Classes/Particles/ParticleEmitter.h
+++ b/Engine/Source/Runtime/Engine/Classes/Particles/ParticleEmitter.h
@@ -6,11 +6,7 @@
 //=============================================================================
 
 #pragma once
-
-#include "ParticleLODLevel.h"
-#include "ParticleHelper.h"
 #include "ParticleEmitterInstances.h"
-
 #include "ParticleEmitter.generated.h"
 
 //=============================================================================
