Fixed new default maps not having physics scenes created

This broke actor placement because new worlds which were explicitly created without PX scenes were not having them correctly created when new maps were created in the level editor.

diff --git a/Engine/Source/Editor/UnrealEd/Private/EditorServer.cpp b/Engine/Source/Editor/UnrealEd/Private/EditorServer.cpp
--- a/Engine/Source/Editor/UnrealEd/Private/EditorServer.cpp
+++ b/Engine/Source/Editor/UnrealEd/Private/EditorServer.cpp
@@ -2201,6 +2201,16 @@ bool UEditorEngine::Map_Load(const TCHAR* Str, FOutputDevice& Ar)
 				Context.SetCurrentWorld(World);
 				GWorld = World;
 
+				// Inactive worlds may have already been initialized, but still be lacking important scene information
+				if ( Context.World()->WorldType == EWorldType::Inactive )
+				{
+					Context.World()->CreatePhysicsScene();
+					Context.World()->CreateFXSystem();
+
+					// Update components so the scene is populated
+					Context.World()->UpdateWorldComponents(true, true);
+				}
+
 				Context.World()->WorldType = EWorldType::Editor;
 
 				if ( !Context.World()->bIsWorldInitialized )
