Adding FBX error to work with doc link/rich tooltip - like map error

diff --git a/Engine/Source/Editor/UnrealEd/Classes/Factories/FbxSkeletalMeshImportData.h b/Engine/Source/Editor/UnrealEd/Classes/Factories/FbxSkeletalMeshImportData.h
--- a/Engine/Source/Editor/UnrealEd/Classes/Factories/FbxSkeletalMeshImportData.h
+++ b/Engine/Source/Editor/UnrealEd/Classes/Factories/FbxSkeletalMeshImportData.h
@@ -27,6 +27,10 @@ class UFbxSkeletalMeshImportData : public UFbxMeshImportData
 	UPROPERTY(EditAnywhere, config, Category=ImportSettings)
 	uint32 bImportMeshesInBoneHierarchy:1;
 
+	/** If checked, we will import group node (null node) as root. */
+	UPROPERTY(EditAnywhere, config, Category=ImportSettings)
+	uint32 bImportGroupNodeAsRoot:1;
+
 	/** True to import morph target meshes from the FBX file */
 	UPROPERTY(EditAnywhere, config, Category=ImportSettings, meta=(ToolTip="If enabled, creates Unreal morph objects for the imported meshes"))
 	uint32 bImportMorphTargets:1;
diff --git a/Engine/Source/Editor/UnrealEd/Private/Factories/SkeletalMeshImport.cpp b/Engine/Source/Editor/UnrealEd/Private/Factories/SkeletalMeshImport.cpp
--- a/Engine/Source/Editor/UnrealEd/Private/Factories/SkeletalMeshImport.cpp
+++ b/Engine/Source/Editor/UnrealEd/Private/Factories/SkeletalMeshImport.cpp
@@ -10,6 +10,7 @@
 #include "SkeletalMeshSorting.h"
 #include "../../../../Source/Runtime/Engine/Classes/PhysicsEngine/PhysicsAsset.h"
 #include "FbxImporter.h"
+#include "FbxErrors.h"
 
 DEFINE_LOG_CATEGORY_STATIC(LogSkeletalMeshImport, Log, All);
 
@@ -22,7 +23,7 @@ bool SkeletonsAreCompatible( const FReferenceSkeleton& NewSkel, const FReference
 	{
 		UnFbx::FFbxImporter* FFbxImporter = UnFbx::FFbxImporter::GetInstance();	
 		FFbxImporter->AddTokenizedErrorMessage(FTokenizedMessage::Create(EMessageSeverity::Error, FText::Format(LOCTEXT("MeshHasDifferentRoot", "Root Bone is '{0}' instead of '{1}'.\nDiscarding existing LODs."),
-			FText::FromName(NewSkel.GetBoneName(0)), FText::FromName(ExistSkel.GetBoneName(0)))));
+			FText::FromName(NewSkel.GetBoneName(0)), FText::FromName(ExistSkel.GetBoneName(0)))), FFbxErrors::SkeletalMesh_DifferentRoots);
 		return false;
 	}
 
@@ -43,7 +44,7 @@ bool SkeletonsAreCompatible( const FReferenceSkeleton& NewSkel, const FReference
 			{
 				UnFbx::FFbxImporter* FFbxImporter = UnFbx::FFbxImporter::GetInstance();			
 				FFbxImporter->AddTokenizedErrorMessage(FTokenizedMessage::Create(EMessageSeverity::Error, FText::Format(LOCTEXT("MeshHasDifferentRoot", "Root Bone is '{0}' instead of '{1}'.\nDiscarding existing LODs."),
-					FText::FromName(NewBoneName), FText::FromName(NewParentName))));
+					FText::FromName(NewBoneName), FText::FromName(NewParentName))), FFbxErrors::SkeletalMesh_DifferentRoots);
 				return false;
 			}
 		}
@@ -207,7 +208,7 @@ bool ProcessImportMeshSkeleton(FReferenceSkeleton& RefSkeleton, int32& SkeletalD
 		if(RefSkeleton.FindBoneIndex(BoneInfo.Name) != INDEX_NONE)
 		{
 			UnFbx::FFbxImporter* FFbxImporter = UnFbx::FFbxImporter::GetInstance();
-			FFbxImporter->AddTokenizedErrorMessage(FTokenizedMessage::Create(EMessageSeverity::Error, FText::Format(LOCTEXT("SkeletonHasDuplicateBones", "Skeleton has non-unique bone names.\nBone named '{0}' encountered more than once."), FText::FromName(BoneInfo.Name))));
+			FFbxImporter->AddTokenizedErrorMessage(FTokenizedMessage::Create(EMessageSeverity::Error, FText::Format(LOCTEXT("SkeletonHasDuplicateBones", "Skeleton has non-unique bone names.\nBone named '{0}' encountered more than once."), FText::FromName(BoneInfo.Name))), FFbxErrors::SkeletalMesh_DuplicateBones);
 			return false;
 		}
 
@@ -328,7 +329,7 @@ void ProcessImportMeshInfluences(FSkeletalMeshImportData& ImportData)
 	{
 		UnFbx::FFbxImporter* FFbxImporter = UnFbx::FFbxImporter::GetInstance();
 		// warn about no influences
-		FFbxImporter->AddTokenizedErrorMessage(FTokenizedMessage::Create(EMessageSeverity::Warning, LOCTEXT("WarningNoSkelInfluences", "Warning skeletal mesh is has no vertex influences")));
+		FFbxImporter->AddTokenizedErrorMessage(FTokenizedMessage::Create(EMessageSeverity::Warning, LOCTEXT("WarningNoSkelInfluences", "Warning skeletal mesh is has no vertex influences")), FFbxErrors::SkeletalMesh_NoInfluences);
 		// add one for each wedge entry
 		Influences.AddUninitialized(Wedges.Num());
 		for( int32 WedgeIdx=0; WedgeIdx<Wedges.Num(); WedgeIdx++ )
@@ -575,7 +576,8 @@ void FSavedCustomSortSectionInfo::Restore(USkeletalMesh* NewSkelMesh, int32 LODM
 						FText::FromString(TriangleSortOptionToString((ETriangleSortOption)SavedSortOption)),
 						FText::AsNumber(SavedSectionIdx),
 						FText::AsNumber(NumMismatchedStrips),
-						FText::AsNumber(OldStrips[0].Num()))));
+						FText::AsNumber(OldStrips[0].Num()))), 
+						FFbxErrors::SkeletalMesh_RestoreSortingMismatchedStrips);
 				}
 
 				// Restore the settings saved in the LODInfo (used for the UI)
@@ -614,7 +616,7 @@ void FSavedCustomSortSectionInfo::Restore(USkeletalMesh* NewSkelMesh, int32 LODM
 	{
 		UnFbx::FFbxImporter* FFbxImporter = UnFbx::FFbxImporter::GetInstance();
 		FFbxImporter->AddTokenizedErrorMessage(FTokenizedMessage::Create(EMessageSeverity::Warning, FText::Format(LOCTEXT("FailedRestoreSortingNoSectionMatch", "Unable to restore triangle sort setting \"{0}\" for section number {1} in the old mesh, as a matching section could not be found in the new mesh. The custom sorting information has been lost."),
-			FText::FromString(TriangleSortOptionToString((ETriangleSortOption)SavedSortOption)), FText::AsNumber(SavedSectionIdx))));
+			FText::FromString(TriangleSortOptionToString((ETriangleSortOption)SavedSortOption)), FText::AsNumber(SavedSectionIdx))), FFbxErrors::SkeletalMesh_RestoreSortingNoSectionMatch);
 	}
 }
 
@@ -661,7 +663,8 @@ void FSavedCustomSortInfo::Restore(USkeletalMesh* NewSkeletalMesh, int32 LODMode
 				UnFbx::FFbxImporter* FFbxImporter = UnFbx::FFbxImporter::GetInstance();
 				FFbxImporter->AddTokenizedErrorMessage(FTokenizedMessage::Create(EMessageSeverity::Error, FText::Format(LOCTEXT("FailedRestoreSortingForSectionNumber", "Unable to restore triangle sort setting \"{0}\" for section {1} as the new mesh does not contain that many sections. Please find the matching section and apply manually."),
 					FText::FromString(TriangleSortOptionToString((ETriangleSortOption)SortSectionInfo.SavedSortOption)),
-					FText::AsNumber(SortSectionInfo.SavedSectionIdx))));
+					FText::AsNumber(SortSectionInfo.SavedSectionIdx))), 
+					FFbxErrors::SkeletalMesh_RestoreSortingForSectionNumber);
 				continue;
 			}
 
@@ -870,7 +873,7 @@ void RestoreExistingSkelMeshData(ExistingSkelMeshData* MeshData, USkeletalMesh*
 				if(bMissingBone)
 				{
 					UnFbx::FFbxImporter* FFbxImporter = UnFbx::FFbxImporter::GetInstance();
-					FFbxImporter->AddTokenizedErrorMessage(FTokenizedMessage::Create(EMessageSeverity::Warning, FText::Format(LOCTEXT("NewMeshMissingBoneFromLOD", "New mesh is missing bone '{0}' required by an LOD. LOD will be removed."), FText::FromName(MissingBoneName))));
+					FFbxImporter->AddTokenizedErrorMessage(FTokenizedMessage::Create(EMessageSeverity::Warning, FText::Format(LOCTEXT("NewMeshMissingBoneFromLOD", "New mesh is missing bone '{0}' required by an LOD. LOD will be removed."), FText::FromName(MissingBoneName))), FFbxErrors::SkeletalMesh_LOD_MissingBone);
 				}
 				else
 				{
diff --git a/Engine/Source/Editor/UnrealEd/Private/Fbx/FbxFactory.cpp b/Engine/Source/Editor/UnrealEd/Private/Fbx/FbxFactory.cpp
--- a/Engine/Source/Editor/UnrealEd/Private/Fbx/FbxFactory.cpp
+++ b/Engine/Source/Editor/UnrealEd/Private/Fbx/FbxFactory.cpp
@@ -35,7 +35,7 @@
 
 #include "FbxImporter.h"
 
-
+#include "FbxErrors.h"
 #include "AssetRegistryModule.h"
 
 #define LOCTEXT_NAMESPACE "FBXFactory"
@@ -402,7 +402,7 @@ UObject* UFbxFactory::FactoryCreateBinary
 								}
 								else
 								{
-									FbxImporter->AddTokenizedErrorMessage(FTokenizedMessage::Create(EMessageSeverity::Error, LOCTEXT("FailedToImport_SkeletalMeshLOD", "Failed to import Skeletal mesh LOD.")));
+									FbxImporter->AddTokenizedErrorMessage(FTokenizedMessage::Create(EMessageSeverity::Error, LOCTEXT("FailedToImport_SkeletalMeshLOD", "Failed to import Skeletal mesh LOD.")), FFbxErrors::SkeletalMesh_LOD_FailedToImport);
 								}
 							}
 						
@@ -437,7 +437,8 @@ UObject* UFbxFactory::FactoryCreateBinary
 					// if total nodes we found is 0, we didn't find anything. 
 					if (TotalNumNodes == 0)
 					{
-						FbxImporter->AddTokenizedErrorMessage(FTokenizedMessage::Create(EMessageSeverity::Error, LOCTEXT("FailedToImport_NoMeshFoundOnRoot", "Could not find any valid mesh on the root hierarchy. If you have mesh in the sub hierarchy, please enable option of [Import Meshes In Bone Hierarchy] when import.")));
+						FbxImporter->AddTokenizedErrorMessage(FTokenizedMessage::Create(EMessageSeverity::Error, LOCTEXT("FailedToImport_NoMeshFoundOnRoot", "Could not find any valid mesh on the root hierarchy. If you have mesh in the sub hierarchy, please enable option of [Import Meshes In Bone Hierarchy] when import.")), 
+							FFbxErrors::SkeletalMesh_NoMeshFoundOnRoot);
 					}
 				}
 				else if ( ImportUI->MeshTypeToImport == FBXIT_Animation )// animation
@@ -453,22 +454,22 @@ UObject* UFbxFactory::FactoryCreateBinary
 			{
 				if (RootNodeToImport == NULL)
 				{
-					FbxImporter->AddTokenizedErrorMessage(FTokenizedMessage::Create(EMessageSeverity::Error, LOCTEXT("FailedToImport_InvalidRoot", "Could not find root node.")));
+					FbxImporter->AddTokenizedErrorMessage(FTokenizedMessage::Create(EMessageSeverity::Error, LOCTEXT("FailedToImport_InvalidRoot", "Could not find root node.")), FFbxErrors::SkeletalMesh_InvalidRoot);
 				}
 				else if (ImportUI->MeshTypeToImport == FBXIT_SkeletalMesh)
 				{
-					FbxImporter->AddTokenizedErrorMessage(FTokenizedMessage::Create(EMessageSeverity::Error, LOCTEXT("FailedToImport_InvalidBone", "Failed to find any bone hierarchy. Try to import as Rigid Mesh option enabled. ")));
+					FbxImporter->AddTokenizedErrorMessage(FTokenizedMessage::Create(EMessageSeverity::Error, LOCTEXT("FailedToImport_InvalidBone", "Failed to find any bone hierarchy. Try to import as Rigid Mesh option enabled. ")), FFbxErrors::SkeletalMesh_InvalidBone);
 				}
 				else
 				{
-					FbxImporter->AddTokenizedErrorMessage(FTokenizedMessage::Create(EMessageSeverity::Error, LOCTEXT("FailedToImport_InvalidNode", "Could not find any node.")));
+					FbxImporter->AddTokenizedErrorMessage(FTokenizedMessage::Create(EMessageSeverity::Error, LOCTEXT("FailedToImport_InvalidNode", "Could not find any node.")), FFbxErrors::SkeletalMesh_InvalidNode);
 				}
 			}
 		}
 
 		if (NewObject == NULL)
 		{
-			FbxImporter->AddTokenizedErrorMessage(FTokenizedMessage::Create(EMessageSeverity::Error, LOCTEXT("FailedToImport_NoObject", "Import failed.")));
+			FbxImporter->AddTokenizedErrorMessage(FTokenizedMessage::Create(EMessageSeverity::Error, LOCTEXT("FailedToImport_NoObject", "Import failed.")), FFbxErrors::Generic_ImportingNewObjectFailed);
 		}
 
 		FbxImporter->ReleaseScene();
diff --git a/Engine/Source/Editor/UnrealEd/Private/Fbx/FbxMainImport.cpp b/Engine/Source/Editor/UnrealEd/Private/Fbx/FbxMainImport.cpp
--- a/Engine/Source/Editor/UnrealEd/Private/Fbx/FbxMainImport.cpp
+++ b/Engine/Source/Editor/UnrealEd/Private/Fbx/FbxMainImport.cpp
@@ -196,6 +196,7 @@ void ApplyImportUIToImportOptions(UFbxImportUI* ImportUI, FBXImportOptions& InOu
 	InOutImportOptions.bOneConvexHullPerUCX = ImportUI->StaticMeshImportData->bOneConvexHullPerUCX;
 	InOutImportOptions.StaticMeshLODGroup = ImportUI->StaticMeshImportData->StaticMeshLODGroup;
 	InOutImportOptions.bImportMeshesInBoneHierarchy = ImportUI->SkeletalMeshImportData->bImportMeshesInBoneHierarchy;
+	InOutImportOptions.bImportGroupNodeAsRoot = ImportUI->SkeletalMeshImportData->bImportGroupNodeAsRoot;
 	InOutImportOptions.bCreatePhysicsAsset = ImportUI->bCreatePhysicsAsset;
 	InOutImportOptions.PhysicsAsset = ImportUI->PhysicsAsset;
 	// animation options
@@ -663,7 +664,7 @@ bool FFbxImporter::ImportFile(FString Filename)
 	else
 	{
 		ErrorMessage = ANSI_TO_TCHAR(Importer->GetStatus().GetErrorString());
-		AddTokenizedErrorMessage(FTokenizedMessage::Create(EMessageSeverity::Warning, FText::Format(LOCTEXT("FbxSkeletaLMeshimport_TriangulatingFailed", "FBX Scene Loading Failed : '{0}'"), FText::FromString(ErrorMessage))));
+		AddTokenizedErrorMessage(FTokenizedMessage::Create(EMessageSeverity::Warning, FText::Format(LOCTEXT("FbxSkeletaLMeshimport_FileLoadingFailed", "FBX Scene Loading Failed : '{0}'"), FText::FromString(ErrorMessage))), FFbxErrors::Generic_LoadingSceneFailed);
 		CleanUp();
 		Result = false;
 		CurPhase = NOTSTARTED;
@@ -1078,14 +1079,18 @@ void DumpFBXNode(FbxNode* Node)
 	if(Mesh)
 	{
 		int DeformerCount = Mesh->GetDeformerCount();
-		UE_LOG(LogFbx, Log,TEXT("Dumping Node [%s] : Total Deformer Count %d."), *NodeName, DeformerCount);
+		UE_LOG(LogFbx, Log,TEXT("================================================="));
+		UE_LOG(LogFbx, Log,TEXT("Dumping Node START [%s] : Total Deformer Count %d."), *NodeName, DeformerCount);
 		for(int i=0; i<DeformerCount; i++)
 		{
 			FbxDeformer* Deformer = Mesh->GetDeformer(i);
 			const FString DeformerName(Deformer->GetName());
 			const FString DeformerTypeName(Deformer->GetTypeName());
 			UE_LOG(LogFbx, Log,TEXT("\t[Node %d] %s (Type %s)."), i+1, *DeformerName, *DeformerTypeName);
+			UE_LOG(LogFbx, Log,TEXT("================================================="));
 		}
+
+		UE_LOG(LogFbx, Log,TEXT("Dumping Node END [%s]"), *NodeName);
 	}
 
 	for(int ChildIdx=0; ChildIdx < Node->GetChildCount(); ChildIdx++)
@@ -1175,7 +1180,8 @@ void FFbxImporter::RecursiveFindFbxSkelMesh(FbxNode* Node, TArray< TArray<FbxNod
 					FTokenizedMessage::Create(
 						EMessageSeverity::Warning, 
 						FText::Format( LOCTEXT("FBX_NoWeightsOnDeformer", "Ignoring mesh {0} because it but no weights."), FText::FromString( ANSI_TO_TCHAR(SkelMeshNode->GetName()) ) )
-					)
+					), 
+					FFbxErrors::SkeletalMesh_NoWeightsOnDeformer
 				);
 			}
 		}
@@ -1290,6 +1296,40 @@ void FFbxImporter::FillFbxSkelMeshArrayInScene(FbxNode* Node, TArray< TArray<Fbx
 		RecursiveFixSkeleton(SkeletonArray[SkelIndex], *outSkelMeshArray[SkelIndex], ImportOptions->bImportMeshesInBoneHierarchy );
 	}
 
+	// if it doesn't want group node (or null node) for root, remove them
+	/*
+	if (ImportOptions->bImportGroupNodeAsRoot == false)
+	{
+		// find the last node
+		for ( int32 SkelMeshIndex=0; SkelMeshIndex < outSkelMeshArray.Num() ; ++SkelMeshIndex )
+		{
+			auto SkelMesh = outSkelMeshArray[SkelMeshIndex];
+			if ( SkelMesh->Num() > 0 )
+			{
+				auto Node = SkelMesh->Last();
+
+				if(Node)
+				{
+					DumpFBXNode(Node);
+
+					FbxNodeAttribute* Attr = Node->GetNodeAttribute();
+					if(Attr && Attr->GetAttributeType() == FbxNodeAttribute::eNull)
+					{
+						// if root is null, just remove
+						SkelMesh->Remove(Node);
+						// SkelMesh is still valid?
+						if ( SkelMesh->Num() == 0 )
+						{
+							// we remove this from outSkelMeshArray
+							outSkelMeshArray.RemoveAt(SkelMeshIndex);
+							--SkelMeshIndex;
+						}
+					}
+				}
+			}
+		}
+	}*/
+
 	SkeletonArray.Empty();
 	// b) find rigid mesh
 	
@@ -1461,7 +1501,7 @@ void FFbxImporter::CheckSmoothingInfo(FbxMesh* FbxMesh)
 		FbxLayer* LayerSmoothing = FbxMesh->GetLayer(0, FbxLayerElement::eSmoothing);
 		if (!LayerSmoothing)
 		{
-			AddTokenizedErrorMessage(FTokenizedMessage::Create(EMessageSeverity::Warning, LOCTEXT("Prompt_NoSmoothgroupForFBXScene", "Warning: No smoothing group information was found in this FBX scene.  Please make sure to enable the 'Export Smoothing Groups' option in the FBX Exporter plug-in before exporting the file.  Even for tools that don't support smoothing groups, the FBX Exporter will generate appropriate smoothing data at export-time so that correct vertex normals can be inferred while importing.")));
+			AddTokenizedErrorMessage(FTokenizedMessage::Create(EMessageSeverity::Warning, LOCTEXT("Prompt_NoSmoothgroupForFBXScene", "Warning: No smoothing group information was found in this FBX scene.  Please make sure to enable the 'Export Smoothing Groups' option in the FBX Exporter plug-in before exporting the file.  Even for tools that don't support smoothing groups, the FBX Exporter will generate appropriate smoothing data at export-time so that correct vertex normals can be inferred while importing.")), FFbxErrors::Generic_Mesh_NoSmoothingGroup);
 		}
 	}
 }
diff --git a/Engine/Source/Editor/UnrealEd/Private/Fbx/FbxSkeletalMeshImport.cpp b/Engine/Source/Editor/UnrealEd/Private/Fbx/FbxSkeletalMeshImport.cpp
--- a/Engine/Source/Editor/UnrealEd/Private/Fbx/FbxSkeletalMeshImport.cpp
+++ b/Engine/Source/Editor/UnrealEd/Private/Fbx/FbxSkeletalMeshImport.cpp
@@ -53,6 +53,8 @@
 #include "MessageLogModule.h"
 #include "ComponentReregisterContext.h"
 
+#include "FbxErrors.h"
+
 #define LOCTEXT_NAMESPACE "FBXImpoter"
 
 using namespace UnFbx;
@@ -769,7 +771,7 @@ bool UnFbx::FFbxImporter::ImportBone(TArray<FbxNode*>& NodeArray, FSkeletalMeshI
 
 		if ( NbPoses != Default_NbPoses)
 		{
-			AddTokenizedErrorMessage(FTokenizedMessage::Create(EMessageSeverity::Warning, LOCTEXT("FbxSkeletaLMeshimport_SceneMissingBinding", "The imported scene has no initial binding position (Bind Pose) for the skin. The plug-in will compute one automatically. However, it may create unexpected results.")));
+			AddTokenizedErrorMessage(FTokenizedMessage::Create(EMessageSeverity::Warning, LOCTEXT("FbxSkeletaLMeshimport_SceneMissingBinding", "The imported scene has no initial binding position (Bind Pose) for the skin. The plug-in will compute one automatically. However, it may create unexpected results.")), FFbxErrors::SkeletalMesh_NoBindPoseInScene);
 		}
 
 		//
@@ -792,7 +794,7 @@ bool UnFbx::FFbxImporter::ImportBone(TArray<FbxNode*>& NodeArray, FSkeletalMeshI
 
 		if (ClusterArray.Num() == 0)
 		{
-			AddTokenizedErrorMessage(FTokenizedMessage::Create(EMessageSeverity::Warning, LOCTEXT("FbxSkeletaLMeshimport_NoAssociatedCluster", "No associated clusters")));
+			AddTokenizedErrorMessage(FTokenizedMessage::Create(EMessageSeverity::Warning, LOCTEXT("FbxSkeletaLMeshimport_NoAssociatedCluster", "No associated clusters")), FFbxErrors::SkeletalMesh_NoAssociatedCluster);
 			return false;
 		}
 
@@ -833,7 +835,7 @@ bool UnFbx::FFbxImporter::ImportBone(TArray<FbxNode*>& NodeArray, FSkeletalMeshI
 	// if no bond is found
 	if(SortedLinks.Num() == 0)
 	{
-		AddTokenizedErrorMessage(FTokenizedMessage::Create(EMessageSeverity::Warning, FText::Format(LOCTEXT("FbxSkeletaLMeshimport_NoBone", "'{0}' has no bones"), FText::FromString(NodeArray[0]->GetName()))));
+		AddTokenizedErrorMessage(FTokenizedMessage::Create(EMessageSeverity::Warning, FText::Format(LOCTEXT("FbxSkeletaLMeshimport_NoBone", "'{0}' has no bones"), FText::FromString(NodeArray[0]->GetName()))), FFbxErrors::SkeletalMesh_NoBoneFound);
 		return false;
 	}
 
@@ -841,7 +843,7 @@ bool UnFbx::FFbxImporter::ImportBone(TArray<FbxNode*>& NodeArray, FSkeletalMeshI
 	if (!bUseTime0AsRefPose && PoseArray.GetCount() == 0)
 	{
 		// add to tokenized error message
-		AddTokenizedErrorMessage(FTokenizedMessage::Create(EMessageSeverity::Warning, LOCTEXT("FbxSkeletaLMeshimport_InvalidBindPose", "Could not find the bind pose.  It will use time 0 as bind pose.")));
+		AddTokenizedErrorMessage(FTokenizedMessage::Create(EMessageSeverity::Warning, LOCTEXT("FbxSkeletaLMeshimport_InvalidBindPose", "Could not find the bind pose.  It will use time 0 as bind pose.")), FFbxErrors::SkeletalMesh_InvalidBindPose);
 		bUseTime0AsRefPose = true;
 	}
 
@@ -860,7 +862,7 @@ bool UnFbx::FFbxImporter::ImportBone(TArray<FbxNode*>& NodeArray, FSkeletalMeshI
 			{
 				FString RawBoneName = ANSI_TO_TCHAR(Link->GetName());
 				AddTokenizedErrorMessage(FTokenizedMessage::Create(EMessageSeverity::Error, FText::Format(LOCTEXT("Error_DuplicateBoneName", "Error, Could not import {0}.\nDuplicate bone name found ('{1}'). Each bone must have a unique name."),
-					FText::FromString(NodeArray[0]->GetName()), FText::FromString(RawBoneName))));
+					FText::FromString(NodeArray[0]->GetName()), FText::FromString(RawBoneName))), FFbxErrors::SkeletalMesh_DuplicateBones);
 				return false;
 			}
 		}
@@ -909,7 +911,7 @@ bool UnFbx::FFbxImporter::ImportBone(TArray<FbxNode*>& NodeArray, FSkeletalMeshI
 			++NumberOfRoot;
 			if (NumberOfRoot > 1)
 			{
-				AddTokenizedErrorMessage(FTokenizedMessage::Create(EMessageSeverity::Error, LOCTEXT("MultipleRootsFound", "Multiple roots are found in the bone hierarchy. We only support single root bone.")));
+				AddTokenizedErrorMessage(FTokenizedMessage::Create(EMessageSeverity::Error, LOCTEXT("MultipleRootsFound", "Multiple roots are found in the bone hierarchy. We only support single root bone.")), FFbxErrors::SkeletalMesh_MultipleRoots);
 				return false;
 			}
 		}
@@ -1038,7 +1040,7 @@ bool UnFbx::FFbxImporter::ImportBone(TArray<FbxNode*>& NodeArray, FSkeletalMeshI
 	
 	if(bAnyLinksNotInBindPose)
 	{
-		AddTokenizedErrorMessage(FTokenizedMessage::Create(EMessageSeverity::Warning, FText::Format(LOCTEXT("FbxSkeletaLMeshimport_BonesAreMissingFromBindPose", "Warning: The following bones are missing from the bind pose.  This can happen for bones that are not vert weighted.  If they are not in the correct orientation after importing, please set the \"Use T0 as ref pose\" option or add them to the bind pose and reimport the skeletal mesh \n\n{0}"), FText::FromString(LinksWithoutBindPoses))));
+		AddTokenizedErrorMessage(FTokenizedMessage::Create(EMessageSeverity::Warning, FText::Format(LOCTEXT("FbxSkeletaLMeshimport_BonesAreMissingFromBindPose", "Warning: The following bones are missing from the bind pose.  This can happen for bones that are not vert weighted.  If they are not in the correct orientation after importing, please set the \"Use T0 as ref pose\" option or add them to the bind pose and reimport the skeletal mesh \n\n{0}"), FText::FromString(LinksWithoutBindPoses))), FFbxErrors::SkeletalMesh_BonesAreMissingFromBindPose);
 	}
 	
 	return true;
@@ -1059,7 +1061,7 @@ USkeletalMesh* UnFbx::FFbxImporter::ImportSkeletalMesh(UObject* InParent, TArray
 
 	if( !FbxMesh )
 	{
-		AddTokenizedErrorMessage(FTokenizedMessage::Create(EMessageSeverity::Warning, FText::Format(LOCTEXT("FbxSkeletaLMeshimport_NodeInvalidSkeletalMesh", "Fbx node: '{0}' is not a valid skeletal mesh"), FText::FromString(Node->GetName()))));
+		AddTokenizedErrorMessage(FTokenizedMessage::Create(EMessageSeverity::Warning, FText::Format(LOCTEXT("FbxSkeletaLMeshimport_NodeInvalidSkeletalMesh", "Fbx node: '{0}' is not a valid skeletal mesh"), FText::FromString(Node->GetName()))), FFbxErrors::Generic_Mesh_MeshNotFound);
 		return NULL;
 	}
 	if (FbxMesh->GetDeformerCount(FbxDeformer::eSkin) == 0)
@@ -1093,7 +1095,7 @@ USkeletalMesh* UnFbx::FFbxImporter::ImportSkeletalMesh(UObject* InParent, TArray
 		// if any other object exists, we can't import with this name
 		else if (ExistingObject)
 		{
-			AddTokenizedErrorMessage(FTokenizedMessage::Create(EMessageSeverity::Warning, FText::Format(LOCTEXT("FbxSkeletaLMeshimport_OverlappingName", "Same name but different class: '{0}' already exists"), FText::FromString(ExistingObject->GetName()))));
+			AddTokenizedErrorMessage(FTokenizedMessage::Create(EMessageSeverity::Warning, FText::Format(LOCTEXT("FbxSkeletaLMeshimport_OverlappingName", "Same name but different class: '{0}' already exists"), FText::FromString(ExistingObject->GetName()))), FFbxErrors::Generic_SameNameAssetExists);
 			return NULL;
 		}
 	}
@@ -1120,7 +1122,7 @@ USkeletalMesh* UnFbx::FFbxImporter::ImportSkeletalMesh(UObject* InParent, TArray
 	// from popping up again on each additional pass.  
 	if ( !ImportBone(NodeArray, *SkelMeshImportDataPtr, SortedLinkArray,  bDiffPose, (FbxShapeArray != NULL), bUseTime0AsRefPose))
 	{
-		AddTokenizedErrorMessage(FTokenizedMessage::Create(EMessageSeverity::Error, LOCTEXT("FbxSkeletaLMeshimport_MultipleRootFound", "Multiple roots found")));
+		AddTokenizedErrorMessage(FTokenizedMessage::Create(EMessageSeverity::Error, LOCTEXT("FbxSkeletaLMeshimport_MultipleRootFound", "Multiple roots found")), FFbxErrors::SkeletalMesh_MultipleRoots);
 		// I can't delete object here since this is middle of import
 		// but I can move to transient package, and GC will automatically collect it
 		SkeletalMesh->ClearFlags( RF_Standalone );
@@ -1262,7 +1264,7 @@ USkeletalMesh* UnFbx::FFbxImporter::ImportSkeletalMesh(UObject* InParent, TArray
 		{
 			for (auto Message : WarningMessages)
 			{
-				AddTokenizedErrorMessage(FTokenizedMessage::Create(EMessageSeverity::Error, Message));
+				AddTokenizedErrorMessage(FTokenizedMessage::Create(EMessageSeverity::Error, Message), FFbxErrors::SkeletalMesh_BuildingSkeletalMesh);
 			}
 
 			SkeletalMesh->MarkPendingKill();
@@ -1272,7 +1274,7 @@ USkeletalMesh* UnFbx::FFbxImporter::ImportSkeletalMesh(UObject* InParent, TArray
 		{
 			for (auto Message : WarningMessages)
 			{
-				AddTokenizedErrorMessage(FTokenizedMessage::Create(EMessageSeverity::Warning, Message));
+				AddTokenizedErrorMessage(FTokenizedMessage::Create(EMessageSeverity::Warning, Message), FFbxErrors::SkeletalMesh_BuildingSkeletalMesh);
 			}
 		}
 
@@ -1320,7 +1322,7 @@ USkeletalMesh* UnFbx::FFbxImporter::ImportSkeletalMesh(UObject* InParent, TArray
 				UPhysicsAsset * NewPhysicsAsset = CreateAsset<UPhysicsAsset>(InParent->GetName(), ObjectName, true); 
 				if (!NewPhysicsAsset)
 				{
-					AddTokenizedErrorMessage(FTokenizedMessage::Create(EMessageSeverity::Warning, FText::Format(LOCTEXT("FbxSkeletaLMeshimport_CouldNotCreatePhysicsAsset", "Could not create Physics Asset ('{0}') for '{1}'"), FText::FromString(ObjectName), FText::FromString(SkeletalMesh->GetName()))));
+					AddTokenizedErrorMessage(FTokenizedMessage::Create(EMessageSeverity::Warning, FText::Format(LOCTEXT("FbxSkeletaLMeshimport_CouldNotCreatePhysicsAsset", "Could not create Physics Asset ('{0}') for '{1}'"), FText::FromString(ObjectName), FText::FromString(SkeletalMesh->GetName()))), FFbxErrors::SkeletalMesh_FailedToCreatePhyscisAsset);
 				}
 				else
 				{
@@ -1330,7 +1332,7 @@ USkeletalMesh* UnFbx::FFbxImporter::ImportSkeletalMesh(UObject* InParent, TArray
 					bool bSuccess = FPhysicsAssetUtils::CreateFromSkeletalMesh(NewPhysicsAsset, SkeletalMesh, NewBodyData, CreationErrorMessage);
 					if (!bSuccess)
 					{
-						AddTokenizedErrorMessage(FTokenizedMessage::Create(EMessageSeverity::Warning, CreationErrorMessage));
+						AddTokenizedErrorMessage(FTokenizedMessage::Create(EMessageSeverity::Warning, CreationErrorMessage), FFbxErrors::SkeletalMesh_FailedToCreatePhyscisAsset);
 						// delete the asset since we could not have create physics asset
 						TArray<UObject*> ObjectsToDelete;
 						ObjectsToDelete.Add(NewPhysicsAsset);
@@ -1360,8 +1362,8 @@ USkeletalMesh* UnFbx::FFbxImporter::ImportSkeletalMesh(UObject* InParent, TArray
 				// if not skeleton, we're done, we can't create skeleton with same name
 				// @todo in the future, we'll allow them to rename
 				if (!Skeleton)
-			{
-					AddTokenizedErrorMessage(FTokenizedMessage::Create(EMessageSeverity::Error, FText::Format(LOCTEXT("FbxSkeletaLMeshimport_SkeletonRecreateError", "'{0}' already exists. It fails to recreate it."), FText::FromString(ObjectName))));
+				{
+					AddTokenizedErrorMessage(FTokenizedMessage::Create(EMessageSeverity::Error, FText::Format(LOCTEXT("FbxSkeletaLMeshimport_SkeletonRecreateError", "'{0}' already exists. It fails to recreate it."), FText::FromString(ObjectName))), FFbxErrors::SkeletalMesh_SkeletonRecreateError);
 					return SkeletalMesh;
 				}
 			}
@@ -1410,7 +1412,7 @@ UObject* UnFbx::FFbxImporter::CreateAssetOfClass(UClass* AssetClass, FString Par
 	if ((Object != NULL) && (Object->GetClass() != AssetClass))
 	{
 		UnFbx::FFbxImporter* FFbxImporter = UnFbx::FFbxImporter::GetInstance();
-		FFbxImporter->AddTokenizedErrorMessage(FTokenizedMessage::Create(EMessageSeverity::Error, LOCTEXT("Error_AssetExist", "Asset with same name exists. Can't overwrite another asset")));
+		FFbxImporter->AddTokenizedErrorMessage(FTokenizedMessage::Create(EMessageSeverity::Error, LOCTEXT("Error_AssetExist", "Asset with same name exists. Can't overwrite another asset")), FFbxErrors::Generic_SameNameAssetExists);
 		return NULL;
 	}
 
@@ -1423,7 +1425,7 @@ UObject* UnFbx::FFbxImporter::CreateAssetOfClass(UClass* AssetClass, FString Par
 		// 			return NULL;
 		// 		}
 		UnFbx::FFbxImporter* FFbxImporter = UnFbx::FFbxImporter::GetInstance();
-		FFbxImporter->AddTokenizedErrorMessage(FTokenizedMessage::Create(EMessageSeverity::Warning, FText::Format(LOCTEXT("FbxSkeletaLMeshimport_SameNameExist", "Asset with the name ('{0}') exists. Overwriting..."), FText::FromString(ParentPath))));
+		FFbxImporter->AddTokenizedErrorMessage(FTokenizedMessage::Create(EMessageSeverity::Warning, FText::Format(LOCTEXT("FbxSkeletaLMeshimport_SameNameExist", "Asset with the name ('{0}') exists. Overwriting..."), FText::FromString(ParentPath))), FFbxErrors::Generic_SameNameAssetOverriding);
 	}
 
 	if (Object == NULL)
@@ -1578,7 +1580,7 @@ USkeletalMesh* UnFbx::FFbxImporter::ReimportSkeletalMesh(USkeletalMesh* Mesh, UF
 	else
 	{
 		// no mesh found in the FBX file
-		AddTokenizedErrorMessage(FTokenizedMessage::Create(EMessageSeverity::Warning, FText::Format(LOCTEXT("FbxSkeletaLMeshimport_NoFBXMeshMatch", "No FBX mesh matches the Unreal mesh '{0}'."), FText::FromString(Mesh->GetName()))));
+		AddTokenizedErrorMessage(FTokenizedMessage::Create(EMessageSeverity::Warning, FText::Format(LOCTEXT("FbxSkeletaLMeshimport_NoFBXMeshMatch", "No FBX mesh matches the Unreal mesh '{0}'."), FText::FromString(Mesh->GetName()))), FFbxErrors::Generic_Mesh_MeshNotFound);
 	}
 
 	return NewMesh;
@@ -1655,14 +1657,14 @@ bool UnFbx::FFbxImporter::FillSkelMeshImporterFromFbx( FSkeletalMeshImportData&
 	FbxLayer* BaseLayer = Mesh->GetLayer(0);
 	if (BaseLayer == NULL)
 	{
-		AddTokenizedErrorMessage(FTokenizedMessage::Create(EMessageSeverity::Error, LOCTEXT("FbxSkeletaLMeshimport_NoGeometry", "There is no geometry information in mesh")));
+		AddTokenizedErrorMessage(FTokenizedMessage::Create(EMessageSeverity::Error, LOCTEXT("FbxSkeletaLMeshimport_NoGeometry", "There is no geometry information in mesh")), FFbxErrors::Generic_Mesh_NoGeometry);
 		return false;
 	}
 
 	// Do some checks before proceeding, check to make sure the number of bones does not exceed the maximum supported
 	if(SortedLinks.Num() > MAX_BONES)
 	{
-		AddTokenizedErrorMessage(FTokenizedMessage::Create(EMessageSeverity::Error, FText::Format(LOCTEXT("FbxSkeletaLMeshimport_ExceedsMaxBoneCount", "'{0}' mesh has '{1}' bones which exceeds the maximum allowed bone count of {2}."), FText::FromString(Node->GetName()), FText::AsNumber(SortedLinks.Num()), FText::AsNumber(MAX_BONES))));
+		AddTokenizedErrorMessage(FTokenizedMessage::Create(EMessageSeverity::Error, FText::Format(LOCTEXT("FbxSkeletaLMeshimport_ExceedsMaxBoneCount", "'{0}' mesh has '{1}' bones which exceeds the maximum allowed bone count of {2}."), FText::FromString(Node->GetName()), FText::AsNumber(SortedLinks.Num()), FText::AsNumber(MAX_BONES))), FFbxErrors::SkeletalMesh_ExceedsMaxBoneCount);
 		return false;
 	}
 
@@ -1819,7 +1821,7 @@ bool UnFbx::FFbxImporter::FillSkelMeshImporterFromFbx( FSkeletalMeshImportData&
 		}
 		else
 		{
-			AddTokenizedErrorMessage(FTokenizedMessage::Create(EMessageSeverity::Warning, FText::Format(LOCTEXT("FbxSkeletaLMeshimport_TriangulatingFailed", "Unable to triangulate mesh '{0}'. Check detail for Ouput Log."), FText::FromString(Node->GetName()))));
+			AddTokenizedErrorMessage(FTokenizedMessage::Create(EMessageSeverity::Warning, FText::Format(LOCTEXT("FbxSkeletaLMeshimport_TriangulatingFailed", "Unable to triangulate mesh '{0}'. Check detail for Ouput Log."), FText::FromString(Node->GetName()))), FFbxErrors::Generic_Mesh_TriangulationFailed);
 			return false;
 		}
 	}
@@ -1843,7 +1845,7 @@ bool UnFbx::FFbxImporter::FillSkelMeshImporterFromFbx( FSkeletalMeshImportData&
 	}
 	else
 	{
-		AddTokenizedErrorMessage(FTokenizedMessage::Create(EMessageSeverity::Warning, FText::Format(LOCTEXT("FbxSkeletaLMeshimport_NoUVSet", "Mesh '{0}' has no UV set. Creating a default set."), FText::FromString(Node->GetName()))));
+		AddTokenizedErrorMessage(FTokenizedMessage::Create(EMessageSeverity::Warning, FText::Format(LOCTEXT("FbxSkeletaLMeshimport_NoUVSet", "Mesh '{0}' has no UV set. Creating a default set."), FText::FromString(Node->GetName()))), FFbxErrors::SkeletalMesh_NoUVSet);
 	}
 	
 	LayerCount = Mesh->GetLayerCount();
@@ -1892,7 +1894,7 @@ bool UnFbx::FFbxImporter::FillSkelMeshImporterFromFbx( FSkeletalMeshImportData&
 		{
 			if (!GeometryConverter->ComputePolygonSmoothingFromEdgeSmoothing(Mesh))
 			{
-				AddTokenizedErrorMessage(FTokenizedMessage::Create(EMessageSeverity::Warning, FText::Format(LOCTEXT("FbxSkeletaLMeshimport_ConvertSmoothingGroupFailed", "Unable to fully convert the smoothing groups for mesh '{0}'"), FText::FromString(Mesh->GetName()))));
+				AddTokenizedErrorMessage(FTokenizedMessage::Create(EMessageSeverity::Warning, FText::Format(LOCTEXT("FbxSkeletaLMeshimport_ConvertSmoothingGroupFailed", "Unable to fully convert the smoothing groups for mesh '{0}'"), FText::FromString(Mesh->GetName()))), FFbxErrors::Generic_Mesh_ConvertSmoothingGroupFailed);
 				bSmoothingAvailable = false;
 			}
 		}
@@ -2024,7 +2026,7 @@ bool UnFbx::FFbxImporter::FillSkelMeshImporterFromFbx( FSkeletalMeshImportData&
 				}
 				else
 				{
-					AddTokenizedErrorMessage(FTokenizedMessage::Create(EMessageSeverity::Warning, FText::Format(LOCTEXT("FbxSkeletaLMeshimport_Unsupportingsmoothinggroup", "Unsupported Smoothing group mapping mode on mesh '{0}'"), FText::FromString(Mesh->GetName()))));
+					AddTokenizedErrorMessage(FTokenizedMessage::Create(EMessageSeverity::Warning, FText::Format(LOCTEXT("FbxSkeletaLMeshimport_Unsupportingsmoothinggroup", "Unsupported Smoothing group mapping mode on mesh '{0}'"), FText::FromString(Mesh->GetName()))), FFbxErrors::Generic_Mesh_UnsupportingSmoothingGroup);
 				}
 			}
 		}
@@ -2104,7 +2106,7 @@ bool UnFbx::FFbxImporter::FillSkelMeshImporterFromFbx( FSkeletalMeshImportData&
 						int32 Index = LayerElementMaterial->GetIndexArray().GetAt(LocalIndex);							
 						if (!MaterialMapping.IsValidIndex(Index))
 						{
-							AddTokenizedErrorMessage(FTokenizedMessage::Create(EMessageSeverity::Warning, LOCTEXT("FbxSkeletaLMeshimport_MaterialIndexInconsistency", "Face material index inconsistency - forcing to 0")));
+							AddTokenizedErrorMessage(FTokenizedMessage::Create(EMessageSeverity::Warning, LOCTEXT("FbxSkeletaLMeshimport_MaterialIndexInconsistency", "Face material index inconsistency - forcing to 0")), FFbxErrors::Generic_Mesh_MaterialIndexInconsistency);
 						}
 						else
 						{
@@ -2119,7 +2121,7 @@ bool UnFbx::FFbxImporter::FillSkelMeshImporterFromFbx( FSkeletalMeshImportData&
 			// because we don't import material for morph, so the ImportData.Materials contains zero material
 			if ( !FbxShape && (Triangle.MatIndex < 0 ||  Triangle.MatIndex >= FbxMaterials.Num() ) )
 			{
-				AddTokenizedErrorMessage(FTokenizedMessage::Create(EMessageSeverity::Warning, LOCTEXT("FbxSkeletaLMeshimport_MaterialIndexInconsistency", "Face material index inconsistency - forcing to 0")));
+				AddTokenizedErrorMessage(FTokenizedMessage::Create(EMessageSeverity::Warning, LOCTEXT("FbxSkeletaLMeshimport_MaterialIndexInconsistency", "Face material index inconsistency - forcing to 0")), FFbxErrors::Generic_Mesh_MaterialIndexInconsistency);
 				Triangle.MatIndex = 0;
 			}
 		}
@@ -2429,7 +2431,7 @@ bool UnFbx::FFbxImporter::ImportSkeletalMeshLOD(USkeletalMesh* InSkeletalMesh, U
 	if (InSkeletalMesh->RefSkeleton.GetBoneName(0) != BaseSkeletalMesh->RefSkeleton.GetBoneName(0))
 	{
 		AddTokenizedErrorMessage(FTokenizedMessage::Create(EMessageSeverity::Error, FText::Format(LOCTEXT("LODRootNameIncorrect", "Root bone in LOD is '{0}' instead of '{1}'.\nImport failed."),
-			FText::FromName(InSkeletalMesh->RefSkeleton.GetBoneName(0)), FText::FromName(BaseSkeletalMesh->RefSkeleton.GetBoneName(0)))));
+			FText::FromName(InSkeletalMesh->RefSkeleton.GetBoneName(0)), FText::FromName(BaseSkeletalMesh->RefSkeleton.GetBoneName(0)))), FFbxErrors::SkeletalMesh_LOD_RootNameIncorrect);
 
 		return false;
 	}
@@ -2444,7 +2446,7 @@ bool UnFbx::FFbxImporter::ImportSkeletalMeshLOD(USkeletalMesh* InSkeletalMesh, U
 		{
 			// If we could not find the bone from this LOD in base mesh - we fail.
 			AddTokenizedErrorMessage(FTokenizedMessage::Create(EMessageSeverity::Error, FText::Format(LOCTEXT("LODBoneDoesNotMatch", "Bone '{0}' not found in base SkeletalMesh '{1}'.\nImport failed."),
-				FText::FromName(LODBoneName), FText::FromString(BaseSkeletalMesh->GetName()))));
+				FText::FromName(LODBoneName), FText::FromString(BaseSkeletalMesh->GetName()))), FFbxErrors::SkeletalMesh_LOD_BonesDoNotMatch);
 
 			return false;
 		}
@@ -2461,7 +2463,7 @@ bool UnFbx::FFbxImporter::ImportSkeletalMeshLOD(USkeletalMesh* InSkeletalMesh, U
 			{
 				// If bone has different parents, display an error and don't allow import.
 				AddTokenizedErrorMessage(FTokenizedMessage::Create(EMessageSeverity::Error, FText::Format(LOCTEXT("LODBoneHasIncorrectParent", "Bone '{0}' in LOD has parent '{1}' instead of '{2}'"),
-					FText::FromName(LODBoneName), FText::FromName(LODParentName), FText::FromName(BaseParentName))));
+					FText::FromName(LODBoneName), FText::FromName(LODParentName), FText::FromName(BaseParentName))), FFbxErrors::SkeletalMesh_LOD_IncorrectParent);
 
 				return false;
 			}
@@ -2480,7 +2482,7 @@ bool UnFbx::FFbxImporter::ImportSkeletalMeshLOD(USkeletalMesh* InSkeletalMesh, U
 		{
 			if (NewLODModel.Chunks[ChunkIndex].SoftVertices.Num() > 0)
 			{
-				AddTokenizedErrorMessage(FTokenizedMessage::Create(EMessageSeverity::Warning, FText(LOCTEXT("LODHasSoftVertices", "Warning: The mesh LOD you are importing has some vertices with more than one influence."))));
+				AddTokenizedErrorMessage(FTokenizedMessage::Create(EMessageSeverity::Warning, FText(LOCTEXT("LODHasSoftVertices", "Warning: The mesh LOD you are importing has some vertices with more than one influence."))), FFbxErrors::SkeletalMesh_LOD_HasSoftVerts);
 			}
 		}
 	}
@@ -2501,7 +2503,7 @@ bool UnFbx::FFbxImporter::ImportSkeletalMeshLOD(USkeletalMesh* InSkeletalMesh, U
 			if (SocketBoneIndex == INDEX_NONE)
 			{
 				AddTokenizedErrorMessage(FTokenizedMessage::Create(EMessageSeverity::Error, FText::Format(LOCTEXT("LODMissingSocketBone", "This LOD is missing bone '{0}' used by socket '{1}'.\nAborting import."),
-					FText::FromName(Socket->BoneName), FText::FromName(Socket->SocketName))));
+					FText::FromName(Socket->BoneName), FText::FromName(Socket->SocketName))), FFbxErrors::SkeletalMesh_LOD_MissingSocketBone);
 
 				return false;
 			}
@@ -2777,7 +2779,8 @@ void UnFbx::FFbxImporter::ImportMorphTargetsInternal( TArray<FbxNode*>& SkelMesh
 								else
 								{
 									AddTokenizedErrorMessage(FTokenizedMessage::Create(EMessageSeverity::Error, FText::Format(FText::FromString("Could not find the {0} morphtarget for LOD {1}. \
-																													Make sure the name for morphtarget matches with LOD 0"), FText::FromString(ShapeName), FText::FromString(FString::FromInt(LODIndex)))));
+																													Make sure the name for morphtarget matches with LOD 0"), FText::FromString(ShapeName), FText::FromString(FString::FromInt(LODIndex)))), 
+																													FFbxErrors::SkeletalMesh_LOD_MissingMorphTarget);
 								}
 							}
 
@@ -2873,18 +2876,18 @@ void UnFbx::FFbxImporter::ImportFbxMorphTarget(TArray<FbxNode*> &SkelMeshNodeArr
 	}
 }
 
-void UnFbx::FFbxImporter::AddErrorMessage(const FText & Error)
-{
-	AddTokenizedErrorMessage(FTokenizedMessage::Create(EMessageSeverity::Error, Error));
-}
-
-void UnFbx::FFbxImporter::AddTokenizedErrorMessage(TSharedRef<FTokenizedMessage> Error)
+void UnFbx::FFbxImporter::AddTokenizedErrorMessage(TSharedRef<FTokenizedMessage> Error, FName FbxErrorName)
 {
 	// check to see if Logger exists, this way, we guarantee only prints to FBX import
 	// when we meant to print
 	if (Logger)
 	{
 		Logger->TokenizedErrorMessages.Add(Error);
+
+		if(FbxErrorName != NAME_None)
+		{
+			Error->AddToken(FFbxErrorToken::Create(FbxErrorName));
+		}
 	}
 	else
 	{
@@ -2905,7 +2908,7 @@ void UnFbx::FFbxImporter::FlushToTokenizedErrorMessage(EMessageSeverity::Type Se
 {
 	if (!ErrorMessage.IsEmpty())
 	{
-		AddTokenizedErrorMessage(FTokenizedMessage::Create(Severity, FText::Format(FText::FromString("{0}"), FText::FromString(ErrorMessage))));
+		AddTokenizedErrorMessage(FTokenizedMessage::Create(Severity, FText::Format(FText::FromString("{0}"), FText::FromString(ErrorMessage))), NAME_None);
 	}
 }
 
diff --git a/Engine/Source/Editor/UnrealEd/Private/Fbx/FbxSkeletalMeshImportData.cpp b/Engine/Source/Editor/UnrealEd/Private/Fbx/FbxSkeletalMeshImportData.cpp
--- a/Engine/Source/Editor/UnrealEd/Private/Fbx/FbxSkeletalMeshImportData.cpp
+++ b/Engine/Source/Editor/UnrealEd/Private/Fbx/FbxSkeletalMeshImportData.cpp
@@ -5,6 +5,7 @@
 UFbxSkeletalMeshImportData::UFbxSkeletalMeshImportData(const class FPostConstructInitializeProperties& PCIP)
 	: Super(PCIP)
 	, bImportMeshesInBoneHierarchy(true)
+	, bImportGroupNodeAsRoot(true)
 {
 }
 
diff --git a/Engine/Source/Editor/UnrealEd/Private/Fbx/FbxStaticMeshImport.cpp b/Engine/Source/Editor/UnrealEd/Private/Fbx/FbxStaticMeshImport.cpp
--- a/Engine/Source/Editor/UnrealEd/Private/Fbx/FbxStaticMeshImport.cpp
+++ b/Engine/Source/Editor/UnrealEd/Private/Fbx/FbxStaticMeshImport.cpp
@@ -45,6 +45,7 @@
 #include "FbxImporter.h"
 #include "../Private/GeomFitUtils.h"
 #include "TargetPlatform.h"
+#include "FbxErrors.h"
 
 #define LOCTEXT_NAMESPACE "FbxStaticMeshImport"
 
@@ -104,7 +105,7 @@ bool UnFbx::FFbxImporter::BuildStaticMeshFromGeometry(FbxMesh* Mesh, UStaticMesh
 	FbxLayer* BaseLayer = Mesh->GetLayer(0);
 	if (BaseLayer == NULL)
 	{
-		AddTokenizedErrorMessage(FTokenizedMessage::Create(EMessageSeverity::Error, FText::Format(LOCTEXT("Error_NoGeometryInMesh", "There is no geometry information in mesh '{0}'"), FText::FromString(Mesh->GetName()))));
+		AddTokenizedErrorMessage(FTokenizedMessage::Create(EMessageSeverity::Error, FText::Format(LOCTEXT("Error_NoGeometryInMesh", "There is no geometry information in mesh '{0}'"), FText::FromString(Mesh->GetName()))), FFbxErrors::Generic_Mesh_NoGeometry);
 		return false;
 	}
 
@@ -253,7 +254,7 @@ bool UnFbx::FFbxImporter::BuildStaticMeshFromGeometry(FbxMesh* Mesh, UStaticMesh
 		}
 		else
 		{
-			AddTokenizedErrorMessage(FTokenizedMessage::Create(EMessageSeverity::Warning, FText::Format(LOCTEXT("Error_FailedToTriangulate", "Unable to triangulate mesh '{0}'"), FText::FromString(Mesh->GetName()))));
+			AddTokenizedErrorMessage(FTokenizedMessage::Create(EMessageSeverity::Warning, FText::Format(LOCTEXT("Error_FailedToTriangulate", "Unable to triangulate mesh '{0}'"), FText::FromString(Mesh->GetName()))), FFbxErrors::Generic_Mesh_TriangulationFailed);
 			return false; // not clean, missing some dealloc
 		}
 	}
@@ -329,7 +330,7 @@ bool UnFbx::FFbxImporter::BuildStaticMeshFromGeometry(FbxMesh* Mesh, UStaticMesh
 		{
 			if (!GeometryConverter->ComputePolygonSmoothingFromEdgeSmoothing(Mesh))
 			{
-				AddTokenizedErrorMessage(FTokenizedMessage::Create(EMessageSeverity::Warning, FText::Format(LOCTEXT("Error_FailedtoConvertSmoothingGroup", "Unable to fully convert the smoothing groups for mesh '{0}'"), FText::FromString(Mesh->GetName()))));
+				AddTokenizedErrorMessage(FTokenizedMessage::Create(EMessageSeverity::Warning, FText::Format(LOCTEXT("Error_FailedtoConvertSmoothingGroup", "Unable to fully convert the smoothing groups for mesh '{0}'"), FText::FromString(Mesh->GetName()))), FFbxErrors::Generic_Mesh_ConvertSmoothingGroupFailed);
 				bSmoothingAvailable = false;
 			}
 		}
@@ -594,7 +595,7 @@ bool UnFbx::FFbxImporter::BuildStaticMeshFromGeometry(FbxMesh* Mesh, UStaticMesh
 			}
 			else
 			{
-				AddTokenizedErrorMessage(FTokenizedMessage::Create(EMessageSeverity::Warning, FText::Format(LOCTEXT("Error_UnsupportedSmoothingGroup", "Unsupported Smoothing group mapping mode on mesh  '{0}'"), FText::FromString(Mesh->GetName()))));
+				AddTokenizedErrorMessage(FTokenizedMessage::Create(EMessageSeverity::Warning, FText::Format(LOCTEXT("Error_UnsupportedSmoothingGroup", "Unsupported Smoothing group mapping mode on mesh  '{0}'"), FText::FromString(Mesh->GetName()))), FFbxErrors::Generic_Mesh_UnsupportingSmoothingGroup);
 			}
 		}
 
@@ -653,7 +654,7 @@ bool UnFbx::FFbxImporter::BuildStaticMeshFromGeometry(FbxMesh* Mesh, UStaticMesh
 
 		if (MaterialIndex >= MaterialCount + MaterialIndexOffset || MaterialIndex < 0)
 		{
-			AddTokenizedErrorMessage(FTokenizedMessage::Create(EMessageSeverity::Warning, LOCTEXT("Error_MaterialIndexInconsistency", "Face material index inconsistency - forcing to 0")));
+			AddTokenizedErrorMessage(FTokenizedMessage::Create(EMessageSeverity::Warning, LOCTEXT("Error_MaterialIndexInconsistency", "Face material index inconsistency - forcing to 0")), FFbxErrors::Generic_Mesh_MaterialIndexInconsistency);
 			MaterialIndex = 0;
 		}
 	
@@ -780,7 +781,7 @@ UStaticMesh* UnFbx::FFbxImporter::ReimportStaticMesh(UStaticMesh* Mesh, UFbxStat
 		}
 		else // no mesh found in the FBX file
 		{
-			AddTokenizedErrorMessage(FTokenizedMessage::Create(EMessageSeverity::Warning, FText::Format(LOCTEXT("Error_NoFBXMeshFound", "No FBX mesh found when reimport Unreal mesh '{0}'. The FBX file is crashed."), FText::FromString(Mesh->GetName()))));
+			AddTokenizedErrorMessage(FTokenizedMessage::Create(EMessageSeverity::Warning, FText::Format(LOCTEXT("Error_NoFBXMeshFound", "No FBX mesh found when reimport Unreal mesh '{0}'. The FBX file is crashed."), FText::FromString(Mesh->GetName()))), FFbxErrors::Generic_Mesh_MeshNotFound);
 		}
 	}
 	return NewMesh;
@@ -883,7 +884,7 @@ UStaticMesh* UnFbx::FFbxImporter::ImportStaticMeshAsSingle(UObject* InParent, TA
 		else
 		{
 			// failed to delete
-			AddTokenizedErrorMessage(FTokenizedMessage::Create(EMessageSeverity::Error, FText::Format(LOCTEXT("ContentBrowser_CannotDeleteReferenced", "{0} wasn't created.\n\nThe asset is referenced by other content."), FText::FromString(MeshName))));
+			AddTokenizedErrorMessage(FTokenizedMessage::Create(EMessageSeverity::Error, FText::Format(LOCTEXT("ContentBrowser_CannotDeleteReferenced", "{0} wasn't created.\n\nThe asset is referenced by other content."), FText::FromString(MeshName))), FFbxErrors::Generic_CannotDeleteReferenced);
 			return NULL;
 		}
 
@@ -1002,7 +1003,8 @@ UStaticMesh* UnFbx::FFbxImporter::ImportStaticMeshAsSingle(UObject* InParent, TA
 				FText::Format(LOCTEXT("Error_TooManyMaterials", "StaticMesh has too many({1}) materials. Clamping materials to {0} which may produce unexpected results. Break apart your mesh into multiple pieces to fix this."),
 				FText::AsNumber(MAX_MESH_MATERIAL_INDEX),
 				FText::AsNumber(UniqueMaterials.Num())
-				)));
+				)), 
+				FFbxErrors::StaticMesh_TooManyMaterials);
 		}
 
 		// Sort materials based on _SkinXX in the name.
@@ -1165,7 +1167,7 @@ UStaticMesh* UnFbx::FFbxImporter::ImportStaticMeshAsSingle(UObject* InParent, TA
 			// NOTE: We don't care about missing UV sets here, just bad ones!
 			if( BadUVSets.Num() > 0 )
 			{
-				AddTokenizedErrorMessage(FTokenizedMessage::Create(EMessageSeverity::Warning, FText::Format(LOCTEXT("Error_UVSetLayoutProblem", "Warning: The light map UV set for static mesh '{0}' appears to have layout problems.  Either the triangle UVs are overlapping one another or the UV are out of bounds (0.0 - 1.0 range.)"), FText::FromString(MeshName))));
+				AddTokenizedErrorMessage(FTokenizedMessage::Create(EMessageSeverity::Warning, FText::Format(LOCTEXT("Error_UVSetLayoutProblem", "Warning: The light map UV set for static mesh '{0}' appears to have layout problems.  Either the triangle UVs are overlapping one another or the UV are out of bounds (0.0 - 1.0 range.)"), FText::FromString(MeshName))), FFbxErrors::StaticMesh_UVSetLayoutProblem);
 			}
 		}
 	}
@@ -1383,7 +1385,7 @@ bool UnFbx::FFbxImporter::ImportCollisionModels(UStaticMesh* StaticMesh, const F
 			}
 			else
 			{
-				AddTokenizedErrorMessage(FTokenizedMessage::Create(EMessageSeverity::Warning, FText::Format(LOCTEXT("Error_FailedToTriangulate", "Unable to triangulate mesh '{0}'"), FText::FromString(NodeName))));
+				AddTokenizedErrorMessage(FTokenizedMessage::Create(EMessageSeverity::Warning, FText::Format(LOCTEXT("Error_FailedToTriangulate", "Unable to triangulate mesh '{0}'"), FText::FromString(NodeName))), FFbxErrors::Generic_Mesh_TriangulationFailed);
 				return false;
 			}
 		}
diff --git a/Engine/Source/Editor/UnrealEd/Private/FbxImporter.h b/Engine/Source/Editor/UnrealEd/Private/FbxImporter.h
--- a/Engine/Source/Editor/UnrealEd/Private/FbxImporter.h
+++ b/Engine/Source/Editor/UnrealEd/Private/FbxImporter.h
@@ -125,6 +125,7 @@ struct FBXImportOptions
 	bool bPreserveSmoothingGroups;
 	bool bKeepOverlappingVertices;
 	bool bImportMeshesInBoneHierarchy;
+	bool bImportGroupNodeAsRoot;
 	bool bCreatePhysicsAsset;
 	UPhysicsAsset *PhysicsAsset;
 	// Animation option
@@ -980,8 +981,7 @@ public:
 	void SetupAnimationDataFromMesh(USkeletalMesh * SkeletalMesh, UObject * InParent, TArray<FbxNode*>& NodeArray, UFbxAnimSequenceImportData* ImportData, const FString & Filename);
 
 	/** error message handler */
-	void AddTokenizedErrorMessage(TSharedRef<FTokenizedMessage> Error);
-	void AddErrorMessage(const FText & Error);
+	void AddTokenizedErrorMessage(TSharedRef<FTokenizedMessage> Error, FName FbxErrorName );
 	void ClearTokenizedErrorMessages();
 	void FlushToTokenizedErrorMessage(enum EMessageSeverity::Type Severity);
 
diff --git a/Engine/Source/Editor/UnrealEd/Private/FbxMeshUtils.cpp b/Engine/Source/Editor/UnrealEd/Private/FbxMeshUtils.cpp
--- a/Engine/Source/Editor/UnrealEd/Private/FbxMeshUtils.cpp
+++ b/Engine/Source/Editor/UnrealEd/Private/FbxMeshUtils.cpp
@@ -16,6 +16,8 @@
 
 #include "ComponentReregisterContext.h"
 
+#include "FbxErrors.h"
+
 DEFINE_LOG_CATEGORY_STATIC(LogExportMeshUtils, Log, All);
 
 #define LOCTEXT_NAMESPACE "FbxMeshUtil"
@@ -330,7 +332,7 @@ namespace FbxMeshUtils
 				// Nothing found, error out
 				if (LODNodeList.Num() == 0)
 				{
-					FFbxImporter->AddTokenizedErrorMessage(FTokenizedMessage::Create(EMessageSeverity::Error, FText(LOCTEXT("Prompt_NoMeshFound", "No meshes were found in file."))));
+					FFbxImporter->AddTokenizedErrorMessage(FTokenizedMessage::Create(EMessageSeverity::Error, FText(LOCTEXT("Prompt_NoMeshFound", "No meshes were found in file."))), FFbxErrors::Generic_Mesh_MeshNotFound);
 
 					FFbxImporter->ReleaseScene();
 					return;
@@ -341,7 +343,7 @@ namespace FbxMeshUtils
 			if (LODLevel > BaseStaticMesh->GetNumLODs())
 			{
 				// Make sure they don't manage to select a bad LOD index
-				FFbxImporter->AddTokenizedErrorMessage(FTokenizedMessage::Create(EMessageSeverity::Warning, FText::Format(LOCTEXT("Prompt_InvalidLODIndex", "Invalid mesh LOD index {0}, as no prior LOD index exists!"), FText::AsNumber(LODLevel))));
+				FFbxImporter->AddTokenizedErrorMessage(FTokenizedMessage::Create(EMessageSeverity::Warning, FText::Format(LOCTEXT("Prompt_InvalidLODIndex", "Invalid mesh LOD index {0}, as no prior LOD index exists!"), FText::AsNumber(LODLevel))), FFbxErrors::Generic_Mesh_LOD_InvalidIndex);
 			}
 			else
 			{
@@ -407,7 +409,7 @@ namespace FbxMeshUtils
 			if ( !FFbxImporter->ImportFromFile( *Filename, FPaths::GetExtension( Filename ) ) )
 			{
 				// Log the error message and fail the import.
-				FFbxImporter->AddTokenizedErrorMessage(FTokenizedMessage::Create(EMessageSeverity::Error, LOCTEXT("FBXImport_ParseFailed", "FBX file parsing failed.")));
+				FFbxImporter->AddTokenizedErrorMessage(FTokenizedMessage::Create(EMessageSeverity::Error, LOCTEXT("FBXImport_ParseFailed", "FBX file parsing failed.")), FFbxErrors::Generic_FBXFileParseFailed);
 			}
 			else
 			{
@@ -423,7 +425,7 @@ namespace FbxMeshUtils
 				// Nothing found, error out
 				if (MeshArray.Num() == 0)
 				{
-					FFbxImporter->AddTokenizedErrorMessage(FTokenizedMessage::Create(EMessageSeverity::Error, LOCTEXT("FBXImport_NoMesh", "No meshes were found in file.")));
+					FFbxImporter->AddTokenizedErrorMessage(FTokenizedMessage::Create(EMessageSeverity::Error, LOCTEXT("FBXImport_NoMesh", "No meshes were found in file.")), FFbxErrors::Generic_MeshNotFound);
 					FFbxImporter->ReleaseScene();
 					return;
 				}
@@ -464,7 +466,7 @@ namespace FbxMeshUtils
 				if (SelectedLOD > SelectedSkelMesh->LODInfo.Num())
 				{
 					// Make sure they don't manage to select a bad LOD index
-					FFbxImporter->AddTokenizedErrorMessage(FTokenizedMessage::Create(EMessageSeverity::Warning, FText::Format(LOCTEXT("FBXImport_InvalidLODIdx", "Invalid mesh LOD index {0}, no prior LOD index exists"), FText::AsNumber(SelectedLOD))));
+					FFbxImporter->AddTokenizedErrorMessage(FTokenizedMessage::Create(EMessageSeverity::Warning, FText::Format(LOCTEXT("FBXImport_InvalidLODIdx", "Invalid mesh LOD index {0}, no prior LOD index exists"), FText::AsNumber(SelectedLOD))), FFbxErrors::Generic_Mesh_LOD_InvalidIndex);
 				}
 				else
 				{
@@ -603,12 +605,12 @@ namespace FbxMeshUtils
 			if( OpenFilenames.Num() == 0)
 			{
 				UnFbx::FFbxImporter* FFbxImporter = UnFbx::FFbxImporter::GetInstance();
-				FFbxImporter->AddTokenizedErrorMessage(FTokenizedMessage::Create(EMessageSeverity::Error, LOCTEXT("NoFileSelectedForLOD", "No file was selected for the LOD.")));
+				FFbxImporter->AddTokenizedErrorMessage(FTokenizedMessage::Create(EMessageSeverity::Error, LOCTEXT("NoFileSelectedForLOD", "No file was selected for the LOD.")), FFbxErrors::Generic_Mesh_LOD_NoFileSelected);
 			}
 			else if(OpenFilenames.Num() > 1)
 			{
 				UnFbx::FFbxImporter* FFbxImporter = UnFbx::FFbxImporter::GetInstance();
-				FFbxImporter->AddTokenizedErrorMessage(FTokenizedMessage::Create(EMessageSeverity::Error, LOCTEXT("MultipleFilesSelectedForLOD", "You may only select one file for the LOD.")));
+				FFbxImporter->AddTokenizedErrorMessage(FTokenizedMessage::Create(EMessageSeverity::Error, LOCTEXT("MultipleFilesSelectedForLOD", "You may only select one file for the LOD.")), FFbxErrors::Generic_Mesh_LOD_MultipleFilesSelected);
 			}
 			else
 			{
diff --git a/Engine/Source/Editor/UnrealEd/Private/FbxOptionWindow.h b/Engine/Source/Editor/UnrealEd/Private/FbxOptionWindow.h
--- a/Engine/Source/Editor/UnrealEd/Private/FbxOptionWindow.h
+++ b/Engine/Source/Editor/UnrealEd/Private/FbxOptionWindow.h
@@ -110,6 +110,7 @@ private:
 	void SetSkeletalMesh_ReserveSmoothingGroups(ESlateCheckBoxState::Type NewType);
 	void SetSkeletalMesh_KeepOverlappingVertices(ESlateCheckBoxState::Type NewType);
 	void SetSkeletalMesh_ImportMeshesInBoneHierarchy(ESlateCheckBoxState::Type NewType);
+	void SetSkeletalMesh_ImportGroupNodeAsRoot(ESlateCheckBoxState::Type NewType);
 	void SetSkeletalMesh_CreatePhysicsAsset(ESlateCheckBoxState::Type NewType);
 	void SetStaticMesh_ImportMeshLODs(ESlateCheckBoxState::Type NewType);
 	void SetStaticMesh_LODGroup(TSharedPtr<FString> NewValue, ESelectInfo::Type SelectInfo);
diff --git a/Engine/Source/Editor/UnrealEd/Private/SkeletalMeshEdit.cpp b/Engine/Source/Editor/UnrealEd/Private/SkeletalMeshEdit.cpp
--- a/Engine/Source/Editor/UnrealEd/Private/SkeletalMeshEdit.cpp
+++ b/Engine/Source/Editor/UnrealEd/Private/SkeletalMeshEdit.cpp
@@ -10,6 +10,7 @@
 #include "ObjectTools.h"
 #include "AssetRegistryModule.h"
 #include "FbxImporter.h"
+#include "FbxErrors.h"
 
 #define LOCTEXT_NAMESPACE "SkeletalMeshEdit"
 
@@ -43,7 +44,7 @@ UAnimSequence * UEditorEngine::ImportFbxAnimation( USkeleton * Skeleton, UObject
 
 		if (!SkeletonRoot)
 		{
-			FFbxImporter->AddTokenizedErrorMessage(FTokenizedMessage::Create(EMessageSeverity::Error, FText::Format(LOCTEXT("Error_CouldNotFindFbxTrack", "Could not find needed track ({0}).\nImport failed."), FText::FromName(Skeleton->GetReferenceSkeleton().GetBoneName(0)))));
+			FFbxImporter->AddTokenizedErrorMessage(FTokenizedMessage::Create(EMessageSeverity::Error, FText::Format(LOCTEXT("Error_CouldNotFindFbxTrack", "Could not find needed track ({0}).\nImport failed."), FText::FromName(Skeleton->GetReferenceSkeleton().GetBoneName(0)))), FFbxErrors::Animation_CouldNotFindRootTrack);
 
 			FFbxImporter->ReleaseScene();
 			GWarn->EndSlowTask();
@@ -68,7 +69,8 @@ UAnimSequence * UEditorEngine::ImportFbxAnimation( USkeleton * Skeleton, UObject
 
 		if(SortedLinks.Num() == 0)
 		{
-			FFbxImporter->AddTokenizedErrorMessage(FTokenizedMessage::Create(EMessageSeverity::Warning, LOCTEXT("Error_CouldNotBuildValidSkeleton", "Could not create a valid skeleton from the import data that matches the given Skeletal Mesh.  Check the bone names of both the Skeletal Mesh for this AnimSet and the animation data you are trying to import.")));
+			FFbxImporter->AddTokenizedErrorMessage(FTokenizedMessage::Create(EMessageSeverity::Warning, LOCTEXT("Error_CouldNotBuildValidSkeleton", "Could not create a valid skeleton from the import data that matches the given Skeletal Mesh.  Check the bone names of both the Skeletal Mesh for this AnimSet and the animation data you are trying to import.")), 
+				FFbxErrors::Animation_CouldNotBuildSkeleton);
 		}
 		else
 		{
@@ -139,7 +141,7 @@ bool UEditorEngine::ReimportFbxAnimation( USkeleton * Skeleton, UAnimSequence *
 
 		if (!SkeletonRoot)
 		{
-			FbxImporter->AddTokenizedErrorMessage(FTokenizedMessage::Create(EMessageSeverity::Error, FText::Format(LOCTEXT("Error_CouldNotFindFbxTrack", "Could not find needed track ({0}).\nImport failed."), FText::FromName(Skeleton->GetReferenceSkeleton().GetBoneName(0)))));
+			FbxImporter->AddTokenizedErrorMessage(FTokenizedMessage::Create(EMessageSeverity::Error, FText::Format(LOCTEXT("Error_CouldNotFindFbxTrack", "Could not find needed track ({0}).\nImport failed."), FText::FromName(Skeleton->GetReferenceSkeleton().GetBoneName(0)))), FFbxErrors::Animation_CouldNotFindTrack);
 
 			FbxImporter->ReleaseScene();
 			GWarn->EndSlowTask();
@@ -166,7 +168,7 @@ bool UEditorEngine::ReimportFbxAnimation( USkeleton * Skeleton, UAnimSequence *
 
 		if(SortedLinks.Num() == 0)
 		{
-			FbxImporter->AddTokenizedErrorMessage(FTokenizedMessage::Create(EMessageSeverity::Warning, LOCTEXT("Error_CouldNotBuildValidSkeleton", "Could not create a valid skeleton from the import data that matches the given Skeletal Mesh.  Check the bone names of both the Skeletal Mesh for this AnimSet and the animation data you are trying to import.")));
+			FbxImporter->AddTokenizedErrorMessage(FTokenizedMessage::Create(EMessageSeverity::Warning, LOCTEXT("Error_CouldNotBuildValidSkeleton", "Could not create a valid skeleton from the import data that matches the given Skeletal Mesh.  Check the bone names of both the Skeletal Mesh for this AnimSet and the animation data you are trying to import.")), FFbxErrors::Animation_CouldNotBuildSkeleton);
 		}
 		else
 		{
@@ -199,7 +201,7 @@ bool UEditorEngine::ReimportFbxAnimation( USkeleton * Skeleton, UAnimSequence *
 			{
 				// no track is found
 
-				FbxImporter->AddTokenizedErrorMessage(FTokenizedMessage::Create(EMessageSeverity::Error, LOCTEXT("Error_CouldNotFindTrack", "Could not find needed track.")));
+				FbxImporter->AddTokenizedErrorMessage(FTokenizedMessage::Create(EMessageSeverity::Error, LOCTEXT("Error_CouldNotFindTrack", "Could not find needed track.")), FFbxErrors::Animation_CouldNotFindTrack);
 
 				FbxImporter->ReleaseScene();
 				GWarn->EndSlowTask();
@@ -292,7 +294,7 @@ bool UnFbx::FFbxImporter::IsValidAnimationData(TArray<FbxNode*>& SortedLinks, TA
 		FbxTimeSpan AnimTimeSpan = GetAnimationTimeSpan(SortedLinks[0], CurAnimStack);
 		if (AnimTimeSpan.GetDuration() <= 0)
 		{
-			AddTokenizedErrorMessage(FTokenizedMessage::Create(EMessageSeverity::Warning, FText::Format(LOCTEXT("FBXImport_ZeroLength", "Animation Stack {0} does not contain any valid key. Try different time options when import."), FText::FromString(CurAnimStack->GetName()))));
+			AddTokenizedErrorMessage(FTokenizedMessage::Create(EMessageSeverity::Warning, FText::Format(LOCTEXT("FBXImport_ZeroLength", "Animation Stack {0} does not contain any valid key. Try different time options when import."), FText::FromString(CurAnimStack->GetName()))), FFbxErrors::Animation_ZeroLength);
 			continue;
 		}
 
@@ -353,7 +355,7 @@ void UnFbx::FFbxImporter::FillAndVerifyBoneNames(USkeleton* Skeleton, TArray<Fbx
 	// make sure at least root bone matches
 	if ( OutRawBoneNames[0] != RefSkeleton.GetBoneName(0) )
 	{
-		AddTokenizedErrorMessage(FTokenizedMessage::Create(EMessageSeverity::Error, FText::Format(LOCTEXT("FBXImport_RootMatchFail", "Root bone name does not match (FBX: {0} | Skeleton: {1})"), FText::FromName(OutRawBoneNames[0]), FText::FromName(RefSkeleton.GetBoneName(0)))));
+		AddTokenizedErrorMessage(FTokenizedMessage::Create(EMessageSeverity::Error, FText::Format(LOCTEXT("FBXImport_RootMatchFail", "Root bone name does not match (FBX: {0} | Skeleton: {1})"), FText::FromName(OutRawBoneNames[0]), FText::FromName(RefSkeleton.GetBoneName(0)))), FFbxErrors::Animation_RootTrackMismatch);
 
 		return;
 	}
@@ -366,7 +368,7 @@ void UnFbx::FFbxImporter::FillAndVerifyBoneNames(USkeleton* Skeleton, TArray<Fbx
 			if (OutRawBoneNames[I] == OutRawBoneNames[J])
 			{
 				FString RawBoneName = OutRawBoneNames[J].ToString();
-				AddTokenizedErrorMessage(FTokenizedMessage::Create(EMessageSeverity::Warning, FText::Format(LOCTEXT("FBXImport_DupeBone", "Could not import {0}.\nDuplicate bone name found ('{1}'). Each bone must have a unique name."), FText::FromString(Filename), FText::FromString(RawBoneName))));
+				AddTokenizedErrorMessage(FTokenizedMessage::Create(EMessageSeverity::Warning, FText::Format(LOCTEXT("FBXImport_DupeBone", "Could not import {0}.\nDuplicate bone name found ('{1}'). Each bone must have a unique name."), FText::FromString(Filename), FText::FromString(RawBoneName))), FFbxErrors::Animation_DuplicatedBone);
 			}
 		}
 	}
@@ -386,7 +388,7 @@ void UnFbx::FFbxImporter::FillAndVerifyBoneNames(USkeleton* Skeleton, TArray<Fbx
 	if (BoneNames.IsEmpty() == false)
 	{
 		// warn user
-		AddTokenizedErrorMessage(FTokenizedMessage::Create(EMessageSeverity::Warning, FText::Format(LOCTEXT("FBXImport_MissingBone", "The following bones exist in the imported animation, but not in the Skeleton asset {0}.  Any animation on these bones will not be imported: \n\n {1}"), FText::FromString(Skeleton->GetName()), FText::FromString(BoneNames) )));
+		AddTokenizedErrorMessage(FTokenizedMessage::Create(EMessageSeverity::Warning, FText::Format(LOCTEXT("FBXImport_MissingBone", "The following bones exist in the imported animation, but not in the Skeleton asset {0}.  Any animation on these bones will not be imported: \n\n {1}"), FText::FromString(Skeleton->GetName()), FText::FromString(BoneNames) )), FFbxErrors::Animation_MissingBones);
 	}
 }
 //-------------------------------------------------------------------------
@@ -468,7 +470,7 @@ UAnimSequence * UnFbx::FFbxImporter::ImportAnimations(USkeleton * Skeleton, UObj
 	int32 ValidTakeCount = 0;
 	if (IsValidAnimationData(SortedLinks, NodeArray, ValidTakeCount) == false)
 	{
-		AddTokenizedErrorMessage(FTokenizedMessage::Create(EMessageSeverity::Warning, LOCTEXT("FBXImport_InvalidAnimationData", "This does not contain any valid animation takes.")));
+		AddTokenizedErrorMessage(FTokenizedMessage::Create(EMessageSeverity::Warning, LOCTEXT("FBXImport_InvalidAnimationData", "This does not contain any valid animation takes.")), FFbxErrors::Animation_InvalidData);
 		return NULL;
 	}
 
@@ -526,7 +528,7 @@ UAnimSequence * UnFbx::FFbxImporter::ImportAnimations(USkeleton * Skeleton, UObj
 		// if object with same name exists, warn user
 		if (Object && !DestSeq)
 		{
-			AddTokenizedErrorMessage(FTokenizedMessage::Create(EMessageSeverity::Error, LOCTEXT("Error_AssetExist", "Asset with same name exists. Can't overwrite another asset")));
+			AddTokenizedErrorMessage(FTokenizedMessage::Create(EMessageSeverity::Error, LOCTEXT("Error_AssetExist", "Asset with same name exists. Can't overwrite another asset")), FFbxErrors::Generic_SameNameAssetExists);
 			continue; // Move on to next sequence...
 		}
 
@@ -897,7 +899,7 @@ namespace AnimationTransformDebug
 						UnFbx::FFbxImporter* FFbxImporter = UnFbx::FFbxImporter::GetInstance();
 						// now print information - it doesn't match well, find out what it is
 						FFbxImporter->AddTokenizedErrorMessage(FTokenizedMessage::Create(EMessageSeverity::Warning, FText::Format(NSLOCTEXT("UnrealEd", "FBXImport_TransformError", "Imported bone transform is different from original. Please check Output Log to see detail of error. "),
-							FText::FromName(Data.BoneName), FText::AsNumber(Data.BoneIndex), FText::FromString(Data.SourceGlobalTransform[Key].ToString()), FText::FromString(GlobalTransform.ToString()))));
+							FText::FromName(Data.BoneName), FText::AsNumber(Data.BoneIndex), FText::FromString(Data.SourceGlobalTransform[Key].ToString()), FText::FromString(GlobalTransform.ToString()))), FFbxErrors::Animation_TransformError);
 
 						bShouldOutputToMessageLog = false;
 					}
@@ -1047,7 +1049,7 @@ bool UnFbx::FFbxImporter::ImportAnimation(USkeleton * Skeleton, UAnimSequence *
 				{
 					bSuccess = false;
 					AddTokenizedErrorMessage(FTokenizedMessage::Create(EMessageSeverity::Error, FText::Format(LOCTEXT("Error_InvalidTransform",
-						"Track {0} contains invalid transform. Could not import the track."), FText::FromName(BoneName))));
+						"Track {0} contains invalid transform. Could not import the track."), FText::FromName(BoneName))), FFbxErrors::Animation_TransformError);
 					break;
 				}
 
@@ -1056,7 +1058,7 @@ bool UnFbx::FFbxImporter::ImportAnimation(USkeleton * Skeleton, UAnimSequence *
 				{
 					bSuccess = false;
 					AddTokenizedErrorMessage(FTokenizedMessage::Create(EMessageSeverity::Error, FText::Format(LOCTEXT("Error_InvalidUnrealTransform",
-										"Track {0} did not yeild valid transform. Please report this to animation team."), FText::FromName(BoneName))));
+										"Track {0} did not yeild valid transform. Please report this to animation team."), FText::FromName(BoneName))), FFbxErrors::Animation_TransformError);
 					break;
 				}
 				// debug data
@@ -1088,7 +1090,7 @@ bool UnFbx::FFbxImporter::ImportAnimation(USkeleton * Skeleton, UAnimSequence *
 				{
 					bSuccess = false;
 					AddTokenizedErrorMessage(FTokenizedMessage::Create(EMessageSeverity::Error, FText::Format(LOCTEXT("Error_InvalidUnrealLocalTransform",
-										"Track {0} did not yeild valid local transform. Please report this to animation team."), FText::FromName(BoneName))));
+										"Track {0} did not yeild valid local transform. Please report this to animation team."), FText::FromName(BoneName))), FFbxErrors::Animation_TransformError);
 					break;
 				}
 
diff --git a/Engine/Source/Runtime/Core/Private/Misc/FbxErrors.cpp b/Engine/Source/Runtime/Core/Private/Misc/FbxErrors.cpp
new file mode 100644
--- /dev/null
+++ b/Engine/Source/Runtime/Core/Private/Misc/FbxErrors.cpp
@@ -0,0 +1,133 @@
+// Copyright 1998-2014 Epic Games, Inc. All Rights Reserved.
+
+#include "Core.h"
+#include "FbxErrors.h"
+	
+/** Generic */
+FName FFbxErrors::Generic_ImportingNewObjectFailed(TEXT("Generic_ImportingNewObjectFailed"));
+
+FName FFbxErrors::Generic_LoadingSceneFailed(TEXT("Generic_LoadingSceneFailed"));
+
+FName FFbxErrors::Generic_SameNameAssetExists(TEXT("Generic_SameNameAssetExists"));
+
+FName FFbxErrors::Generic_SameNameAssetOverriding(TEXT("Generic_SameNameAssetOverriding"));
+
+FName FFbxErrors::Generic_CannotDeleteReferenced(TEXT("Generic_CannotDeleteReferenced"));
+
+FName FFbxErrors::Generic_FBXFileParseFailed(TEXT("Generic_FBXFileParseFailed"));
+
+FName FFbxErrors::Generic_MeshNotFound(TEXT("Generic_MeshNotFound"));
+
+/** Mesh Generic **/
+
+FName FFbxErrors::Generic_Mesh_NoGeometry(TEXT("Generic_Mesh_NoGeometry"));
+
+FName FFbxErrors::Generic_Mesh_TriangulationFailed(TEXT("Generic_Mesh_TriangulationFailed"));
+
+FName FFbxErrors::Generic_Mesh_ConvertSmoothingGroupFailed(TEXT("Generic_Mesh_ConvertSmoothingGroupFailed"));
+
+FName FFbxErrors::Generic_Mesh_UnsupportingSmoothingGroup(TEXT("Generic_Mesh_UnsupportingSmoothingGroup"));
+
+FName FFbxErrors::Generic_Mesh_MaterialIndexInconsistency(TEXT("Generic_Mesh_MaterialIndexInconsistency"));
+
+FName FFbxErrors::Generic_Mesh_MeshNotFound(TEXT("Generic_Mesh_MeshNotFound"));
+
+FName FFbxErrors::Generic_Mesh_NoSmoothingGroup(TEXT("Generic_Mesh_NoSmoothingGroup"));
+
+FName FFbxErrors::Generic_Mesh_LOD_InvalidIndex(TEXT("Generic_Mesh_LOD_InvalidIndex"));
+
+FName FFbxErrors::Generic_Mesh_LOD_NoFileSelected(TEXT("Generic_Mesh_LOD_NoFileSelected"));
+
+FName FFbxErrors::Generic_Mesh_LOD_MultipleFilesSelected(TEXT("Generic_Mesh_LOD_MultipleFilesSelected"));
+
+/** Static Mesh **/
+FName FFbxErrors::StaticMesh_TooManyMaterials(TEXT("StaticMesh_TooManyMaterials"));
+
+FName FFbxErrors::StaticMesh_UVSetLayoutProblem(TEXT("StaticMesh_UVSetLayoutProblem"));
+
+/** SkeletalMesh **/
+FName FFbxErrors::SkeletalMesh_DifferentRoots(TEXT("SkeletalMesh_DifferentRoot"));
+
+FName FFbxErrors::SkeletalMesh_DuplicateBones(TEXT("SkeletalMesh_DuplicateBones"));
+
+FName FFbxErrors::SkeletalMesh_NoInfluences(TEXT("SkeletalMesh_NoInfluences"));
+
+FName FFbxErrors::SkeletalMesh_RestoreSortingMismatchedStrips(TEXT("SkeletalMesh_RestoreSortingMismatchedStrips"));
+
+FName FFbxErrors::SkeletalMesh_RestoreSortingNoSectionMatch(TEXT("SkeletalMesh_RestoreSortingNoSectionMatch"));
+
+FName FFbxErrors::SkeletalMesh_RestoreSortingForSectionNumber(TEXT("SkeletalMesh_RestoreSortingForSectionNumber"));
+
+FName FFbxErrors::SkeletalMesh_NoMeshFoundOnRoot(TEXT("SkeletalMesh_NoMeshFoundOnRoot"));
+
+FName FFbxErrors::SkeletalMesh_InvalidRoot(TEXT("SkeletalMesh_InvalidRoot"));
+
+FName FFbxErrors::SkeletalMesh_InvalidBone(TEXT("SkeletalMesh_InvalidBone"));
+
+FName FFbxErrors::SkeletalMesh_InvalidNode(TEXT("SkeletalMesh_InvalidNode"));
+
+FName FFbxErrors::SkeletalMesh_NoWeightsOnDeformer(TEXT("SkeletalMesh_NoWeightsOnDeformer"));
+
+FName FFbxErrors::SkeletalMesh_NoBindPoseInScene(TEXT("SkeletalMesh_NoBindPoseInScene"));
+
+FName FFbxErrors::SkeletalMesh_NoAssociatedCluster(TEXT("SkeletalMesh_NoAssociatedCluster"));
+
+FName FFbxErrors::SkeletalMesh_NoBoneFound(TEXT("SkeletalMesh_NoBoneFound"));
+
+FName FFbxErrors::SkeletalMesh_InvalidBindPose(TEXT("SkeletalMesh_InvalidBindPose"));
+
+FName FFbxErrors::SkeletalMesh_MultipleRoots(TEXT("SkeletalMesh_MultipleRoots"));
+
+FName FFbxErrors::SkeletalMesh_BonesAreMissingFromBindPose(TEXT("SkeletalMesh_BonesAreMissingFromBindPose"));
+
+FName FFbxErrors::SkeletalMesh_BuildingSkeletalMesh(TEXT("SkeletalMesh_BuildingSkeletalMesh"));
+
+FName FFbxErrors::SkeletalMesh_FailedToCreatePhyscisAsset(TEXT("SkeletalMesh_FailedToCreatePhyscisAsset"));
+
+FName FFbxErrors::SkeletalMesh_SkeletonRecreateError(TEXT("SkeletalMesh_SkeletonRecreateError"));
+
+FName FFbxErrors::SkeletalMesh_ExceedsMaxBoneCount(TEXT("SkeletalMesh_ExceedsMaxBoneCount"));
+
+FName FFbxErrors::SkeletalMesh_NoUVSet(TEXT("SkeletalMesh_NoUVSet"));
+
+FName FFbxErrors::SkeletalMesh_LOD_MissingBone(TEXT("SkeletalMesh_LOD_MissingBone"));
+
+FName FFbxErrors::SkeletalMesh_LOD_FailedToImport(TEXT("SkeletalMesh_LOD_FailedToImport"));
+
+FName FFbxErrors::SkeletalMesh_LOD_RootNameIncorrect(TEXT("SkeletalMesh_LOD_RootNameIncorrect"));
+
+FName FFbxErrors::SkeletalMesh_LOD_BonesDoNotMatch(TEXT("SkeletalMesh_LOD_BonesDoNotMatch"));
+
+FName FFbxErrors::SkeletalMesh_LOD_IncorrectParent(TEXT("SkeletalMesh_LOD_IncorrectParent"));
+
+FName FFbxErrors::SkeletalMesh_LOD_HasSoftVerts(TEXT("SkeletalMesh_LOD_HasSoftVerts"));
+
+FName FFbxErrors::SkeletalMesh_LOD_MissingSocketBone(TEXT("SkeletalMesh_LOD_MissingSocketBone"));
+
+FName FFbxErrors::SkeletalMesh_LOD_MissingMorphTarget(TEXT("SkeletalMesh_LOD_MissingMorphTarget"));
+
+/** Animation **/
+FName FFbxErrors::Animation_CouldNotFindRootTrack(TEXT("Animation_CouldNotFindRootTrack"));
+
+FName FFbxErrors::Animation_CouldNotBuildSkeleton(TEXT("Animation_CouldNotBuildSkeleton"));
+
+FName FFbxErrors::Animation_CouldNotFindTrack(TEXT("Animation_CouldNotFindTrack"));
+
+FName FFbxErrors::Animation_ZeroLength(TEXT("Animation_ZeroLength"));
+
+FName FFbxErrors::Animation_RootTrackMismatch(TEXT("Animation_RootTrackMismatch"));
+
+FName FFbxErrors::Animation_DuplicatedBone(TEXT("Animation_DuplicatedBone"));
+
+FName FFbxErrors::Animation_MissingBones(TEXT("Animation_MissingBones"));
+
+FName FFbxErrors::Animation_InvalidData(TEXT("Animation_InvalidData"));
+
+FName FFbxErrors::Animation_TransformError(TEXT("Animation_TransformError"));
+
+static const FString FbxErrorsPath = TEXT("Shared/Editor/FbxErrors");
+
+FFbxErrorToken::FFbxErrorToken(const FName& InErrorName)
+	: FDocumentationToken(FbxErrorsPath, FbxErrorsPath, InErrorName.ToString())
+{
+}
\ No newline at end of file
diff --git a/Engine/Source/Runtime/Core/Public/Misc/FbxErrors.h b/Engine/Source/Runtime/Core/Public/Misc/FbxErrors.h
new file mode 100644
--- /dev/null
+++ b/Engine/Source/Runtime/Core/Public/Misc/FbxErrors.h
@@ -0,0 +1,151 @@
+// Copyright 1998-2014 Epic Games, Inc. All Rights Reserved.
+
+#pragma once
+
+#include "TokenizedMessage.h"
+
+/**
+ * This file contains known map errors that can be referenced by name.
+ * Documentation for these errors is assumed to lie in UDN at: Engine\Documentation\Source\Shared\Editor\MapErrors
+ */
+struct CORE_API FFbxErrors
+{
+	/** Generic */
+	static FName Generic_ImportingNewObjectFailed;
+
+	static FName Generic_LoadingSceneFailed;
+
+	static FName Generic_SameNameAssetExists;
+
+	static FName Generic_SameNameAssetOverriding;
+
+	static FName Generic_CannotDeleteReferenced;
+
+	static FName Generic_FBXFileParseFailed;
+
+	static FName Generic_MeshNotFound;
+
+	/** Mesh Generic **/
+
+	static FName Generic_Mesh_NoGeometry;
+
+	static FName Generic_Mesh_TriangulationFailed;
+
+	static FName Generic_Mesh_ConvertSmoothingGroupFailed;
+
+	static FName Generic_Mesh_UnsupportingSmoothingGroup;
+
+	static FName Generic_Mesh_MaterialIndexInconsistency;
+
+	static FName Generic_Mesh_MeshNotFound;
+
+	static FName Generic_Mesh_NoSmoothingGroup;
+
+	static FName Generic_Mesh_LOD_InvalidIndex;
+
+	static FName Generic_Mesh_LOD_NoFileSelected;
+
+	static FName Generic_Mesh_LOD_MultipleFilesSelected;
+
+	/** Static Mesh **/
+	static FName StaticMesh_TooManyMaterials;
+
+	static FName StaticMesh_UVSetLayoutProblem;
+
+	/** SkeletalMesh **/
+	static FName SkeletalMesh_DifferentRoots;
+
+	static FName SkeletalMesh_DuplicateBones;
+
+	static FName SkeletalMesh_NoInfluences;
+
+	static FName SkeletalMesh_RestoreSortingMismatchedStrips;
+
+	static FName SkeletalMesh_RestoreSortingNoSectionMatch;
+
+	static FName SkeletalMesh_RestoreSortingForSectionNumber;
+
+	static FName SkeletalMesh_NoMeshFoundOnRoot;
+
+	static FName SkeletalMesh_InvalidRoot;
+
+	static FName SkeletalMesh_InvalidBone;
+
+	static FName SkeletalMesh_InvalidNode;
+
+	static FName SkeletalMesh_NoWeightsOnDeformer;
+
+	static FName SkeletalMesh_NoBindPoseInScene;
+
+	static FName SkeletalMesh_NoAssociatedCluster;
+
+	static FName SkeletalMesh_NoBoneFound;
+
+	static FName SkeletalMesh_InvalidBindPose;
+
+	static FName SkeletalMesh_MultipleRoots;
+
+	static FName SkeletalMesh_BonesAreMissingFromBindPose;
+
+	static FName SkeletalMesh_BuildingSkeletalMesh;
+
+	static FName SkeletalMesh_FailedToCreatePhyscisAsset;
+
+	static FName SkeletalMesh_SkeletonRecreateError;
+
+	static FName SkeletalMesh_ExceedsMaxBoneCount;
+
+	static FName SkeletalMesh_NoUVSet;
+
+	static FName SkeletalMesh_LOD_MissingBone;
+
+	static FName SkeletalMesh_LOD_FailedToImport;
+
+	static FName SkeletalMesh_LOD_RootNameIncorrect;
+
+	static FName SkeletalMesh_LOD_BonesDoNotMatch;
+
+	static FName SkeletalMesh_LOD_IncorrectParent;
+
+	static FName SkeletalMesh_LOD_HasSoftVerts;
+
+	static FName SkeletalMesh_LOD_MissingSocketBone;
+
+	static FName SkeletalMesh_LOD_MissingMorphTarget;
+
+	/** Animation **/
+	static FName Animation_CouldNotFindRootTrack;
+
+	static FName Animation_CouldNotBuildSkeleton;
+
+	static FName Animation_CouldNotFindTrack;
+
+	static FName Animation_ZeroLength;
+
+	static FName Animation_RootTrackMismatch;
+
+	static FName Animation_DuplicatedBone;
+
+	static FName Animation_MissingBones;
+
+	static FName Animation_InvalidData;
+
+	static FName Animation_TransformError;
+};
+
+/**
+ * Map error specific message token.
+ */
+class FFbxErrorToken : public FDocumentationToken
+{
+public:
+	/** Factory method, tokens can only be constructed as shared refs */
+	CORE_API static TSharedRef<FFbxErrorToken> Create( const FName& InErrorName )
+	{
+		return MakeShareable(new FFbxErrorToken(InErrorName));
+	}
+
+private:
+	/** Private constructor */
+	FFbxErrorToken( const FName& InErrorName );
+};
\ No newline at end of file
