#ue4 Tooltips for Blueprint Interfaces, Function Libraries, and Macro Libraries

diff --git a/Engine/Source/Editor/UnrealEd/Classes/Factories/BlueprintFunctionLibraryFactory.h b/Engine/Source/Editor/UnrealEd/Classes/Factories/BlueprintFunctionLibraryFactory.h
--- a/Engine/Source/Editor/UnrealEd/Classes/Factories/BlueprintFunctionLibraryFactory.h
+++ b/Engine/Source/Editor/UnrealEd/Classes/Factories/BlueprintFunctionLibraryFactory.h
@@ -16,6 +16,7 @@ class UBlueprintFunctionLibraryFactory : public UBlueprintFactory
 	virtual FText GetDisplayName() const override;
 	virtual FName GetNewAssetThumbnailOverride() const override;
 	virtual uint32 GetMenuCategories() const override;
+	virtual FText GetToolTip() const override;
 	virtual UObject* FactoryCreateNew(UClass* Class, UObject* InParent, FName Name, EObjectFlags Flags, UObject* Context, FFeedbackContext* Warn, FName CallingContext) override;
 	virtual bool ConfigureProperties() override;
 	virtual FString GetDefaultNewAssetName() const override;
diff --git a/Engine/Source/Editor/UnrealEd/Classes/Factories/BlueprintInterfaceFactory.h b/Engine/Source/Editor/UnrealEd/Classes/Factories/BlueprintInterfaceFactory.h
--- a/Engine/Source/Editor/UnrealEd/Classes/Factories/BlueprintInterfaceFactory.h
+++ b/Engine/Source/Editor/UnrealEd/Classes/Factories/BlueprintInterfaceFactory.h
@@ -16,6 +16,7 @@ class UBlueprintInterfaceFactory : public UFactory
 	virtual FText GetDisplayName() const override;
 	virtual FName GetNewAssetThumbnailOverride() const override;
 	virtual uint32 GetMenuCategories() const override;
+	virtual FText GetToolTip() const override;
 	virtual UObject* FactoryCreateNew(UClass* Class, UObject* InParent, FName Name, EObjectFlags Flags, UObject* Context, FFeedbackContext* Warn, FName CallingContext) override;
 	virtual UObject* FactoryCreateNew(UClass* Class, UObject* InParent, FName Name, EObjectFlags Flags, UObject* Context, FFeedbackContext* Warn) override;
 	virtual FString GetDefaultNewAssetName() const override;
diff --git a/Engine/Source/Editor/UnrealEd/Classes/Factories/BlueprintMacroFactory.h b/Engine/Source/Editor/UnrealEd/Classes/Factories/BlueprintMacroFactory.h
--- a/Engine/Source/Editor/UnrealEd/Classes/Factories/BlueprintMacroFactory.h
+++ b/Engine/Source/Editor/UnrealEd/Classes/Factories/BlueprintMacroFactory.h
@@ -16,6 +16,7 @@ class UBlueprintMacroFactory : public UBlueprintFactory
 	virtual FText GetDisplayName() const override;
 	virtual FName GetNewAssetThumbnailOverride() const override;
 	virtual uint32 GetMenuCategories() const override;
+	virtual FText GetToolTip() const override;
 	virtual UObject* FactoryCreateNew(UClass* Class, UObject* InParent, FName Name, EObjectFlags Flags, UObject* Context, FFeedbackContext* Warn, FName CallingContext) override;
 	virtual FString GetDefaultNewAssetName() const override;
 	// End of UFactory interface
diff --git a/Engine/Source/Editor/UnrealEd/Private/Factories/EditorFactories.cpp b/Engine/Source/Editor/UnrealEd/Private/Factories/EditorFactories.cpp
--- a/Engine/Source/Editor/UnrealEd/Private/Factories/EditorFactories.cpp
+++ b/Engine/Source/Editor/UnrealEd/Private/Factories/EditorFactories.cpp
@@ -6252,6 +6252,11 @@ uint32 UBlueprintMacroFactory::GetMenuCategories() const
 	return EAssetTypeCategories::Blueprint;
 }
 
+FText UBlueprintMacroFactory::GetToolTip() const
+{
+	return LOCTEXT("BlueprintMacroLibraryTooltip", "Blueprint Macro Libraries are containers of macros to be used in other blueprints. They cannot contain variables, inherit from other blueprints, or be placed in levels. Changes to macros in a Blueprint Macro Library will not take effect until client blueprints are recompiled.");
+}
+
 UObject* UBlueprintMacroFactory::FactoryCreateNew(UClass* Class, UObject* InParent, FName Name, EObjectFlags Flags, UObject* Context, FFeedbackContext* Warn, FName CallingContext)
 {
 	// Make sure we are trying to factory a blueprint, then create and init one
@@ -6304,6 +6309,11 @@ uint32 UBlueprintFunctionLibraryFactory::GetMenuCategories() const
 	return EAssetTypeCategories::Blueprint;
 }
 
+FText UBlueprintFunctionLibraryFactory::GetToolTip() const
+{
+	return LOCTEXT("BlueprintFunctionLibraryTooltip", "Blueprint Function Libraries are containers of functions to be used in other blueprints. They cannot contain variables, inherit from other blueprints, or be placed in levels. Changes to functions in a Blueprint Function Library will take effect without recompiling the client blueprints.");
+}
+
 UObject* UBlueprintFunctionLibraryFactory::FactoryCreateNew(UClass* Class, UObject* InParent, FName Name, EObjectFlags Flags, UObject* Context, FFeedbackContext* Warn, FName CallingContext)
 {
 	// Make sure we are trying to factory a blueprint, then create and init one
@@ -6359,6 +6369,11 @@ uint32 UBlueprintInterfaceFactory::GetMenuCategories() const
 	return EAssetTypeCategories::Blueprint;
 }
 
+FText UBlueprintInterfaceFactory::GetToolTip() const
+{
+	return LOCTEXT("BlueprintInterfaceTooltip", "A Blueprint Interface is a collection of one or more functions - name only, no implementation - that can be added to other Blueprints. These other Blueprints are then expected to implement the functions of the Blueprint Interface in a unique manner.");
+}
+
 UObject* UBlueprintInterfaceFactory::FactoryCreateNew(UClass* Class, UObject* InParent, FName Name, EObjectFlags Flags, UObject* Context, FFeedbackContext* Warn, FName CallingContext)
 {
 	// Make sure we are trying to factory a blueprint, then create and init one
