UE4: Consolidated physics scene and FXSystem creation logic for inactive worlds during map transitions. This should handle the general case.

diff --git a/Engine/Source/Editor/UnrealEd/Private/EditorServer.cpp b/Engine/Source/Editor/UnrealEd/Private/EditorServer.cpp
--- a/Engine/Source/Editor/UnrealEd/Private/EditorServer.cpp
+++ b/Engine/Source/Editor/UnrealEd/Private/EditorServer.cpp
@@ -2105,15 +2105,6 @@ bool UEditorEngine::Map_Load(const TCHAR* Str, FOutputDevice& Ar)
 					if (ExistingPackage)
 					{
 						ExistingWorld = UWorld::FindWorldInPackage(ExistingPackage);
-
-						// If we are using an existing inactive world, make sure it has a physics scene and FXSystem.
-						// Inactive worlds are already initialized but lack these two objects for memory reasons.
-						if ( ExistingWorld->WorldType == EWorldType::Inactive )
-						{
-							ExistingWorld->ClearWorldComponents();
-							ExistingWorld->CreatePhysicsScene();
-							ExistingWorld->CreateFXSystem();
-						}
 					}
 					else
 					{
@@ -2201,21 +2192,27 @@ bool UEditorEngine::Map_Load(const TCHAR* Str, FOutputDevice& Ar)
 				Context.SetCurrentWorld(World);
 				GWorld = World;
 
-				// Inactive worlds may have already been initialized, but still be lacking important scene information
-				if ( Context.World()->WorldType == EWorldType::Inactive )
-				{
-					Context.World()->CreatePhysicsScene();
-					Context.World()->CreateFXSystem();
+				World->WorldType = EWorldType::Editor;
 
-					// Update components so the scene is populated
-					Context.World()->UpdateWorldComponents(true, true);
-				}
+				if (World->bIsWorldInitialized)
+				{
+					// If we are using a previously initialized world, make sure it has a physics scene and FXSystem.
+					// Inactive worlds are already initialized but lack these two objects for memory reasons.
+					World->ClearWorldComponents();
 
-				Context.World()->WorldType = EWorldType::Editor;
+					if ( World->GetPhysicsScene() == nullptr )
+					{
+						World->CreatePhysicsScene();
+					}
 
-				if ( !Context.World()->bIsWorldInitialized )
+					if ( World->FXSystem == nullptr )
+					{
+						World->CreateFXSystem();
+					}
+				}
+				else
 				{
-					Context.World()->InitWorld( GetEditorWorldInitializationValues() );
+					World->InitWorld( GetEditorWorldInitializationValues() );
 				}
 
 				{
