Add default try get owner node and blueprint update event node for anim BP graph

diff --git a/Engine/Source/Editor/UnrealEd/Private/Kismet2/Kismet2.cpp b/Engine/Source/Editor/UnrealEd/Private/Kismet2/Kismet2.cpp
--- a/Engine/Source/Editor/UnrealEd/Private/Kismet2/Kismet2.cpp
+++ b/Engine/Source/Editor/UnrealEd/Private/Kismet2/Kismet2.cpp
@@ -243,7 +243,7 @@ UBlueprint* FKismetEditorUtilities::CreateBlueprint(UClass* ParentClass, UObject
 		CreateDefaultEventGraphs(NewBP);
 	}
 
-	//@TODO: ANIMREFACTOR: This kind of code should be on a per-blueprint basis; not centralized here
+	//@TODO: ANIMREFACTOR 1: This kind of code should be on a per-blueprint basis; not centralized here
 	if (UAnimBlueprint* AnimBP = Cast<UAnimBlueprint>(NewBP))
 	{
 		UAnimBlueprint* RootAnimBP = UAnimBlueprint::FindRootAnimBlueprint(AnimBP);
@@ -271,6 +271,39 @@ UBlueprint* FKismetEditorUtilities::CreateBlueprint(UClass* ParentClass, UObject
 	FKismetCompilerOptions CompileOptions;
 	Compiler.CompileBlueprint(NewBP, CompileOptions, Results);
 
+	//@TODO: ANIMREFACTOR 2: This kind of code should be on a per-blueprint basis; not centralized here
+	if(UAnimBlueprint* AnimBP = Cast<UAnimBlueprint>(NewBP))
+	{
+		// add default nodes to event graph
+		UEdGraph* Graph = NewBP->UbergraphPages[0];
+
+		if(Graph->Nodes.Num() == 0)
+		{
+			// add update event graph
+			UK2Node_Event* EventNode = NewObject<UK2Node_Event>(Graph);
+			EventNode->EventSignatureClass = UAnimInstance::StaticClass();
+			EventNode->EventSignatureName = FName(TEXT("BlueprintUpdateAnimation"));
+			EventNode->bInternalEvent = true;
+			EventNode->CreateNewGuid();
+			EventNode->PostPlacedNewNode();
+			EventNode->AllocateDefaultPins();
+
+			Graph->AddNode(EventNode);
+
+			// add try get owner node
+			UK2Node_CallFunction* GetOwnerNode = NewObject<UK2Node_CallFunction>(Graph);
+			UFunction* MakeNodeFunction = FindObject<UClass>(ANY_PACKAGE, TEXT("AnimInstance"))->FindFunctionByName(TEXT("TryGetPawnOwner"));
+			GetOwnerNode->CreateNewGuid();
+			GetOwnerNode->PostPlacedNewNode();
+			GetOwnerNode->SetFromFunction(MakeNodeFunction);
+			GetOwnerNode->AllocateDefaultPins();
+			GetOwnerNode->NodePosX = EventNode->NodePosX;
+			GetOwnerNode->NodePosY = EventNode->NodePosY + EventNode->NodeHeight + 100;
+
+			Graph->AddNode(GetOwnerNode);
+		}
+	}
+
 	// Report blueprint creation to analytics
 	if (FEngineAnalytics::IsAvailable())
 	{
