Header cleanup: Remove #include "Engine.h" from UnrealEd.h

diff --git a/Engine/Source/Editor/DetailCustomizations/Private/MathStructProxyCustomizations.h b/Engine/Source/Editor/DetailCustomizations/Private/MathStructProxyCustomizations.h
--- a/Engine/Source/Editor/DetailCustomizations/Private/MathStructProxyCustomizations.h
+++ b/Engine/Source/Editor/DetailCustomizations/Private/MathStructProxyCustomizations.h
@@ -1,7 +1,7 @@
 // Copyright 1998-2014 Epic Games, Inc. All Rights Reserved.
 
 #pragma once
-
+#include "IPropertyUtilities.h"
 #include "MathStructCustomizations.h"
 
 /** 
diff --git a/Engine/Source/Editor/UnrealEd/Classes/ActorFactories/ActorFactory.h b/Engine/Source/Editor/UnrealEd/Classes/ActorFactories/ActorFactory.h
--- a/Engine/Source/Editor/UnrealEd/Classes/ActorFactories/ActorFactory.h
+++ b/Engine/Source/Editor/UnrealEd/Classes/ActorFactories/ActorFactory.h
@@ -4,7 +4,9 @@
 #pragma once
 #include "ActorFactory.generated.h"
 
-#pragma once
+class FAssetData;
+class UBlueprint;
+class ULevel;
 
 UCLASS(collapsecategories, hidecategories=Object, editinlinenew, config=Editor, abstract, transient)
 class UNREALED_API UActorFactory : public UObject
diff --git a/Engine/Source/Editor/UnrealEd/Classes/ActorFactories/ActorFactoryAmbientSound.h b/Engine/Source/Editor/UnrealEd/Classes/ActorFactories/ActorFactoryAmbientSound.h
--- a/Engine/Source/Editor/UnrealEd/Classes/ActorFactories/ActorFactoryAmbientSound.h
+++ b/Engine/Source/Editor/UnrealEd/Classes/ActorFactories/ActorFactoryAmbientSound.h
@@ -6,7 +6,7 @@
  *
  * The base class of all ambient sound types
  */
-
+#include "ActorFactories/ActorFactory.h"
 #include "ActorFactoryAmbientSound.generated.h"
 
 UCLASS(MinimalAPI, config=Editor, collapsecategories, hidecategories=Object)
diff --git a/Engine/Source/Editor/UnrealEd/Classes/ActorFactories/ActorFactoryAnimationAsset.h b/Engine/Source/Editor/UnrealEd/Classes/ActorFactories/ActorFactoryAnimationAsset.h
--- a/Engine/Source/Editor/UnrealEd/Classes/ActorFactories/ActorFactoryAnimationAsset.h
+++ b/Engine/Source/Editor/UnrealEd/Classes/ActorFactories/ActorFactoryAnimationAsset.h
@@ -2,8 +2,11 @@
 
 
 #pragma once
+#include "ActorFactories/ActorFactorySkeletalMesh.h"
 #include "ActorFactoryAnimationAsset.generated.h"
 
+class USkeletalMesh;
+
 UCLASS(MinimalAPI, config=Editor, hidecategories=Object)
 class UActorFactoryAnimationAsset : public UActorFactorySkeletalMesh
 {
diff --git a/Engine/Source/Editor/UnrealEd/Classes/ActorFactories/ActorFactoryBlueprint.h b/Engine/Source/Editor/UnrealEd/Classes/ActorFactories/ActorFactoryBlueprint.h
--- a/Engine/Source/Editor/UnrealEd/Classes/ActorFactories/ActorFactoryBlueprint.h
+++ b/Engine/Source/Editor/UnrealEd/Classes/ActorFactories/ActorFactoryBlueprint.h
@@ -2,6 +2,7 @@
 
 
 #pragma once
+#include "ActorFactories/ActorFactory.h"
 #include "ActorFactoryBlueprint.generated.h"
 
 UCLASS(MinimalAPI, config=Editor, collapsecategories, hidecategories=Object)
diff --git a/Engine/Source/Editor/UnrealEd/Classes/ActorFactories/ActorFactoryBoxVolume.h b/Engine/Source/Editor/UnrealEd/Classes/ActorFactories/ActorFactoryBoxVolume.h
--- a/Engine/Source/Editor/UnrealEd/Classes/ActorFactories/ActorFactoryBoxVolume.h
+++ b/Engine/Source/Editor/UnrealEd/Classes/ActorFactories/ActorFactoryBoxVolume.h
@@ -2,6 +2,7 @@
 
 
 #pragma once
+#include "ActorFactories/ActorFactory.h"
 #include "ActorFactoryBoxVolume.generated.h"
 
 UCLASS(MinimalAPI, config=Editor, collapsecategories, hidecategories=Object)
diff --git a/Engine/Source/Editor/UnrealEd/Classes/ActorFactories/ActorFactoryClass.h b/Engine/Source/Editor/UnrealEd/Classes/ActorFactories/ActorFactoryClass.h
--- a/Engine/Source/Editor/UnrealEd/Classes/ActorFactories/ActorFactoryClass.h
+++ b/Engine/Source/Editor/UnrealEd/Classes/ActorFactories/ActorFactoryClass.h
@@ -2,6 +2,7 @@
 
 
 #pragma once
+#include "ActorFactories/ActorFactory.h"
 #include "ActorFactoryClass.generated.h"
 
 UCLASS(MinimalAPI, config=Editor, collapsecategories, hidecategories=Object)
diff --git a/Engine/Source/Editor/UnrealEd/Classes/ActorFactories/ActorFactoryCylinderVolume.h b/Engine/Source/Editor/UnrealEd/Classes/ActorFactories/ActorFactoryCylinderVolume.h
--- a/Engine/Source/Editor/UnrealEd/Classes/ActorFactories/ActorFactoryCylinderVolume.h
+++ b/Engine/Source/Editor/UnrealEd/Classes/ActorFactories/ActorFactoryCylinderVolume.h
@@ -2,6 +2,7 @@
 
 
 #pragma once
+#include "ActorFactories/ActorFactory.h"
 #include "ActorFactoryCylinderVolume.generated.h"
 
 UCLASS(MinimalAPI, config=Editor, collapsecategories, hidecategories=Object)
diff --git a/Engine/Source/Editor/UnrealEd/Classes/ActorFactories/ActorFactoryDeferredDecal.h b/Engine/Source/Editor/UnrealEd/Classes/ActorFactories/ActorFactoryDeferredDecal.h
--- a/Engine/Source/Editor/UnrealEd/Classes/ActorFactories/ActorFactoryDeferredDecal.h
+++ b/Engine/Source/Editor/UnrealEd/Classes/ActorFactories/ActorFactoryDeferredDecal.h
@@ -2,6 +2,7 @@
 
 
 #pragma once
+#include "ActorFactories/ActorFactory.h"
 #include "ActorFactoryDeferredDecal.generated.h"
 
 UCLASS(MinimalAPI, config=Editor)
diff --git a/Engine/Source/Editor/UnrealEd/Classes/ActorFactories/ActorFactoryDestructible.h b/Engine/Source/Editor/UnrealEd/Classes/ActorFactories/ActorFactoryDestructible.h
--- a/Engine/Source/Editor/UnrealEd/Classes/ActorFactories/ActorFactoryDestructible.h
+++ b/Engine/Source/Editor/UnrealEd/Classes/ActorFactories/ActorFactoryDestructible.h
@@ -2,6 +2,7 @@
 
 
 #pragma once
+#include "ActorFactories/ActorFactory.h"
 #include "ActorFactoryDestructible.generated.h"
 
 UCLASS(MinimalAPI, config=Editor)
diff --git a/Engine/Source/Editor/UnrealEd/Classes/ActorFactories/ActorFactoryEmitter.h b/Engine/Source/Editor/UnrealEd/Classes/ActorFactories/ActorFactoryEmitter.h
--- a/Engine/Source/Editor/UnrealEd/Classes/ActorFactories/ActorFactoryEmitter.h
+++ b/Engine/Source/Editor/UnrealEd/Classes/ActorFactories/ActorFactoryEmitter.h
@@ -2,6 +2,7 @@
 
 
 #pragma once
+#include "ActorFactories/ActorFactory.h"
 #include "ActorFactoryEmitter.generated.h"
 
 UCLASS(MinimalAPI, config=Editor, collapsecategories, hidecategories=Object)
diff --git a/Engine/Source/Editor/UnrealEd/Classes/ActorFactories/ActorFactoryMatineeActor.h b/Engine/Source/Editor/UnrealEd/Classes/ActorFactories/ActorFactoryMatineeActor.h
--- a/Engine/Source/Editor/UnrealEd/Classes/ActorFactories/ActorFactoryMatineeActor.h
+++ b/Engine/Source/Editor/UnrealEd/Classes/ActorFactories/ActorFactoryMatineeActor.h
@@ -2,6 +2,7 @@
 
 
 #pragma once
+#include "ActorFactories/ActorFactory.h"
 #include "ActorFactoryMatineeActor.generated.h"
 
 UCLASS(MinimalAPI, config=Editor)
diff --git a/Engine/Source/Editor/UnrealEd/Classes/ActorFactories/ActorFactoryPhysicsAsset.h b/Engine/Source/Editor/UnrealEd/Classes/ActorFactories/ActorFactoryPhysicsAsset.h
--- a/Engine/Source/Editor/UnrealEd/Classes/ActorFactories/ActorFactoryPhysicsAsset.h
+++ b/Engine/Source/Editor/UnrealEd/Classes/ActorFactories/ActorFactoryPhysicsAsset.h
@@ -2,6 +2,7 @@
 
 
 #pragma once
+#include "ActorFactories/ActorFactory.h"
 #include "ActorFactoryPhysicsAsset.generated.h"
 
 UCLASS(MinimalAPI, config=Editor, collapsecategories, hidecategories=Object)
diff --git a/Engine/Source/Editor/UnrealEd/Classes/ActorFactories/ActorFactorySkeletalMesh.h b/Engine/Source/Editor/UnrealEd/Classes/ActorFactories/ActorFactorySkeletalMesh.h
--- a/Engine/Source/Editor/UnrealEd/Classes/ActorFactories/ActorFactorySkeletalMesh.h
+++ b/Engine/Source/Editor/UnrealEd/Classes/ActorFactories/ActorFactorySkeletalMesh.h
@@ -2,8 +2,11 @@
 
 
 #pragma once
+#include "ActorFactories/ActorFactory.h"
 #include "ActorFactorySkeletalMesh.generated.h"
 
+class USkeletalMesh;
+
 UCLASS(MinimalAPI, config=Editor)
 class UActorFactorySkeletalMesh : public UActorFactory
 {
diff --git a/Engine/Source/Editor/UnrealEd/Classes/Animation/DebugSkelMeshComponent.h b/Engine/Source/Editor/UnrealEd/Classes/Animation/DebugSkelMeshComponent.h
--- a/Engine/Source/Editor/UnrealEd/Classes/Animation/DebugSkelMeshComponent.h
+++ b/Engine/Source/Editor/UnrealEd/Classes/Animation/DebugSkelMeshComponent.h
@@ -2,6 +2,7 @@
 
 
 #pragma once
+#include "Components/SkeletalMeshComponent.h"
 #include "DebugSkelMeshComponent.generated.h"
 
 USTRUCT()
diff --git a/Engine/Source/Editor/UnrealEd/Classes/Animation/EditorAnimBaseObj.h b/Engine/Source/Editor/UnrealEd/Classes/Animation/EditorAnimBaseObj.h
--- a/Engine/Source/Editor/UnrealEd/Classes/Animation/EditorAnimBaseObj.h
+++ b/Engine/Source/Editor/UnrealEd/Classes/Animation/EditorAnimBaseObj.h
@@ -8,6 +8,8 @@
 #pragma once
 #include "EditorAnimBaseObj.generated.h"
 
+class UAnimSequenceBase;
+
 DECLARE_DELEGATE_TwoParams( FOnAnimObjectChange, class UObject*, bool)
 
 
diff --git a/Engine/Source/Editor/UnrealEd/Classes/Animation/EditorAnimCompositeSegment.h b/Engine/Source/Editor/UnrealEd/Classes/Animation/EditorAnimCompositeSegment.h
--- a/Engine/Source/Editor/UnrealEd/Classes/Animation/EditorAnimCompositeSegment.h
+++ b/Engine/Source/Editor/UnrealEd/Classes/Animation/EditorAnimCompositeSegment.h
@@ -6,6 +6,8 @@
  */
 
 #pragma once
+#include "Animation/AnimCompositeBase.h"
+#include "Animation/EditorAnimBaseObj.h"
 #include "EditorAnimCompositeSegment.generated.h"
 
 DECLARE_DELEGATE_OneParam( FOnAnimSegmentChanged, class UEditorAnimSegment*)
diff --git a/Engine/Source/Editor/UnrealEd/Classes/Animation/EditorAnimSegment.h b/Engine/Source/Editor/UnrealEd/Classes/Animation/EditorAnimSegment.h
--- a/Engine/Source/Editor/UnrealEd/Classes/Animation/EditorAnimSegment.h
+++ b/Engine/Source/Editor/UnrealEd/Classes/Animation/EditorAnimSegment.h
@@ -6,6 +6,8 @@
  */
 
 #pragma once
+#include "Animation/AnimCompositeBase.h"
+#include "Animation/EditorAnimBaseObj.h"
 #include "EditorAnimSegment.generated.h"
 
 DECLARE_DELEGATE_OneParam( FOnAnimSegmentChanged, class UEditorAnimSegment*)
diff --git a/Engine/Source/Editor/UnrealEd/Classes/Animation/EditorBranchPoint.h b/Engine/Source/Editor/UnrealEd/Classes/Animation/EditorBranchPoint.h
--- a/Engine/Source/Editor/UnrealEd/Classes/Animation/EditorBranchPoint.h
+++ b/Engine/Source/Editor/UnrealEd/Classes/Animation/EditorBranchPoint.h
@@ -6,6 +6,8 @@
  */
 
 #pragma once
+#include "Animation/AnimMontage.h"
+#include "Animation/EditorAnimBaseObj.h"
 #include "EditorBranchPoint.generated.h"
 
 
diff --git a/Engine/Source/Editor/UnrealEd/Classes/Animation/EditorCompositeSection.h b/Engine/Source/Editor/UnrealEd/Classes/Animation/EditorCompositeSection.h
--- a/Engine/Source/Editor/UnrealEd/Classes/Animation/EditorCompositeSection.h
+++ b/Engine/Source/Editor/UnrealEd/Classes/Animation/EditorCompositeSection.h
@@ -6,6 +6,8 @@
  */
 
 #pragma once
+#include "Animation/AnimMontage.h"
+#include "Animation/EditorAnimBaseObj.h"
 #include "EditorCompositeSection.generated.h"
 
 UCLASS(hidecategories=UObject, MinimalAPI, BlueprintType)
diff --git a/Engine/Source/Editor/UnrealEd/Classes/Animation/EditorParentPlayerListObj.h b/Engine/Source/Editor/UnrealEd/Classes/Animation/EditorParentPlayerListObj.h
--- a/Engine/Source/Editor/UnrealEd/Classes/Animation/EditorParentPlayerListObj.h
+++ b/Engine/Source/Editor/UnrealEd/Classes/Animation/EditorParentPlayerListObj.h
@@ -3,6 +3,9 @@
 #pragma once
 #include "EditorParentPlayerListObj.generated.h"
 
+struct FAnimParentNodeAssetOverride;
+class UAnimGraphNode_Base;
+
 UCLASS()
 class UNREALED_API UEditorParentPlayerListObj : public UObject
 {
diff --git a/Engine/Source/Editor/UnrealEd/Classes/Commandlets/GatherTextFromAssetsCommandlet.h b/Engine/Source/Editor/UnrealEd/Classes/Commandlets/GatherTextFromAssetsCommandlet.h
--- a/Engine/Source/Editor/UnrealEd/Classes/Commandlets/GatherTextFromAssetsCommandlet.h
+++ b/Engine/Source/Editor/UnrealEd/Classes/Commandlets/GatherTextFromAssetsCommandlet.h
@@ -1,7 +1,8 @@
 // Copyright 1998-2014 Epic Games, Inc. All Rights Reserved.
 
-#pragma once 
-#include "Commandlets/Commandlet.h"
+#pragma once
+#include "Commandlets/GatherTextCommandletBase.h"
+#include "Sound/DialogueTypes.h"
 #include "GatherTextFromAssetsCommandlet.generated.h"
 
 namespace EAssetTextGatherStatus
@@ -22,7 +23,7 @@ namespace EAssetTextGatherStatus
 }
 
 class FLocMetadataValue;
-
+class UDialogueWave;
 /**
  *	UGatherTextFromAssetsCommandlet: Localization commandlet that collects all text to be localized from the game assets.
  */
diff --git a/Engine/Source/Editor/UnrealEd/Classes/Editor/EditorEngine.h b/Engine/Source/Editor/UnrealEd/Classes/Editor/EditorEngine.h
--- a/Engine/Source/Editor/UnrealEd/Classes/Editor/EditorEngine.h
+++ b/Engine/Source/Editor/UnrealEd/Classes/Editor/EditorEngine.h
@@ -2,17 +2,19 @@
 
 #pragma once
 
+#include "Engine/Brush.h"
+#include "Engine/Engine.h"
 #include "EditorUserSettings.h"
 #include "Transactor.h"
 #include "../Settings/LevelEditorPlaySettings.h"
 #include "../Settings/LevelEditorViewportSettings.h"
 #include "EditorEngine.generated.h"
 
+class APlayerStart;
 class FAssetData;
+class FPoly;
 class UAnimSequence;
 class USkeleton;
-class FPoly;
-
 //
 // Things to set in mapSetBrush.
 //
diff --git a/Engine/Source/Editor/UnrealEd/Classes/Editor/GroupActor.h b/Engine/Source/Editor/UnrealEd/Classes/Editor/GroupActor.h
--- a/Engine/Source/Editor/UnrealEd/Classes/Editor/GroupActor.h
+++ b/Engine/Source/Editor/UnrealEd/Classes/Editor/GroupActor.h
@@ -5,6 +5,7 @@
 //=============================================================================
 
 #pragma once
+#include "GameFramework/Actor.h"
 #include "GroupActor.generated.h"
 
 UCLASS(hidedropdown,MinimalAPI, notplaceable, NotBlueprintable)
diff --git a/Engine/Source/Editor/UnrealEd/Classes/Exporters/AnimSequenceExporterFBX.h b/Engine/Source/Editor/UnrealEd/Classes/Exporters/AnimSequenceExporterFBX.h
--- a/Engine/Source/Editor/UnrealEd/Classes/Exporters/AnimSequenceExporterFBX.h
+++ b/Engine/Source/Editor/UnrealEd/Classes/Exporters/AnimSequenceExporterFBX.h
@@ -5,6 +5,7 @@
 //=============================================================================
 
 #pragma once
+#include "Exporters/Exporter.h"
 #include "AnimSequenceExporterFBX.generated.h"
 
 UCLASS()
diff --git a/Engine/Source/Editor/UnrealEd/Classes/Factories/DataAssetFactory.h b/Engine/Source/Editor/UnrealEd/Classes/Factories/DataAssetFactory.h
--- a/Engine/Source/Editor/UnrealEd/Classes/Factories/DataAssetFactory.h
+++ b/Engine/Source/Editor/UnrealEd/Classes/Factories/DataAssetFactory.h
@@ -1,8 +1,11 @@
 // Copyright 1998-2014 Epic Games, Inc. All Rights Reserved.
 
 #pragma once
+#include "Factories/Factory.h"
 #include "DataAssetFactory.generated.h"
 
+class UDataAsset;
+
 UCLASS(hidecategories=Object, MinimalAPI)
 class UDataAssetFactory : public UFactory
 {
diff --git a/Engine/Source/Editor/UnrealEd/Classes/Factories/FbxAssetImportData.h b/Engine/Source/Editor/UnrealEd/Classes/Factories/FbxAssetImportData.h
--- a/Engine/Source/Editor/UnrealEd/Classes/Factories/FbxAssetImportData.h
+++ b/Engine/Source/Editor/UnrealEd/Classes/Factories/FbxAssetImportData.h
@@ -1,6 +1,7 @@
 // Copyright 1998-2014 Epic Games, Inc. All Rights Reserved.
 
 #pragma once
+#include "EditorFramework/AssetImportData.h"
 #include "FbxAssetImportData.generated.h"
 
 /**
diff --git a/Engine/Source/Editor/UnrealEd/Classes/MaterialEditor/MaterialEditorMeshComponent.h b/Engine/Source/Editor/UnrealEd/Classes/MaterialEditor/MaterialEditorMeshComponent.h
--- a/Engine/Source/Editor/UnrealEd/Classes/MaterialEditor/MaterialEditorMeshComponent.h
+++ b/Engine/Source/Editor/UnrealEd/Classes/MaterialEditor/MaterialEditorMeshComponent.h
@@ -2,6 +2,7 @@
 
 
 #pragma once
+#include "Components/StaticMeshComponent.h"
 #include "MaterialEditorMeshComponent.generated.h"
 
 UCLASS()
diff --git a/Engine/Source/Editor/UnrealEd/Classes/MaterialGraph/MaterialGraph.h b/Engine/Source/Editor/UnrealEd/Classes/MaterialGraph/MaterialGraph.h
--- a/Engine/Source/Editor/UnrealEd/Classes/MaterialGraph/MaterialGraph.h
+++ b/Engine/Source/Editor/UnrealEd/Classes/MaterialGraph/MaterialGraph.h
@@ -2,6 +2,7 @@
 
 #pragma once
 #include "EdGraph/EdGraph.h"
+#include "Materials/Material.h"
 #include "MaterialGraph.generated.h"
 
 DECLARE_DELEGATE_RetVal( bool, FRealtimeStateGetter );
diff --git a/Engine/Source/Editor/UnrealEd/Classes/Settings/LevelEditorViewportSettings.h b/Engine/Source/Editor/UnrealEd/Classes/Settings/LevelEditorViewportSettings.h
--- a/Engine/Source/Editor/UnrealEd/Classes/Settings/LevelEditorViewportSettings.h
+++ b/Engine/Source/Editor/UnrealEd/Classes/Settings/LevelEditorViewportSettings.h
@@ -8,6 +8,7 @@
 
 #include "Viewports.h"
 #include "Editor/UnrealEdTypes.h"
+#include "Engine/EngineBaseTypes.h"
 #include "LevelEditorViewportSettings.generated.h"
 
 
diff --git a/Engine/Source/Editor/UnrealEd/Classes/TexAligner/TexAligner.h b/Engine/Source/Editor/UnrealEd/Classes/TexAligner/TexAligner.h
--- a/Engine/Source/Editor/UnrealEd/Classes/TexAligner/TexAligner.h
+++ b/Engine/Source/Editor/UnrealEd/Classes/TexAligner/TexAligner.h
@@ -10,6 +10,8 @@
 
 #include "TexAligner.generated.h"
 
+class FPoly;
+
 /** Alignment types. */
 UENUM()
 enum ETexAlign
diff --git a/Engine/Source/Editor/UnrealEd/Private/Fbx/FbxMainImport.cpp b/Engine/Source/Editor/UnrealEd/Private/Fbx/FbxMainImport.cpp
--- a/Engine/Source/Editor/UnrealEd/Private/Fbx/FbxMainImport.cpp
+++ b/Engine/Source/Editor/UnrealEd/Private/Fbx/FbxMainImport.cpp
@@ -38,6 +38,7 @@
 #include "Factories.h"
 #include "Engine.h"
 #include "SkelImport.h"
+#include "FbxErrors.h"
 #include "FbxImporter.h"
 #include "FbxOptionWindow.h"
 #include "FbxErrors.h"
diff --git a/Engine/Source/Editor/UnrealEd/Public/EditorComponents.h b/Engine/Source/Editor/UnrealEd/Public/EditorComponents.h
--- a/Engine/Source/Editor/UnrealEd/Public/EditorComponents.h
+++ b/Engine/Source/Editor/UnrealEd/Public/EditorComponents.h
@@ -5,7 +5,14 @@
 =============================================================================*/
 
 #pragma once
-
+#include "Engine/EngineTypes.h"
+#include "UObject/GCObject.h"
+
+class FPrimitiveDrawInterface;
+class FReferenceCollector;
+class FSceneView;
+class UMaterial;
+class UMaterialInstanceDynamic;
 
 class FGridWidget : public FGCObject
 {
diff --git a/Engine/Source/Editor/UnrealEd/Public/TickableEditorObject.h b/Engine/Source/Editor/UnrealEd/Public/TickableEditorObject.h
--- a/Engine/Source/Editor/UnrealEd/Public/TickableEditorObject.h
+++ b/Engine/Source/Editor/UnrealEd/Public/TickableEditorObject.h
@@ -1,6 +1,7 @@
 // Copyright 1998-2014 Epic Games, Inc. All Rights Reserved.
 
 #pragma once
+#include "Tickable.h"
 
 /**
  * This class provides common registration for gamethread editor only tickable objects. It is an
diff --git a/Engine/Source/Editor/UnrealEd/Public/UnrealEd.h b/Engine/Source/Editor/UnrealEd/Public/UnrealEd.h
--- a/Engine/Source/Editor/UnrealEd/Public/UnrealEd.h
+++ b/Engine/Source/Editor/UnrealEd/Public/UnrealEd.h
@@ -8,7 +8,6 @@
 #if WITH_EDITOR
 
 #include "EngineDefines.h"
-#include "Engine.h"
 
 #include "Slate.h"
 #include "EditorStyle.h"
diff --git a/Engine/Source/Runtime/Engine/Classes/Animation/AnimInstance.h b/Engine/Source/Runtime/Engine/Classes/Animation/AnimInstance.h
--- a/Engine/Source/Runtime/Engine/Classes/Animation/AnimInstance.h
+++ b/Engine/Source/Runtime/Engine/Classes/Animation/AnimInstance.h
@@ -18,6 +18,7 @@ class UAnimationAsset;
 class UCanvas;
 class UWorld;
 class FTransform;
+class FDebugDisplayInfo;
 
 DECLARE_DELEGATE_TwoParams(FOnMontageEnded, UAnimMontage*, bool /*bInterrupted*/)
 DECLARE_DELEGATE_TwoParams(FOnMontageBlendingOutStarted, UAnimMontage*, bool /*bInterrupted*/)
diff --git a/Engine/Source/Runtime/Engine/Classes/Distributions/Distribution.h b/Engine/Source/Runtime/Engine/Classes/Distributions/Distribution.h
--- a/Engine/Source/Runtime/Engine/Classes/Distributions/Distribution.h
+++ b/Engine/Source/Runtime/Engine/Classes/Distributions/Distribution.h
@@ -2,6 +2,7 @@
 
 
 #pragma once
+#include "Math/CurveEdInterface.h"
 #include "Distribution.generated.h"
 
 /** Lookup table for distributions. */
diff --git a/Engine/Source/Runtime/Engine/Classes/Distributions/DistributionFloat.h b/Engine/Source/Runtime/Engine/Classes/Distributions/DistributionFloat.h
--- a/Engine/Source/Runtime/Engine/Classes/Distributions/DistributionFloat.h
+++ b/Engine/Source/Runtime/Engine/Classes/Distributions/DistributionFloat.h
@@ -3,6 +3,7 @@
 
 #pragma once
 #include "Distributions.h"
+#include "Distributions/Distribution.h"
 #include "DistributionFloat.generated.h"
 
 /** Type-safe floating point distribution. */
diff --git a/Engine/Source/Runtime/Engine/Classes/Distributions/DistributionFloatConstant.h b/Engine/Source/Runtime/Engine/Classes/Distributions/DistributionFloatConstant.h
--- a/Engine/Source/Runtime/Engine/Classes/Distributions/DistributionFloatConstant.h
+++ b/Engine/Source/Runtime/Engine/Classes/Distributions/DistributionFloatConstant.h
@@ -2,6 +2,7 @@
 
 
 #pragma once
+#include "Distributions/DistributionFloat.h"
 #include "DistributionFloatConstant.generated.h"
 
 UCLASS(collapsecategories, hidecategories=Object, editinlinenew,MinimalAPI)
diff --git a/Engine/Source/Runtime/Engine/Classes/Distributions/DistributionFloatConstantCurve.h b/Engine/Source/Runtime/Engine/Classes/Distributions/DistributionFloatConstantCurve.h
--- a/Engine/Source/Runtime/Engine/Classes/Distributions/DistributionFloatConstantCurve.h
+++ b/Engine/Source/Runtime/Engine/Classes/Distributions/DistributionFloatConstantCurve.h
@@ -2,6 +2,7 @@
 
 
 #pragma once
+#include "Distributions/DistributionFloat.h"
 #include "DistributionFloatConstantCurve.generated.h"
 
 UCLASS(collapsecategories, hidecategories=Object, editinlinenew,MinimalAPI)
diff --git a/Engine/Source/Runtime/Engine/Classes/Distributions/DistributionFloatUniform.h b/Engine/Source/Runtime/Engine/Classes/Distributions/DistributionFloatUniform.h
--- a/Engine/Source/Runtime/Engine/Classes/Distributions/DistributionFloatUniform.h
+++ b/Engine/Source/Runtime/Engine/Classes/Distributions/DistributionFloatUniform.h
@@ -3,6 +3,7 @@
 
 #pragma once
 #include "Distributions.h"
+#include "Distributions/DistributionFloat.h"
 #include "DistributionFloatUniform.generated.h"
 
 UCLASS(collapsecategories, hidecategories=Object, editinlinenew, MinimalAPI)
diff --git a/Engine/Source/Runtime/Engine/Classes/Distributions/DistributionFloatUniformCurve.h b/Engine/Source/Runtime/Engine/Classes/Distributions/DistributionFloatUniformCurve.h
--- a/Engine/Source/Runtime/Engine/Classes/Distributions/DistributionFloatUniformCurve.h
+++ b/Engine/Source/Runtime/Engine/Classes/Distributions/DistributionFloatUniformCurve.h
@@ -3,6 +3,7 @@
 
 #pragma once
 #include "Distributions.h"
+#include "Distributions/DistributionFloat.h"
 #include "DistributionFloatUniformCurve.generated.h"
 
 UCLASS(collapsecategories, hidecategories=Object, editinlinenew, MinimalAPI)
diff --git a/Engine/Source/Runtime/Engine/Classes/Distributions/DistributionVector.h b/Engine/Source/Runtime/Engine/Classes/Distributions/DistributionVector.h
--- a/Engine/Source/Runtime/Engine/Classes/Distributions/DistributionVector.h
+++ b/Engine/Source/Runtime/Engine/Classes/Distributions/DistributionVector.h
@@ -3,6 +3,7 @@
 
 #pragma once
 #include "Distributions.h"
+#include "Distributions/Distribution.h"
 #include "DistributionVector.generated.h"
 
 UENUM()
diff --git a/Engine/Source/Runtime/Engine/Classes/Distributions/DistributionVectorConstant.h b/Engine/Source/Runtime/Engine/Classes/Distributions/DistributionVectorConstant.h
--- a/Engine/Source/Runtime/Engine/Classes/Distributions/DistributionVectorConstant.h
+++ b/Engine/Source/Runtime/Engine/Classes/Distributions/DistributionVectorConstant.h
@@ -2,6 +2,7 @@
 
 
 #pragma once
+#include "Distributions/DistributionVector.h"
 #include "DistributionVectorConstant.generated.h"
 
 UCLASS(collapsecategories, hidecategories=Object, editinlinenew)
diff --git a/Engine/Source/Runtime/Engine/Classes/Distributions/DistributionVectorConstantCurve.h b/Engine/Source/Runtime/Engine/Classes/Distributions/DistributionVectorConstantCurve.h
--- a/Engine/Source/Runtime/Engine/Classes/Distributions/DistributionVectorConstantCurve.h
+++ b/Engine/Source/Runtime/Engine/Classes/Distributions/DistributionVectorConstantCurve.h
@@ -2,6 +2,7 @@
 
 
 #pragma once
+#include "Distributions/DistributionVector.h"
 #include "DistributionVectorConstantCurve.generated.h"
 
 UCLASS(collapsecategories, hidecategories=Object, editinlinenew,MinimalAPI)
diff --git a/Engine/Source/Runtime/Engine/Classes/Engine/GameViewportClient.h b/Engine/Source/Runtime/Engine/Classes/Engine/GameViewportClient.h
--- a/Engine/Source/Runtime/Engine/Classes/Engine/GameViewportClient.h
+++ b/Engine/Source/Runtime/Engine/Classes/Engine/GameViewportClient.h
@@ -20,6 +20,7 @@
 #include "GameViewportClient.generated.h"
 
 class UGameInstance;
+class UNetDriver;
 
 /**
  * Enum of the different splitscreen types
diff --git a/Engine/Source/Runtime/Engine/Classes/Engine/World.h b/Engine/Source/Runtime/Engine/Classes/Engine/World.h
--- a/Engine/Source/Runtime/Engine/Classes/Engine/World.h
+++ b/Engine/Source/Runtime/Engine/Classes/Engine/World.h
@@ -11,6 +11,7 @@
 #include "EngineDefines.h"
 #include "Engine/LatentActionManager.h"
 #include "Runtime/RHI/Public/RHIDefinitions.h"
+#include "GameFramework/Actor.h"
 #include "World.generated.h"
 
 class FPhysScene;
@@ -22,6 +23,9 @@ class ABrush;
 class UModel;
 class APhysicsVolume;
 class UTexture2D;
+class APlayerController;
+class AMatineeActor;
+class AWorldSettings;
 
 template<typename,typename> class TOctree;
 
diff --git a/Engine/Source/Runtime/Engine/Classes/Exporters/Exporter.h b/Engine/Source/Runtime/Engine/Classes/Exporters/Exporter.h
--- a/Engine/Source/Runtime/Engine/Classes/Exporters/Exporter.h
+++ b/Engine/Source/Runtime/Engine/Classes/Exporters/Exporter.h
@@ -9,6 +9,8 @@
 #pragma once
 #include "Exporter.generated.h"
 
+class UActorComponent;
+
 UCLASS(abstract, transient, MinimalAPI)
 class UExporter : public UObject
 {
diff --git a/Engine/Source/Runtime/Engine/Private/World.cpp b/Engine/Source/Runtime/Engine/Private/World.cpp
--- a/Engine/Source/Runtime/Engine/Private/World.cpp
+++ b/Engine/Source/Runtime/Engine/Private/World.cpp
@@ -31,6 +31,8 @@
 #include "LevelUtils.h"
 #include "PhysicsPublic.h"
 #include "AI/AISystemBase.h"
+#include "SceneInterface.h"
+
 #if WITH_EDITOR
 	#include "Editor/UnrealEd/Public/Kismet2/KismetEditorUtilities.h"
 	#include "Editor/UnrealEd/Public/Kismet2/BlueprintEditorUtils.h"
diff --git a/Engine/Source/Runtime/Engine/Public/Model.h b/Engine/Source/Runtime/Engine/Public/Model.h
--- a/Engine/Source/Runtime/Engine/Public/Model.h
+++ b/Engine/Source/Runtime/Engine/Public/Model.h
@@ -11,10 +11,13 @@
 #include "RenderUtils.h"
 #include "SceneTypes.h"
 
+struct FLightmassPrimitiveSettings;
 struct FStaticLightingVertex;
 class ABrush;
-class UModel;
+class ULevel;
 class ULightComponent;
+class UMaterialInterface;
+class UModel;
 class UPolys;
 
 //
diff --git a/Engine/Source/Runtime/Engine/Public/ParticleEmitterInstances.h b/Engine/Source/Runtime/Engine/Public/ParticleEmitterInstances.h
--- a/Engine/Source/Runtime/Engine/Public/ParticleEmitterInstances.h
+++ b/Engine/Source/Runtime/Engine/Public/ParticleEmitterInstances.h
@@ -8,6 +8,8 @@
 #pragma once
 
 #include "Distributions.h"
+#include "Distributions/DistributionFloat.h"
+#include "Distributions/DistributionVector.h"
 #include "ParticleHelper.h"
 
 /*-----------------------------------------------------------------------------
