CORE:  In-editor fixup redirects not working with string asset ref

TTP# 341126

diff --git a/Engine/Source/Developer/AssetTools/Private/AssetFixUpRedirectors.cpp b/Engine/Source/Developer/AssetTools/Private/AssetFixUpRedirectors.cpp
--- a/Engine/Source/Developer/AssetTools/Private/AssetFixUpRedirectors.cpp
+++ b/Engine/Source/Developer/AssetTools/Private/AssetFixUpRedirectors.cpp
@@ -84,6 +84,9 @@ void FAssetFixUpRedirectors::ExecuteFixUp(TArray<TWeakObjectPtr<UObjectRedirecto
 				// If any referencing packages are left read-only, the checkout failed or SCC was not enabled. Trim them from the save list and leave redirectors.
 				DetectReadOnlyPackages(RedirectorRefsList, ReferencingPackagesToSave);
 
+				// Fix up referencing FStringAssetReferences
+				FixUpStringAssetReferences(RedirectorRefsList, ReferencingPackagesToSave);
+
 				// Save all packages that were referencing any of the assets that were moved without redirectors
 				SaveReferencingPackages(ReferencingPackagesToSave);
 
@@ -380,4 +383,19 @@ void FAssetFixUpRedirectors::ReportFailures(const TArray<FRedirectorRefs>& Redir
 	EditorErrors.Open();
 }
 
+void FAssetFixUpRedirectors::FixUpStringAssetReferences(const TArray<FRedirectorRefs>& RedirectorsToFix, const TArray<UPackage*>& InReferencingPackagesToSave) const
+{
+	TArray<UPackage *> PackagesToCheck(InReferencingPackagesToSave);
+
+	FEditorFileUtils::GetDirtyWorldPackages(PackagesToCheck);
+	FEditorFileUtils::GetDirtyContentPackages(PackagesToCheck);
+
+	for (auto& RedirectorRef : RedirectorsToFix)
+	{
+		FAssetRenameManager::RenameReferencingStringAssetReferences(PackagesToCheck,
+			RedirectorRef.Redirector->GetPathName(),
+			RedirectorRef.Redirector->DestinationObject->GetPathName());
+	}
+}
+
 #undef LOCTEXT_NAMESPACE
diff --git a/Engine/Source/Developer/AssetTools/Private/AssetFixUpRedirectors.h b/Engine/Source/Developer/AssetTools/Private/AssetFixUpRedirectors.h
--- a/Engine/Source/Developer/AssetTools/Private/AssetFixUpRedirectors.h
+++ b/Engine/Source/Developer/AssetTools/Private/AssetFixUpRedirectors.h
@@ -37,6 +37,9 @@ private:
 	/** Finds any read only packages and removes them from the save list. Redirectors referenced by these packages will not be fixed up. */ 
 	void DetectReadOnlyPackages(TArray<FRedirectorRefs>& RedirectorsToFix, TArray<UPackage*>& InOutReferencingPackagesToSave) const;
 
+	/** FixUp string asset references. */
+	void FixUpStringAssetReferences(const TArray<FRedirectorRefs>& RedirectorsToFix, const TArray<UPackage*>& InReferencingPackagesToSave) const;
+
 	/** Saves all the referencing packages and updates SCC state */
 	void SaveReferencingPackages(const TArray<UPackage*>& ReferencingPackagesToSave) const;
 
diff --git a/Engine/Source/Developer/AssetTools/Private/AssetRenameManager.cpp b/Engine/Source/Developer/AssetTools/Private/AssetRenameManager.cpp
--- a/Engine/Source/Developer/AssetTools/Private/AssetRenameManager.cpp
+++ b/Engine/Source/Developer/AssetTools/Private/AssetRenameManager.cpp
@@ -504,11 +504,72 @@ void FAssetRenameManager::DetectReadOnlyPackages(TArray<FAssetRenameDataWithRefe
 	}
 }
 
+/**
+ * Function that renames all FStringAssetReference object with the old asset path to the new one.
+ *
+ * @param PackagesToCheck Packages to check for referencing FStringAssetReference.
+ * @param OldAssetPath Old path.
+ * @param NewAssetPath New path.
+ */
+void FAssetRenameManager::RenameReferencingStringAssetReferences(const TArray<UPackage *> PackagesToCheck, const FString& OldAssetPath, const FString& NewAssetPath)
+{
+	struct FStringAssetReferenceRenameSerializer : public FArchiveUObject
+	{
+		FStringAssetReferenceRenameSerializer(const FString& OldAssetPath, const FString& NewAssetPath)
+			: OldAssetPath(OldAssetPath), NewAssetPath(NewAssetPath)
+		{ }
+
+		FArchive& operator<<(FStringAssetReference& Reference) override
+		{
+			if (Reference.AssetLongPathname == OldAssetPath)
+			{
+				Reference.AssetLongPathname = NewAssetPath;
+			}
+
+			// Generated class path support.
+			if (Reference.AssetLongPathname == OldAssetPath + "_C")
+			{
+				Reference.AssetLongPathname = NewAssetPath + "_C";
+			}
+
+			return *this;
+		}
+
+	private:
+		const FString& OldAssetPath;
+		const FString& NewAssetPath;
+	};
+
+	FStringAssetReferenceRenameSerializer RenameSerializer(OldAssetPath, NewAssetPath);
+
+	for (auto* Package : PackagesToCheck)
+	{
+		TArray<UObject*> ObjectsInPackage;
+		GetObjectsWithOuter(Package, ObjectsInPackage);
+
+		for (auto* Object : ObjectsInPackage)
+		{
+			Object->Serialize(RenameSerializer);
+		}
+	}
+}
+
 void FAssetRenameManager::PerformAssetRename(TArray<FAssetRenameDataWithReferencers>& AssetsToRename) const
 {
 	const FText AssetRenameSlowTask = LOCTEXT("AssetRenameSlowTask", "Renaming Assets");
 	GWarn->BeginSlowTask( AssetRenameSlowTask, true );
 
+	/**
+	 * We need to collect and check those cause dependency graph is only
+	 * representing on-disk state and we want to support rename for in-memory
+	 * objects. It is only needed for string references as in memory references
+	 * for other objects are pointers, so renames doesn't apply to those.
+	 */
+	TArray<UPackage *> DirtyPackagesToCheckForStringReferences;
+	
+	FEditorFileUtils::GetDirtyWorldPackages(DirtyPackagesToCheckForStringReferences);
+	FEditorFileUtils::GetDirtyContentPackages(DirtyPackagesToCheckForStringReferences);
+
 	TArray<UPackage*> PackagesToSave;
 	TArray<UPackage*> PotentialPackagesToDelete;
 	for ( int32 AssetIdx = 0; AssetIdx < AssetsToRename.Num(); ++AssetIdx )
@@ -524,6 +585,7 @@ void FAssetRenameManager::PerformAssetRename(TArray<FAssetRenameDataWithReferenc
 		}
 
 		UObject* Asset = RenameData.Asset.Get();
+		FString OldAssetPath = Asset->GetPathName();
 
 		if ( !Asset )
 		{
@@ -575,6 +637,15 @@ void FAssetRenameManager::PerformAssetRename(TArray<FAssetRenameDataWithReferenc
 			RenameData.bRenameFailed = true;
 			RenameData.FailureReason = ErrorMessage;
 		}
+
+		TArray<UPackage *> PackagesToCheck(DirtyPackagesToCheckForStringReferences);
+
+		for (auto PackageNameIt = RenameData.ReferencingPackageNames.CreateConstIterator(); PackageNameIt; ++PackageNameIt)
+		{
+			PackagesToCheck.Add(LoadPackage(NULL, *PackageNameIt->ToString(), LOAD_None));
+		}
+
+		RenameReferencingStringAssetReferences(PackagesToCheck, OldAssetPath, Asset->GetPathName());
 	}
 
 	GWarn->EndSlowTask();
diff --git a/Engine/Source/Developer/AssetTools/Private/AssetRenameManager.h b/Engine/Source/Developer/AssetTools/Private/AssetRenameManager.h
--- a/Engine/Source/Developer/AssetTools/Private/AssetRenameManager.h
+++ b/Engine/Source/Developer/AssetTools/Private/AssetRenameManager.h
@@ -19,6 +19,15 @@ public:
 	/** Renames assets using the specified names. */
 	void RenameAssets(const TArray<FAssetRenameData>& AssetsAndNames) const;
 
+	/**
+	 * Function that renames all FStringAssetReference object with the old asset path to the new one.
+	 *
+	 * @param PackagesToCheck Packages to check for referencing FStringAssetReference.
+	 * @param OldAssetPath Old path.
+	 * @param NewAssetPath New path.
+	 */
+	static void RenameReferencingStringAssetReferences(const TArray<UPackage *> PackagesToCheck, const FString& OldAssetPath, const FString& NewAssetPath);
+
 private:
 	/** Attempts to load and fix redirector references for the supplied assets */
 	void FixReferencesAndRename(TArray<FAssetRenameDataWithReferencers> AssetsToRename) const;
diff --git a/Engine/Source/Editor/KismetCompiler/Private/KismetCompilerVMBackend.cpp b/Engine/Source/Editor/KismetCompiler/Private/KismetCompilerVMBackend.cpp
--- a/Engine/Source/Editor/KismetCompiler/Private/KismetCompilerVMBackend.cpp
+++ b/Engine/Source/Editor/KismetCompiler/Private/KismetCompilerVMBackend.cpp
@@ -65,6 +65,11 @@ public:
 		return FArchive::operator<<(AssetPtr);
 	}
 
+	FArchive& operator<<(FStringAssetReference& Value) override
+	{
+		return FArchiveUObject::operator<<(Value);
+	}
+
 	FArchive& operator<<(TCHAR* S)
 	{
 		Serialize(S, FCString::Strlen(S) + 1); 
diff --git a/Engine/Source/Editor/UnrealEd/Private/FileHelpers.cpp b/Engine/Source/Editor/UnrealEd/Private/FileHelpers.cpp
--- a/Engine/Source/Editor/UnrealEd/Private/FileHelpers.cpp
+++ b/Engine/Source/Editor/UnrealEd/Private/FileHelpers.cpp
@@ -2365,28 +2365,9 @@ bool FEditorFileUtils::SaveDirtyPackages(const bool bPromptUserToSave, const boo
 	// A list of all packages that need to be saved
 	TArray<UPackage*> PackagesToSave;
 
-	// Need to track the number of packages we're not ignoring for save.
-	int32 NumPackagesNotIgnored = 0;
-
 	if( bSaveMapPackages )
 	{
-		// If we are saving map packages, collect all valid worlds and see if their package is dirty
-		TArray<UWorld*> Worlds;
-		EditorLevelUtils::GetWorlds( GWorld, Worlds, true );
-
-		for( int32 WorldIdx = 0; WorldIdx < Worlds.Num(); ++WorldIdx  )
-		{
-			UPackage* WorldPackage = Worlds[ WorldIdx ]->GetOutermost();
-			if( WorldPackage->IsDirty() && (WorldPackage->PackageFlags & PKG_PlayInEditor) == 0 
-				&& !WorldPackage->HasAnyFlags(RF_Transient))
-			{
-				// Count the number of packages to not ignore.
-				NumPackagesNotIgnored += (PackagesNotSavedDuringSaveAll.Find(WorldPackage->GetName())==NULL) ? 1 : 0;
-
-				// IF the package is dirty and its not a pie package, add the world package to the list of packages to save
-				PackagesToSave.Add( WorldPackage );
-			}
-		}
+		GetDirtyWorldPackages(PackagesToSave);
 	}
 
 	// Don't iterate through content packages if we dont plan on saving them
@@ -2394,41 +2375,16 @@ bool FEditorFileUtils::SaveDirtyPackages(const bool bPromptUserToSave, const boo
 	{
 		CollectGarbage(GARBAGE_COLLECTION_KEEPFLAGS);
 
-		// Make a list of all content packages that we should save
-		for ( TObjectIterator<UPackage> It; It; ++It )
-		{
-			UPackage*	Package					= *It;
-			bool		bShouldIgnorePackage	= false;
-
-			// Only look at root packages.
-			bShouldIgnorePackage |= Package->GetOuter() != NULL;
-			// Don't try to save "Transient" package.
-			bShouldIgnorePackage |= Package == GetTransientPackage();
-			// Ignore PIE packages.
-			bShouldIgnorePackage |= (Package->PackageFlags & PKG_PlayInEditor) != 0;
-			// Ignore packages that haven't been modified.
-			bShouldIgnorePackage |= !Package->IsDirty();
-
-			if ( !bShouldIgnorePackage )
-			{
-				UWorld*		AssociatedWorld			= UWorld::FindWorldInPackage(Package);
-				const bool	bIsMapPackage			= AssociatedWorld != NULL;
-
-				// Ignore map packages, they are caught above.
-				bShouldIgnorePackage |= bIsMapPackage; 
-
-				// Ignore packages with long, invalid names. This culls out packages with paths in read-only roots such as /Temp.
-				bShouldIgnorePackage |= (!FPackageName::IsShortPackageName(Package->GetFName()) && !FPackageName::IsValidLongPackageName(Package->GetName(), /*bIncludeReadOnlyRoots=*/false));
-			}
+		GetDirtyContentPackages(PackagesToSave);
+	}
 
-			if( !bShouldIgnorePackage )
-			{
-				// Count the number of packages to not ignore.
-				NumPackagesNotIgnored += (PackagesNotSavedDuringSaveAll.Find(Package->GetName())==NULL) ? 1 : 0;
+	// Need to track the number of packages we're not ignoring for save.
+	int32 NumPackagesNotIgnored = 0;
 
-				PackagesToSave.Add( Package );
-			}
-		}
+	for (auto* Package : PackagesToSave)
+	{
+		// Count the number of packages to not ignore.
+		NumPackagesNotIgnored += (PackagesNotSavedDuringSaveAll.Find(Package->GetName()) == NULL) ? 1 : 0;
 	}
 
 	return InternalSavePackages(PackagesToSave, NumPackagesNotIgnored, bPromptUserToSave, bFastSave, bNotifyNoPackagesSaved, bOutPackagesNeededSaving);
@@ -3015,4 +2971,58 @@ FString FEditorFileUtils::ExtractPackageName(const FString& ObjectPath)
 	return ObjectPath;
 }
 
+void FEditorFileUtils::GetDirtyWorldPackages(TArray<UPackage*>& OutDirtyPackages)
+{
+	// If we are saving map packages, collect all valid worlds and see if their package is dirty
+	TArray<UWorld*> Worlds;
+	EditorLevelUtils::GetWorlds(GWorld, Worlds, true);
+
+	for (int32 WorldIdx = 0; WorldIdx < Worlds.Num(); ++WorldIdx)
+	{
+		UPackage* WorldPackage = Worlds[WorldIdx]->GetOutermost();
+		if (WorldPackage->IsDirty() && (WorldPackage->PackageFlags & PKG_PlayInEditor) == 0
+			&& !WorldPackage->HasAnyFlags(RF_Transient))
+		{
+			// IF the package is dirty and its not a pie package, add the world package to the list of packages to save
+			OutDirtyPackages.Add(WorldPackage);
+		}
+	}
+}
+
+void FEditorFileUtils::GetDirtyContentPackages(TArray<UPackage*>& OutDirtyPackages)
+{
+	// Make a list of all content packages that we should save
+	for (TObjectIterator<UPackage> It; It; ++It)
+	{
+		UPackage*	Package = *It;
+		bool		bShouldIgnorePackage = false;
+
+		// Only look at root packages.
+		bShouldIgnorePackage |= Package->GetOuter() != NULL;
+		// Don't try to save "Transient" package.
+		bShouldIgnorePackage |= Package == GetTransientPackage();
+		// Ignore PIE packages.
+		bShouldIgnorePackage |= (Package->PackageFlags & PKG_PlayInEditor) != 0;
+		// Ignore packages that haven't been modified.
+		bShouldIgnorePackage |= !Package->IsDirty();
+
+		if (!bShouldIgnorePackage)
+		{
+			UWorld*		AssociatedWorld = UWorld::FindWorldInPackage(Package);
+			const bool	bIsMapPackage = AssociatedWorld != NULL;
+
+			// Ignore map packages, they are caught above.
+			bShouldIgnorePackage |= bIsMapPackage;
+
+			// Ignore packages with long, invalid names. This culls out packages with paths in read-only roots such as /Temp.
+			bShouldIgnorePackage |= (!FPackageName::IsShortPackageName(Package->GetFName()) && !FPackageName::IsValidLongPackageName(Package->GetName(), /*bIncludeReadOnlyRoots=*/false));
+		}
+
+		if (!bShouldIgnorePackage)
+		{
+			OutDirtyPackages.Add(Package);
+		}
+	}
+}
+
 #undef LOCTEXT_NAMESPACE
diff --git a/Engine/Source/Editor/UnrealEd/Public/FileHelpers.h b/Engine/Source/Editor/UnrealEd/Public/FileHelpers.h
--- a/Engine/Source/Editor/UnrealEd/Public/FileHelpers.h
+++ b/Engine/Source/Editor/UnrealEd/Public/FileHelpers.h
@@ -156,6 +156,20 @@ public:
 	*/
 	UNREALED_API static bool SaveDirtyContentPackages(TArray<UClass*>& SaveContentClasses, const bool bPromptUserToSave, const bool bFastSave = false, const bool bNotifyNoPackagesSaved = false);
 
+	/**
+	 * Appends array with all currently dirty world packages.
+	 *
+	 * @param OutDirtyPackages Array to append dirty packages to.
+	 */
+	UNREALED_API static void GetDirtyWorldPackages(TArray<UPackage*>& OutDirtyPackages);
+
+	/**
+	 * Appends array with all currently dirty content packages.
+	 *
+	 * @param OutDirtyPackages Array to append dirty packages to.
+	 */
+	UNREALED_API static void GetDirtyContentPackages(TArray<UPackage*>& OutDirtyPackages);
+
 	/**
 	 * Saves the active level, prompting the use for checkout if necessary.
 	 *
diff --git a/Engine/Source/Runtime/AssetRegistry/Private/AssetDataGatherer.cpp b/Engine/Source/Runtime/AssetRegistry/Private/AssetDataGatherer.cpp
--- a/Engine/Source/Runtime/AssetRegistry/Private/AssetDataGatherer.cpp
+++ b/Engine/Source/Runtime/AssetRegistry/Private/AssetDataGatherer.cpp
@@ -3,7 +3,7 @@
 #include "AssetRegistryPCH.h"
 
 #define MAX_FILES_TO_PROCESS_BEFORE_FLUSH 250
-#define CACHE_SERIALIZATION_VERSION 1
+#define CACHE_SERIALIZATION_VERSION 2
 
 FAssetDataGatherer::FAssetDataGatherer(const TArray<FString>& InPaths, bool bInIsSynchronous, bool bInLoadAndSaveCache)
 	: StopTaskCounter( 0 )
diff --git a/Engine/Source/Runtime/AssetRegistry/Private/AssetRegistry.cpp b/Engine/Source/Runtime/AssetRegistry/Private/AssetRegistry.cpp
--- a/Engine/Source/Runtime/AssetRegistry/Private/AssetRegistry.cpp
+++ b/Engine/Source/Runtime/AssetRegistry/Private/AssetRegistry.cpp
@@ -1418,6 +1418,13 @@ void FAssetRegistry::DependencyDataGathered(const double TickStartTime, TArray<F
 		{
 			PackageDependencies.Add( Result.GetImportPackageName(ImportIdx) );
 		}
+		for (int32 StringAssetRefIdx = 0; StringAssetRefIdx < Result.StringAssetReferencesMap.Num(); ++StringAssetRefIdx)
+		{
+			FString PackageName, ObjName;
+			Result.StringAssetReferencesMap[StringAssetRefIdx].Split(".", &PackageName, &ObjName);
+			PackageDependencies.Add(*PackageName);
+		}
+
 
 		// Doubly-link all new dependencies for this package
 		for (auto NewDependsIt = PackageDependencies.CreateConstIterator(); NewDependsIt; ++NewDependsIt)
diff --git a/Engine/Source/Runtime/AssetRegistry/Private/PackageDependencyData.h b/Engine/Source/Runtime/AssetRegistry/Private/PackageDependencyData.h
--- a/Engine/Source/Runtime/AssetRegistry/Private/PackageDependencyData.h
+++ b/Engine/Source/Runtime/AssetRegistry/Private/PackageDependencyData.h
@@ -23,8 +23,7 @@ public:
 		// serialize out the asset info
 		Ar << DependencyData.PackageName;
 		Ar << DependencyData.ImportMap;
-		//Ar << DependencyData.ExportMap;
-		//Ar << DependencyData.DependsMap;
+		Ar << DependencyData.StringAssetReferencesMap;
 
 		return Ar;
 	}
diff --git a/Engine/Source/Runtime/AssetRegistry/Private/PackageReader.cpp b/Engine/Source/Runtime/AssetRegistry/Private/PackageReader.cpp
--- a/Engine/Source/Runtime/AssetRegistry/Private/PackageReader.cpp
+++ b/Engine/Source/Runtime/AssetRegistry/Private/PackageReader.cpp
@@ -260,6 +260,7 @@ bool FPackageReader::ReadDependencyData (FPackageDependencyData& OutDependencyDa
 
 	SerializeNameMap();
 	SerializeImportMap(OutDependencyData.ImportMap);
+	SerializeStringAssetReferencesMap(OutDependencyData.StringAssetReferencesMap);
 
 	return true;
 }
@@ -298,6 +299,21 @@ void FPackageReader::SerializeImportMap(TArray<FObjectImport>& OutImportMap)
 	}
 }
 
+void FPackageReader::SerializeStringAssetReferencesMap(TArray<FString>& OutStringAssetReferencesMap)
+{
+	if (UE4Ver() >= VER_UE4_ADD_STRING_ASSET_REFERENCES_MAP && PackageFileSummary.StringAssetReferencesCount > 0)
+	{
+		Seek(PackageFileSummary.StringAssetReferencesOffset);
+
+		for (int32 ReferenceIdx = 0; ReferenceIdx < PackageFileSummary.StringAssetReferencesCount; ++ReferenceIdx)
+		{
+			FString Buf;
+			*this << Buf;
+			OutStringAssetReferencesMap.Add(MoveTemp(Buf));
+		}
+	}
+}
+
 void FPackageReader::Serialize( void* V, int64 Length )
 {
 	check(Loader);
diff --git a/Engine/Source/Runtime/AssetRegistry/Private/PackageReader.h b/Engine/Source/Runtime/AssetRegistry/Private/PackageReader.h
--- a/Engine/Source/Runtime/AssetRegistry/Private/PackageReader.h
+++ b/Engine/Source/Runtime/AssetRegistry/Private/PackageReader.h
@@ -23,6 +23,7 @@ public:
 	/** Serializers for different package maps */
 	void SerializeNameMap();
 	void SerializeImportMap(TArray<FObjectImport>& OutImportMap);
+	void SerializeStringAssetReferencesMap(TArray<FString>& OutStringAssetReferencesMap);
 
 	// Farchive implementation to redirect requests to the Loader
 	void Serialize( void* V, int64 Length );
diff --git a/Engine/Source/Runtime/Core/Private/Serialization/Archive.cpp b/Engine/Source/Runtime/Core/Private/Serialization/Archive.cpp
--- a/Engine/Source/Runtime/Core/Private/Serialization/Archive.cpp
+++ b/Engine/Source/Runtime/Core/Private/Serialization/Archive.cpp
@@ -154,6 +154,13 @@ FArchive& FArchive::operator<<( class FAssetPtr& AssetPtr )
 	return *this;
 }
 
+FArchive& FArchive::operator<<(struct FStringAssetReference& Value)
+{
+	// The base FArchive does not implement this method. Use FArchiveUOBject instead.
+	UE_LOG(LogSerialization, Fatal, TEXT("FArchive does not support FAssetPtr serialization. Use FArchiveUObject instead."));
+	return *this;
+}
+
 const FCustomVersionContainer& FArchive::GetCustomVersions() const
 {
 	if (bCustomVersionsAreReset)
diff --git a/Engine/Source/Runtime/Core/Public/Serialization/ArchiveBase.h b/Engine/Source/Runtime/Core/Public/Serialization/ArchiveBase.h
--- a/Engine/Source/Runtime/Core/Public/Serialization/ArchiveBase.h
+++ b/Engine/Source/Runtime/Core/Public/Serialization/ArchiveBase.h
@@ -92,6 +92,15 @@ public:
 	 */
 	virtual FArchive& operator<<(class FAssetPtr& Value);
 
+	/**
+	 * Serializes string asset reference from or into this archive.
+	 *
+	 * @param Value String asset reference to serialize.
+	 *
+	 * @return This instance.
+	 */
+	virtual FArchive& operator<<(struct FStringAssetReference& Value);
+
 
 public:
 
@@ -1061,6 +1070,11 @@ public:
 		return InnerArchive << Value;
 	}
 
+	virtual FArchive& operator<<(struct FStringAssetReference& Value) override
+	{
+		return InnerArchive << Value;
+	}
+
 	virtual void Serialize(void* V, int64 Length)
 	{
 		InnerArchive.Serialize(V, Length);
diff --git a/Engine/Source/Runtime/Core/Public/UObject/ObjectVersion.h b/Engine/Source/Runtime/Core/Public/UObject/ObjectVersion.h
--- a/Engine/Source/Runtime/Core/Public/UObject/ObjectVersion.h
+++ b/Engine/Source/Runtime/Core/Public/UObject/ObjectVersion.h
@@ -580,6 +580,8 @@ enum EUnrealEngineObjectUE4Version
 	VER_UE4_FIX_MATERIAL_PROPERTY_OVERRIDE_SERIALIZE,
 	// Addition of linear color sampler. color sample type is changed to linear sampler if source texture !sRGB
 	VER_UE4_ADD_LINEAR_COLOR_SAMPLER,
+	// Added StringAssetReferencesMap to support renames of FStringAssetReference properties.
+	VER_UE4_ADD_STRING_ASSET_REFERENCES_MAP,
 
 	// -----<new versions can be added before this line>-------------------------------------------------
 	// - this needs to be the last line (see note below)
diff --git a/Engine/Source/Runtime/CoreUObject/Private/Misc/StringAssetReference.cpp b/Engine/Source/Runtime/CoreUObject/Private/Misc/StringAssetReference.cpp
--- a/Engine/Source/Runtime/CoreUObject/Private/Misc/StringAssetReference.cpp
+++ b/Engine/Source/Runtime/CoreUObject/Private/Misc/StringAssetReference.cpp
@@ -21,7 +21,7 @@ bool FStringAssetReference::Serialize(FArchive& Ar)
 		AssetLongPathname = FCoreDelegates::StringAssetReferenceSaving.Execute(AssetLongPathname);
 	}
 #endif // WITH_EDITOR
-	Ar << AssetLongPathname;
+	Ar << *this;
 #if WITH_EDITOR
 	if (Ar.IsLoading() && Ar.IsPersistent() && FCoreDelegates::StringAssetReferenceLoaded.IsBound())
 	{
diff --git a/Engine/Source/Runtime/CoreUObject/Private/Serialization/ArchiveUObject.cpp b/Engine/Source/Runtime/CoreUObject/Private/Serialization/ArchiveUObject.cpp
--- a/Engine/Source/Runtime/CoreUObject/Private/Serialization/ArchiveUObject.cpp
+++ b/Engine/Source/Runtime/CoreUObject/Private/Serialization/ArchiveUObject.cpp
@@ -65,6 +65,12 @@ FArchive& FArchiveUObject::operator<<( class FAssetPtr& AssetPtr )
 	return Ar;
 }
 
+FArchive& FArchiveUObject::operator<<(struct FStringAssetReference& Value)
+{
+	*this << Value.AssetLongPathname;
+	return *this;
+}
+
 /*----------------------------------------------------------------------------
 	FObjectAndNameAsStringProxyArchive.
 ----------------------------------------------------------------------------*/
diff --git a/Engine/Source/Runtime/CoreUObject/Private/Serialization/DuplicateDataReader.cpp b/Engine/Source/Runtime/CoreUObject/Private/Serialization/DuplicateDataReader.cpp
--- a/Engine/Source/Runtime/CoreUObject/Private/Serialization/DuplicateDataReader.cpp
+++ b/Engine/Source/Runtime/CoreUObject/Private/Serialization/DuplicateDataReader.cpp
@@ -65,7 +65,7 @@ FArchive& FDuplicateDataReader::operator<<( FAssetPtr& AssetPtr)
 {
 	FArchive& Ar = *this;
 	FStringAssetReference ID;
-	Ar << ID;
+	ID.Serialize(Ar);
 
 	AssetPtr = ID;
 	return Ar;
diff --git a/Engine/Source/Runtime/CoreUObject/Private/Serialization/DuplicateDataWriter.cpp b/Engine/Source/Runtime/CoreUObject/Private/Serialization/DuplicateDataWriter.cpp
--- a/Engine/Source/Runtime/CoreUObject/Private/Serialization/DuplicateDataWriter.cpp
+++ b/Engine/Source/Runtime/CoreUObject/Private/Serialization/DuplicateDataWriter.cpp
@@ -54,7 +54,8 @@ FArchive& FDuplicateDataWriter::operator<<(FLazyObjectPtr& LazyObjectPtr)
 FArchive& FDuplicateDataWriter::operator<<(FAssetPtr& AssetPtr)
 {
 	FStringAssetReference ID = AssetPtr.GetUniqueID();
-	return *this << ID;
+	ID.Serialize(*this);
+	return *this;
 }
 
 /**
diff --git a/Engine/Source/Runtime/CoreUObject/Private/Serialization/ObjectReader.cpp b/Engine/Source/Runtime/CoreUObject/Private/Serialization/ObjectReader.cpp
--- a/Engine/Source/Runtime/CoreUObject/Private/Serialization/ObjectReader.cpp
+++ b/Engine/Source/Runtime/CoreUObject/Private/Serialization/ObjectReader.cpp
@@ -36,12 +36,17 @@ FArchive& FObjectReader::operator<<( class FAssetPtr& AssetPtr )
 {
 	FArchive& Ar = *this;
 	FStringAssetReference ID;
-	Ar << ID;
+	ID.Serialize(Ar);
 
 	AssetPtr = ID;
 	return Ar;
 }
 
+FArchive& FObjectReader::operator<<(FStringAssetReference& Value)
+{
+	return *this << Value.AssetLongPathname;
+}
+
 FString FObjectReader::GetArchiveName() const
 {
 	return TEXT("FObjectReader");
diff --git a/Engine/Source/Runtime/CoreUObject/Private/Serialization/ObjectWriter.cpp b/Engine/Source/Runtime/CoreUObject/Private/Serialization/ObjectWriter.cpp
--- a/Engine/Source/Runtime/CoreUObject/Private/Serialization/ObjectWriter.cpp
+++ b/Engine/Source/Runtime/CoreUObject/Private/Serialization/ObjectWriter.cpp
@@ -29,7 +29,13 @@ FArchive& FObjectWriter::operator<<( class FLazyObjectPtr& LazyObjectPtr )
 FArchive& FObjectWriter::operator<<( class FAssetPtr& AssetPtr )
 {
 	FStringAssetReference ID = AssetPtr.GetUniqueID();
-	return *this << ID;
+	ID.Serialize(*this);
+	return *this;
+}
+
+FArchive& FObjectWriter::operator<<(FStringAssetReference& Value)
+{
+	return *this << Value.AssetLongPathname;
 }
 
 FString FObjectWriter::GetArchiveName() const
diff --git a/Engine/Source/Runtime/CoreUObject/Private/UObject/CoreNet.cpp b/Engine/Source/Runtime/CoreUObject/Private/UObject/CoreNet.cpp
--- a/Engine/Source/Runtime/CoreUObject/Private/UObject/CoreNet.cpp
+++ b/Engine/Source/Runtime/CoreUObject/Private/UObject/CoreNet.cpp
@@ -174,6 +174,11 @@ FArchive& FNetBitWriter::operator<<( UObject*& Object )
 	return *this;
 }
 
+FArchive& FNetBitWriter::operator<<(FStringAssetReference& Value)
+{
+	return *this << Value.AssetLongPathname;
+}
+
 
 // ----------------------------------------------------------------
 //	FNetBitReader
@@ -196,6 +201,11 @@ FArchive& FNetBitReader::operator<<( class FName& N )
 	return *this;
 }
 
+FArchive& FNetBitReader::operator<<(FStringAssetReference& Value)
+{
+	return *this << Value.AssetLongPathname;
+}
+
 static const TCHAR * GLastRPCFailedReason = NULL;
 
 void RPC_ResetLastFailedReason()
diff --git a/Engine/Source/Runtime/CoreUObject/Private/UObject/Linker.cpp b/Engine/Source/Runtime/CoreUObject/Private/UObject/Linker.cpp
--- a/Engine/Source/Runtime/CoreUObject/Private/UObject/Linker.cpp
+++ b/Engine/Source/Runtime/CoreUObject/Private/UObject/Linker.cpp
@@ -141,6 +141,11 @@ void ULinker::Serialize( FArchive& Ar )
 		Ar << ImportMap;
 		Ar << ExportMap;
 		Ar << DependsMap;
+
+		if (Ar.IsSaving() || Ar.UE4Ver() >= VER_UE4_ADD_STRING_ASSET_REFERENCES_MAP)
+		{
+			Ar << StringAssetReferencesMap;
+		}
 	}
 
 	// Prevent garbage collecting of linker's names and package.
diff --git a/Engine/Source/Runtime/CoreUObject/Private/UObject/LinkerSave.cpp b/Engine/Source/Runtime/CoreUObject/Private/UObject/LinkerSave.cpp
--- a/Engine/Source/Runtime/CoreUObject/Private/UObject/LinkerSave.cpp
+++ b/Engine/Source/Runtime/CoreUObject/Private/UObject/LinkerSave.cpp
@@ -165,7 +165,9 @@ FArchive& ULinkerSave::operator<<( FAssetPtr& AssetPtr)
 	{
 		ID = AssetPtr.GetUniqueID();
 	}
-	return *this << ID;
+
+	ID.Serialize(*this);
+	return *this;
 }
 
 IMPLEMENT_CORE_INTRINSIC_CLASS(ULinkerSave, ULinker,
diff --git a/Engine/Source/Runtime/CoreUObject/Private/UObject/PackageFileSummary.cpp b/Engine/Source/Runtime/CoreUObject/Private/UObject/PackageFileSummary.cpp
--- a/Engine/Source/Runtime/CoreUObject/Private/UObject/PackageFileSummary.cpp
+++ b/Engine/Source/Runtime/CoreUObject/Private/UObject/PackageFileSummary.cpp
@@ -122,11 +122,17 @@ FArchive& operator<<( FArchive& Ar, FPackageFileSummary& Sum )
 		Ar << Sum.ExportCount   << Sum.ExportOffset;
 		Ar << Sum.ImportCount   << Sum.ImportOffset;
 		Ar << Sum.DependsOffset;
+
 		if (Ar.IsLoading() && (Sum.FileVersionUE3 < VER_MIN_ENGINE_UE3 || Sum.FileVersionUE4 < VER_UE4_OLDEST_LOADABLE_PACKAGE || Sum.FileVersionUE4 > GPackageFileUE4Version))
 		{
 			return Ar; // we can't safely load more than this because the below was different in older files.
 		}
 
+		if (Ar.IsSaving() || Sum.FileVersionUE4 >= VER_UE4_ADD_STRING_ASSET_REFERENCES_MAP)
+		{
+			Ar << Sum.StringAssetReferencesCount << Sum.StringAssetReferencesOffset;
+		}
+
 		Ar << Sum.ThumbnailTableOffset;
 
 		int32 GenerationCount = Sum.Generations.Num();
diff --git a/Engine/Source/Runtime/CoreUObject/Private/UObject/SavePackage.cpp b/Engine/Source/Runtime/CoreUObject/Private/UObject/SavePackage.cpp
--- a/Engine/Source/Runtime/CoreUObject/Private/UObject/SavePackage.cpp
+++ b/Engine/Source/Runtime/CoreUObject/Private/UObject/SavePackage.cpp
@@ -518,9 +518,10 @@ class FArchiveSaveTagImports : public FArchiveUObject
 public:
 	ULinkerSave* Linker;
 	TArray<UObject*> Dependencies;
+	TArray<FString>& StringAssetReferencesMap;
 
-	FArchiveSaveTagImports( ULinkerSave* InLinker)
-	: Linker( InLinker )
+	FArchiveSaveTagImports(ULinkerSave* InLinker, TArray<FString>& StringAssetReferencesMap)
+		: Linker(InLinker), StringAssetReferencesMap(StringAssetReferencesMap)
 	{
 		ArIsSaving				= true;
 		ArIsPersistent			= true;
@@ -532,10 +533,15 @@ public:
 			SetCookingTarget(Linker->CookingTarget());
 		}
 	}
-	FArchive& operator<<( UObject*& Obj );
-	FArchive& operator<<( FLazyObjectPtr& LazyObjectPtr);
-	FArchive& operator<<( FAssetPtr& AssetPtr);
-	FArchive& operator<<( FName& Name )
+	FArchive& operator<<(UObject*& Obj) override;
+	FArchive& operator<<(FLazyObjectPtr& LazyObjectPtr) override;
+	FArchive& operator<<(FAssetPtr& AssetPtr) override;
+	FArchive& operator<<(FStringAssetReference& Value) override
+	{
+		StringAssetReferencesMap.Add(Value.ToString());
+		return *this;
+	}
+	FArchive& operator<<(FName& Name) override
 	{
 		SavePackageState->MarkNameAsReferenced(Name);
 		return *this;
@@ -2464,7 +2470,7 @@ bool UPackage::SavePackage( UPackage* InOuter, UObject* Base, EObjectFlags TopLe
 		TempFilename = FPaths::CreateTempFilename(*FPaths::GameSavedDir(), *BaseFilename);
 
 		// Init.
-		const int32 SimpleSaveSteps = 28;
+		const int32 SimpleSaveSteps = 29;
 		const int32 ExportSaveSteps = 101;
 		const int32 TotalSaveSteps = SimpleSaveSteps + ExportSaveSteps;
 		int32 CurSaveStep = 0;
@@ -2638,6 +2644,7 @@ bool UPackage::SavePackage( UPackage* InOuter, UObject* Base, EObjectFlags TopLe
 				}
 
 				// Import objects & names.
+				TArray<FString> StringAssetReferencesMap;
 				{
 					TArray<UObject*> TagExpObjects;
 					GetObjectsWithAnyMarks(TagExpObjects, OBJECTMARK_TagExp);
@@ -2646,7 +2653,7 @@ bool UPackage::SavePackage( UPackage* InOuter, UObject* Base, EObjectFlags TopLe
 						UObject* Obj = TagExpObjects[Index];
 						check(Obj->HasAnyMarks(OBJECTMARK_TagExp));
 						// Build list.
-						FArchiveSaveTagImports ImportTagger( Linker );
+						FArchiveSaveTagImports ImportTagger(Linker, StringAssetReferencesMap);
 						ImportTagger.SetPortFlags(ComparisonFlags);
 						ImportTagger.SetFilterEditorOnly(FilterEditorOnly);
 
@@ -3221,6 +3228,19 @@ bool UPackage::SavePackage( UPackage* InOuter, UObject* Base, EObjectFlags TopLe
 
 
 				UE_LOG_COOK_TIME(TEXT("Serialize Dependency Map"));
+
+				if (EndSavingIfCancelled(Linker, TempFilename)) { return false; }
+				GWarn->UpdateProgress(++CurSaveStep, TotalSaveSteps);
+
+				// Save string asset reference map
+				Linker->Summary.StringAssetReferencesOffset = Linker->Tell();
+				Linker->Summary.StringAssetReferencesCount = StringAssetReferencesMap.Num();
+				for (int32 i = 0; i < StringAssetReferencesMap.Num(); i++)
+				{
+					*Linker << StringAssetReferencesMap[i];
+				}
+
+				UE_LOG_COOK_TIME(TEXT("Serialize StringAssetReference Map"));
 	
 				// Save thumbnails
 				UPackage::SaveThumbnails( InOuter, Linker );
diff --git a/Engine/Source/Runtime/CoreUObject/Public/Misc/StringAssetReference.h b/Engine/Source/Runtime/CoreUObject/Public/Misc/StringAssetReference.h
--- a/Engine/Source/Runtime/CoreUObject/Public/Misc/StringAssetReference.h
+++ b/Engine/Source/Runtime/CoreUObject/Public/Misc/StringAssetReference.h
@@ -85,12 +85,6 @@ struct COREUOBJECT_API FStringAssetReference
 		return GetTypeHash(This.AssetLongPathname);
 	}
 
-	friend FArchive& operator<<(FArchive& Ar,FStringAssetReference& P)
-	{
-		P.Serialize(Ar);
-		return Ar;
-	}
-
 	/** Code needed by AssetPtr to track rather object references should be rechecked */
 
 	static int32 GetCurrentTag()
diff --git a/Engine/Source/Runtime/CoreUObject/Public/Serialization/ArchiveUObject.h b/Engine/Source/Runtime/CoreUObject/Public/Serialization/ArchiveUObject.h
--- a/Engine/Source/Runtime/CoreUObject/Public/Serialization/ArchiveUObject.h
+++ b/Engine/Source/Runtime/CoreUObject/Public/Serialization/ArchiveUObject.h
@@ -99,6 +99,7 @@ public:
 	COREUOBJECT_API virtual FArchive& operator<<( class UObject*& Res ) override;
 	COREUOBJECT_API virtual FArchive& operator<<( FLazyObjectPtr& LazyObjectPtr );
 	COREUOBJECT_API virtual FArchive& operator<<( FAssetPtr& AssetPtr );
+	COREUOBJECT_API virtual FArchive& operator<<(FStringAssetReference& AssetPtr) override;
 	COREUOBJECT_API virtual FString GetArchiveName() const override;
 	// End of FArchive interface
 
@@ -131,6 +132,7 @@ public:
 	COREUOBJECT_API virtual FArchive& operator<<( class UObject*& Res ) override;
 	COREUOBJECT_API virtual FArchive& operator<<( FLazyObjectPtr& LazyObjectPtr );
 	COREUOBJECT_API virtual FArchive& operator<<( FAssetPtr& AssetPtr );
+	COREUOBJECT_API virtual FArchive& operator<<(FStringAssetReference& AssetPtr) override;
 	COREUOBJECT_API virtual FString GetArchiveName() const override;
 	// End of FArchive interface
 
diff --git a/Engine/Source/Runtime/CoreUObject/Public/Serialization/ArchiveUObjectBase.h b/Engine/Source/Runtime/CoreUObject/Public/Serialization/ArchiveUObjectBase.h
--- a/Engine/Source/Runtime/CoreUObject/Public/Serialization/ArchiveUObjectBase.h
+++ b/Engine/Source/Runtime/CoreUObject/Public/Serialization/ArchiveUObjectBase.h
@@ -16,5 +16,6 @@ public:
 	// Begin FArchive Interface
 	virtual FArchive& operator<< (class FLazyObjectPtr& Value) override;
 	virtual FArchive& operator<< (class FAssetPtr& Value) override;
+	virtual FArchive& operator<< (struct FStringAssetReference& Value) override;
 	// End FArchive Interface
 };
diff --git a/Engine/Source/Runtime/CoreUObject/Public/UObject/CoreNet.h b/Engine/Source/Runtime/CoreUObject/Public/UObject/CoreNet.h
--- a/Engine/Source/Runtime/CoreUObject/Public/UObject/CoreNet.h
+++ b/Engine/Source/Runtime/CoreUObject/Public/UObject/CoreNet.h
@@ -230,6 +230,7 @@ public:
 
 	FArchive& operator<<( FName& Name );
 	FArchive& operator<<( UObject*& Object );
+	FArchive& operator<<(FStringAssetReference& Value) override;
 };
 
 /**
@@ -247,6 +248,7 @@ public:
 
 	FArchive& operator<<( FName& Name );
 	FArchive& operator<<( UObject*& Object );
+	FArchive& operator<<(FStringAssetReference& Value) override;
 };
 
 
diff --git a/Engine/Source/Runtime/CoreUObject/Public/UObject/Linker.h b/Engine/Source/Runtime/CoreUObject/Public/UObject/Linker.h
--- a/Engine/Source/Runtime/CoreUObject/Public/UObject/Linker.h
+++ b/Engine/Source/Runtime/CoreUObject/Public/UObject/Linker.h
@@ -588,7 +588,7 @@ public:
 	/**
 	 * Number of imports contained in this package
 	 */
-	int32     ImportCount;
+	int32		ImportCount;
 
 	/**
 	 * Location into the file on disk for the ImportMap data
@@ -600,6 +600,16 @@ public:
 	*/
 	int32		DependsOffset;
 
+	/**
+	 * Number of references contained in this package
+	 */
+	int32		StringAssetReferencesCount;
+
+	/**
+	 * Location into the file on disk for the string asset references map data
+	 */
+	int32		StringAssetReferencesOffset;
+
 	/**
 	 * Thumbnail table offset
 	 */
@@ -723,6 +733,8 @@ public:
 	TArray<FObjectExport> ExportMap;
 	/** List of dependency lists for each export */
 	TArray<TArray<FPackageIndex> > DependsMap;
+	/** Map that holds info about string asset references from the package. */
+	TArray<FString> StringAssetReferencesMap;
 
 	/**
 	 * Check that this Index is non-null and return an import or export
