Fixed rotations offset by a pivot location not working correctly

I had previously removed some code inside ApplyDeltaToActor which was causing surface snapping to misbehave (CL 2223237). This code is actually required for rotations round a pivot to work correctly, so has been reinstated. I have changed the surface snapping code to set world transforms directly.

diff --git a/Engine/Source/Editor/UnrealEd/Private/Editor.cpp b/Engine/Source/Editor/UnrealEd/Private/Editor.cpp
--- a/Engine/Source/Editor/UnrealEd/Private/Editor.cpp
+++ b/Engine/Source/Editor/UnrealEd/Private/Editor.cpp
@@ -2160,6 +2160,16 @@ void UEditorEngine::ApplyDeltaToActor(AActor* InActor,
 					InActor->SetActorRotation( InDeltaRot );
 				}
 			}
+
+			if ( bDelta )
+			{
+				FVector NewActorLocation = InActor->GetActorLocation();
+				NewActorLocation -= GLevelEditorModeTools().PivotLocation;
+				NewActorLocation = FRotationMatrix( InDeltaRot ).TransformPosition( NewActorLocation );
+				NewActorLocation += GLevelEditorModeTools().PivotLocation;
+				NewActorLocation -= InActor->GetActorLocation();
+				InActor->EditorApplyTranslation( NewActorLocation, bAltDown, bShiftDown, bControlDown );
+			}
 		}
 	}
 
diff --git a/Engine/Source/Editor/UnrealEd/Private/LevelEditorViewport.cpp b/Engine/Source/Editor/UnrealEd/Private/LevelEditorViewport.cpp
--- a/Engine/Source/Editor/UnrealEd/Private/LevelEditorViewport.cpp
+++ b/Engine/Source/Editor/UnrealEd/Private/LevelEditorViewport.cpp
@@ -1986,15 +1986,18 @@ void FLevelEditorViewportClient::ProjectActorsIntoWorld(const TArray<AActor*>& A
 			const auto ActorLocation = Actor->GetActorLocation();
 			const auto ActorRotation = Actor->GetActorRotation().Quaternion();
 
-			FTransform* PreDragActorTransform = PreDragActorTransforms.Find(Actor);
+			const FTransform* PreDragActorTransform = PreDragActorTransforms.Find(Actor);
 			check(PreDragActorTransform);
 
 			// Compute the surface aligned transform. Note we do not use the snapped version here as our DragDelta is already snapped
-			const FTransform ActorTransform = FActorPositioning::GetSurfaceAlignedTransform(Factory,
+			FTransform ActorTransform = FActorPositioning::GetSurfaceAlignedTransform(Factory,
 				TraceResult.Location, TraceResult.SurfaceNormal, Actor->GetPlacementExtent(), *PreDragActorTransform);
 			
-			const FRotator DeltaRotation(ActorTransform.GetRotation() * ActorRotation.Inverse());
-			ApplyDeltaToActor(Actor, ActorTransform.GetLocation() - ActorLocation, DeltaRotation, FVector(0.f, 0.f, 0.f));
+			ActorTransform.SetScale3D(Actor->GetActorScale3D());
+			if (auto* RootComponent = Actor->GetRootComponent())
+			{
+				RootComponent->SetWorldTransform(ActorTransform);
+			}
 		}
 		else
 		{
