Destructible meshes now snap to surfaces correctly

diff --git a/Engine/Source/Editor/UnrealEd/Classes/ActorFactories/ActorFactoryDestructible.h b/Engine/Source/Editor/UnrealEd/Classes/ActorFactories/ActorFactoryDestructible.h
--- a/Engine/Source/Editor/UnrealEd/Classes/ActorFactories/ActorFactoryDestructible.h
+++ b/Engine/Source/Editor/UnrealEd/Classes/ActorFactories/ActorFactoryDestructible.h
@@ -15,6 +15,7 @@ class UActorFactoryDestructible : public UActorFactory
 	virtual void PostCreateBlueprint( UObject* Asset, AActor* CDO ) override;
 	virtual bool CanCreateActorFrom( const FAssetData& AssetData, FText& OutErrorMsg ) override;
 	virtual UObject* GetAssetFromActorInstance(AActor* ActorInstance) override;
+	virtual FQuat AlignObjectToSurfaceNormal(const FVector& InSurfaceNormal, const FQuat& ActorRotation) const override;
 	// End UActorFactory Interface
 };
 
diff --git a/Engine/Source/Editor/UnrealEd/Private/Factories/ActorFactory.cpp b/Engine/Source/Editor/UnrealEd/Private/Factories/ActorFactory.cpp
--- a/Engine/Source/Editor/UnrealEd/Private/Factories/ActorFactory.cpp
+++ b/Engine/Source/Editor/UnrealEd/Private/Factories/ActorFactory.cpp
@@ -1430,6 +1430,7 @@ UActorFactoryDestructible::UActorFactoryDestructible(const class FPostConstructI
 {
 	DisplayName = LOCTEXT("DestructibleDisplayName", "Destructible");
 	NewActorClass = ADestructibleActor::StaticClass();
+	bUseSurfaceOrientation = true;
 }
 
 bool UActorFactoryDestructible::CanCreateActorFrom( const FAssetData& AssetData, FText& OutErrorMsg )
@@ -1480,6 +1481,12 @@ void UActorFactoryDestructible::PostCreateBlueprint( UObject* Asset, AActor* CDO
 	}
 }
 
+FQuat UActorFactoryDestructible::AlignObjectToSurfaceNormal(const FVector& InSurfaceNormal, const FQuat& ActorRotation) const
+{
+	// Meshes align the Z (up) axis with the surface normal
+	return FindActorAlignmentRotation(ActorRotation, FVector(0.f, 0.f, 1.f), InSurfaceNormal);
+}
+
 
 /*-----------------------------------------------------------------------------
 UActorFactoryVectorField
