Use GameMode static class for property info rather than requiring a game mode have been chosen already

diff --git a/Engine/Source/Editor/UnrealEd/Public/GameModeInfoCustomizer.h b/Engine/Source/Editor/UnrealEd/Public/GameModeInfoCustomizer.h
--- a/Engine/Source/Editor/UnrealEd/Public/GameModeInfoCustomizer.h
+++ b/Engine/Source/Editor/UnrealEd/Public/GameModeInfoCustomizer.h
@@ -42,15 +42,9 @@ public:
 	void CustomizeGameModeDefaultClass(IDetailGroup& Group, FName DefaultClassPropertyName)
 	{
 		// Find the metaclass of this property
-		UClass* MetaClass = UObject::StaticClass();
-		const UClass* GameModeClass = GetCurrentGameModeClass();
-		bool bAllowNone = false;
-		if (GameModeClass != NULL)
-		{
-			UClassProperty* ClassProp = FindFieldChecked<UClassProperty>(GameModeClass, DefaultClassPropertyName);
-			MetaClass = ClassProp->MetaClass;
-			bAllowNone = !(ClassProp->PropertyFlags & CPF_NoClear);
-		}
+		UClassProperty* ClassProp = FindFieldChecked<UClassProperty>(AGameMode::StaticClass(), DefaultClassPropertyName);
+		UClass* MetaClass = ClassProp->MetaClass;
+		const bool bAllowNone = !(ClassProp->PropertyFlags & CPF_NoClear);
 
 		TAttribute<bool> CanBrowseAtrribute = TAttribute<bool>::Create( TAttribute<bool>::FGetter::CreateSP( this, &FGameModeInfoCustomizer::CanBrowseDefaultClass, DefaultClassPropertyName) ) ;
 
