Crash: Blueprint Asset will cause crash on load in Standalone Mode or Packaged Game if placed in default leve.. ttp341388

diff --git a/Engine/Source/Editor/UnrealEd/Private/Kismet2/BlueprintEditorUtils.cpp b/Engine/Source/Editor/UnrealEd/Private/Kismet2/BlueprintEditorUtils.cpp
--- a/Engine/Source/Editor/UnrealEd/Private/Kismet2/BlueprintEditorUtils.cpp
+++ b/Engine/Source/Editor/UnrealEd/Private/Kismet2/BlueprintEditorUtils.cpp
@@ -3009,7 +3009,7 @@ bool FBlueprintEditorUtils::IsPinTypeValid(const FEdGraphPinType& Type)
 void FBlueprintEditorUtils::GetClassVariableList(const UBlueprint* Blueprint, TArray<FName>& VisibleVariables, bool bIncludePrivateVars) 
 {
 	// Existing variables in the parent class and above
-	check(Blueprint->bIsRegeneratingOnLoad || (Blueprint->SkeletonGeneratedClass != NULL));
+	check(!Blueprint->bHasBeenRegenerated || Blueprint->bIsRegeneratingOnLoad || (Blueprint->SkeletonGeneratedClass != NULL));
 	if (Blueprint->SkeletonGeneratedClass != NULL)
 	{
 		for (TFieldIterator<UProperty> PropertyIt(Blueprint->SkeletonGeneratedClass, EFieldIteratorFlags::IncludeSuper); PropertyIt; ++PropertyIt)
