Add link to import doc in the import window

diff --git a/Engine/Source/Editor/UnrealEd/Private/Fbx/FbxOptionWindow.cpp b/Engine/Source/Editor/UnrealEd/Private/Fbx/FbxOptionWindow.cpp
--- a/Engine/Source/Editor/UnrealEd/Private/Fbx/FbxOptionWindow.cpp
+++ b/Engine/Source/Editor/UnrealEd/Private/Fbx/FbxOptionWindow.cpp
@@ -44,9 +44,33 @@ public:
 
 				+SVerticalBox::Slot().AutoHeight().Padding(0, 2)
 				[
-					SNew(STextBlock)
-					.TextStyle(FEditorStyle::Get(), "FBXLargeFont")
-					.Text(LOCTEXT("FbxOptionWindow_SelectType", "Select Asset Type"))
+					SNew(SHorizontalBox)
+					
+					+SHorizontalBox::Slot()
+					.Padding(0, 10)
+					[
+						SNew(STextBlock)
+						.TextStyle(FEditorStyle::Get(), "FBXLargeFont")
+						.Text(LOCTEXT("FbxOptionWindow_SelectType", "Select Asset Type"))
+					]
+
+// 					+SHorizontalBox::Slot()
+// 					.HAlign(HAlign_Right)
+// 					[
+// 						SNew(SBox)
+// 						.VAlign(VAlign_Bottom)
+// 						[
+// 							SNew(STextBlock)
+// 							.Text(LOCTEXT("FbxOptionWindow_Help", "[Help]"))
+// 						]
+// 					]
+
+					+SHorizontalBox::Slot()
+					.HAlign(HAlign_Left)
+					.FillWidth(1)
+					[
+						IDocumentation::Get()->CreateAnchor(FString("Engine/Content/FBX/ImportOptions"))
+					]
 				]
 
 				+SVerticalBox::Slot().AutoHeight().Padding(0, 5)
