Fixed MMB+drag to move the pivot.
#ttp 342197 - CRITICAL: Regression: EDITOR: GROUP: Middle mouse plus drag does not set pivot
#branch UE4
#proj Editor.UnrealEd
Only set pivot location back to its default in FLevelEditorViewportClient::TrackingStopped if it wasn't moved independently.

diff --git a/Engine/Source/Editor/UnrealEd/Private/LevelEditorViewport.cpp b/Engine/Source/Editor/UnrealEd/Private/LevelEditorViewport.cpp
--- a/Engine/Source/Editor/UnrealEd/Private/LevelEditorViewport.cpp
+++ b/Engine/Source/Editor/UnrealEd/Private/LevelEditorViewport.cpp
@@ -2532,7 +2532,10 @@ void FLevelEditorViewportClient::TrackingStopped()
 			GEditor->BroadcastEndObjectMovement( *Actor );
 		}
 
-		GUnrealEd->UpdatePivotLocationForSelection();
+		if (!bPivotMovedIndependantly)
+		{
+			GUnrealEd->UpdatePivotLocationForSelection();
+		}
 	}
 
 	// End the transaction here if one was started in StartTransaction()
