Fixed crash when closing editor with Paint or Landscape mode active.
#ttp 341775 - CRITICAL: Regression: EDITOR: CRASH: Closing with Paint or Landscape mode enabled
#branch UE4
#proj Editor.UnrealEd
The default editor mode is not ever reset to default, so do this prior to editor shutdown (in ExitEditor()) and set the current editor mode to default so that the previous editor mode can be shutdown gracefully.

diff --git a/Engine/Source/Editor/UnrealEd/Private/UnrealEd.cpp b/Engine/Source/Editor/UnrealEd/Private/UnrealEd.cpp
--- a/Engine/Source/Editor/UnrealEd/Private/UnrealEd.cpp
+++ b/Engine/Source/Editor/UnrealEd/Private/UnrealEd.cpp
@@ -168,6 +168,7 @@ void EditorExit()
 	if( GLevelEditorModeTools().IsModeActive(FBuiltinEditorModes::EM_MeshPaint) ||
 		GLevelEditorModeTools().IsModeActive(FBuiltinEditorModes::EM_InterpEdit) )
 	{
+		GLevelEditorModeTools().SetDefaultMode(FBuiltinEditorModes::EM_Default);
 		GLevelEditorModeTools().ActivateDefaultMode();
 	}
 
