Fixed Undo Marquee Select in perspective viewport.
#ttp 341501 - EDITOR: Undo doesn't work after Shift+Marquee selecting additional assets
#branch UE4
#proj Editor.UnrealEd
Added Modify on SelectedObjects object prior to changing selection, within the select transaction.

diff --git a/Engine/Source/Editor/UnrealEd/Private/DragTool_FrustumSelect.cpp b/Engine/Source/Editor/UnrealEd/Private/DragTool_FrustumSelect.cpp
--- a/Engine/Source/Editor/UnrealEd/Private/DragTool_FrustumSelect.cpp
+++ b/Engine/Source/Editor/UnrealEd/Private/DragTool_FrustumSelect.cpp
@@ -69,7 +69,7 @@ void FDragTool_ActorFrustumSelect::EndDrag()
 	FScopedTransaction Transaction( NSLOCTEXT("ActorFrustumSelect", "MarqueeSelectTransation", "Marquee Select" ) );
 
 	bool bShouldSelect = true;
-
+	GEditor->GetSelectedActors()->Modify();
 
 	if( !bShiftDown )
 	{
