Removed transactions from importing normal maps.
#ttp 341181 - EDITOR: Normal textures are shown as being referenced even after a fresh import.
#branch UE4
#proj Editor.UnrealEd
#remove Removed scoped transactions from NormalMapIdentification::HandleAssetPostImport and NormalMapImportNotificationHandler::RevertSetting as they are unnecessary and lead to a confusing dialog when attempting to delete them.

diff --git a/Engine/Source/Editor/UnrealEd/Private/NormalMapIdentification.cpp b/Engine/Source/Editor/UnrealEd/Private/NormalMapIdentification.cpp
--- a/Engine/Source/Editor/UnrealEd/Private/NormalMapIdentification.cpp
+++ b/Engine/Source/Editor/UnrealEd/Private/NormalMapIdentification.cpp
@@ -3,7 +3,6 @@
 #include "UnrealEd.h"
 #include "NormalMapIdentification.h"
 #include "Factories/TextureFactory.h"
-#include "ScopedTransaction.h"
 
 #define NORMALMAP_IDENTIFICATION_TIMING	(1)
 
@@ -423,8 +422,6 @@ public:
 
 				{
 					Texture2D->SetFlags(RF_Transactional);
-
-					const FScopedTransaction Transaction( NSLOCTEXT("AutoNormalMapID", "UndoRevertingNormalMap", "Revert Normal Map Settings") );
 					Texture2D->Modify();
 					Texture2D->PreEditChange(NULL);
 					{
@@ -461,7 +458,6 @@ void NormalMapIdentification::HandleAssetPostImport( UFactory* InFactory, UObjec
 			// Set the compression settings and no gamma correction for a normal map
 			{
 				Texture->SetFlags(RF_Transactional);
-				const FScopedTransaction Transaction( NSLOCTEXT("AutoNormalMapID", "ImportAsNormalMap", "Import As Normal Map") );
 				Texture->Modify();
 				Texture->CompressionSettings = TC_Normalmap;
 				Texture->SRGB = false;
