Adding documentation excerpt name and link for material nodes

Added Shared link to base class and override excerpt name in MaterialGraphNode so that it can associate documentation with the MaterialExpression type

diff --git a/Engine/Source/Editor/UnrealEd/Classes/MaterialGraph/MaterialGraphNode.h b/Engine/Source/Editor/UnrealEd/Classes/MaterialGraph/MaterialGraphNode.h
--- a/Engine/Source/Editor/UnrealEd/Classes/MaterialGraph/MaterialGraphNode.h
+++ b/Engine/Source/Editor/UnrealEd/Classes/MaterialGraph/MaterialGraphNode.h
@@ -59,6 +59,7 @@ public:
 	virtual void PostPlacedNewNode() override;
 	virtual void OnRenameNode(const FString& NewName) override;
 	virtual void GetPinHoverText(const UEdGraphPin& Pin, FString& HoverTextOut) const override;
+	virtual FString GetDocumentationExcerptName() const override;
 	// End UEdGraphNode interface.
 
 	// UMaterialGraphNode_Base interface
diff --git a/Engine/Source/Editor/UnrealEd/Classes/MaterialGraph/MaterialGraphNode_Base.h b/Engine/Source/Editor/UnrealEd/Classes/MaterialGraph/MaterialGraphNode_Base.h
--- a/Engine/Source/Editor/UnrealEd/Classes/MaterialGraph/MaterialGraphNode_Base.h
+++ b/Engine/Source/Editor/UnrealEd/Classes/MaterialGraph/MaterialGraphNode_Base.h
@@ -45,6 +45,7 @@ class UMaterialGraphNode_Base : public UEdGraphNode
 	virtual void ReconstructNode() override;
 	virtual void AutowireNewNode(UEdGraphPin* FromPin) override;
 	virtual bool CanCreateUnderSpecifiedSchema(const UEdGraphSchema* Schema) const override;
+	virtual FString GetDocumentationLink() const override;
 	// End UEdGraphNode interface.
 
 protected:
diff --git a/Engine/Source/Editor/UnrealEd/Private/MaterialGraphNode.cpp b/Engine/Source/Editor/UnrealEd/Private/MaterialGraphNode.cpp
--- a/Engine/Source/Editor/UnrealEd/Private/MaterialGraphNode.cpp
+++ b/Engine/Source/Editor/UnrealEd/Private/MaterialGraphNode.cpp
@@ -737,4 +737,13 @@ bool UMaterialGraphNode::UsesEventColour(UMaterialExpression* Expression)
 	return false;
 }
 
+FString UMaterialGraphNode::GetDocumentationExcerptName() const
+{
+	// Default the node to searching for an excerpt named for the C++ node class name, including the U prefix.
+	// This is done so that the excerpt name in the doc file can be found by find-in-files when searching for the full class name.
+	UClass* MyClass = (MaterialExpression != NULL) ? MaterialExpression->GetClass() : this->GetClass();
+	return FString::Printf(TEXT("%s%s"), MyClass->GetPrefixCPP(), *MyClass->GetName());
+}
+
+
 #undef LOCTEXT_NAMESPACE
diff --git a/Engine/Source/Editor/UnrealEd/Private/MaterialGraphNode_Base.cpp b/Engine/Source/Editor/UnrealEd/Private/MaterialGraphNode_Base.cpp
--- a/Engine/Source/Editor/UnrealEd/Private/MaterialGraphNode_Base.cpp
+++ b/Engine/Source/Editor/UnrealEd/Private/MaterialGraphNode_Base.cpp
@@ -315,3 +315,8 @@ void UMaterialGraphNode_Base::ModifyAndCopyPersistentPinData(UEdGraphPin& Target
 
 	TargetPin.CopyPersistentDataFromOldPin(SourcePin);
 }
+
+FString UMaterialGraphNode_Base::GetDocumentationLink() const
+{
+	return TEXT("Shared/GraphNodes/Material");
+}
