#UE4 FEditorScriptExecutionGuard is not reentrent, as it resets the recursion counter, so Nick W. and I decided to move it up a frame in the callstack. Fix for TTP#341121

diff --git a/Engine/Source/Editor/UnrealEd/Private/Kismet2/KismetReinstanceUtilities.cpp b/Engine/Source/Editor/UnrealEd/Private/Kismet2/KismetReinstanceUtilities.cpp
--- a/Engine/Source/Editor/UnrealEd/Private/Kismet2/KismetReinstanceUtilities.cpp
+++ b/Engine/Source/Editor/UnrealEd/Private/Kismet2/KismetReinstanceUtilities.cpp
@@ -369,6 +369,9 @@ void FBlueprintCompileReinstancer::ReplaceInstancesOfClass(UClass* OldClass, UCl
 
 				NewActor->RegisterAllComponents(); // Register native components
 
+				// Because this is an editor context it's important to use this execution guard:
+				FEditorScriptExecutionGuard ScriptGuard;
+
 				// Run the construction script, which will use the properties we just copied over
 				NewActor->ExecuteConstruction(WorldTransform, &InstanceDataCache);
 
