#ue4 BP dependency gathering - optimization & fixes

diff --git a/Engine/Source/Editor/KismetCompiler/Private/KismetCompiler.cpp b/Engine/Source/Editor/KismetCompiler/Private/KismetCompiler.cpp
--- a/Engine/Source/Editor/KismetCompiler/Private/KismetCompiler.cpp
+++ b/Engine/Source/Editor/KismetCompiler/Private/KismetCompiler.cpp
@@ -3414,20 +3414,13 @@ void FKismetCompilerContext::Compile()
 	// For full compiles, find other blueprints that may need refreshing, and mark them dirty, in case they try to run
 	if( bIsFullCompile && !Blueprint->bIsRegeneratingOnLoad )
 	{
-		TArray<UObject*> AllBlueprints;
-		GetObjectsOfClass(UBlueprint::StaticClass(), AllBlueprints, true);
-  
-		// Mark any blueprints that implement this interface as dirty
-		for( auto CurrentObj = AllBlueprints.CreateIterator(); CurrentObj; ++CurrentObj )
+		TArray<UBlueprint*> DependentBlueprints;
+		FBlueprintEditorUtils::GetDependentBlueprints(Blueprint, DependentBlueprints);
+		for (auto CurrentBP : DependentBlueprints)
 		{
-			UBlueprint* CurrentBP = Cast<UBlueprint>(*CurrentObj);
-  
-			if( FBlueprintEditorUtils::IsBlueprintDependentOn(CurrentBP, Blueprint) )
-			{
-				CurrentBP->Status = BS_Dirty;
-				FBlueprintEditorUtils::RefreshExternalBlueprintDependencyNodes(CurrentBP);
-				CurrentBP->BroadcastChanged();
-			}
+			CurrentBP->Status = BS_Dirty;
+			FBlueprintEditorUtils::RefreshExternalBlueprintDependencyNodes(CurrentBP, NewClass);
+			CurrentBP->BroadcastChanged();
 		}
 	}
 
diff --git a/Engine/Source/Editor/UnrealEd/Private/Kismet2/BlueprintEditorUtils.cpp b/Engine/Source/Editor/UnrealEd/Private/Kismet2/BlueprintEditorUtils.cpp
--- a/Engine/Source/Editor/UnrealEd/Private/Kismet2/BlueprintEditorUtils.cpp
+++ b/Engine/Source/Editor/UnrealEd/Private/Kismet2/BlueprintEditorUtils.cpp
@@ -301,20 +301,44 @@ void FBlueprintEditorUtils::RefreshAllNodes(UBlueprint* Blueprint)
 	}
 }
 
-void FBlueprintEditorUtils::RefreshExternalBlueprintDependencyNodes(UBlueprint* Blueprint)
+void FBlueprintEditorUtils::RefreshExternalBlueprintDependencyNodes(UBlueprint* Blueprint, UStruct* RefreshOnlyChild)
 {
 	TArray<UK2Node*> AllNodes;
 	FBlueprintEditorUtils::GetAllNodesOfClass(Blueprint, AllNodes);
 
-
-	for( auto NodeIt = AllNodes.CreateIterator(); NodeIt; ++NodeIt )
+	if (!RefreshOnlyChild)
 	{
-		UK2Node* Node = *NodeIt;
-		if( Node->HasExternalBlueprintDependencies() )
+		for (auto NodeIt = AllNodes.CreateIterator(); NodeIt; ++NodeIt)
 		{
-			//@todo:  Do we really need per-schema refreshing?
-			const UEdGraphSchema* Schema = Node->GetGraph()->GetSchema();
-			Schema->ReconstructNode(*Node, true);
+			UK2Node* Node = *NodeIt;
+			if (Node->HasExternalBlueprintDependencies())
+			{
+				//@todo:  Do we really need per-schema refreshing?
+				const UEdGraphSchema* Schema = Node->GetGraph()->GetSchema();
+				Schema->ReconstructNode(*Node, true);
+			}
+		}
+	}
+	else
+	{
+		for (auto NodeIt = AllNodes.CreateIterator(); NodeIt; ++NodeIt)
+		{
+			UK2Node* Node = *NodeIt;
+			TArray<UStruct*> Dependencies;
+			if (Node->HasExternalBlueprintDependencies(&Dependencies))
+			{
+				for (UStruct* Struct : Dependencies)
+				{
+					if (Struct->IsChildOf(RefreshOnlyChild))
+					{
+						//@todo:  Do we really need per-schema refreshing?
+						const UEdGraphSchema* Schema = Node->GetGraph()->GetSchema();
+						Schema->ReconstructNode(*Node, true);
+
+						break;
+					}
+				}
+			}
 		}
 	}
 }
@@ -1460,6 +1484,7 @@ void FBlueprintEditorUtils::UpdateDelegatesInBlueprint(UBlueprint* Blueprint)
 // Blueprint has materially changed.  Recompile the skeleton, notify observers, and mark the package as dirty.
 void FBlueprintEditorUtils::MarkBlueprintAsStructurallyModified(UBlueprint* Blueprint)
 {
+	Blueprint->bCachedDependenciesUpTpDate = false;
 	if (Blueprint->Status != BS_BeingCreated && !Blueprint->bBeingCompiled)
 	{
 		{
@@ -1486,6 +1511,7 @@ void FBlueprintEditorUtils::MarkBlueprintAsStructurallyModified(UBlueprint* Blue
 // Blueprint has changed in some manner that invalidates the compiled data (link made/broken, default value changed, etc...)
 void FBlueprintEditorUtils::MarkBlueprintAsModified(UBlueprint* Blueprint)
 {
+	Blueprint->bCachedDependenciesUpTpDate = false;
 	if (Blueprint->Status != BS_BeingCreated)
 	{
 		Blueprint->Status = BS_Dirty;
@@ -2017,89 +2043,72 @@ UK2Node_Event* FBlueprintEditorUtils::FindOverrideForFunction(const UBlueprint*
 	return NULL;
 }
 
-// returns if Blueprint depends on TestBlueprint
-bool FBlueprintEditorUtils::IsBlueprintDependentOn(UBlueprint const* Blueprint, UBlueprint const* TestBlueprint)
+void FBlueprintEditorUtils::GatherDependencies(const UBlueprint* Blueprint, TSet<TWeakObjectPtr<UBlueprint>>& Dependencies)
 {
-	if ( Blueprint && (Blueprint != TestBlueprint) )
+	struct FGatherDependenciesHelper
 	{
-		// is TestBlueprint an interface we implement?
-		for (int32 BPIdx=0; BPIdx<Blueprint->ImplementedInterfaces.Num(); BPIdx++)
+		static UBlueprint* GetGeneratingBlueprint(const UObject* Obj)
 		{
-			FBPInterfaceDescription const& InterfaceDesc = Blueprint->ImplementedInterfaces[BPIdx];
-			if ( (TestBlueprint->GeneratedClass == InterfaceDesc.Interface) && (TestBlueprint->GeneratedClass != NULL) )
+			const UBlueprintGeneratedClass* BPGC = NULL;
+			while (!BPGC && Obj)
 			{
-				return true;
+				BPGC = Cast<const UBlueprintGeneratedClass>(Obj);
+				Obj = Obj->GetOuter();
 			}
+
+			return UBlueprint::GetBlueprintFromClass(BPGC);
 		}
 
-		TArray<UEdGraph*> Graphs;
-		Blueprint->GetAllGraphs(Graphs);
-		for (auto GraphIt = Graphs.CreateIterator(); GraphIt; ++GraphIt)
+		static void ProcessHierarchy(const UStruct* Struct, TSet<TWeakObjectPtr<UBlueprint>>& Dependencies)
 		{
-			UEdGraph* const Graph = (*GraphIt);
-			if (!IsGraphIntermediate(Graph))
+			for (UBlueprint* Blueprint = GetGeneratingBlueprint(Struct);
+				Blueprint;
+				Blueprint = UBlueprint::GetBlueprintFromClass(Cast<UBlueprintGeneratedClass>(Blueprint->ParentClass)))
 			{
-				// do we have any function calls that ref BP's class or subclasses?
-				TArray<UK2Node_CallFunction*> CallFunctionNodes;
-				Graph->GetNodesOfClass(CallFunctionNodes);
-				for (auto NodeIt = CallFunctionNodes.CreateIterator(); NodeIt; ++NodeIt)
+				bool bAlreadyProcessed = false;
+				Dependencies.Add(Blueprint, &bAlreadyProcessed);
+				if (bAlreadyProcessed)
 				{
-					UK2Node_CallFunction* const Node = *NodeIt;
-					if (Node)
-					{
-						UClass* const FuncClass = Node->FunctionReference.GetMemberParentClass(Node);
-						if (FuncClass->IsChildOf(TestBlueprint->GeneratedClass))
-						{
-							return true;
-						}
-					}
-
-					UK2Node_CallFunctionOnMember* const OnMemberNode = Cast<UK2Node_CallFunctionOnMember>(Node);
-					if (OnMemberNode)
-					{
-						UClass* const MemberFuncClass = OnMemberNode->MemberVariableToCallOn.GetMemberParentClass(OnMemberNode);
-						if (MemberFuncClass->IsChildOf(TestBlueprint->GeneratedClass))
-						{
-							return true;
-						}
-					}
+					return;
 				}
+			}
+		}
+	};
 
-				// do we have any variable refs to objects of BP's class subclasses?
-				TArray<UK2Node_Variable*> VarNodes;
-				Graph->GetNodesOfClass(VarNodes);
-				for (auto NodeIt = VarNodes.CreateIterator(); NodeIt; ++NodeIt)
-				{
-					UK2Node_Variable* const Node = *NodeIt;
-					UClass* const VarClass = Node->VariableReference.GetMemberParentClass(Node);
-
-					if (VarClass->IsChildOf(TestBlueprint->GeneratedClass))
-					{
-						return true;
-					}
-				}
+	check(Blueprint);
+	Dependencies.Empty();
+	for (const auto& InterfaceDesc : Blueprint->ImplementedInterfaces)
+	{
+		UBlueprint* InterfaceBP = InterfaceDesc.Interface ? Cast<UBlueprint>(InterfaceDesc.Interface->ClassGeneratedBy) : NULL;
+		if (InterfaceBP)
+		{
+			Dependencies.Add(InterfaceBP);
+		}
+	}
 
-				// do we have any delegates to objects of BP's class subclasses?
-				TArray<UK2Node_BaseMCDelegate*> DelNodes;
-				Graph->GetNodesOfClass(DelNodes);
-				for (auto NodeIt = DelNodes.CreateIterator(); NodeIt; ++NodeIt)
+	TArray<UEdGraph*> Graphs;
+	Blueprint->GetAllGraphs(Graphs);
+	for (auto Graph : Graphs)
+	{
+		if (Graph && !FBlueprintEditorUtils::IsGraphIntermediate(Graph))
+		{
+			TArray<UK2Node*> Nodes;
+			Graph->GetNodesOfClass(Nodes);
+			for (auto Node : Nodes)
+			{
+				TArray<UStruct*> LocalDependentStructures;
+				if (Node && Node->HasExternalBlueprintDependencies(&LocalDependentStructures))
 				{
-					UK2Node_BaseMCDelegate* const Node = *NodeIt;
-					UClass* const DelClass = Node->DelegateReference.GetMemberParentClass(Node);
-
-					if (DelClass->IsChildOf(TestBlueprint->GeneratedClass))
+					for (auto Struct : LocalDependentStructures)
 					{
-						return true;
+						FGatherDependenciesHelper::ProcessHierarchy(Struct, Dependencies);
 					}
 				}
 			}
 		}
 	}
-
-	return false;
 }
 
-
 void FBlueprintEditorUtils::GetDependentBlueprints(UBlueprint* Blueprint, TArray<UBlueprint*>& DependentBlueprints)
 {
 	DependentBlueprints.Empty();
@@ -2108,13 +2117,22 @@ void FBlueprintEditorUtils::GetDependentBlueprints(UBlueprint* Blueprint, TArray
 	bool const bIncludeDerivedClasses = true;
 	GetObjectsOfClass(UBlueprint::StaticClass(), AllBlueprints, bIncludeDerivedClasses );
 
-	for ( auto ObjIt = AllBlueprints.CreateConstIterator(); ObjIt; ++ObjIt )
+	for (auto ObjIt = AllBlueprints.CreateIterator(); ObjIt; ++ObjIt)
 	{
 		// we know the class is correct so a fast cast is ok here
-		UBlueprint* const TestBP = (UBlueprint*) *ObjIt;
-		if ( FBlueprintEditorUtils::IsBlueprintDependentOn(TestBP, Blueprint) && !TestBP->HasAnyFlags(RF_PendingKill) )
+		UBlueprint* TestBP = (UBlueprint*) *ObjIt;
+		if (TestBP && !TestBP->HasAnyFlags(RF_PendingKill))
 		{
-			DependentBlueprints.Add(TestBP);
+			if (!TestBP->bCachedDependenciesUpTpDate)
+			{
+				GatherDependencies(TestBP, TestBP->CachedDependencies);
+				TestBP->bCachedDependenciesUpTpDate = true;
+			}
+
+			if (TestBP->CachedDependencies.Contains(Blueprint))
+			{
+				DependentBlueprints.Add(TestBP);
+			}
 		}
 	}
 }
diff --git a/Engine/Source/Editor/UnrealEd/Private/Kismet2/KismetReinstanceUtilities.cpp b/Engine/Source/Editor/UnrealEd/Private/Kismet2/KismetReinstanceUtilities.cpp
--- a/Engine/Source/Editor/UnrealEd/Private/Kismet2/KismetReinstanceUtilities.cpp
+++ b/Engine/Source/Editor/UnrealEd/Private/Kismet2/KismetReinstanceUtilities.cpp
@@ -92,16 +92,7 @@ FBlueprintCompileReinstancer::FBlueprintCompileReinstancer(UClass* InClassToRein
 		check(GeneratingBP || GIsAutomationTesting);
 		if(GeneratingBP)
 		{
-			TArray<UObject*> AllBlueprints;
-			GetObjectsOfClass(UBlueprint::StaticClass(), AllBlueprints, true);
-			for( auto BPIt = AllBlueprints.CreateIterator(); BPIt; ++BPIt )
-			{
-				UBlueprint* CurrentBP = Cast<UBlueprint>(*BPIt);
-				if( FBlueprintEditorUtils::IsBlueprintDependentOn(CurrentBP, GeneratingBP) )
-				{
-					Dependencies.Add(CurrentBP);
-				}
-			}
+			FBlueprintEditorUtils::GetDependentBlueprints(GeneratingBP, Dependencies);
 		}
 	}
 }
diff --git a/Engine/Source/Editor/UnrealEd/Public/Kismet2/BlueprintEditorUtils.h b/Engine/Source/Editor/UnrealEd/Public/Kismet2/BlueprintEditorUtils.h
--- a/Engine/Source/Editor/UnrealEd/Public/Kismet2/BlueprintEditorUtils.h
+++ b/Engine/Source/Editor/UnrealEd/Public/Kismet2/BlueprintEditorUtils.h
@@ -49,7 +49,7 @@ public:
 	/**
 	 * Optimized refresh of nodes that depend on external blueprints.  Refreshes the nodes, but does not recompile the skeleton class
 	 */
-	static void RefreshExternalBlueprintDependencyNodes(UBlueprint* Blueprint);
+	static void RefreshExternalBlueprintDependencyNodes(UBlueprint* Blueprint, UStruct* RefreshOnlyChild = NULL);
 
 	/**
 	 * Refresh the nodes of an individual graph.
@@ -314,8 +314,8 @@ public:
 		return nullptr;
 	}
 
-	/** @return true if this Blueprint is dependent on the given blueprint, false otherwise */
-	static bool IsBlueprintDependentOn(UBlueprint const* Blueprint, UBlueprint const* TestBlueprint);
+	/** Gather all bps that Blueprint depends on */
+	static void GatherDependencies(const UBlueprint* Blueprint, TSet<TWeakObjectPtr<UBlueprint>>& OutDependencies);
 
 	/** Returns a list of loaded Blueprints that are dependent on the given Blueprint. */
 	static void GetDependentBlueprints(UBlueprint* Blueprint, TArray<UBlueprint*>& DependentBlueprints);
diff --git a/Engine/Source/Runtime/Engine/Classes/Engine/Blueprint.h b/Engine/Source/Runtime/Engine/Classes/Engine/Blueprint.h
--- a/Engine/Source/Runtime/Engine/Classes/Engine/Blueprint.h
+++ b/Engine/Source/Runtime/Engine/Classes/Engine/Blueprint.h
@@ -413,6 +413,9 @@ public:
 #endif // WITH_EDITORONLY_DATA
 
 #if WITH_EDITORONLY_DATA
+	bool bCachedDependenciesUpTpDate;
+	TSet<TWeakObjectPtr<UBlueprint>> CachedDependencies;
+
 	bool IsUpToDate() const
 	{
 		return BS_UpToDate == Status || BS_UpToDateWithWarnings == Status;
@@ -498,6 +501,7 @@ public:
 	virtual UClass* RegenerateClass(UClass* ClassToRegenerate, UObject* PreviousCDO, TArray<UObject*>& ObjLoaded) override;
 	virtual void PostLoad() override;
 	virtual void PostLoadSubobjects( FObjectInstancingGraph* OuterInstanceGraph ) override;
+	virtual bool Modify(bool bAlwaysMarkDirty = true) override;
 	// End of UObject interface
 
 	/** Consigns the GeneratedClass and the SkeletonGeneratedClass to oblivion, and nulls their references */
diff --git a/Engine/Source/Runtime/Engine/Private/Blueprint.cpp b/Engine/Source/Runtime/Engine/Private/Blueprint.cpp
--- a/Engine/Source/Runtime/Engine/Private/Blueprint.cpp
+++ b/Engine/Source/Runtime/Engine/Private/Blueprint.cpp
@@ -164,6 +164,7 @@ void UBlueprintCore::GenerateDeterministicGuid()
 
 UBlueprint::UBlueprint(const class FPostConstructInitializeProperties& PCIP)
 	: Super(PCIP)
+	, bCachedDependenciesUpTpDate(false)
 {
 
 #if WITH_EDITOR
@@ -1066,6 +1067,12 @@ void UBlueprint::PostLoadSubobjects(FObjectInstancingGraph* OuterInstanceGraph)
 	ChangeOwnerOfTemplates();
 }
 
+bool UBlueprint::Modify(bool bAlwaysMarkDirty)
+{
+	bCachedDependenciesUpTpDate = false;
+	return Super::Modify(bAlwaysMarkDirty);
+}
+
 #endif
 
 #if WITH_EDITORONLY_DATA
