BP: SpawnActorFromClass node and ExposeOnSpawn fixes.

diff --git a/Engine/Source/Editor/BlueprintGraph/Classes/K2Node_SpawnActorFromClass.h b/Engine/Source/Editor/BlueprintGraph/Classes/K2Node_SpawnActorFromClass.h
--- a/Engine/Source/Editor/BlueprintGraph/Classes/K2Node_SpawnActorFromClass.h
+++ b/Engine/Source/Editor/BlueprintGraph/Classes/K2Node_SpawnActorFromClass.h
@@ -28,6 +28,7 @@ class UK2Node_SpawnActorFromClass : public UK2Node
 	virtual void GetNodeAttributes( TArray<TKeyValuePair<FString, FString>>& OutNodeAttributes ) const override;
 	virtual void GetMenuActions(TArray<UBlueprintNodeSpawner*>& ActionListOut) const override;
 	virtual FText GetMenuCategory() const override;
+	virtual class FNodeHandlingFunctor* CreateNodeHandler(class FKismetCompilerContext& CompilerContext) const override;
 	// End UK2Node interface
 
 
diff --git a/Engine/Source/Editor/BlueprintGraph/Private/K2Node_CastByteToEnum.cpp b/Engine/Source/Editor/BlueprintGraph/Private/K2Node_CastByteToEnum.cpp
--- a/Engine/Source/Editor/BlueprintGraph/Private/K2Node_CastByteToEnum.cpp
+++ b/Engine/Source/Editor/BlueprintGraph/Private/K2Node_CastByteToEnum.cpp
@@ -135,7 +135,7 @@ FNodeHandlingFunctor* UK2Node_CastByteToEnum::CreateNodeHandler(FKismetCompilerC
 	{
 		return new FKCHandler_CastByteToEnum(CompilerContext);
 	}
-	return NULL;
+	return new FNodeHandlingFunctor(CompilerContext);;
 }
 
 void UK2Node_CastByteToEnum::GetMenuActions(TArray<UBlueprintNodeSpawner*>& ActionListOut) const
diff --git a/Engine/Source/Editor/BlueprintGraph/Private/K2Node_SpawnActorFromClass.cpp b/Engine/Source/Editor/BlueprintGraph/Private/K2Node_SpawnActorFromClass.cpp
--- a/Engine/Source/Editor/BlueprintGraph/Private/K2Node_SpawnActorFromClass.cpp
+++ b/Engine/Source/Editor/BlueprintGraph/Private/K2Node_SpawnActorFromClass.cpp
@@ -327,6 +327,11 @@ FText UK2Node_SpawnActorFromClass::GetMenuCategory() const
 	return FEditorCategoryUtils::GetCommonCategory(FCommonEditorCategory::Gameplay);
 }
 
+FNodeHandlingFunctor* UK2Node_SpawnActorFromClass::CreateNodeHandler(FKismetCompilerContext& CompilerContext) const
+{
+	return new FNodeHandlingFunctor(CompilerContext);
+}
+
 void UK2Node_SpawnActorFromClass::ExpandNode(class FKismetCompilerContext& CompilerContext, UEdGraph* SourceGraph)
 {
 	Super::ExpandNode(CompilerContext, SourceGraph);
diff --git a/Engine/Source/Editor/UnrealEd/Private/Kismet2/BlueprintEditorUtils.cpp b/Engine/Source/Editor/UnrealEd/Private/Kismet2/BlueprintEditorUtils.cpp
--- a/Engine/Source/Editor/UnrealEd/Private/Kismet2/BlueprintEditorUtils.cpp
+++ b/Engine/Source/Editor/UnrealEd/Private/Kismet2/BlueprintEditorUtils.cpp
@@ -2646,6 +2646,35 @@ void FBlueprintEditorUtils::SetVariableSaveGameFlag(UBlueprint* InBlueprint, con
 	FBlueprintEditorUtils::MarkBlueprintAsStructurallyModified(InBlueprint);
 }
 
+struct FMetaDataDependencyHelper
+{
+	static void OnChange(UBlueprint* Blueprint, FName MetaDataKey)
+	{
+		if (Blueprint && (FBlueprintMetadata::MD_ExposeOnSpawn == MetaDataKey))
+		{
+			TArray<UEdGraph*> AllGraphs;
+			Blueprint->GetAllGraphs(AllGraphs);
+			for (UEdGraph* Graph : AllGraphs)
+			{
+				if (Graph)
+				{
+					const UEdGraphSchema* Schema = Graph->GetSchema();
+					TArray<UK2Node_SpawnActorFromClass*> LocalSpawnNodes;
+					Graph->GetNodesOfClass(LocalSpawnNodes);
+					for (UK2Node_SpawnActorFromClass* Node : LocalSpawnNodes)
+					{
+						UClass* ClassToSpawn = Node ? Node->GetClassToSpawn() : NULL;
+						if (ClassToSpawn && ClassToSpawn->IsChildOf(Blueprint->GeneratedClass))
+						{
+							Schema->ReconstructNode(*Node, true);
+						}
+					}
+				}
+			}
+		}
+	}
+};
+
 void FBlueprintEditorUtils::SetBlueprintVariableMetaData(UBlueprint* Blueprint, const FName& VarName, const UStruct* InLocalVarScope, const FName& MetaDataKey, const FString& MetaDataValue)
 {
 	// If there is a local var scope, we know we are looking at a local variable
@@ -2684,8 +2713,10 @@ void FBlueprintEditorUtils::SetBlueprintVariableMetaData(UBlueprint* Blueprint,
 		{
 			Blueprint->NewVariables[VarIndex].SetMetaData(MetaDataKey, MetaDataValue);
 			UProperty* Property = FindField<UProperty>(Blueprint->SkeletonGeneratedClass, VarName);
-			Property->SetMetaData(MetaDataKey, *MetaDataValue);
-
+			if (Property)
+			{
+				Property->SetMetaData(MetaDataKey, *MetaDataValue);
+			}
 			Property = FindField<UProperty>(Blueprint->GeneratedClass, VarName);
 			if (Property)
 			{
@@ -2693,6 +2724,9 @@ void FBlueprintEditorUtils::SetBlueprintVariableMetaData(UBlueprint* Blueprint,
 			}
 		}
 	}
+
+	FMetaDataDependencyHelper::OnChange(Blueprint, MetaDataKey);
+
 	FBlueprintEditorUtils::MarkBlueprintAsStructurallyModified(Blueprint);
 }
 
@@ -2802,9 +2836,21 @@ void FBlueprintEditorUtils::RemoveBlueprintVariableMetaData(UBlueprint* Blueprin
 		else
 		{
 			Blueprint->NewVariables[VarIndex].RemoveMetaData(MetaDataKey);
+			UProperty* Property = FindField<UProperty>(Blueprint->SkeletonGeneratedClass, VarName);
+			if (Property)
+			{
+				Property->RemoveMetaData(MetaDataKey);
+			}
+			Property = FindField<UProperty>(Blueprint->GeneratedClass, VarName);
+			if (Property)
+			{
+				Property->RemoveMetaData(MetaDataKey);
+			}
 		}
 	}
 
+	FMetaDataDependencyHelper::OnChange(Blueprint, MetaDataKey);
+
 	FBlueprintEditorUtils::MarkBlueprintAsStructurallyModified(Blueprint);
 }
 
