Eliminate Classes.INT.udn and instead use an excerpt with the name of the class in the appropriate Types INT.udn file

diff --git a/Engine/Source/Editor/UnrealEd/Private/EditorClassUtils.cpp b/Engine/Source/Editor/UnrealEd/Private/EditorClassUtils.cpp
--- a/Engine/Source/Editor/UnrealEd/Private/EditorClassUtils.cpp
+++ b/Engine/Source/Editor/UnrealEd/Private/EditorClassUtils.cpp
@@ -6,16 +6,25 @@
 #include "IDocumentation.h"
 #include "SourceCodeNavigation.h"
 
-static const FString ClassDocsPage(TEXT("Shared/Classes"));
+FString GetDocumentationPage(const UClass* Class)
+{
+	return FString::Printf( TEXT("Shared/Types/%s%s"), Class->GetPrefixCPP(), *Class->GetName() );
+}
+
+FString GetDocumentationExcerpt(const UClass* Class)
+{
+	return FString::Printf( TEXT("%s%s"), Class->GetPrefixCPP(), *Class->GetName() );
+}
 
 TSharedRef<SToolTip> FEditorClassUtils::GetTooltip(const UClass* Class)
 {
-	return IDocumentation::Get()->CreateToolTip(Class->GetToolTipText(), nullptr, ClassDocsPage, Class->GetName());
+	return IDocumentation::Get()->CreateToolTip(Class->GetToolTipText(), nullptr, GetDocumentationPage(Class), GetDocumentationExcerpt(Class));
 }
 
 FString FEditorClassUtils::GetDocumentationLink(const UClass* Class)
 {
 	FString DocumentationLink;
+	const FString ClassDocsPage = GetDocumentationPage(Class);
 
 	TSharedRef<IDocumentation> Documentation = IDocumentation::Get();
 	if (Documentation->PageExists(ClassDocsPage))
@@ -23,7 +32,7 @@ FString FEditorClassUtils::GetDocumentationLink(const UClass* Class)
 		TSharedRef<IDocumentationPage> ClassDocs = Documentation->GetPage(ClassDocsPage, NULL);
 
 		FExcerpt Excerpt;
-		if (ClassDocs->GetExcerpt(Class->GetName(), Excerpt))
+		if (ClassDocs->GetExcerpt(GetDocumentationExcerpt(Class), Excerpt))
 		{
 			FString* FullDocumentationLink = Excerpt.Variables.Find( TEXT("ToolTipFullLink") );
 			if (FullDocumentationLink)
