Fix crash due to NULL local variable

diff --git a/Engine/Source/Editor/UnrealEd/Private/PlayLevel.cpp b/Engine/Source/Editor/UnrealEd/Private/PlayLevel.cpp
--- a/Engine/Source/Editor/UnrealEd/Private/PlayLevel.cpp
+++ b/Engine/Source/Editor/UnrealEd/Private/PlayLevel.cpp
@@ -304,48 +304,51 @@ void UEditorEngine::TeardownPlaySession(FWorldContext &PieWorldContext)
 		CleanupGameViewport();
 	
 		// Clean up the slate PIE viewport if we have one
-		if (SlatePlayInEditorSession && SlatePlayInEditorSession->DestinationSlateViewport.IsValid())
+		if (SlatePlayInEditorSession)
 		{
-			TSharedPtr<ILevelViewport> Viewport = SlatePlayInEditorSession->DestinationSlateViewport.Pin();
-
-			if( !bIsSimulatingInEditor)
+			if (SlatePlayInEditorSession->DestinationSlateViewport.IsValid())
 			{
-				// Set the editor viewport location to match that of Play in Viewport if we aren't simulating in the editor, we have a valid player to get the location from 
-				if (SlatePlayInEditorSession->EditorPlayer.IsValid() && SlatePlayInEditorSession->EditorPlayer.Get()->PlayerController)
-				{
-					FVector ViewLocation;
-					FRotator ViewRotation;
-					SlatePlayInEditorSession->EditorPlayer.Get()->PlayerController->GetPlayerViewPoint(ViewLocation, ViewRotation);
-					Viewport->GetLevelViewportClient().SetViewLocation( ViewLocation );
+				TSharedPtr<ILevelViewport> Viewport = SlatePlayInEditorSession->DestinationSlateViewport.Pin();
 
-					if( Viewport->GetLevelViewportClient().IsPerspective() )
+				if( !bIsSimulatingInEditor)
+				{
+					// Set the editor viewport location to match that of Play in Viewport if we aren't simulating in the editor, we have a valid player to get the location from 
+					if (SlatePlayInEditorSession->EditorPlayer.IsValid() && SlatePlayInEditorSession->EditorPlayer.Get()->PlayerController)
 					{
-						// Rotation only matters for perspective viewports not orthographic
-						Viewport->GetLevelViewportClient().SetViewRotation( ViewRotation );
+						FVector ViewLocation;
+						FRotator ViewRotation;
+						SlatePlayInEditorSession->EditorPlayer.Get()->PlayerController->GetPlayerViewPoint(ViewLocation, ViewRotation);
+						Viewport->GetLevelViewportClient().SetViewLocation( ViewLocation );
+
+						if( Viewport->GetLevelViewportClient().IsPerspective() )
+						{
+							// Rotation only matters for perspective viewports not orthographic
+							Viewport->GetLevelViewportClient().SetViewRotation( ViewRotation );
+						}
 					}
 				}
-			}
 
-			// No longer simulating in the viewport
-			Viewport->GetLevelViewportClient().SetIsSimulateInEditorViewport( false );
+				// No longer simulating in the viewport
+				Viewport->GetLevelViewportClient().SetIsSimulateInEditorViewport( false );
 
-			// Clear out the hit proxies before GC'ing
-			Viewport->GetLevelViewportClient().Viewport->InvalidateHitProxy();
-		}
-		else if (SlatePlayInEditorSession->SlatePlayInEditorWindow.IsValid())
-		{
-			// Unregister the game viewport from slate.  This sends a final message to the viewport
-			// so it can have a chance to release mouse capture, mouse lock, etc.		
-			FSlateApplication::Get().UnregisterGameViewport();
+				// Clear out the hit proxies before GC'ing
+				Viewport->GetLevelViewportClient().Viewport->InvalidateHitProxy();
+			}
+			else if (SlatePlayInEditorSession->SlatePlayInEditorWindow.IsValid())
+			{
+				// Unregister the game viewport from slate.  This sends a final message to the viewport
+				// so it can have a chance to release mouse capture, mouse lock, etc.		
+				FSlateApplication::Get().UnregisterGameViewport();
 
-			// Viewport client is cleaned up.  Make sure its not being accessed
-			SlatePlayInEditorSession->SlatePlayInEditorWindowViewport->SetViewportClient(NULL);
+				// Viewport client is cleaned up.  Make sure its not being accessed
+				SlatePlayInEditorSession->SlatePlayInEditorWindowViewport->SetViewportClient(NULL);
 
-			// The window may have already been destroyed in the case that the PIE window close box was pressed 
-			if (SlatePlayInEditorSession->SlatePlayInEditorWindow.IsValid())
-			{
-				// Destroy the SWindow
-				FSlateApplication::Get().DestroyWindowImmediately(SlatePlayInEditorSession->SlatePlayInEditorWindow.Pin().ToSharedRef());
+				// The window may have already been destroyed in the case that the PIE window close box was pressed 
+				if (SlatePlayInEditorSession->SlatePlayInEditorWindow.IsValid())
+				{
+					// Destroy the SWindow
+					FSlateApplication::Get().DestroyWindowImmediately(SlatePlayInEditorSession->SlatePlayInEditorWindow.Pin().ToSharedRef());
+				}
 			}
 		}
 	
