342312	FBX import failure in 4.4

- Two Nan Issues : 1. divide by 0 2. FMatrix inverse produces NaN if det == 0, so changed it to FTransform.

diff --git a/Engine/Source/Editor/UnrealEd/Private/PhysicsAssetUtils.cpp b/Engine/Source/Editor/UnrealEd/Private/PhysicsAssetUtils.cpp
--- a/Engine/Source/Editor/UnrealEd/Private/PhysicsAssetUtils.cpp
+++ b/Engine/Source/Editor/UnrealEd/Private/PhysicsAssetUtils.cpp
@@ -300,7 +300,10 @@ FVector ComputeEigenVector(const FMatrix & A)
 	for (int32 i = 0; i < 32; ++i)
 	{
 		float Length = Bk.Size();
-		Bk = A.TransformVector(Bk) / Length;
+		if ( Length > 0.f )
+		{
+			Bk = A.TransformVector(Bk) / Length;
+		}
 	}
 
 	return Bk.SafeNormal();
@@ -373,13 +376,16 @@ bool CreateCollisionFromBone( UBodySetup* bs, USkeletalMesh* skelMesh, int32 Bon
 	}
 	
 
-
+	// convert to FTransform now
+	// Matrix inverse doesn't handle well when DET == 0, so 
+	// convert to FTransform and use that data
+	FTransform ElementTransform(ElemTM);
 	// Get the (Unreal scale) bounding box for this bone using the rotation.
 	const FBoneVertInfo* BoneInfo = &Infos[BoneIndex];
 	FBox BoneBox(0);
 	for(int32 j=0; j<BoneInfo->Positions.Num(); j++)
 	{
-		BoneBox += ElemTM.InverseTransformPosition( BoneInfo->Positions[j]  );
+		BoneBox += ElementTransform.InverseTransformPosition( BoneInfo->Positions[j]  );
 	}
 
 	FVector BoxCenter(0,0,0), BoxExtent(0,0,0);
@@ -396,15 +402,15 @@ bool CreateCollisionFromBone( UBodySetup* bs, USkeletalMesh* skelMesh, int32 Bon
 		BoxExtent = FVector(DefaultPrimSize, DefaultPrimSize, DefaultPrimSize);
 	}
 
-	FVector BoneOrigin = ElemTM.TransformPosition( BoxCenter );
-	ElemTM.SetOrigin( BoneOrigin );
+	FVector BoneOrigin = ElementTransform.TransformPosition( BoxCenter );
+	ElementTransform.SetTranslation( BoneOrigin );
 
 	if(Params.GeomType == EFG_Box)
 	{
 		// Add a new box geometry to this body the size of the bounding box.
 		FKBoxElem BoxElem;
 
-		BoxElem.SetTransform(FTransform(ElemTM));
+		BoxElem.SetTransform(ElementTransform);
 
 		BoxElem.X = BoxExtent.X * 2.0f * 1.01f; // Side Lengths (add 1% to avoid graphics glitches)
 		BoxElem.Y = BoxExtent.Y * 2.0f * 1.01f;
@@ -416,7 +422,7 @@ bool CreateCollisionFromBone( UBodySetup* bs, USkeletalMesh* skelMesh, int32 Bon
 	{
 		FKSphereElem SphereElem;
 
-		SphereElem.Center = ElemTM.GetOrigin();
+		SphereElem.Center = ElementTransform.GetTranslation();
 		SphereElem.Radius = BoxExtent.GetMax() * 1.01f;
 
 		bs->AggGeom.SphereElems.Add(SphereElem);
@@ -508,7 +514,7 @@ bool CreateCollisionFromBone( UBodySetup* bs, USkeletalMesh* skelMesh, int32 Bon
 		
 		FKSphylElem SphylElem;
 
-		SphylElem.SetTransform(FTransform(ElemTM));
+		SphylElem.SetTransform(ElementTransform);
 
 		SphylElem.Radius = FMath::Max(BoxExtent.X, BoxExtent.Y) * 1.01f;
 		SphylElem.Length = BoxExtent.Z * 1.01f;
