Added tooltip for bAllowBackgroundAudio setting

diff --git a/Engine/Source/Editor/UnrealEd/Classes/Settings/LevelEditorMiscSettings.h b/Engine/Source/Editor/UnrealEd/Classes/Settings/LevelEditorMiscSettings.h
--- a/Engine/Source/Editor/UnrealEd/Classes/Settings/LevelEditorMiscSettings.h
+++ b/Engine/Source/Editor/UnrealEd/Classes/Settings/LevelEditorMiscSettings.h
@@ -37,7 +37,7 @@ public:
 	uint32 bReplaceRespectsScale:1;
 
 public:
-
+	/** If checked audio playing in the editor will continue to play even if the editor is in the background */
 	UPROPERTY(EditAnywhere, config, Category=Sound)
 	uint32 bAllowBackgroundAudio:1;
 
