Fixed incorrect tooltip on bAutoApplyLightingEnable

diff --git a/Engine/Source/Editor/UnrealEd/Classes/Settings/LevelEditorMiscSettings.h b/Engine/Source/Editor/UnrealEd/Classes/Settings/LevelEditorMiscSettings.h
--- a/Engine/Source/Editor/UnrealEd/Classes/Settings/LevelEditorMiscSettings.h
+++ b/Engine/Source/Editor/UnrealEd/Classes/Settings/LevelEditorMiscSettings.h
@@ -20,7 +20,9 @@ class UNREALED_API ULevelEditorMiscSettings
 
 public:
 
-	/** Whether to automatically save after a time interval */
+	/** If checked lighting will be automatically applied to the level after a static lighting build is complete.
+	 * Otherwise you will be prompted to apply it
+	 */
 	UPROPERTY(EditAnywhere, config, Category=Editing, meta=(DisplayName="Apply Lighting Automatically"))
 	uint32 bAutoApplyLightingEnable:1;
 
