342506	AnimBP: Event Blueprint Update Animation node does not tick

diff --git a/Engine/Source/Editor/UnrealEd/Private/Kismet2/Kismet2.cpp b/Engine/Source/Editor/UnrealEd/Private/Kismet2/Kismet2.cpp
--- a/Engine/Source/Editor/UnrealEd/Private/Kismet2/Kismet2.cpp
+++ b/Engine/Source/Editor/UnrealEd/Private/Kismet2/Kismet2.cpp
@@ -283,7 +283,7 @@ UBlueprint* FKismetEditorUtilities::CreateBlueprint(UClass* ParentClass, UObject
 			UK2Node_Event* EventNode = NewObject<UK2Node_Event>(Graph);
 			EventNode->EventSignatureClass = UAnimInstance::StaticClass();
 			EventNode->EventSignatureName = FName(TEXT("BlueprintUpdateAnimation"));
-			EventNode->bInternalEvent = true;
+			EventNode->bOverrideFunction=true;
 			EventNode->CreateNewGuid();
 			EventNode->PostPlacedNewNode();
 			EventNode->AllocateDefaultPins();
