Linux editor: minor fixes.

- Part of PR #306 by amigo and #ue4linux community.
- Mostly things like initialization order changes, but also platform-specific additions/improvements.

diff --git a/Engine/Source/Developer/MessageLog/Private/Presentation/MessageLogListingViewModel.h b/Engine/Source/Developer/MessageLog/Private/Presentation/MessageLogListingViewModel.h
--- a/Engine/Source/Developer/MessageLog/Private/Presentation/MessageLogListingViewModel.h
+++ b/Engine/Source/Developer/MessageLog/Private/Presentation/MessageLogListingViewModel.h
@@ -139,14 +139,14 @@ public:
 
 private:
 	FMessageLogListingViewModel( TSharedPtr< FMessageLogListingModel > InMessageLogListingModel, const FText& InLogLabel, const FMessageLogInitializationOptions& InitializationOptions )
-		: MessageLogListingModel( InMessageLogListingModel )
-		, bShowFilters( InitializationOptions.bShowFilters )
+		: bShowFilters( InitializationOptions.bShowFilters )
 		, bShowPages( InitializationOptions.bShowPages )
 		, bDiscardDuplicates( InitializationOptions.bDiscardDuplicates )
 		, MaxPageCount( InitializationOptions.MaxPageCount )
 		, CurrentPageIndex( 0 )
 		, bIsRefreshing( false )
 		, LogLabel( InLogLabel )
+		, MessageLogListingModel( InMessageLogListingModel )
 	{}
 
 	/** Rebuilds the list of filtered messages */
diff --git a/Engine/Source/Developer/MessageLog/Private/UserInterface/SMessageLogListing.cpp b/Engine/Source/Developer/MessageLog/Private/UserInterface/SMessageLogListing.cpp
--- a/Engine/Source/Developer/MessageLog/Private/UserInterface/SMessageLogListing.cpp
+++ b/Engine/Source/Developer/MessageLog/Private/UserInterface/SMessageLogListing.cpp
@@ -238,8 +238,8 @@ protected:
 
 
 SMessageLogListing::SMessageLogListing()
-	: bUpdatingSelection( false )
-	, UICommandList( MakeShareable( new FUICommandList ) )
+	: UICommandList( MakeShareable( new FUICommandList ) )
+	, bUpdatingSelection( false )
 {
 
 }
diff --git a/Engine/Source/Editor/GameProjectGeneration/Private/GameProjectUtils.cpp b/Engine/Source/Editor/GameProjectGeneration/Private/GameProjectUtils.cpp
--- a/Engine/Source/Editor/GameProjectGeneration/Private/GameProjectUtils.cpp
+++ b/Engine/Source/Editor/GameProjectGeneration/Private/GameProjectUtils.cpp
@@ -475,6 +475,8 @@ bool GameProjectUtils::OpenCodeIDE(const FString& ProjectFile, FText& OutFailRea
 	CodeSolutionFile = SolutionFilenameWithoutExtension + TEXT(".sln");
 #elif PLATFORM_MAC
 	CodeSolutionFile = SolutionFilenameWithoutExtension + TEXT(".xcodeproj");
+#elif PLATFORM_LINUX
+	STUBBED("Linux solution filename");
 #else
 	OutFailReason = LOCTEXT( "OpenCodeIDE_UnknownPlatform", "could not open the code editing IDE. The operating system is unknown." );
 	return false;
@@ -482,7 +484,7 @@ bool GameProjectUtils::OpenCodeIDE(const FString& ProjectFile, FText& OutFailRea
 
 	// Open the solution with the default application
 	const FString FullPath = FPaths::Combine(*SolutionFolder, *CodeSolutionFile);
-#if PLATFORM_MAC || PLATFORM_LINUX
+#if PLATFORM_MAC
 	if ( IFileManager::Get().DirectoryExists(*FullPath) )
 #else
 	if ( FPaths::FileExists(FullPath) )
diff --git a/Engine/Source/Editor/PropertyEditor/Public/PropertyCustomizationHelpers.h b/Engine/Source/Editor/PropertyEditor/Public/PropertyCustomizationHelpers.h
--- a/Engine/Source/Editor/PropertyEditor/Public/PropertyCustomizationHelpers.h
+++ b/Engine/Source/Editor/PropertyEditor/Public/PropertyCustomizationHelpers.h
@@ -134,8 +134,8 @@ public:
 	DECLARE_DELEGATE( FOnPropertyValueChanged );
 
 	SLATE_BEGIN_ARGS( SProperty )
-		: _DisplayResetToDefault( true )
-		, _ShouldDisplayName( true )
+		: _ShouldDisplayName( true )
+		, _DisplayResetToDefault( true )
 		{}
 		/** The display name to use in the default property widget */
 		SLATE_ATTRIBUTE( FString, DisplayName )
diff --git a/Engine/Source/Editor/UnrealEd/Private/Fbx/FbxMainImport.cpp b/Engine/Source/Editor/UnrealEd/Private/Fbx/FbxMainImport.cpp
--- a/Engine/Source/Editor/UnrealEd/Private/Fbx/FbxMainImport.cpp
+++ b/Engine/Source/Editor/UnrealEd/Private/Fbx/FbxMainImport.cpp
@@ -234,12 +234,12 @@ void FImportedMaterialData::Clear()
 }
 
 FFbxImporter::FFbxImporter()
-	: bFirstMesh(true)
-	, Importer( NULL )
+	: Scene(NULL)
 	, ImportOptions(NULL)
 	, GeometryConverter(NULL)
-	, Scene(NULL)
 	, SdkManager(NULL)
+	, Importer( NULL )
+	, bFirstMesh(true)
 	, Logger(NULL)
 {
 	// Create the SdkManager
diff --git a/Engine/Source/Editor/UnrealEd/Private/Lightmass/Lightmass.cpp b/Engine/Source/Editor/UnrealEd/Private/Lightmass/Lightmass.cpp
--- a/Engine/Source/Editor/UnrealEd/Private/Lightmass/Lightmass.cpp
+++ b/Engine/Source/Editor/UnrealEd/Private/Lightmass/Lightmass.cpp
@@ -2471,7 +2471,10 @@ bool FLightmassProcessor::BeginRun()
 		TEXT("../DotNET/Linux/AgentInterface.dll"),
 		TEXT("../Linux/libUnrealLightmass-Core-Linux-Debug.so"),
 		TEXT("../Linux/libUnrealLightmass-Projects-Linux-Debug.so"),
-		TEXT("../Linux/libUnrealLightmass-SwarmInterface-Linux-Debug.so")
+		TEXT("../Linux/libUnrealLightmass-SwarmInterface-Linux-Debug.so"),
+		TEXT("../Linux/libUnrealLightmass-Networking-Linux-Debug.so"),
+		TEXT("../Linux/libUnrealLightmass-Messaging-Linux-Debug.so")
+		TEXT("../../Plugins/Messaging/UdpMessaging/Binaries/Linux/libUnrealLightmass-UdpMessaging-Linux-Debug.so")
 	};
 #else
 	const TCHAR* LightmassExecutable64 = TEXT("../Linux/UnrealLightmass");
@@ -2480,10 +2483,13 @@ bool FLightmassProcessor::BeginRun()
 		TEXT("../DotNET/Linux/AgentInterface.dll"),
 		TEXT("../Linux/libUnrealLightmass-Core.so"),
 		TEXT("../Linux/libUnrealLightmass-Projects.so"),
-		TEXT("../Linux/libUnrealLightmass-SwarmInterface.so")
+		TEXT("../Linux/libUnrealLightmass-SwarmInterface.so"),
+		TEXT("../Linux/libUnrealLightmass-Networking.so"),
+		TEXT("../Linux/libUnrealLightmass-Messaging.so"),
+		TEXT("../../Plugins/Messaging/UdpMessaging/Binaries/Linux/libUnrealLightmass-UdpMessaging.so")
 	};
-#endif
-#else // platform
+#endif // UE_BUILD_DEBUG
+#else // PLATFORM_LINUX
 #error "Unknown Lightmass platform"
 #endif
 	const int32 RequiredDependencyPaths64Count = ARRAY_COUNT(RequiredDependencyPaths64);
diff --git a/Engine/Source/Editor/UnrealEd/Private/PerformanceMonitor.h b/Engine/Source/Editor/UnrealEd/Private/PerformanceMonitor.h
--- a/Engine/Source/Editor/UnrealEd/Private/PerformanceMonitor.h
+++ b/Engine/Source/Editor/UnrealEd/Private/PerformanceMonitor.h
@@ -7,13 +7,13 @@ struct FMovingAverage
 {
 	/** Constructor from a sample size, and sample rate */
 	FMovingAverage(const int32 InSampleSize = 0, const float InSampleRateSeconds = 1.f)
-		: SampleRateSeconds(InSampleRateSeconds)
+		: CurrentSampleCount(0)
 		, CurrentSampleAccumulator(0)
 		, CurrentSampleStartTime(0)
+		, SampleRateSeconds(InSampleRateSeconds)
 		, SampleSize(InSampleSize)
-		, CurrentSampleCount(0)
-		, NextSampleIndex(0)
 		, SampleAverage(0)
+		, NextSampleIndex(0)
 	{
 		Samples.Reserve(SampleSize);
 	}
diff --git a/Engine/Source/Editor/UnrealEd/Private/UnrealEd.cpp b/Engine/Source/Editor/UnrealEd/Private/UnrealEd.cpp
--- a/Engine/Source/Editor/UnrealEd/Private/UnrealEd.cpp
+++ b/Engine/Source/Editor/UnrealEd/Private/UnrealEd.cpp
@@ -160,7 +160,8 @@ int32 EditorInit( IEngineLoop& EngineLoop )
 		}
 	}
 
-	return 1;
+	// this will be ultimately returned from main(), so no error should be 0.
+	return 0;
 }
 
 void EditorExit()
diff --git a/Engine/Source/Editor/UnrealEd/Public/SComponentClassCombo.h b/Engine/Source/Editor/UnrealEd/Public/SComponentClassCombo.h
--- a/Engine/Source/Editor/UnrealEd/Public/SComponentClassCombo.h
+++ b/Engine/Source/Editor/UnrealEd/Public/SComponentClassCombo.h
@@ -13,13 +13,13 @@ class FComponentClassComboEntry: public TSharedFromThis<FComponentClassComboEntr
 {
 public:
 	FComponentClassComboEntry( const FString& InHeadingText, TSubclassOf<UActorComponent> InComponentClass )
-		: HeadingText(InHeadingText)
-		, ComponentClass(InComponentClass)
+		: ComponentClass(InComponentClass)
+		, HeadingText(InHeadingText)
 	{}
 
 	FComponentClassComboEntry( const FString& InHeadingText )
-		: HeadingText(InHeadingText)
-		, ComponentClass(NULL)
+		: ComponentClass(NULL)
+		, HeadingText(InHeadingText)
 	{}
 
 	FComponentClassComboEntry()
