Typo fixed: bCachedDependenciesUpTpDate -> bCachedDependenciesUpToDate

diff --git a/Engine/Source/Editor/UnrealEd/Private/Kismet2/BlueprintEditorUtils.cpp b/Engine/Source/Editor/UnrealEd/Private/Kismet2/BlueprintEditorUtils.cpp
--- a/Engine/Source/Editor/UnrealEd/Private/Kismet2/BlueprintEditorUtils.cpp
+++ b/Engine/Source/Editor/UnrealEd/Private/Kismet2/BlueprintEditorUtils.cpp
@@ -1484,7 +1484,7 @@ void FBlueprintEditorUtils::UpdateDelegatesInBlueprint(UBlueprint* Blueprint)
 // Blueprint has materially changed.  Recompile the skeleton, notify observers, and mark the package as dirty.
 void FBlueprintEditorUtils::MarkBlueprintAsStructurallyModified(UBlueprint* Blueprint)
 {
-	Blueprint->bCachedDependenciesUpTpDate = false;
+	Blueprint->bCachedDependenciesUpToDate = false;
 	if (Blueprint->Status != BS_BeingCreated && !Blueprint->bBeingCompiled)
 	{
 		{
@@ -1511,7 +1511,7 @@ void FBlueprintEditorUtils::MarkBlueprintAsStructurallyModified(UBlueprint* Blue
 // Blueprint has changed in some manner that invalidates the compiled data (link made/broken, default value changed, etc...)
 void FBlueprintEditorUtils::MarkBlueprintAsModified(UBlueprint* Blueprint)
 {
-	Blueprint->bCachedDependenciesUpTpDate = false;
+	Blueprint->bCachedDependenciesUpToDate = false;
 	if (Blueprint->Status != BS_BeingCreated)
 	{
 		Blueprint->Status = BS_Dirty;
@@ -2123,10 +2123,10 @@ void FBlueprintEditorUtils::GetDependentBlueprints(UBlueprint* Blueprint, TArray
 		UBlueprint* TestBP = (UBlueprint*) *ObjIt;
 		if (TestBP && !TestBP->HasAnyFlags(RF_PendingKill))
 		{
-			if (!TestBP->bCachedDependenciesUpTpDate)
+			if (!TestBP->bCachedDependenciesUpToDate)
 			{
 				GatherDependencies(TestBP, TestBP->CachedDependencies);
-				TestBP->bCachedDependenciesUpTpDate = true;
+				TestBP->bCachedDependenciesUpToDate = true;
 			}
 
 			if (TestBP->CachedDependencies.Contains(Blueprint))
diff --git a/Engine/Source/Runtime/Engine/Classes/Engine/Blueprint.h b/Engine/Source/Runtime/Engine/Classes/Engine/Blueprint.h
--- a/Engine/Source/Runtime/Engine/Classes/Engine/Blueprint.h
+++ b/Engine/Source/Runtime/Engine/Classes/Engine/Blueprint.h
@@ -413,7 +413,7 @@ public:
 #endif // WITH_EDITORONLY_DATA
 
 #if WITH_EDITORONLY_DATA
-	bool bCachedDependenciesUpTpDate;
+	bool bCachedDependenciesUpToDate;
 	TSet<TWeakObjectPtr<UBlueprint>> CachedDependencies;
 
 	bool IsUpToDate() const
diff --git a/Engine/Source/Runtime/Engine/Private/Blueprint.cpp b/Engine/Source/Runtime/Engine/Private/Blueprint.cpp
--- a/Engine/Source/Runtime/Engine/Private/Blueprint.cpp
+++ b/Engine/Source/Runtime/Engine/Private/Blueprint.cpp
@@ -165,7 +165,7 @@ void UBlueprintCore::GenerateDeterministicGuid()
 UBlueprint::UBlueprint(const class FPostConstructInitializeProperties& PCIP)
 	: Super(PCIP)
 #if WITH_EDITORONLY_DATA
-	, bCachedDependenciesUpTpDate(false)
+	, bCachedDependenciesUpToDate(false)
 #endif
 #if WITH_EDITOR
 	, bRunConstructionScriptOnDrag(true)
@@ -1070,7 +1070,7 @@ void UBlueprint::PostLoadSubobjects(FObjectInstancingGraph* OuterInstanceGraph)
 
 bool UBlueprint::Modify(bool bAlwaysMarkDirty)
 {
-	bCachedDependenciesUpTpDate = false;
+	bCachedDependenciesUpToDate = false;
 	return Super::Modify(bAlwaysMarkDirty);
 }
 
