Improved accuracy of numbers being serialised to JSON.
Addresses issue brought to light by UDN question 215588.

diff --git a/Engine/Source/Editor/UnrealEd/Private/Tests/EditorAutomationTests.cpp b/Engine/Source/Editor/UnrealEd/Private/Tests/EditorAutomationTests.cpp
--- a/Engine/Source/Editor/UnrealEd/Private/Tests/EditorAutomationTests.cpp
+++ b/Engine/Source/Editor/UnrealEd/Private/Tests/EditorAutomationTests.cpp
@@ -514,13 +514,13 @@ bool FJsonAutomationTest::RunTest(const FString& Parameters)
 		check( FJsonSerializer::Deserialize( Reader, Object ) );
 		check( Object.IsValid() );
 
-		float TestValues[] = {2.544e+15f, -0.544e-2f, 251e3, -0.0, 843};
+		double TestValues[] = {2.544e+15, -0.544e-2, 251e3, -0.0, 843};
 		for (int32 i = 0; i < 5; ++i)
 		{
 			const TSharedPtr<FJsonValue>* Value = Object->Values.Find(FString::Printf(TEXT("Value%i"), i + 1));
 			check(Value && (*Value)->Type == EJson::Number);
-			const float Float = (*Value)->AsNumber();
-			check(FMath::Abs(Float - TestValues[i]) < KINDA_SMALL_NUMBER);
+			const double Number = (*Value)->AsNumber();
+			check(Number == TestValues[i]);
 		}
 
 		FString OutputString;
@@ -530,11 +530,11 @@ bool FJsonAutomationTest::RunTest(const FString& Parameters)
 		// %g isn't standardized, so we use the same %g format that is used inside PrintJson instead of hardcoding the values here
 		const FString TestOutput = FString::Printf(
 			TEXT("{")
-			TEXT(	"\"Value1\":%g,")
-			TEXT(	"\"Value2\":%g,")
-			TEXT(	"\"Value3\":%g,")
-			TEXT(	"\"Value4\":%g,")
-			TEXT(	"\"Value5\":%g")
+			TEXT(	"\"Value1\":%.17g,")
+			TEXT(	"\"Value2\":%.17g,")
+			TEXT(	"\"Value3\":%.17g,")
+			TEXT(	"\"Value4\":%.17g,")
+			TEXT(	"\"Value5\":%.17g")
 			TEXT("}"),
 			TestValues[0], TestValues[1], TestValues[2], TestValues[3], TestValues[4]);
 		check(OutputString == TestOutput);
diff --git a/Engine/Source/Runtime/Core/Public/Serialization/Json/JsonWriter.h b/Engine/Source/Runtime/Core/Public/Serialization/Json/JsonWriter.h
--- a/Engine/Source/Runtime/Core/Public/Serialization/Json/JsonWriter.h
+++ b/Engine/Source/Runtime/Core/Public/Serialization/Json/JsonWriter.h
@@ -328,7 +328,9 @@ protected:
 
 	FORCEINLINE void WriteNumberValue( const double Value)
 	{
-		PrintPolicy::WriteString(Stream, FString::Printf(TEXT("%g"), Value));
+		// Specify 17 significant digits, the most that can ever be useful from a double
+		// In particular, this ensures large integers are written correctly
+		PrintPolicy::WriteString(Stream, FString::Printf(TEXT("%.17g"), Value));
 	}
 
 	FORCEINLINE void WriteIntegerValue( const int64 Value)
