Added some extra info to the check in RecompileBlueprintBytecode to indicate which blueprint has an invalid genrated class.

diff --git a/Engine/Source/Editor/UnrealEd/Private/Kismet2/Kismet2.cpp b/Engine/Source/Editor/UnrealEd/Private/Kismet2/Kismet2.cpp
--- a/Engine/Source/Editor/UnrealEd/Private/Kismet2/Kismet2.cpp
+++ b/Engine/Source/Editor/UnrealEd/Private/Kismet2/Kismet2.cpp
@@ -487,7 +487,7 @@ void FKismetEditorUtilities::GenerateBlueprintSkeleton(UBlueprint* BlueprintObj,
 void FKismetEditorUtilities::RecompileBlueprintBytecode(UBlueprint* BlueprintObj, TArray<UObject*>* ObjLoaded)
 {
 	check(BlueprintObj);
-	check(BlueprintObj->GeneratedClass);
+	checkf(BlueprintObj->GeneratedClass, TEXT("Invalid generated class for %s"), *BlueprintObj->GetName());
 
 	IKismetCompilerInterface& Compiler = FModuleManager::LoadModuleChecked<IKismetCompilerInterface>(KISMET_COMPILER_MODULENAME);
 
