- Reduced the workscale for auto detection in scalability settings.
- Give the user the chance to not auto-apply scalability settings via a notification.

diff --git a/Engine/Source/Editor/UnrealEd/Classes/Editor/EditorGameAgnosticSettings.h b/Engine/Source/Editor/UnrealEd/Classes/Editor/EditorGameAgnosticSettings.h
--- a/Engine/Source/Editor/UnrealEd/Classes/Editor/EditorGameAgnosticSettings.h
+++ b/Engine/Source/Editor/UnrealEd/Classes/Editor/EditorGameAgnosticSettings.h
@@ -43,6 +43,12 @@ class UEditorGameAgnosticSettings : public UObject
 	/** Load the engine scalability benchmark results. Performs a benchmark if not yet valid. */
 	void LoadScalabilityBenchmark();
 
+	/** Auto detects and applies the scalability benchmark */
+	void AutoApplyScalabilityBenchmark();
+
+	/** @return true if the scalability benchmark is valid */
+	bool IsScalabilityBenchmarkValid() const;
+
 	// Begin UObject Interface
 	virtual void PostEditChangeProperty( struct FPropertyChangedEvent& PropertyChangedEvent) override;
 	// End UObject Interface
diff --git a/Engine/Source/Editor/UnrealEd/Private/Editor/EditorGameAgnosticSettings.cpp b/Engine/Source/Editor/UnrealEd/Private/Editor/EditorGameAgnosticSettings.cpp
--- a/Engine/Source/Editor/UnrealEd/Private/Editor/EditorGameAgnosticSettings.cpp
+++ b/Engine/Source/Editor/UnrealEd/Private/Editor/EditorGameAgnosticSettings.cpp
@@ -37,12 +37,9 @@ void UEditorGameAgnosticSettings::LoadScalabilityBenchmark()
 
 	const TCHAR* Section = TEXT("EngineBenchmarkResult");
 
-	bool bIsValid = false;
-	bIsValid = GConfig->GetBool(Section, TEXT("Valid"), bIsValid, GEditorGameAgnosticIni) && bIsValid;
-
 	Scalability::FQualityLevels Temporary;
 
-	if (bIsValid)
+	if (IsScalabilityBenchmarkValid())
 	{
 		GConfig->GetInt(Section, TEXT("ResolutionQuality"),		Temporary.ResolutionQuality,	GEditorGameAgnosticIni);
 		GConfig->GetInt(Section, TEXT("ViewDistanceQuality"),		Temporary.ViewDistanceQuality,	GEditorGameAgnosticIni);
@@ -53,25 +50,40 @@ void UEditorGameAgnosticSettings::LoadScalabilityBenchmark()
 		GConfig->GetInt(Section, TEXT("EffectsQuality"),			Temporary.EffectsQuality,		GEditorGameAgnosticIni);
 		EngineBenchmarkResult = Temporary;
 	}
-	else
-	{
-		GWarn->StatusUpdate(0, 1, NSLOCTEXT("UnrealEd", "RunningEngineBenchmark", "Running engine benchmark..."));
-		GWarn->PushStatus();
-		
-		Temporary = Scalability::BenchmarkQualityLevels();
-
-		GConfig->SetBool(Section,TEXT("Valid"),					true,							GEditorGameAgnosticIni);
-		GConfig->SetInt(Section, TEXT("ResolutionQuality"),		Temporary.ResolutionQuality,	GEditorGameAgnosticIni);
-		GConfig->SetInt(Section, TEXT("ViewDistanceQuality"),		Temporary.ViewDistanceQuality,	GEditorGameAgnosticIni);
-		GConfig->SetInt(Section, TEXT("AntiAliasingQuality"),		Temporary.AntiAliasingQuality,	GEditorGameAgnosticIni);
-		GConfig->SetInt(Section, TEXT("ShadowQuality"),			Temporary.ShadowQuality,		GEditorGameAgnosticIni);
-		GConfig->SetInt(Section, TEXT("PostProcessQuality"),		Temporary.PostProcessQuality,	GEditorGameAgnosticIni);
-		GConfig->SetInt(Section, TEXT("TextureQuality"),			Temporary.TextureQuality,		GEditorGameAgnosticIni);
-		GConfig->SetInt(Section, TEXT("EffectsQuality"),			Temporary.EffectsQuality,		GEditorGameAgnosticIni);
-
-		Scalability::SetQualityLevels(Temporary);
-		Scalability::SaveState(GEditorGameAgnosticIni);
-		
-		GWarn->PopStatus();
-	}
+}
+
+void UEditorGameAgnosticSettings::AutoApplyScalabilityBenchmark()
+{
+	const TCHAR* Section = TEXT("EngineBenchmarkResult");
+
+	GWarn->StatusUpdate(0, 1, NSLOCTEXT("UnrealEd", "RunningEngineBenchmark", "Running engine benchmark..."));
+	GWarn->PushStatus();
+
+	uint32 WorkScale = 5.0f;
+
+	Scalability::FQualityLevels Temporary = Scalability::BenchmarkQualityLevels(WorkScale);
+
+	GConfig->SetBool(Section, TEXT("Valid"), true, GEditorGameAgnosticIni);
+	GConfig->SetInt(Section, TEXT("ResolutionQuality"), Temporary.ResolutionQuality, GEditorGameAgnosticIni);
+	GConfig->SetInt(Section, TEXT("ViewDistanceQuality"), Temporary.ViewDistanceQuality, GEditorGameAgnosticIni);
+	GConfig->SetInt(Section, TEXT("AntiAliasingQuality"), Temporary.AntiAliasingQuality, GEditorGameAgnosticIni);
+	GConfig->SetInt(Section, TEXT("ShadowQuality"), Temporary.ShadowQuality, GEditorGameAgnosticIni);
+	GConfig->SetInt(Section, TEXT("PostProcessQuality"), Temporary.PostProcessQuality, GEditorGameAgnosticIni);
+	GConfig->SetInt(Section, TEXT("TextureQuality"), Temporary.TextureQuality, GEditorGameAgnosticIni);
+	GConfig->SetInt(Section, TEXT("EffectsQuality"), Temporary.EffectsQuality, GEditorGameAgnosticIni);
+
+	Scalability::SetQualityLevels(Temporary);
+	Scalability::SaveState(GEditorGameAgnosticIni);
+
+	GWarn->PopStatus();
+}
+
+bool UEditorGameAgnosticSettings::IsScalabilityBenchmarkValid() const
+{
+	const TCHAR* Section = TEXT("EngineBenchmarkResult");
+
+	bool bIsValid = false;
+	GConfig->GetBool(Section, TEXT("Valid"), bIsValid, GEditorGameAgnosticIni);
+
+	return bIsValid;
 }
\ No newline at end of file
diff --git a/Engine/Source/Editor/UnrealEd/Private/PerformanceMonitor.cpp b/Engine/Source/Editor/UnrealEd/Private/PerformanceMonitor.cpp
--- a/Engine/Source/Editor/UnrealEd/Private/PerformanceMonitor.cpp
+++ b/Engine/Source/Editor/UnrealEd/Private/PerformanceMonitor.cpp
@@ -157,6 +157,8 @@ Scalability::FQualityLevels FPerformanceMonitor::GetAutoScalabilityQualityLevels
 
 void FPerformanceMonitor::AutoApplyScalability()
 {
+	GEditor->AccessGameAgnosticSettings().AutoApplyScalabilityBenchmark();
+
 	Scalability::FQualityLevels NewLevels = FPerformanceMonitor::GetAutoScalabilityQualityLevels();
 
 	Scalability::SetQualityLevels(NewLevels);
@@ -224,7 +226,7 @@ void FPerformanceMonitor::HidePerformanceWarning()
 
 void FPerformanceMonitor::Tick(float DeltaTime)
 {
-	if (GEngine->ShouldThrottleCPUUsage())
+	if (GEngine->ShouldThrottleCPUUsage() && (!GShaderCompilingManager || !GShaderCompilingManager->IsCompiling() ) )
 	{
 		return;
 	}
@@ -290,9 +292,7 @@ void FPerformanceMonitor::Tick(float DeltaTime)
 	}
 	else
 	{
-		// Before we warn the user lets check if we already have auto-scalability enabled, no sense in warning them
-		// if we can't do anything about it.
-		if ( !WillAutoScalabilityHelp() )
+		if( GEditor->GetGameAgnosticSettings().IsScalabilityBenchmarkValid() )
 		{
 			return;
 		}
diff --git a/Engine/Source/Editor/UnrealEd/Private/UnrealEd.cpp b/Engine/Source/Editor/UnrealEd/Private/UnrealEd.cpp
--- a/Engine/Source/Editor/UnrealEd/Private/UnrealEd.cpp
+++ b/Engine/Source/Editor/UnrealEd/Private/UnrealEd.cpp
@@ -91,13 +91,6 @@ int32 EditorInit( IEngineLoop& EngineLoop )
 	// Set up the actor folders singleton
 	FActorFolders::Init();
 
-	if ( !GEditorGameAgnosticIni.IsEmpty() )
-	{
-		// If we have a game agnostic ini config, ensure that the benchmark has been loaded for it
-		GEditor->AccessGameAgnosticSettings().LoadScalabilityBenchmark();
-		GEditor->SaveGameAgnosticSettings();
-	}
-
 	// =================== CORE EDITOR INIT FINISHED ===================
 
 	// Hide the splash screen now that everything is ready to go
diff --git a/Engine/Source/Runtime/Engine/Private/Scalability.cpp b/Engine/Source/Runtime/Engine/Private/Scalability.cpp
--- a/Engine/Source/Runtime/Engine/Private/Scalability.cpp
+++ b/Engine/Source/Runtime/Engine/Private/Scalability.cpp
@@ -231,14 +231,14 @@ static int32 GetRenderScaleLevelFromQualityLevel(int32 InQualityLevel)
 	return ScalesForQuality[InQualityLevel];
 }
 
-FQualityLevels BenchmarkQualityLevels()
+FQualityLevels BenchmarkQualityLevels(uint32 WorkScale)
 {
 	// benchmark the system
 
 	FQualityLevels Results;
 
 	FSynthBenchmarkResults SynthBenchmark;
-	ISynthBenchmark::Get().Run(SynthBenchmark, true);
+	ISynthBenchmark::Get().Run(SynthBenchmark, true, WorkScale );
 
 	float CPUPerfIndex = SynthBenchmark.ComputeCPUPerfIndex();
 	float GPUPerfIndex = SynthBenchmark.ComputeGPUPerfIndex();
diff --git a/Engine/Source/Runtime/Engine/Public/Scalability.h b/Engine/Source/Runtime/Engine/Public/Scalability.h
--- a/Engine/Source/Runtime/Engine/Public/Scalability.h
+++ b/Engine/Source/Runtime/Engine/Public/Scalability.h
@@ -69,7 +69,7 @@ namespace Scalability
 	ENGINE_API void RecordQualityLevelsAnalytics(bool bAutoApplied);
 
 	/** Run synthbenchmark and configure scalability based on results **/
-	ENGINE_API FQualityLevels BenchmarkQualityLevels();
+	ENGINE_API FQualityLevels BenchmarkQualityLevels(uint32 WorkScale=10);
 
 	/** Process a console command line **/
 	ENGINE_API void ProcessCommand(const TCHAR* Cmd, FOutputDevice& Ar);
