Fix new class dialog class tooltips pointing to the wrong documentation page location

diff --git a/Engine/Source/Editor/GameProjectGeneration/Private/SNewClassDialog.cpp b/Engine/Source/Editor/GameProjectGeneration/Private/SNewClassDialog.cpp
--- a/Engine/Source/Editor/GameProjectGeneration/Private/SNewClassDialog.cpp
+++ b/Engine/Source/Editor/GameProjectGeneration/Private/SNewClassDialog.cpp
@@ -621,13 +621,14 @@ TSharedRef<ITableRow> SNewClassDialog::MakeParentClassListViewWidget(TSharedPtr<
 	const FString ClassName = ParentClassItem->ParentClassInfo.GetClassName();
 	const FString ClassDescription = ParentClassItem->ParentClassInfo.GetClassDescription();
 	const FSlateBrush* const ClassBrush = ParentClassItem->ParentClassInfo.GetClassIcon();
+	const UClass* Class = ParentClassItem->ParentClassInfo.BaseClass;
 
 	const int32 ItemHeight = 64;
 	const int32 DescriptionIndent = 32;
 	return
 		SNew( STableRow<TSharedPtr<FParentClassItem>>, OwnerTable )
 		.Style(FEditorStyle::Get(), "NewClassDialog.ParentClassListView.TableRow")
-		.ToolTip(IDocumentation::Get()->CreateToolTip(FText::FromString(ClassDescription), nullptr, "Shared/Classes", ClassName))
+		.ToolTip(IDocumentation::Get()->CreateToolTip(FText::FromString(ClassDescription), nullptr, FEditorClassUtils::GetDocumentationPage(Class), FEditorClassUtils::GetDocumentationExcerpt(Class)))
 		[
 			SNew(SBox).HeightOverride(ItemHeight)
 			[
diff --git a/Engine/Source/Editor/UnrealEd/Private/EditorClassUtils.cpp b/Engine/Source/Editor/UnrealEd/Private/EditorClassUtils.cpp
--- a/Engine/Source/Editor/UnrealEd/Private/EditorClassUtils.cpp
+++ b/Engine/Source/Editor/UnrealEd/Private/EditorClassUtils.cpp
@@ -6,19 +6,19 @@
 #include "IDocumentation.h"
 #include "SourceCodeNavigation.h"
 
-FString GetDocumentationPage(const UClass* Class)
+FString FEditorClassUtils::GetDocumentationPage(const UClass* Class)
 {
-	return FString::Printf( TEXT("Shared/Types/%s%s"), Class->GetPrefixCPP(), *Class->GetName() );
+	return (Class ? FString::Printf( TEXT("Shared/Types/%s%s"), Class->GetPrefixCPP(), *Class->GetName() ) : FString());
 }
 
-FString GetDocumentationExcerpt(const UClass* Class)
+FString FEditorClassUtils::GetDocumentationExcerpt(const UClass* Class)
 {
-	return FString::Printf( TEXT("%s%s"), Class->GetPrefixCPP(), *Class->GetName() );
+	return (Class ? FString::Printf( TEXT("%s%s"), Class->GetPrefixCPP(), *Class->GetName() ) : FString());
 }
 
 TSharedRef<SToolTip> FEditorClassUtils::GetTooltip(const UClass* Class)
 {
-	return IDocumentation::Get()->CreateToolTip(Class->GetToolTipText(), nullptr, GetDocumentationPage(Class), GetDocumentationExcerpt(Class));
+	return (Class ? IDocumentation::Get()->CreateToolTip(Class->GetToolTipText(), nullptr, GetDocumentationPage(Class), GetDocumentationExcerpt(Class)) : SNew(SToolTip));
 }
 
 FString FEditorClassUtils::GetDocumentationLink(const UClass* Class)
@@ -61,8 +61,9 @@ TSharedRef<SWidget> FEditorClassUtils::GetDocumentationLinkWidget(const UClass*
 TSharedRef<SWidget> FEditorClassUtils::GetSourceLink(const UClass* Class, const TWeakObjectPtr<UObject> ObjectWeakPtr)
 {
 	TSharedRef<SWidget> SourceHyperlink = SNullWidget::NullWidget;
+	UBlueprint* Blueprint = (Class ? Cast<UBlueprint>(Class->ClassGeneratedBy) : nullptr);
 
-	if (UBlueprint* Blueprint = Cast<UBlueprint>(Class->ClassGeneratedBy))
+	if (Blueprint)
 	{
 		struct Local
 		{
diff --git a/Engine/Source/Editor/UnrealEd/Public/EditorClassUtils.h b/Engine/Source/Editor/UnrealEd/Public/EditorClassUtils.h
--- a/Engine/Source/Editor/UnrealEd/Public/EditorClassUtils.h
+++ b/Engine/Source/Editor/UnrealEd/Public/EditorClassUtils.h
@@ -8,6 +8,23 @@ class SToolTip;
 namespace FEditorClassUtils
 {
 
+	/**
+	 * Gets the page that documentation for this class is contained on
+	 *
+	 * @param	InClass		Class we want to find the documentation page of
+	 * @return				Path to the documentation page
+	 */
+	UNREALED_API FString GetDocumentationPage(const UClass* Class);
+
+	/**
+	 * Gets the excerpt to use for this class
+	 * Excerpt will be contained on the page returned by GetDocumentationPage
+	 *
+	 * @param	InClass		Class we want to find the documentation excerpt of
+	 * @return				Name of the to the documentation excerpt
+	 */
+	UNREALED_API FString GetDocumentationExcerpt(const UClass* Class);
+
 	/**
 	 * Gets the tooltip to display for a given class
 	 *
