294048	Add Support for exporting custom curves from DCC apps into UE4

diff --git a/Engine/Source/Editor/DetailCustomizations/Private/WheeledVehicleMovementComponent4WDetails.cpp b/Engine/Source/Editor/DetailCustomizations/Private/WheeledVehicleMovementComponent4WDetails.cpp
--- a/Engine/Source/Editor/DetailCustomizations/Private/WheeledVehicleMovementComponent4WDetails.cpp
+++ b/Engine/Source/Editor/DetailCustomizations/Private/WheeledVehicleMovementComponent4WDetails.cpp
@@ -83,7 +83,8 @@ void FWheeledVehicleMovementComponent4WDetails::CustomizeDetails(IDetailLayoutBu
 				.ViewMaxOutput(1.f)
 				.TimelineLength(150)
 				.HideUI(false)
-				.ZoomToFit(false)
+				.ZoomToFitVertical(false)
+				.ZoomToFitHorizontal(false)
 				.DesiredSize(FVector2D(512, 128))
 			];
 
diff --git a/Engine/Source/Editor/Kismet/Private/STimelineEditor.cpp b/Engine/Source/Editor/Kismet/Private/STimelineEditor.cpp
--- a/Engine/Source/Editor/Kismet/Private/STimelineEditor.cpp
+++ b/Engine/Source/Editor/Kismet/Private/STimelineEditor.cpp
@@ -239,7 +239,7 @@ void STimelineEdTrack::Construct(const FArguments& InArgs, TSharedPtr<FTimelineE
 		}
 
 		//Inform track widget about the curve and whether it is editable or not.
-		TrackWidget->SetZoomToFit(bZoomToFit);
+		TrackWidget->SetZoomToFit(bZoomToFit, bZoomToFit);
 		TrackWidget->SetCurveOwner(CurveBasePtr, !TrackBase->bIsExternalCurve);
 	}
 
diff --git a/Engine/Source/Editor/Persona/Private/SAnimCurvePanel.cpp b/Engine/Source/Editor/Persona/Private/SAnimCurvePanel.cpp
--- a/Engine/Source/Editor/Persona/Private/SAnimCurvePanel.cpp
+++ b/Engine/Source/Editor/Persona/Private/SAnimCurvePanel.cpp
@@ -223,7 +223,8 @@ void SAnimCurveEd::Construct(const FArguments& InArgs)
 		.DataMaxInput(InArgs._DataMaxInput)
 		.ViewMinOutput(0.f)
 		.ViewMaxOutput(1.f)
-		.ZoomToFit(false)
+		.ZoomToFitVertical(true)
+		.ZoomToFitHorizontal(false)
 		.TimelineLength(InArgs._TimelineLength)
 		.DrawCurve(InArgs._DrawCurve)
 		.HideUI(InArgs._HideUI)
diff --git a/Engine/Source/Editor/UnrealEd/Classes/Factories/FbxAnimSequenceImportData.h b/Engine/Source/Editor/UnrealEd/Classes/Factories/FbxAnimSequenceImportData.h
--- a/Engine/Source/Editor/UnrealEd/Classes/Factories/FbxAnimSequenceImportData.h
+++ b/Engine/Source/Editor/UnrealEd/Classes/Factories/FbxAnimSequenceImportData.h
@@ -34,7 +34,7 @@ class UFbxAnimSequenceImportData : public UFbxAssetImportData
 
 	/** Type of asset to import from the FBX file */
 	UPROPERTY(EditAnywhere, Category=ImportSettings, meta=(DisplayName = "Start Frame"))
-	int32 StartFrame;
+	int32	StartFrame;
 	/** Type of asset to import from the FBX file */
 	UPROPERTY(EditAnywhere, Category=ImportSettings, meta=(DisplayName = "End Frame"))
 	int32	EndFrame;
@@ -43,6 +43,10 @@ class UFbxAnimSequenceImportData : public UFbxAssetImportData
 	UPROPERTY(EditAnywhere, Category = ImportSettings, meta = (DisplayName = "Source Animation Name"))
 	FString SourceAnimationName;
 
+	/** Import if custom attribute as a curve within the animation **/
+	UPROPERTY(EditAnywhere, config, Category = ImportSettings, meta = (DisplayName = "Import Custom Attribute"))
+	bool bImportCustomAttribute;
+
 	/** Gets or creates fbx import data for the specified anim sequence */
 	static UFbxAnimSequenceImportData* GetImportDataForAnimSequence(UAnimSequence* AnimSequence, UFbxAnimSequenceImportData* TemplateForCreation);
 };
\ No newline at end of file
diff --git a/Engine/Source/Editor/UnrealEd/Classes/Factories/FbxSkeletalMeshImportData.h b/Engine/Source/Editor/UnrealEd/Classes/Factories/FbxSkeletalMeshImportData.h
--- a/Engine/Source/Editor/UnrealEd/Classes/Factories/FbxSkeletalMeshImportData.h
+++ b/Engine/Source/Editor/UnrealEd/Classes/Factories/FbxSkeletalMeshImportData.h
@@ -27,10 +27,6 @@ class UFbxSkeletalMeshImportData : public UFbxMeshImportData
 	UPROPERTY(EditAnywhere, config, Category=ImportSettings)
 	uint32 bImportMeshesInBoneHierarchy:1;
 
-	/** If checked, we will import group node (null node) as root. */
-	UPROPERTY(EditAnywhere, config, Category=ImportSettings)
-	uint32 bImportGroupNodeAsRoot:1;
-
 	/** True to import morph target meshes from the FBX file */
 	UPROPERTY(EditAnywhere, config, Category=ImportSettings, meta=(ToolTip="If enabled, creates Unreal morph objects for the imported meshes"))
 	uint32 bImportMorphTargets:1;
diff --git a/Engine/Source/Editor/UnrealEd/Private/Fbx/FbxAnimSequenceImportData.cpp b/Engine/Source/Editor/UnrealEd/Private/Fbx/FbxAnimSequenceImportData.cpp
--- a/Engine/Source/Editor/UnrealEd/Private/Fbx/FbxAnimSequenceImportData.cpp
+++ b/Engine/Source/Editor/UnrealEd/Private/Fbx/FbxAnimSequenceImportData.cpp
@@ -4,6 +4,7 @@
 
 UFbxAnimSequenceImportData::UFbxAnimSequenceImportData(const class FPostConstructInitializeProperties& PCIP)
 	: Super(PCIP)
+	, bImportCustomAttribute(true)
 {
 	
 }
diff --git a/Engine/Source/Editor/UnrealEd/Private/Fbx/FbxMainImport.cpp b/Engine/Source/Editor/UnrealEd/Private/Fbx/FbxMainImport.cpp
--- a/Engine/Source/Editor/UnrealEd/Private/Fbx/FbxMainImport.cpp
+++ b/Engine/Source/Editor/UnrealEd/Private/Fbx/FbxMainImport.cpp
@@ -198,7 +198,6 @@ void ApplyImportUIToImportOptions(UFbxImportUI* ImportUI, FBXImportOptions& InOu
 	InOutImportOptions.bOneConvexHullPerUCX = ImportUI->StaticMeshImportData->bOneConvexHullPerUCX;
 	InOutImportOptions.StaticMeshLODGroup = ImportUI->StaticMeshImportData->StaticMeshLODGroup;
 	InOutImportOptions.bImportMeshesInBoneHierarchy = ImportUI->SkeletalMeshImportData->bImportMeshesInBoneHierarchy;
-	InOutImportOptions.bImportGroupNodeAsRoot = ImportUI->SkeletalMeshImportData->bImportGroupNodeAsRoot;
 	InOutImportOptions.bCreatePhysicsAsset = ImportUI->bCreatePhysicsAsset;
 	InOutImportOptions.PhysicsAsset = ImportUI->PhysicsAsset;
 	// animation options
@@ -207,6 +206,7 @@ void ApplyImportUIToImportOptions(UFbxImportUI* ImportUI, FBXImportOptions& InOu
 	InOutImportOptions.AnimationRange.Y = ImportUI->AnimSequenceImportData->EndFrame;
 	InOutImportOptions.AnimationName = ImportUI->AnimationName;
 	InOutImportOptions.bPreserveLocalTransform = ImportUI->bPreserveLocalTransform;
+	InOutImportOptions.bImportCustomAttribute = ImportUI->AnimSequenceImportData->bImportCustomAttribute;
 }
 
 void FImportedMaterialData::AddImportedMaterial( FbxSurfaceMaterial& FbxMaterial, UMaterialInterface& UnrealMaterial )
@@ -1068,30 +1068,47 @@ FbxNode* FFbxImporter::GetRootSkeleton(FbxNode* Link)
 }
 
 
-void DumpFBXNode(FbxNode* Node)
+void FFbxImporter::DumpFBXNode(FbxNode* Node) 
 {
 	FbxMesh* Mesh = Node->GetMesh();
 	const FString NodeName(Node->GetName());
-	FbxNodeAttribute* NodeAttribute = Node->GetNodeAttribute();
-	if(NodeAttribute)
-	{
-		FbxNodeAttribute::EType Type = NodeAttribute->GetAttributeType();
-	}
 
 	if(Mesh)
 	{
+		UE_LOG(LogFbx, Log, TEXT("================================================="));
+		UE_LOG(LogFbx, Log, TEXT("Dumping Node START [%s] "), *NodeName);
 		int DeformerCount = Mesh->GetDeformerCount();
-		UE_LOG(LogFbx, Log,TEXT("================================================="));
-		UE_LOG(LogFbx, Log,TEXT("Dumping Node START [%s] : Total Deformer Count %d."), *NodeName, DeformerCount);
+		UE_LOG(LogFbx, Log,TEXT("\tTotal Deformer Count %d."), *NodeName, DeformerCount);
 		for(int i=0; i<DeformerCount; i++)
 		{
 			FbxDeformer* Deformer = Mesh->GetDeformer(i);
 			const FString DeformerName(Deformer->GetName());
 			const FString DeformerTypeName(Deformer->GetTypeName());
-			UE_LOG(LogFbx, Log,TEXT("\t[Node %d] %s (Type %s)."), i+1, *DeformerName, *DeformerTypeName);
+			UE_LOG(LogFbx, Log,TEXT("\t\t[Node %d] %s (Type %s)."), i+1, *DeformerName, *DeformerTypeName);
 			UE_LOG(LogFbx, Log,TEXT("================================================="));
 		}
 
+		FbxNodeAttribute* NodeAttribute = Node->GetNodeAttribute();
+		if(NodeAttribute)
+		{
+			FString NodeAttributeName(NodeAttribute->GetName());
+			FbxNodeAttribute::EType Type = NodeAttribute->GetAttributeType();
+			UE_LOG(LogFbx, Log,TEXT("\tAttribute (%s) Type (%d)."), *NodeAttributeName, (int32)Type);
+		
+			for (int i=0; i<NodeAttribute->GetNodeCount(); ++i)
+			{
+				FbxNode * Child = NodeAttribute->GetNode(i);
+
+				if (Child)
+				{
+					const FString ChildName(Child->GetName());
+					const FString ChildTypeName(Child->GetTypeName());
+					UE_LOG(LogFbx, Log,TEXT("\t\t[Node Attribute Child %d] %s (Type %s)."), i+1, *ChildName, *ChildTypeName);
+				}
+			}
+
+		}
+
 		UE_LOG(LogFbx, Log,TEXT("Dumping Node END [%s]"), *NodeName);
 	}
 
@@ -1203,6 +1220,8 @@ void FFbxImporter::RecursiveFindRigidMesh(FbxNode* Node, TArray< TArray<FbxNode*
 	bool RigidNodeFound = false;
 	FbxNode* RigidMeshNode = NULL;
 
+	DEBUG_FBX_NODE("", Node);
+
 	if (Node->GetMesh())
 	{
 		// ignore skeletal mesh
diff --git a/Engine/Source/Editor/UnrealEd/Private/Fbx/FbxSkeletalMeshImportData.cpp b/Engine/Source/Editor/UnrealEd/Private/Fbx/FbxSkeletalMeshImportData.cpp
--- a/Engine/Source/Editor/UnrealEd/Private/Fbx/FbxSkeletalMeshImportData.cpp
+++ b/Engine/Source/Editor/UnrealEd/Private/Fbx/FbxSkeletalMeshImportData.cpp
@@ -5,7 +5,6 @@
 UFbxSkeletalMeshImportData::UFbxSkeletalMeshImportData(const class FPostConstructInitializeProperties& PCIP)
 	: Super(PCIP)
 	, bImportMeshesInBoneHierarchy(true)
-	, bImportGroupNodeAsRoot(true)
 {
 }
 
diff --git a/Engine/Source/Editor/UnrealEd/Private/FbxImporter.h b/Engine/Source/Editor/UnrealEd/Private/FbxImporter.h
--- a/Engine/Source/Editor/UnrealEd/Private/FbxImporter.h
+++ b/Engine/Source/Editor/UnrealEd/Private/FbxImporter.h
@@ -125,7 +125,6 @@ struct FBXImportOptions
 	bool bPreserveSmoothingGroups;
 	bool bKeepOverlappingVertices;
 	bool bImportMeshesInBoneHierarchy;
-	bool bImportGroupNodeAsRoot;
 	bool bCreatePhysicsAsset;
 	UPhysicsAsset *PhysicsAsset;
 	// Animation option
@@ -134,6 +133,7 @@ struct FBXImportOptions
 	struct FIntPoint AnimationRange;
 	FString AnimationName;
 	bool	bPreserveLocalTransform;
+	bool	bImportCustomAttribute;
 
 	bool ShouldImportNormals()
 	{
@@ -394,11 +394,6 @@ public:
 	 */
 	FbxTimeSpan GetAnimationTimeSpan(FbxNode* RootNode, FbxAnimStack* AnimStack);
 
-	/**
-	 * Import FbxCurve to Curve
-	 */
-	bool ImportCurve(const FbxAnimCurve* FbxCurve, FFloatCurve * Curve, const FbxTimeSpan &AnimTimeSpan) const;
-
 	/**
 	 * Import one animation from CurAnimStack
 	 *
@@ -550,6 +545,8 @@ public:
 	 */
 	UNREALED_API FBXImportOptions* GetImportOptions() const;
 
+	/** helper function **/
+	UNREALED_API static void DumpFBXNode(FbxNode* Node);
 private:
 	/**
 	 * ActorX plug-in can export mesh and dummy as skeleton.
@@ -996,6 +993,18 @@ private:
 	void ClearLogger();
 
 	FImportedMaterialData ImportedMaterialData;
+
+private:
+	/**
+	 * Import FbxCurve to Curve
+	 */
+	bool ImportCurve(const FbxAnimCurve* FbxCurve, FFloatCurve * Curve, const FbxTimeSpan &AnimTimeSpan, const float ValueScale = 1.f) const;
+
+
+	/*
+	 * Import FbxCurve to anim sequence
+	 */
+	bool ImportCurveToAnimSequence(class UAnimSequence * TargetSequence, const FString & CurveName, const FbxAnimCurve * FbxCurve, int32 CurveFlags,const FbxTimeSpan AnimTimeSpan, const float ValueScale = 1.f) const;
 };
 
 
diff --git a/Engine/Source/Editor/UnrealEd/Private/FbxOptionWindow.h b/Engine/Source/Editor/UnrealEd/Private/FbxOptionWindow.h
--- a/Engine/Source/Editor/UnrealEd/Private/FbxOptionWindow.h
+++ b/Engine/Source/Editor/UnrealEd/Private/FbxOptionWindow.h
@@ -110,7 +110,6 @@ private:
 	void SetSkeletalMesh_ReserveSmoothingGroups(ESlateCheckBoxState::Type NewType);
 	void SetSkeletalMesh_KeepOverlappingVertices(ESlateCheckBoxState::Type NewType);
 	void SetSkeletalMesh_ImportMeshesInBoneHierarchy(ESlateCheckBoxState::Type NewType);
-	void SetSkeletalMesh_ImportGroupNodeAsRoot(ESlateCheckBoxState::Type NewType);
 	void SetSkeletalMesh_CreatePhysicsAsset(ESlateCheckBoxState::Type NewType);
 	void SetStaticMesh_ImportMeshLODs(ESlateCheckBoxState::Type NewType);
 	void SetStaticMesh_LODGroup(TSharedPtr<FString> NewValue, ESelectInfo::Type SelectInfo);
@@ -130,6 +129,7 @@ private:
 	void SetAnimationRangeStart(const FText& Name, ETextCommit::Type CommitInfo);
 	void SetAnimationRangeEnd(const FText& Name, ETextCommit::Type CommitInfo);
 	void SetAnimation_ReserveLocalTransform(ESlateCheckBoxState::Type NewType);
+	void SetAnimation_ImportCustomAttribute(ESlateCheckBoxState::Type NewType);
 
 	TSharedRef<SWidget> MakeSkeletonPickerMenu();
 	void OnAssetSelectedFromSkeletonPicker(const FAssetData& AssetData);
diff --git a/Engine/Source/Editor/UnrealEd/Private/SCurveEditor.cpp b/Engine/Source/Editor/UnrealEd/Private/SCurveEditor.cpp
--- a/Engine/Source/Editor/UnrealEd/Private/SCurveEditor.cpp
+++ b/Engine/Source/Editor/UnrealEd/Private/SCurveEditor.cpp
@@ -36,7 +36,8 @@ void SCurveEditor::Construct(const FArguments& InArgs)
 	ViewMinOutput = InArgs._ViewMinOutput;
 	ViewMaxOutput = InArgs._ViewMaxOutput;
 
-	bZoomToFit = InArgs._ZoomToFit;
+	bZoomToFitVertical = InArgs._ZoomToFitVertical;
+	bZoomToFitHorizontal = InArgs._ZoomToFitHorizontal;
 	DesiredSize = InArgs._DesiredSize;
 
 	bIsUsingSlider = false;
@@ -833,18 +834,23 @@ void SCurveEditor::SetCurveOwner(FCurveOwnerInterface* InCurveOwner, bool bCanEd
 
 	SelectedKeys.Empty();
 
-	if( bZoomToFit )
+	if( bZoomToFitVertical )
 	{
 		ZoomToFitVertical();
+	}
+
+	if ( bZoomToFitHorizontal )
+	{
 		ZoomToFitHorizontal();
 	}
 
 	CurveSelectionWidget.Pin()->SetContent(CreateCurveSelectionWidget());
 }
 
-void SCurveEditor::SetZoomToFit(bool bNewZoomToFit)
+void SCurveEditor::SetZoomToFit(bool bNewZoomToFitVertical, bool bNewZoomToFitHorizontal)
 {
-	bZoomToFit = bNewZoomToFit;
+	bZoomToFitVertical = bNewZoomToFitVertical;
+	bZoomToFitHorizontal = bNewZoomToFitHorizontal;
 }
 
 FCurveOwnerInterface* SCurveEditor::GetCurveOwner() const
diff --git a/Engine/Source/Editor/UnrealEd/Private/SkeletalMeshEdit.cpp b/Engine/Source/Editor/UnrealEd/Private/SkeletalMeshEdit.cpp
--- a/Engine/Source/Editor/UnrealEd/Private/SkeletalMeshEdit.cpp
+++ b/Engine/Source/Editor/UnrealEd/Private/SkeletalMeshEdit.cpp
@@ -692,7 +692,7 @@ bool UnFbx::FFbxImporter::ValidateAnimStack(TArray<FbxNode*>& SortedLinks, TArra
 	return bValidAnimStack;
 }
 
-bool UnFbx::FFbxImporter::ImportCurve(const FbxAnimCurve* FbxCurve, FFloatCurve * Curve, const FbxTimeSpan &AnimTimeSpan) const
+bool UnFbx::FFbxImporter::ImportCurve(const FbxAnimCurve* FbxCurve, FFloatCurve * Curve, const FbxTimeSpan &AnimTimeSpan, const float ValueScale/*=1.f*/) const
 {
 	static float DefaultCurveWeight = FbxAnimCurveDef::sDEFAULT_WEIGHT;
 
@@ -702,7 +702,7 @@ bool UnFbx::FFbxImporter::ImportCurve(const FbxAnimCurve* FbxCurve, FFloatCurve
 		{
 			FbxAnimCurveKey Key = FbxCurve->KeyGet(KeyIndex);
 			FbxTime KeyTime = Key.GetTime() - AnimTimeSpan.GetStart();
-			float Value = Key.GetValue()/100.f;
+			float Value = Key.GetValue() * ValueScale;
 			FKeyHandle NewKeyHandle = Curve->FloatCurve.AddKey(KeyTime.GetSecondDouble(), Value, true);
 
 			FbxAnimCurveDef::ETangentMode KeyTangentMode = Key.GetTangentMode();
@@ -912,6 +912,62 @@ namespace AnimationTransformDebug
 		}
 	}
 }
+
+/**
+ * We only support float values, so these are the numbers we can take
+ */
+bool IsSupportedCurveDataType(EFbxType DatatType)
+{
+
+	switch (DatatType)
+	{
+	case eFbxShort:		//!< 16 bit signed integer.
+	case eFbxUShort:		//!< 16 bit unsigned integer.
+	case eFbxUInt:		//!< 32 bit unsigned integer.
+	case eFbxHalfFloat:	//!< 16 bit floating point.
+	case eFbxInt:		//!< 32 bit signed integer.
+	case eFbxFloat:		//!< Floating point value.
+	case eFbxDouble:		//!< Double width floating point value.
+	case eFbxDouble2:	//!< Vector of two double values.
+	case eFbxDouble3:	//!< Vector of three double values.
+	case eFbxDouble4:	//!< Vector of four double values.
+	case eFbxDouble4x4:	//!< Four vectors of four double values.
+		return true;
+	}
+
+	return false;
+}
+
+bool UnFbx::FFbxImporter::ImportCurveToAnimSequence(class UAnimSequence * TargetSequence, const FString & CurveName, const FbxAnimCurve * FbxCurve, int32 CurveFlags,const FbxTimeSpan AnimTimeSpan, const float ValueScale/*=1.f*/) const
+{
+	if (TargetSequence && FbxCurve)
+	{
+		FName Name = *CurveName;
+
+		FFloatCurve * CurveToImport = TargetSequence->RawCurveData.GetCurveData(Name);
+		if(CurveToImport==NULL)
+		{
+			if(TargetSequence->RawCurveData.AddCurveData(Name, CurveFlags))
+			{
+				CurveToImport = TargetSequence->RawCurveData.GetCurveData(Name);
+			}
+			else
+			{
+				// this should not happen, we already checked before adding
+				ensureMsg(0, TEXT("FBX Import: Critical error: no memory?"));
+			}
+		}
+		else
+		{
+			CurveToImport->FloatCurve.Reset();
+		}
+
+		return ImportCurve(FbxCurve, CurveToImport, AnimTimeSpan, ValueScale);
+	}
+
+	return false;
+}
+
 bool UnFbx::FFbxImporter::ImportAnimation(USkeleton * Skeleton, UAnimSequence * DestSeq, const FString & FileName, TArray<FbxNode*>& SortedLinks, TArray<FbxNode*>& NodeArray, FbxAnimStack* CurAnimStack, const int32 ResampleRate, const FbxTimeSpan AnimTimeSpan)
 {
 	FbxTime SequenceLength = AnimTimeSpan.GetDuration();
@@ -968,26 +1024,7 @@ bool UnFbx::FFbxImporter::ImportAnimation(USkeleton * Skeleton, UAnimSequence *
 							if (Curve && Curve->KeyGetCount() > 0)
 							{
 								// now see if we have one already exists. If so, just overwrite that. if not, add new one. 
-								FName Name = *ChannelName;
-								FFloatCurve * CurveToImport = DestSeq->RawCurveData.GetCurveData(Name);
-								if (CurveToImport==NULL)
-								{
-									if ( DestSeq->RawCurveData.AddCurveData(Name, ACF_DrivesMorphTarget | ACF_TriggerEvent) )
-									{
-										CurveToImport = DestSeq->RawCurveData.GetCurveData(Name);
-									}
-									else
-									{
-										// this should not happen, we already checked before adding
-										ensureMsg (0, TEXT("FBX Import: Critical error: no memory?"));
-									}
-								}
-								else
-								{
-									CurveToImport->FloatCurve.Reset();
-								}
-	
-								ImportCurve(Curve, CurveToImport, AnimTimeSpan);
+								ImportCurveToAnimSequence(DestSeq, *ChannelName, Curve,  ACF_DrivesMorphTarget | ACF_TriggerEvent, AnimTimeSpan, 0.01f /** for some reason blend shape values are coming as 100 scaled **/);
 							}
 						}
 					}
@@ -1000,6 +1037,52 @@ bool UnFbx::FFbxImporter::ImportAnimation(USkeleton * Skeleton, UAnimSequence *
 	//shape animation END
 	//
 
+	// 
+	// importing custom attribute START
+	//
+	if (ImportOptions->bImportCustomAttribute)
+	{
+		for(auto Node: SortedLinks)
+		{
+			FbxProperty Property = Node->GetFirstProperty();
+			while (Property.IsValid())
+			{
+				FbxAnimCurveNode* CurveNode = Property.GetCurveNode();
+				// do this if user defined and animated and leaf node
+				if( CurveNode && Property.GetFlag(FbxPropertyAttr::eUserDefined) && 
+					CurveNode->IsAnimated() && IsSupportedCurveDataType(Property.GetPropertyDataType().GetType()) )
+				{
+					FString CurveName = CurveNode->GetName();
+					UE_LOG(LogFbx, Log, TEXT("CurveName : %s"), *CurveName );
+
+					int32 TotalCount = CurveNode->GetChannelsCount();
+					for (int32 ChannelIndex=0; ChannelIndex<TotalCount; ++ChannelIndex)
+					{
+						FbxAnimCurve * AnimCurve = CurveNode->GetCurve(ChannelIndex);
+						FString ChannelName = CurveNode->GetChannelName(ChannelIndex).Buffer();
+
+						if (AnimCurve)
+						{
+							if (TotalCount == 1)
+							{
+								ImportCurveToAnimSequence(DestSeq, CurveName, AnimCurve,  ACF_DefaultCurve, AnimTimeSpan);
+							}
+							else
+							{
+								ImportCurveToAnimSequence(DestSeq, CurveName + "_" + ChannelName, AnimCurve,  ACF_DefaultCurve, AnimTimeSpan);
+							}
+						}
+											
+					}
+				}
+
+				Property = Node->GetNextProperty(Property); 
+			}
+		}
+	}
+
+	// importing custom attribute END
+
 	DestSeq->TrackToSkeletonMapTable.Empty();
 	DestSeq->RawAnimationData.Empty();
 	DestSeq->AnimationTrackNames.Empty();
diff --git a/Engine/Source/Editor/UnrealEd/Public/SCurveEditor.h b/Engine/Source/Editor/UnrealEd/Public/SCurveEditor.h
--- a/Engine/Source/Editor/UnrealEd/Public/SCurveEditor.h
+++ b/Engine/Source/Editor/UnrealEd/Public/SCurveEditor.h
@@ -91,7 +91,8 @@ public:
 		, _HideUI(true)
 		, _AllowZoomOutput(true)
 		, _AlwaysDisplayColorCurves(false)
-		, _ZoomToFit(true)
+		, _ZoomToFitVertical(true)
+		, _ZoomToFitHorizontal(true)
 		{}
 
 		SLATE_ATTRIBUTE( float, ViewMinInput )
@@ -106,7 +107,8 @@ public:
 		SLATE_ARGUMENT( bool, HideUI )
 		SLATE_ARGUMENT( bool, AllowZoomOutput )
 		SLATE_ARGUMENT( bool, AlwaysDisplayColorCurves )
-		SLATE_ARGUMENT( bool, ZoomToFit )
+		SLATE_ARGUMENT( bool, ZoomToFitVertical )
+		SLATE_ARGUMENT( bool, ZoomToFitHorizontal )
 		SLATE_EVENT( FOnSetInputViewRange, OnSetInputViewRange )
 		SLATE_EVENT( FSimpleDelegate, OnCreateAsset )
 	SLATE_END_ARGS()
@@ -120,7 +122,7 @@ public:
 	UNREALED_API void SetCurveOwner(FCurveOwnerInterface* InCurveOwner, bool bCanEdit = true);
 
 	/** Set new zoom to fit **/
-	UNREALED_API void SetZoomToFit(bool bNewZoomToFit);
+	UNREALED_API void SetZoomToFit(bool bNewZoomToFitVertical, bool bNewZoomToFitHorizontal);
 
 	/** Get the currently edited curve */
 	FCurveOwnerInterface* GetCurveOwner() const;
@@ -485,7 +487,10 @@ protected:
 	float				ViewMaxOutput;
 
 	/** True if you want the curve editor to fit to zoom **/
-	bool				bZoomToFit;
+	bool				bZoomToFitVertical;
+
+	/** True if you want the curve editor to fit to zoom **/
+	bool				bZoomToFitHorizontal;
 
 	/** True if the sliders are being used to adjust point values **/
 	bool bIsUsingSlider;
diff --git a/Engine/Source/Runtime/Core/Public/Math/UnrealMatrix.h b/Engine/Source/Runtime/Core/Public/Math/UnrealMatrix.h
--- a/Engine/Source/Runtime/Core/Public/Math/UnrealMatrix.h
+++ b/Engine/Source/Runtime/Core/Public/Math/UnrealMatrix.h
@@ -166,7 +166,7 @@ FORCEINLINE FVector4 FMatrix::TransformPosition(const FVector &V) const
 /** Inverts the matrix and then transforms V - correctly handles scaling in this matrix. */
 FORCEINLINE FVector FMatrix::InverseTransformPosition(const FVector &V) const
 {
-	FMatrix InvSelf = this->Inverse();
+	FMatrix InvSelf = this->InverseSafe();
 	return InvSelf.TransformPosition(V);
 }
 
