Fix crash when switching levels with the foliage mode open and foliage brush visible

diff --git a/Engine/Source/Editor/MainFrame/Private/Frame/MainFrameHandler.h b/Engine/Source/Editor/MainFrame/Private/Frame/MainFrameHandler.h
--- a/Engine/Source/Editor/MainFrame/Private/Frame/MainFrameHandler.h
+++ b/Engine/Source/Editor/MainFrame/Private/Frame/MainFrameHandler.h
@@ -176,7 +176,8 @@ public:
 			if( GLevelEditorModeTools().IsModeActive( FBuiltinEditorModes::EM_Level ) || 
 				GLevelEditorModeTools().IsModeActive( FBuiltinEditorModes::EM_StreamingLevel) )
 			{
-				GLevelEditorModeTools().ActivateDefaultMode();
+				GLevelEditorModeTools().DeactivateMode(FBuiltinEditorModes::EM_Level);
+				GLevelEditorModeTools().DeactivateMode(FBuiltinEditorModes::EM_StreamingLevel);
 				bOkToExit = false;
 			}
 
diff --git a/Engine/Source/Editor/UnrealEd/Private/EditorLevelUtils.cpp b/Engine/Source/Editor/UnrealEd/Private/EditorLevelUtils.cpp
--- a/Engine/Source/Editor/UnrealEd/Private/EditorLevelUtils.cpp
+++ b/Engine/Source/Editor/UnrealEd/Private/EditorLevelUtils.cpp
@@ -549,7 +549,7 @@ namespace EditorLevelUtils
 	ULevel* CreateNewLevel( UWorld* InWorld, bool bMoveSelectedActorsIntoNewLevel, UClass* LevelStreamingClass, const FString& DefaultFilename )
 	{
 		// Editor modes cannot be active when any level saving occurs.
-		GLevelEditorModeTools().ActivateDefaultMode();
+		GLevelEditorModeTools().DeactivateAllModes();
 
 		UWorld* NewWorld = nullptr;
 		if (FParse::Param(FCommandLine::Get(), TEXT("WorldAssets")))
diff --git a/Engine/Source/Editor/UnrealEd/Private/FileHelpers.cpp b/Engine/Source/Editor/UnrealEd/Private/FileHelpers.cpp
--- a/Engine/Source/Editor/UnrealEd/Private/FileHelpers.cpp
+++ b/Engine/Source/Editor/UnrealEd/Private/FileHelpers.cpp
@@ -1744,23 +1744,8 @@ void FEditorFileUtils::LoadMap(const FString& InFilename, bool LoadAsTemplate, b
 		return;
 	}
 
-	// Change out of Matinee when opening new map, so we avoid editing data in the old one.
-	if( GLevelEditorModeTools().IsModeActive( FBuiltinEditorModes::EM_InterpEdit ) )
-	{
-		GLevelEditorModeTools().ActivateDefaultMode();
-	}
-
-	// Also change out of Landscape mode to ensure all references are cleared.
-	if( GLevelEditorModeTools().IsModeActive( FBuiltinEditorModes::EM_Landscape ) )
-	{
-		GLevelEditorModeTools().DeactivateMode( FBuiltinEditorModes::EM_Landscape );
-	}
-
-	// Change out of mesh paint mode when loading a map
-	if( GLevelEditorModeTools().IsModeActive( FBuiltinEditorModes::EM_MeshPaint ) )
-	{
-		GLevelEditorModeTools().DeactivateMode( FBuiltinEditorModes::EM_MeshPaint );
-	}
+	// Deactivate any editor modes when loading a new map
+	GLevelEditorModeTools().DeactivateAllModes();
 
 	FString LoadCommand = FString::Printf( TEXT("MAP LOAD FILE=\"%s\" TEMPLATE=%d SHOWPROGRESS=%d"), *Filename, LoadAsTemplate, bShowProgress );
 	bool bResult = GUnrealEd->Exec( NULL, *LoadCommand );
diff --git a/Engine/Source/Editor/UnrealEd/Private/UnrealEd.cpp b/Engine/Source/Editor/UnrealEd/Private/UnrealEd.cpp
--- a/Engine/Source/Editor/UnrealEd/Private/UnrealEd.cpp
+++ b/Engine/Source/Editor/UnrealEd/Private/UnrealEd.cpp
@@ -160,12 +160,8 @@ int32 EditorInit( IEngineLoop& EngineLoop )
 
 void EditorExit()
 {
-	if( GLevelEditorModeTools().IsModeActive(FBuiltinEditorModes::EM_MeshPaint) ||
-		GLevelEditorModeTools().IsModeActive(FBuiltinEditorModes::EM_InterpEdit) )
-	{
-		GLevelEditorModeTools().SetDefaultMode(FBuiltinEditorModes::EM_Default);
-		GLevelEditorModeTools().ActivateDefaultMode();
-	}
+	GLevelEditorModeTools().SetDefaultMode(FBuiltinEditorModes::EM_Default);
+	GLevelEditorModeTools().DeactivateAllModes(); // this also activates the default mode
 
 	// Save out any config settings for the editor so they don't get lost
 	GEditor->SaveConfig();
