Fix issue with deleting blendspaces/montages after creating a thumbnail

diff --git a/Engine/Source/Editor/UnrealEd/Private/ThumbnailHelpers.cpp b/Engine/Source/Editor/UnrealEd/Private/ThumbnailHelpers.cpp
--- a/Engine/Source/Editor/UnrealEd/Private/ThumbnailHelpers.cpp
+++ b/Engine/Source/Editor/UnrealEd/Private/ThumbnailHelpers.cpp
@@ -565,6 +565,7 @@ bool FAnimationSequenceThumbnailScene::SetAnimation(UAnimSequenceBase* InAnimati
 	if(!bSetSucessfully)
 	{
 		CleanupComponentChildren(PreviewActor->SkeletalMeshComponent);
+		PreviewActor->SkeletalMeshComponent->SetAnimation(NULL);
 		PreviewActor->SkeletalMeshComponent->SetSkeletalMesh(nullptr);
 	}
 
@@ -690,6 +691,7 @@ bool FBlendSpaceThumbnailScene::SetBlendSpace(class UBlendSpaceBase* InBlendSpac
 	if (!bSetSucessfully)
 	{
 		CleanupComponentChildren(PreviewActor->SkeletalMeshComponent);
+		PreviewActor->SkeletalMeshComponent->SetAnimation(NULL);
 		PreviewActor->SkeletalMeshComponent->SetSkeletalMesh(nullptr);
 	}
 
