Adding search capability to key selection

Added new SKeySelector widget to take care of selecting keys in various situations
Used new widget for Key detail customization and Blueprint Key Pins

TTP#335356 - EDITOR: Input axis dropdown menu not filterable

diff --git a/Engine/Source/Editor/DetailCustomizations/Private/KeyStructCustomization.cpp b/Engine/Source/Editor/DetailCustomizations/Private/KeyStructCustomization.cpp
--- a/Engine/Source/Editor/DetailCustomizations/Private/KeyStructCustomization.cpp
+++ b/Engine/Source/Editor/DetailCustomizations/Private/KeyStructCustomization.cpp
@@ -3,6 +3,7 @@
 #include "DetailCustomizationsPrivatePCH.h"
 #include "KeyStructCustomization.h"
 #include "ScopedTransaction.h"
+#include "SKeySelector.h"
 
 #define LOCTEXT_NAMESPACE "FKeyStructCustomization"
 
@@ -26,32 +27,18 @@ void FKeyStructCustomization::CustomizeHeader( TSharedRef<class IPropertyHandle>
 	TArray<void*> StructPtrs;
 	StructPropertyHandle->AccessRawData(StructPtrs);
 	check(StructPtrs.Num() != 0);
-	const FKey& SelectedKey = *(FKey*)StructPtrs[0];
+	SelectedKey = (FKey*)StructPtrs[0];
 
 	bool bMultipleValues = false;
 	for (int32 StructPtrIndex = 1; StructPtrIndex < StructPtrs.Num(); ++StructPtrIndex)
 	{
-		if (*(FKey*)StructPtrs[StructPtrIndex] != SelectedKey)
+		if (*(FKey*)StructPtrs[StructPtrIndex] != *SelectedKey)
 		{
 			bMultipleValues = true;
 			break;
 		}
 	}
 
-	TSharedPtr<FKey> InitialSelectedItem;
-	TArray<FKey> AllKeys;
-	EKeys::GetAllKeys(AllKeys);
-
-	for (FKey Key : AllKeys)
-	{
-		TSharedPtr<FKey> InputKey = MakeShareable(new FKey(Key));
-		InputKeys.Add(InputKey);
-		if (Key == SelectedKey)
-		{
-			InitialSelectedItem = InputKey;
-		}
-	}
-	
 	// create struct header
 	HeaderRow.NameContent()
 	[
@@ -61,39 +48,22 @@ void FKeyStructCustomization::CustomizeHeader( TSharedRef<class IPropertyHandle>
 	.MinDesiredWidth(125.0f)
 	.MaxDesiredWidth(325.0f)
 	[
-		SNew(SHorizontalBox)
-
-		+ SHorizontalBox::Slot()
-			.FillWidth(1.0f)
-			[
-				// text box
-				SNew(SComboBox< TSharedPtr<FKey> >)
-				.OptionsSource(&InputKeys)
-				.InitiallySelectedItem(InitialSelectedItem)
-				.OnGenerateWidget(this, &FKeyStructCustomization::OnGenerateComboWidget)			
-				.OnSelectionChanged(this, &FKeyStructCustomization::OnSelectionChanged)
-				.Content()
-				[
-					SAssignNew(TextBlock, STextBlock)
-					.Text(bMultipleValues ? LOCTEXT("MultipleValues", "Multiple Values") : SelectedKey.GetDisplayName())
-					.Font(StructCustomizationUtils.GetRegularFont())
-				]
-			]
+		SNew(SKeySelector)
+		.CurrentKey(this, &FKeyStructCustomization::GetCurrentKey)
+		.OnKeyChanged(this, &FKeyStructCustomization::OnKeyChanged)
+		.Font(StructCustomizationUtils.GetRegularFont())
+		.HasMultipleValues(bMultipleValues)
 	];
 }
 
-
-TSharedRef<SWidget> FKeyStructCustomization::OnGenerateComboWidget(TSharedPtr<FKey> Key)
+FKey FKeyStructCustomization::GetCurrentKey() const
 {
-	return
-		SNew(STextBlock)
-		.Text(Key->GetDisplayName());
+	return *SelectedKey;
 }
 
-void FKeyStructCustomization::OnSelectionChanged(TSharedPtr<FKey> SelectedItem, ESelectInfo::Type SelectInfo)
+void FKeyStructCustomization::OnKeyChanged(TSharedPtr<FKey> SelectedKey)
 {
-	PropertyHandle->SetValueFromFormattedString(SelectedItem->ToString());
-	TextBlock->SetText(SelectedItem->GetDisplayName());
+	PropertyHandle->SetValueFromFormattedString(SelectedKey->ToString());
 }
 
 #undef LOCTEXT_NAMESPACE
diff --git a/Engine/Source/Editor/DetailCustomizations/Private/KeyStructCustomization.h b/Engine/Source/Editor/DetailCustomizations/Private/KeyStructCustomization.h
--- a/Engine/Source/Editor/DetailCustomizations/Private/KeyStructCustomization.h
+++ b/Engine/Source/Editor/DetailCustomizations/Private/KeyStructCustomization.h
@@ -26,12 +26,15 @@ public:
 
 private:
 
-	TSharedRef<SWidget> OnGenerateComboWidget(TSharedPtr<FKey> Key);
-	void OnSelectionChanged(TSharedPtr<FKey> SelectedItem, ESelectInfo::Type SelectInfo);
+	/** Gets the current Key being edited. */
+	FKey GetCurrentKey() const;
 
-	// Holds a handle to the property being edited.
+	/** Updates the property when a new key is selected. */
+	void OnKeyChanged(TSharedPtr<FKey> SelectedKey);
+
+	/** Holds a handle to the property being edited. */
 	TSharedPtr<IPropertyHandle> PropertyHandle;
-	TSharedPtr<STextBlock> TextBlock;
 
-	TArray<TSharedPtr<FKey>> InputKeys;
+	/** Pointer to the first key being edited. */
+	FKey* SelectedKey;
 };
diff --git a/Engine/Source/Editor/GraphEditor/Private/KismetPins/SGraphPinKey.cpp b/Engine/Source/Editor/GraphEditor/Private/KismetPins/SGraphPinKey.cpp
--- a/Engine/Source/Editor/GraphEditor/Private/KismetPins/SGraphPinKey.cpp
+++ b/Engine/Source/Editor/GraphEditor/Private/KismetPins/SGraphPinKey.cpp
@@ -5,28 +5,20 @@
 #include "SGraphPinComboBox.h"
 #include "SGraphPinEnum.h"
 #include "SGraphPinKey.h"
+#include "SKeySelector.h"
 
 void SGraphPinKey::Construct(const FArguments& InArgs, UEdGraphPin* InGraphPinObj)
 {
+	TArray<FKey> KeyList;
 	EKeys::GetAllKeys(KeyList);
 
-	const FKey SelectedKey(*InGraphPinObj->GetDefaultAsString());
+	SelectedKey = FKey(*InGraphPinObj->GetDefaultAsString());
 
-	if (SelectedKey.IsValid())
+	if (!SelectedKey.IsValid())
 	{
-		for (int32 KeyIndex = 0; KeyIndex < KeyList.Num(); ++KeyIndex)
-		{
-			if (KeyList[KeyIndex] == SelectedKey)
-			{
-				SelectedIndex = KeyIndex;
-				break;
-			}
-		}
-	}
-	else
-	{
-		SelectedIndex = 0;
-		InGraphPinObj->GetSchema()->TrySetDefaultValue(*InGraphPinObj, KeyList[SelectedIndex].ToString());
+		// Ensure first valid key is always set by default
+		SelectedKey = KeyList[0];
+		InGraphPinObj->GetSchema()->TrySetDefaultValue(*InGraphPinObj, SelectedKey.ToString());
 	}
 
 	SGraphPin::Construct(SGraphPin::FArguments(), InGraphPinObj);
@@ -34,42 +26,18 @@ void SGraphPinKey::Construct(const FArguments& InArgs, UEdGraphPin* InGraphPinOb
 
 TSharedRef<SWidget>	SGraphPinKey::GetDefaultValueWidget()
 {
-	//Get list of enum indexes
-	TArray< TSharedPtr<int32> > ComboItems;
-	GenerateComboBoxIndexes( ComboItems );
-
-	//Create widget
-	return SNew(SPinComboBox)
-		.ComboItemList( ComboItems )
-		.VisibleText( this, &SGraphPinKey::OnGetText )
-		.OnSelectionChanged( this, &SGraphPinKey::ComboBoxSelectionChanged )
-		.Visibility( this, &SGraphPin::GetDefaultValueVisibility )
-		.OnGetDisplayName(this, &SGraphPinKey::OnGetFriendlyName);
-}
-
-FString SGraphPinKey::OnGetFriendlyName(int32 EnumIndex)
-{
-	check(EnumIndex < KeyList.Num());
-	return KeyList[EnumIndex].GetDisplayName().ToString();
+	return SNew(SKeySelector)
+			.CurrentKey(this, &SGraphPinKey::GetCurrentKey)
+			.OnKeyChanged(this, &SGraphPinKey::OnKeyChanged);
 }
 
-void SGraphPinKey::ComboBoxSelectionChanged( TSharedPtr<int32> NewSelection, ESelectInfo::Type /*SelectInfo*/ )
+FKey SGraphPinKey::GetCurrentKey() const
 {
-	check(*NewSelection < KeyList.Num());
-	SelectedIndex = *NewSelection;
-	GraphPinObj->GetSchema()->TrySetDefaultValue(*GraphPinObj, KeyList[SelectedIndex].ToString());
+	return SelectedKey;
 }
 
-FString SGraphPinKey::OnGetText() const 
+void SGraphPinKey::OnKeyChanged(TSharedPtr<FKey> InSelectedKey)
 {
-	return KeyList[SelectedIndex].GetDisplayName().ToString();
-}
-
-void SGraphPinKey::GenerateComboBoxIndexes( TArray< TSharedPtr<int32> >& OutComboBoxIndexes )
-{
-	for (int32 KeyIndex = 0; KeyIndex < KeyList.Num(); ++KeyIndex)
-	{
-		TSharedPtr<int32> EnumIdxPtr(new int32(KeyIndex));
-		OutComboBoxIndexes.Add(EnumIdxPtr);
-	}
+	SelectedKey = *InSelectedKey.Get();
+	GraphPinObj->GetSchema()->TrySetDefaultValue(*GraphPinObj, SelectedKey.ToString());
 }
\ No newline at end of file
diff --git a/Engine/Source/Editor/GraphEditor/Private/KismetPins/SGraphPinKey.h b/Engine/Source/Editor/GraphEditor/Private/KismetPins/SGraphPinKey.h
--- a/Engine/Source/Editor/GraphEditor/Private/KismetPins/SGraphPinKey.h
+++ b/Engine/Source/Editor/GraphEditor/Private/KismetPins/SGraphPinKey.h
@@ -21,35 +21,11 @@ protected:
 
 private:
 
-	/**
-	 *	Function to get current string associated with the combo box selection
-	 *
-	 *	@return currently selected string
-	 */
-	FString OnGetText() const;
-
-	/**
-	 *	Function to generate the list of indexes from the enum object
-	 *
-	 *	@param OutComboBoxIndexes - Int array reference to store the list of indexes
-	 */
-	void GenerateComboBoxIndexes( TArray< TSharedPtr<int32> >& OutComboBoxIndexes );
+	/** Gets the current Key being edited. */
+	FKey GetCurrentKey() const;
 
-	/**
-	 *	Function to set the newly selected index
-	 *
-	 * @param NewSelection The newly selected item in the combo box
-	 * @param SelectInfo Provides context on how the selection changed
-	 */
-	void ComboBoxSelectionChanged( TSharedPtr<int32> NewSelection, ESelectInfo::Type SelectInfo );
-
-	/**
-	 * Returns the friendly name of the enum at index EnumIndex
-	 *
-	 * @param EnumIndex	- The index of the enum to return the friendly name for
-	 */
-	FString OnGetFriendlyName(int32 EnumIndex);
+	/** Updates the pin default when a new key is selected. */
+	void OnKeyChanged(TSharedPtr<FKey> SelectedKey);
 
-	TArray<FKey> KeyList;
-	int32 SelectedIndex;
+	FKey SelectedKey;
 };
diff --git a/Engine/Source/Editor/UnrealEd/Private/SKeySelector.cpp b/Engine/Source/Editor/UnrealEd/Private/SKeySelector.cpp
new file mode 100644
--- /dev/null
+++ b/Engine/Source/Editor/UnrealEd/Private/SKeySelector.cpp
@@ -0,0 +1,421 @@
+// Copyright 1998-2014 Epic Games, Inc. All Rights Reserved.
+
+
+#include "UnrealEd.h"
+#include "SKeySelector.h"
+#include "ScopedTransaction.h"
+#include "IDocumentation.h"
+#include "SListViewSelectorDropdownMenu.h"
+
+#define LOCTEXT_NAMESPACE "KeySelector"
+
+static const FString BigTooltipDocLink = TEXT("Shared/Editor/ProjectSettings");
+
+class FKeyTreeInfo
+{
+public:
+	/** This data item's children */
+	TArray< TSharedPtr<FKeyTreeInfo> > Children;
+
+private:
+	/** This data item's name */
+	FText Name;
+
+	/** The actual key associated with this item */
+	TSharedPtr<FKey> Key;
+
+public:
+	FKeyTreeInfo(FText InName, TSharedPtr<FKey> InKey)
+	: Name(InName)
+	, Key(InKey)
+	{
+	}
+
+	FKeyTreeInfo(TSharedPtr<FKeyTreeInfo> InInfo)
+	: Name(InInfo->Name)
+	, Key(InInfo->Key)
+	{
+	}
+
+	FString GetDescription() const
+	{
+		if (Key.IsValid())
+		{
+			return Key->GetDisplayName().ToString();
+		}
+		else
+		{
+			return Name.ToString();
+		}
+	}
+
+	TSharedPtr<FKey> GetKey() const
+	{
+		return Key;
+	}
+
+	bool MatchesSearchTokens(const TArray<FString>& SearchTokens)
+	{
+		FString Description = GetDescription();
+
+		for (auto Token : SearchTokens)
+		{
+			if (!Description.Contains(Token))
+			{
+				return false;
+			}
+		}
+
+		return true;
+	}
+};
+
+void SKeySelector::Construct(const FArguments& InArgs)
+{
+	SearchText = FText::GetEmpty();
+
+	OnKeyChanged = InArgs._OnKeyChanged;
+	CurrentKey = InArgs._CurrentKey;
+
+	auto PadCategory = *new (KeyTreeRoot) FKeyTreeItem(MakeShareable(new FKeyTreeInfo(LOCTEXT("PadCategory", "Gamepad"), nullptr)));
+	auto KeyCategory = *new (KeyTreeRoot) FKeyTreeItem(MakeShareable(new FKeyTreeInfo(LOCTEXT("KeyCategory", "Keyboard"), nullptr)));
+	auto MouseCategory = *new (KeyTreeRoot) FKeyTreeItem(MakeShareable(new FKeyTreeInfo(LOCTEXT("MouseCategory", "Mouse"), nullptr)));
+
+	TArray<FKey> AllKeys;
+	EKeys::GetAllKeys(AllKeys);
+
+	for (FKey Key : AllKeys)
+	{
+		if (!InArgs._FilterBlueprintBindable || Key.IsBindableInBlueprints())
+		{
+			auto NewKeyItem = MakeShareable(new FKeyTreeInfo(FText(), MakeShareable(new FKey(Key))));
+			if (Key.IsMouseButton())
+			{
+				MouseCategory->Children.Add(NewKeyItem);
+			}
+			else if (Key.IsGamepadKey())
+			{
+				PadCategory->Children.Add(NewKeyItem);
+			}
+			else
+			{
+				KeyCategory->Children.Add(NewKeyItem);
+			}
+		}
+	}
+
+	TreeViewWidth = InArgs._TreeViewWidth;
+	TreeViewHeight = InArgs._TreeViewHeight;
+	CategoryFont = FSlateFontInfo( FPaths::EngineContentDir() / TEXT("Slate/Fonts/Roboto-Bold.ttf"), 9 );
+	KeyFont = FSlateFontInfo( FPaths::EngineContentDir() / TEXT("Slate/Fonts/Roboto-Regular.ttf"), 10);
+	bHasMultipleValues = InArgs._HasMultipleValues;
+
+	FilteredKeyTreeRoot = KeyTreeRoot;
+
+	this->ChildSlot
+	[
+		SAssignNew(KeyComboButton, SComboButton)
+		.OnGetMenuContent(this, &SKeySelector::GetMenuContent)
+		.ContentPadding(0)
+		.ToolTipText(LOCTEXT("KeySelector", "Select the key value."))
+		.ButtonContent()
+		[
+			SNew(SHorizontalBox)
+			+ SHorizontalBox::Slot()
+			.AutoWidth()
+			.VAlign(VAlign_Center)
+			.HAlign(HAlign_Center)
+			[
+				SNew(SImage)
+				.Image(this, &SKeySelector::GetKeyIconImage)
+			]
+			+ SHorizontalBox::Slot()
+			.AutoWidth()
+			.VAlign(VAlign_Center)
+			.HAlign(HAlign_Left)
+			[
+				SNew(STextBlock)
+				.Text(this, &SKeySelector::GetKeyDescription)
+				.Font(InArgs._Font)
+			]
+		]
+	];
+}
+
+//=======================================================================
+// Attribute Helpers
+
+FString SKeySelector::GetKeyDescription() const
+{
+	if (bHasMultipleValues)
+	{
+		return LOCTEXT("MultipleValues", "Multiple Values").ToString();
+	}
+	return CurrentKey.Get().GetDisplayName().ToString();
+}
+
+const FSlateBrush* SKeySelector::GetKeyIconImage() const
+{
+	if (bHasMultipleValues)
+	{
+		return nullptr;
+	}
+	return GetIconFromKey(CurrentKey.Get());
+}
+
+//=======================================================================
+// Key TreeView Support
+TSharedRef<ITableRow> SKeySelector::GenerateKeyTreeRow(FKeyTreeItem InItem, const TSharedRef<STableViewBase>& OwnerTree)
+{
+	const bool bIsCategory = !InItem->GetKey().IsValid();
+	const FString Description = InItem->GetDescription();
+
+	// Determine the best icon the to represents this item
+	const FSlateBrush* IconBrush = nullptr;
+	if (InItem->GetKey().IsValid())
+	{
+		IconBrush = GetIconFromKey(*InItem->GetKey().Get());
+	}
+
+	return SNew(SComboRow<FKeyTreeItem>, OwnerTree)
+		.ToolTip(IDocumentation::Get()->CreateToolTip(FText::FromString(Description), NULL, *BigTooltipDocLink, *Description))
+		[
+			SNew(SHorizontalBox)
+			+ SHorizontalBox::Slot()
+			.AutoWidth()
+			.Padding(1.f)
+			[
+				SNew(SImage)
+				.Image(IconBrush)
+				.Visibility(bIsCategory ? EVisibility::Collapsed : EVisibility::Visible)
+			]
+			+ SHorizontalBox::Slot()
+			.AutoWidth()
+			.Padding(1.f)
+			[
+				SNew(STextBlock)
+				.Text(Description)
+				.HighlightText(SearchText)
+				.Font(bIsCategory ? CategoryFont : KeyFont)
+			]
+		];
+}
+
+void SKeySelector::OnKeySelectionChanged(FKeyTreeItem Selection, ESelectInfo::Type SelectInfo)
+{
+	// When the user is navigating, do not act upon the selection change
+	if (SelectInfo == ESelectInfo::OnNavigation)
+	{
+		return;
+	}
+
+	// Only handle selection for non-read only items, since STreeViewItem doesn't actually support read-only
+	if (Selection.IsValid())
+	{
+		if (Selection->GetKey().IsValid())
+		{
+			const FScopedTransaction Transaction(LOCTEXT("ChangeKey", "Change Key Value"));
+
+			KeyComboButton->SetIsOpen(false);
+
+			OnKeyChanged.ExecuteIfBound(Selection->GetKey());
+			bHasMultipleValues = false;
+		}
+		else
+		{
+			// Expand / contract the category, if applicable
+			if (Selection->Children.Num() > 0)
+			{
+				const bool bIsExpanded = KeyTreeView->IsItemExpanded(Selection);
+				KeyTreeView->SetItemExpansion(Selection, !bIsExpanded);
+
+				if (SelectInfo == ESelectInfo::OnMouseClick)
+				{
+					KeyTreeView->ClearSelection();
+				}
+			}
+		}
+	}
+}
+
+void SKeySelector::GetKeyChildren(FKeyTreeItem InItem, TArray<FKeyTreeItem>& OutChildren)
+{
+	OutChildren = InItem->Children;
+}
+
+TSharedRef<SWidget>	SKeySelector::GetMenuContent()
+{
+	if (!MenuContent.IsValid())
+	{
+		// Pre-build the tree view and search box as it is needed as a parameter for the context menu's container.
+		SAssignNew(KeyTreeView, SKeyTreeView)
+			.TreeItemsSource(&FilteredKeyTreeRoot)
+			.SelectionMode(ESelectionMode::Single)
+			.OnGenerateRow(this, &SKeySelector::GenerateKeyTreeRow)
+			.OnSelectionChanged(this, &SKeySelector::OnKeySelectionChanged)
+			.OnGetChildren(this, &SKeySelector::GetKeyChildren);
+
+		SAssignNew(FilterTextBox, SSearchBox)
+			.OnTextChanged(this, &SKeySelector::OnFilterTextChanged)
+			.OnTextCommitted(this, &SKeySelector::OnFilterTextCommitted);
+
+		MenuContent = SNew(SListViewSelectorDropdownMenu<FKeyTreeItem>, FilterTextBox, KeyTreeView)
+			[
+				SNew(SVerticalBox)
+				+ SVerticalBox::Slot()
+				.AutoHeight()
+				.Padding(4.f, 4.f, 4.f, 4.f)
+				[
+					FilterTextBox.ToSharedRef()
+				]
+				+ SVerticalBox::Slot()
+				.AutoHeight()
+				.Padding(4.f, 4.f, 4.f, 4.f)
+				[
+					SNew(SBox)
+					.HeightOverride(TreeViewHeight)
+					.WidthOverride(TreeViewWidth)
+					[
+						KeyTreeView.ToSharedRef()
+					]
+				]
+			];
+
+
+		KeyComboButton->SetMenuContentWidgetToFocus(FilterTextBox);
+	}
+	else
+	{
+		// Clear the selection in such a way as to also clear the keyboard selector
+		KeyTreeView->SetSelection(NULL, ESelectInfo::OnNavigation);
+		KeyTreeView->ClearExpandedItems();
+	}
+
+	// Clear the filter text box with each opening
+	if (FilterTextBox.IsValid())
+	{
+		FilterTextBox->SetText(FText::GetEmpty());
+	}
+
+	return MenuContent.ToSharedRef();
+}
+
+//=======================================================================
+// Search Support
+void SKeySelector::OnFilterTextChanged(const FText& NewText)
+{
+	SearchText = NewText;
+	FilteredKeyTreeRoot.Empty();
+
+	TArray<FString> Tokens;
+	GetSearchTokens(SearchText.ToString(), Tokens);
+
+	GetChildrenMatchingSearch(Tokens, KeyTreeRoot, FilteredKeyTreeRoot);
+	KeyTreeView->RequestTreeRefresh();
+
+	// Select the first non-category item
+	auto SelectedItems = KeyTreeView->GetSelectedItems();
+	if (FilteredKeyTreeRoot.Num() > 0)
+	{
+		// Categories have children, we don't want to select categories
+		if (FilteredKeyTreeRoot[0]->Children.Num() > 0)
+		{
+			KeyTreeView->SetSelection(FilteredKeyTreeRoot[0]->Children[0], ESelectInfo::OnNavigation);
+		}
+		else
+		{
+			KeyTreeView->SetSelection(FilteredKeyTreeRoot[0], ESelectInfo::OnNavigation);
+		}
+	}
+}
+
+void SKeySelector::OnFilterTextCommitted(const FText& NewText, ETextCommit::Type CommitInfo)
+{
+	if (CommitInfo == ETextCommit::OnEnter)
+	{
+		auto SelectedItems = KeyTreeView->GetSelectedItems();
+		if (SelectedItems.Num() > 0)
+		{
+			KeyTreeView->SetSelection(SelectedItems[0]);
+		}
+	}
+}
+
+void SKeySelector::GetSearchTokens(const FString& SearchString, TArray<FString>& OutTokens) const
+{
+	if (SearchString.Contains("\"") && SearchString.ParseIntoArray(&OutTokens, TEXT("\""), true) > 0)
+	{
+		for (auto &TokenIt : OutTokens)
+		{
+			// we have the token, we don't need the quotes anymore, they'll just confused the comparison later on
+			TokenIt = TokenIt.TrimQuotes();
+			// We remove the spaces as all later comparison strings will also be de-spaced
+			TokenIt = TokenIt.Replace(TEXT(" "), TEXT(""));
+		}
+
+		// due to being able to handle multiple quoted blocks like ("Make Epic" "Game Now") we can end up with
+		// and empty string between (" ") blocks so this simply removes them
+		struct FRemoveMatchingStrings
+		{
+			bool operator()(const FString& RemovalCandidate) const
+			{
+				return RemovalCandidate.IsEmpty();
+			}
+		};
+		OutTokens.RemoveAll(FRemoveMatchingStrings());
+	}
+	else
+	{
+		// unquoted search equivalent to a match-any-of search
+		SearchString.ParseIntoArray(&OutTokens, TEXT(" "), true);
+	}
+}
+
+bool SKeySelector::GetChildrenMatchingSearch(const TArray<FString>& InSearchTokens, const TArray<FKeyTreeItem>& UnfilteredList, TArray<FKeyTreeItem>& OutFilteredList)
+{
+	bool bReturnVal = false;
+
+	for (auto it = UnfilteredList.CreateConstIterator(); it; ++it)
+	{
+		FKeyTreeItem Item = *it;
+		FKeyTreeItem NewInfo = MakeShareable(new FKeyTreeInfo(Item));
+		TArray<FKeyTreeItem> ValidChildren;
+
+		// Have to run GetChildrenMatchingSearch first, so that we can make sure we get valid children for the list!
+		if (GetChildrenMatchingSearch(InSearchTokens, Item->Children, ValidChildren)
+			|| InSearchTokens.Num() == 0
+			|| Item->MatchesSearchTokens(InSearchTokens))
+		{
+			NewInfo->Children = ValidChildren;
+			OutFilteredList.Add(NewInfo);
+
+			KeyTreeView->SetItemExpansion(NewInfo, InSearchTokens.Num() > 0);
+
+			bReturnVal = true;
+		}
+	}
+
+	return bReturnVal;
+}
+
+const FSlateBrush* SKeySelector::GetIconFromKey(FKey Key) const
+{
+	static const FName NAME_MouseEvent(TEXT("GraphEditor.MouseEvent_16x"));
+	static const FName NAME_PadEvent(TEXT("GraphEditor.PadEvent_16x"));
+	static const FName NAME_KeyEvent(TEXT("GraphEditor.KeyEvent_16x"));
+
+	if (Key.IsMouseButton())
+	{
+		return FEditorStyle::GetBrush(NAME_MouseEvent);
+	}
+	else if (Key.IsGamepadKey())
+	{
+		return FEditorStyle::GetBrush(NAME_PadEvent);
+	}
+	else
+	{
+		return FEditorStyle::GetBrush(NAME_KeyEvent);
+	}
+}
+
+#undef LOCTEXT_NAMESPACE
diff --git a/Engine/Source/Editor/UnrealEd/Public/SKeySelector.h b/Engine/Source/Editor/UnrealEd/Public/SKeySelector.h
new file mode 100644
--- /dev/null
+++ b/Engine/Source/Editor/UnrealEd/Public/SKeySelector.h
@@ -0,0 +1,101 @@
+// Copyright 1998-2014 Epic Games, Inc. All Rights Reserved.
+
+#pragma once
+
+struct FKey;
+class FKeyTreeInfo;
+
+DECLARE_DELEGATE_OneParam(FOnKeyChanged, TSharedPtr<FKey>)
+
+//////////////////////////////////////////////////////////////////////////
+// SKeySelector
+
+typedef TSharedPtr<FKeyTreeInfo> FKeyTreeItem;
+typedef STreeView<FKeyTreeItem> SKeyTreeView;
+
+/** Widget for selecting an input key */
+class UNREALED_API SKeySelector : public SCompoundWidget
+{
+public:
+	SLATE_BEGIN_ARGS( SKeySelector )
+		: _CurrentKey()
+		, _TreeViewWidth(300.f)
+		, _TreeViewHeight(400.f)
+		, _Font( FEditorStyle::GetFontStyle( TEXT("NormalFont") ) )
+		, _FilterBlueprintBindable( false )
+		, _HasMultipleValues( false )
+		{}
+		SLATE_ATTRIBUTE( FKey, CurrentKey )
+		SLATE_ATTRIBUTE( FOptionalSize, TreeViewWidth )
+		SLATE_ATTRIBUTE( FOptionalSize, TreeViewHeight )
+		SLATE_EVENT( FOnKeyChanged, OnKeyChanged )
+		SLATE_ATTRIBUTE( FSlateFontInfo, Font )
+		SLATE_ARGUMENT( bool, FilterBlueprintBindable )
+		SLATE_ARGUMENT( bool, HasMultipleValues )
+	SLATE_END_ARGS()
+public:
+	void Construct(const FArguments& InArgs);
+
+protected:
+	/** Gets the icon for the key being manipulated */
+	const FSlateBrush* GetKeyIconImage() const;
+	/** Gets a succinct description for the key being manipulated */
+	FString GetKeyDescription() const;
+
+	/** Treeview support functions */
+	virtual TSharedRef<ITableRow> GenerateKeyTreeRow(FKeyTreeItem InItem, const TSharedRef<STableViewBase>& OwnerTree);
+	void OnKeySelectionChanged(FKeyTreeItem Selection, ESelectInfo::Type SelectInfo);
+	void GetKeyChildren(FKeyTreeItem InItem, TArray<FKeyTreeItem>& OutChildren);
+
+	/** Gets the Menu Content, setting it up if necessary */
+	virtual TSharedRef<SWidget>	GetMenuContent();
+
+	/** Key searching support */
+	void OnFilterTextChanged(const FText& NewText);
+	void OnFilterTextCommitted(const FText& NewText, ETextCommit::Type CommitInfo);
+	void GetSearchTokens(const FString& SearchString, TArray<FString>& OutTokens) const;
+
+	/** Helper to generate the filtered list of keys, based on the search string matching */
+	bool GetChildrenMatchingSearch(const TArray<FString>& SearchTokens, const TArray<FKeyTreeItem>& UnfilteredList, TArray<FKeyTreeItem>& OutFilteredList);
+
+	/** 
+	 * Determine the best icon to represent the given key.
+	 *
+	 * @param Key		The key to get the icon for.
+	 * @param returns a brush that best represents the icon
+	 */
+	const FSlateBrush* GetIconFromKey(FKey Key) const;
+
+protected:
+	/** Combo Button that shows current key and icon */
+	TSharedPtr<SComboButton>	KeyComboButton;
+
+	/** Reference to the menu content that's displayed when the key button is clicked on */
+	TSharedPtr<SWidget>			MenuContent;
+	TSharedPtr<SSearchBox>		FilterTextBox;
+	TSharedPtr<SKeyTreeView>	KeyTreeView;
+	FText						SearchText;
+
+	/** The key attribute that we're modifying with this widget */
+	TAttribute<FKey>			CurrentKey;
+	/** Whether the selector is currently modifying multiple different values */
+	bool						bHasMultipleValues;
+
+	/** Delegate that is called every time the key changes. */
+	FOnKeyChanged				OnKeyChanged;
+
+	/** Desired width of the tree view widget */
+	TAttribute<FOptionalSize>	TreeViewWidth;
+	/** Desired height of the tree view widget */
+	TAttribute<FOptionalSize>	TreeViewHeight;
+
+	/** Font used for category tree entries */
+	FSlateFontInfo				CategoryFont;
+	/** Font used for key tree entries */
+	FSlateFontInfo				KeyFont;
+
+	/** Array containing the unfiltered list of all values this key could possibly have */
+	TArray<FKeyTreeItem>		KeyTreeRoot;
+	/** Array containing a filtered list, according to the text in the searchbox */
+	TArray<FKeyTreeItem>		FilteredKeyTreeRoot;
+};
