Changing the order in which we fully regenerate blueprint classes (for circular dependencies)... Now passing the original ExportObject along with the class that needs regeneration so we can run through the whole process the first time we get into FBlueprintEditorUtils::RegenerateBlueprintClass.
(TTP #342058)

diff --git a/Engine/Source/Editor/UnrealEd/Private/Kismet2/BlueprintEditorUtils.cpp b/Engine/Source/Editor/UnrealEd/Private/Kismet2/BlueprintEditorUtils.cpp
--- a/Engine/Source/Editor/UnrealEd/Private/Kismet2/BlueprintEditorUtils.cpp
+++ b/Engine/Source/Editor/UnrealEd/Private/Kismet2/BlueprintEditorUtils.cpp
@@ -1018,8 +1018,6 @@ UClass* FBlueprintEditorUtils::RegenerateBlueprintClass(UBlueprint* Blueprint, U
 	UPackage* Package = Cast<UPackage>(Blueprint->GetOutermost());
 	bool bIsPackageDirty = Package ? Package->IsDirty() : false;
 
-	bool bNeedsSkelRefRemoval = false;
-
 	// Preload the blueprint and all its parts before refreshing nodes. 
 	// Otherwise, the nodes might not maintain their proper linkages... 
 	//
@@ -1036,13 +1034,11 @@ UClass* FBlueprintEditorUtils::RegenerateBlueprintClass(UBlueprint* Blueprint, U
 	// member further up the stack).
 	if (!Blueprint->bHasBeenRegenerated)
 	{
-		if (PreviousCDO)
-		{
-			FBlueprintEditorUtils::PreloadMembers(PreviousCDO);
-		}
-		FBlueprintEditorUtils::PreloadMembers(Blueprint);
+		check(PreviousCDO != nullptr);
+		FRegenerationHelper::ForcedLoadMembers(PreviousCDO);
+		FRegenerationHelper::ForcedLoadMembers(Blueprint);
 	}
-	
+
 	if( ShouldRegenerateBlueprint(Blueprint) && !Blueprint->bHasBeenRegenerated )
 	{
 		Blueprint->bIsRegeneratingOnLoad = true;
@@ -1173,111 +1169,102 @@ UClass* FBlueprintEditorUtils::RegenerateBlueprintClass(UBlueprint* Blueprint, U
 		}
 
 		bRegenerated = bHasCode;
-	}
 
-	if ( PreviousCDO != NULL )
-	{
-		if( Blueprint->bRecompileOnLoad && !GForceDisableBlueprintCompileOnLoad)
+		if (!FKismetEditorUtilities::IsClassABlueprintSkeleton(ClassToRegenerate))
 		{
-			if (Blueprint->BlueprintType != BPTYPE_MacroLibrary)
+			if (Blueprint->bRecompileOnLoad)
 			{
 				// Verify that we had a skeleton generated class if we had a previous CDO, to make sure we have something to copy into
-				check(Blueprint->SkeletonGeneratedClass);
-			}
+				check((Blueprint->BlueprintType == BPTYPE_MacroLibrary) || Blueprint->SkeletonGeneratedClass);
 
-			const bool bPreviousMatchesGenerated = (PreviousCDO == Blueprint->GeneratedClass->GetDefaultObject());
+				const bool bPreviousMatchesGenerated = (PreviousCDO == Blueprint->GeneratedClass->GetDefaultObject());
 
-			if (Blueprint->BlueprintType != BPTYPE_MacroLibrary)
-			{
-				UObject* CDOThatKickedOffCOL = PreviousCDO;
-				if (Blueprint->IsGeneratedClassAuthoritative() && !bPreviousMatchesGenerated && Blueprint->PRIVATE_InnermostPreviousCDO)
+				if (Blueprint->BlueprintType != BPTYPE_MacroLibrary)
 				{
-					PreviousCDO = Blueprint->PRIVATE_InnermostPreviousCDO;
+					UObject* CDOThatKickedOffCOL = PreviousCDO;
+					if (Blueprint->IsGeneratedClassAuthoritative() && !bPreviousMatchesGenerated && Blueprint->PRIVATE_InnermostPreviousCDO)
+					{
+						PreviousCDO = Blueprint->PRIVATE_InnermostPreviousCDO;
+					}
 				}
-			}
 
-			// If this is the top of the compile-on-load stack for this object, copy the old CDO properties to the newly created one unless they are the same
-			UClass* AuthoritativeClass = (Blueprint->IsGeneratedClassAuthoritative() ? Blueprint->GeneratedClass : Blueprint->SkeletonGeneratedClass);
-			if (AuthoritativeClass != NULL && PreviousCDO != AuthoritativeClass->GetDefaultObject())
-			{
-				TGuardValue<bool> GuardTemplateNameFlag(GCompilingBlueprint, true);
-
-				// Make sure the previous CDO has been fully loaded before we use it
-				if( PreviousCDO )
+				// If this is the top of the compile-on-load stack for this object, copy the old CDO properties to the newly created one unless they are the same
+				UClass* AuthoritativeClass = (Blueprint->IsGeneratedClassAuthoritative() ? Blueprint->GeneratedClass : Blueprint->SkeletonGeneratedClass);
+				if (AuthoritativeClass != NULL && PreviousCDO != AuthoritativeClass->GetDefaultObject())
 				{
-					FBlueprintEditorUtils::PreloadMembers(PreviousCDO);
-				}
+					TGuardValue<bool> GuardTemplateNameFlag(GCompilingBlueprint, true);
 
-				// Copy over the properties from the old CDO to the new
-				PropagateParentBlueprintDefaults(AuthoritativeClass);
-				UObject* NewCDO = AuthoritativeClass->GetDefaultObject();
-				{
-					FSaveActorFlagsHelper SaveActorFlags(AuthoritativeClass);
-					UEditorEngine::FCopyPropertiesForUnrelatedObjectsParams CopyDetails;
-					CopyDetails.bAggressiveDefaultSubobjectReplacement = true;
-					CopyDetails.bDoDelta = false;
-					UEditorEngine::CopyPropertiesForUnrelatedObjects(PreviousCDO, NewCDO, CopyDetails);
-				}
+					// Make sure the previous CDO has been fully loaded before we use it
+					FBlueprintEditorUtils::PreloadMembers(PreviousCDO);
 
-				if( bRegenerated )
-				{
-					// Collect the instanced components in both the old and new CDOs
-					TArray<UObject*> OldComponents, NewComponents;
-					PreviousCDO->CollectDefaultSubobjects(OldComponents,true);
-					NewCDO->CollectDefaultSubobjects(NewComponents,true);
+					// Copy over the properties from the old CDO to the new
+					PropagateParentBlueprintDefaults(AuthoritativeClass);
+					UObject* NewCDO = AuthoritativeClass->GetDefaultObject();
+					{
+						FSaveActorFlagsHelper SaveActorFlags(AuthoritativeClass);
+						UEditorEngine::FCopyPropertiesForUnrelatedObjectsParams CopyDetails;
+						CopyDetails.bAggressiveDefaultSubobjectReplacement = true;
+						CopyDetails.bDoDelta = false;
+						UEditorEngine::CopyPropertiesForUnrelatedObjects(PreviousCDO, NewCDO, CopyDetails);
+					}
 
-					// For all components common to both, patch the linker table with the new version of the component, so things that reference the default (e.g. InternalArchetypes) will have the updated version
-					for( auto OldCompIt = OldComponents.CreateIterator(); OldCompIt; ++OldCompIt )
+					if (bRegenerated)
 					{
-						UObject* OldComponent = (*OldCompIt);
-						const FName OldComponentName = OldComponent->GetFName();
-						for( auto NewCompIt = NewComponents.CreateIterator(); NewCompIt; ++NewCompIt )
+						// Collect the instanced components in both the old and new CDOs
+						TArray<UObject*> OldComponents, NewComponents;
+						PreviousCDO->CollectDefaultSubobjects(OldComponents, true);
+						NewCDO->CollectDefaultSubobjects(NewComponents, true);
+
+						// For all components common to both, patch the linker table with the new version of the component, so things that reference the default (e.g. InternalArchetypes) will have the updated version
+						for (auto OldCompIt = OldComponents.CreateIterator(); OldCompIt; ++OldCompIt)
 						{
-							UObject* NewComponent = *NewCompIt;
-							if( NewComponent->GetFName() == OldComponentName )
+							UObject* OldComponent = (*OldCompIt);
+							const FName OldComponentName = OldComponent->GetFName();
+							for (auto NewCompIt = NewComponents.CreateIterator(); NewCompIt; ++NewCompIt)
 							{
-								ULinkerLoad::PRIVATE_PatchNewObjectIntoExport(OldComponent, NewComponent);
-								break;
+								UObject* NewComponent = *NewCompIt;
+								if (NewComponent->GetFName() == OldComponentName)
+								{
+									ULinkerLoad::PRIVATE_PatchNewObjectIntoExport(OldComponent, NewComponent);
+									break;
+								}
 							}
 						}
+						NewCDO->CheckDefaultSubobjects();
+						// We purposefully do not call post load here, it happens later on in the normal flow
 					}
-					NewCDO->CheckDefaultSubobjects();
-					// We purposefully do not call post load here, it happens later on in the normal flow
 				}
-			}
 
-			Blueprint->PRIVATE_InnermostPreviousCDO = NULL;
-		}
-		else
-		{
-			// If we didn't recompile, we still need to propagate flags, and instance components
-			FKismetEditorUtilities::ConformBlueprintFlagsAndComponents(Blueprint);
-		}
+				Blueprint->PRIVATE_InnermostPreviousCDO = NULL;
+			}
+			else
+			{
+				// If we didn't recompile, we still need to propagate flags, and instance components
+				FKismetEditorUtilities::ConformBlueprintFlagsAndComponents(Blueprint);
+			}
 
-		// If this is the top of the compile-on-load stack for this object, copy the old CDO properties to the newly created one
-		if (!Blueprint->IsGeneratedClassAuthoritative() && Blueprint->GeneratedClass != NULL)
-		{
-			TGuardValue<bool> GuardTemplateNameFlag(GCompilingBlueprint, true);
+			// If this is the top of the compile-on-load stack for this object, copy the old CDO properties to the newly created one
+			if (!Blueprint->IsGeneratedClassAuthoritative() && Blueprint->GeneratedClass != NULL)
+			{
+				TGuardValue<bool> GuardTemplateNameFlag(GCompilingBlueprint, true);
 
-			UObject* SkeletonCDO = Blueprint->SkeletonGeneratedClass->GetDefaultObject();
-			UObject* GeneratedCDO = Blueprint->GeneratedClass->GetDefaultObject();
+				UObject* SkeletonCDO = Blueprint->SkeletonGeneratedClass->GetDefaultObject();
+				UObject* GeneratedCDO = Blueprint->GeneratedClass->GetDefaultObject();
 
-			UEditorEngine::FCopyPropertiesForUnrelatedObjectsParams CopyDetails;
-			CopyDetails.bAggressiveDefaultSubobjectReplacement = false;
-			CopyDetails.bDoDelta = false;
-			UEditorEngine::CopyPropertiesForUnrelatedObjects(SkeletonCDO, GeneratedCDO, CopyDetails);
+				UEditorEngine::FCopyPropertiesForUnrelatedObjectsParams CopyDetails;
+				CopyDetails.bAggressiveDefaultSubobjectReplacement = false;
+				CopyDetails.bDoDelta = false;
+				UEditorEngine::CopyPropertiesForUnrelatedObjects(SkeletonCDO, GeneratedCDO, CopyDetails);
 
-			Blueprint->SetLegacyGeneratedClassIsAuthoritative();
-		}
+				Blueprint->SetLegacyGeneratedClassIsAuthoritative();
+			}
 
-		// Now that the CDO is valid, update the OwnedComponents, in case we've added or removed native components
-		if (AActor* MyActor = Cast<AActor>(Blueprint->GeneratedClass->GetDefaultObject()))
-		{
-			MyActor->ResetOwnedComponents();
+			// Now that the CDO is valid, update the OwnedComponents, in case we've added or removed native components
+			if (AActor* MyActor = Cast<AActor>(Blueprint->GeneratedClass->GetDefaultObject()))
+			{
+				MyActor->ResetOwnedComponents();
+			}
 		}
-
-		Blueprint->ClearFlags(RF_BeingRegenerated);
-		bNeedsSkelRefRemoval = true;
 	}
 
 	if ( bRegenerated )
@@ -1287,7 +1274,8 @@ UClass* FBlueprintEditorUtils::RegenerateBlueprintClass(UBlueprint* Blueprint, U
 		GeneratedClassPackage->GetMetaData()->RemoveMetaDataOutsidePackage();
 	}
 
-	if ( bNeedsSkelRefRemoval && Blueprint->bLegacyNeedToPurgeSkelRefs)
+	bool const bNeedsSkelRefRemoval = !FKismetEditorUtilities::IsClassABlueprintSkeleton(ClassToRegenerate) && (Blueprint->SkeletonGeneratedClass != nullptr);
+	if (bNeedsSkelRefRemoval && Blueprint->bLegacyNeedToPurgeSkelRefs)
 	{
 		// Remove any references to the skeleton class, replacing them with refs to the generated class instead
 		FArchiveMoveSkeletalRefs SkelRefArchiver(Blueprint);
diff --git a/Engine/Source/Runtime/CoreUObject/Private/Blueprint/BlueprintSupport.cpp b/Engine/Source/Runtime/CoreUObject/Private/Blueprint/BlueprintSupport.cpp
--- a/Engine/Source/Runtime/CoreUObject/Private/Blueprint/BlueprintSupport.cpp
+++ b/Engine/Source/Runtime/CoreUObject/Private/Blueprint/BlueprintSupport.cpp
@@ -37,15 +37,19 @@ struct FPreloadMembersHelper
  */
 bool ULinkerLoad::RegenerateBlueprintClass(UClass* LoadClass, UObject* ExportObject)
 {
-	UObject* PreviousCDO = NULL;
-	if( !LoadClass->ClassGeneratedBy->HasAnyFlags(RF_BeingRegenerated) )
-	{
-		// If this is the first time we're attempting to recompile this class on load, save off the original CDO here, so we can copy its properties after the class has been recompiled
-		PreviousCDO = ExportObject;
+	// determine if somewhere further down the callstack, we're already in this
+	// function for this class
+	bool const bAlreadyRegenerating = LoadClass->ClassGeneratedBy->HasAnyFlags(RF_BeingRegenerated);
+	// Flag the class source object, so we know we're already in the process of compiling this class
+	LoadClass->ClassGeneratedBy->SetFlags(RF_BeingRegenerated);
 
-		// Flag the class source object, so we know we're already in the process of compiling this class
-		LoadClass->ClassGeneratedBy->SetFlags(RF_BeingRegenerated);
-	}
+	// Cache off the current CDO, and specify the CDO for the load class 
+	// manually... do this before we Preload() any children members so that if 
+	// one of those preloads subsequently ends up back here for this class, 
+	// then the ExportObject is carried along and used in the eventual RegenerateClass() call
+	UObject* CurrentCDO = ExportObject;
+	check(!bAlreadyRegenerating || (LoadClass->ClassDefaultObject == ExportObject));
+	LoadClass->ClassDefaultObject = CurrentCDO;
 
 	// Finish loading the class here, so we have all the appropriate data to copy over to the new CDO
 	TArray<UObject*> AllChildMembers;
@@ -55,63 +59,81 @@ bool ULinkerLoad::RegenerateBlueprintClass(UClass* LoadClass, UObject* ExportObj
 		UObject* Member = AllChildMembers[Index];
 		Preload(Member);
 	}
-	Preload(LoadClass);
 
-	LoadClass->StaticLink(true);
-	Preload(ExportObject);
+	// if this was subsequently regenerated from one of the above preloads, then 
+	// we don't have to finish this off, it was already done
+	bool const bWasSubsequentlyRegenerated = !LoadClass->ClassGeneratedBy->HasAnyFlags(RF_BeingRegenerated);
+	// @TODO: find some other condition to block this if we've already  
+	//        regenerated the class (not just if we've regenerated the class 
+	//        from an above Preload(Member))... UBlueprint::RegenerateClass() 
+	//        has an internal conditional to block getting into it again, but we
+	//        can't check UBlueprint members from this module
+	if (!bWasSubsequentlyRegenerated)
+	{
+		Preload(LoadClass);
 
+		LoadClass->StaticLink(true);
+		Preload(CurrentCDO);
 
-	// Cache off the current CDO, and specify the CDO for the load class manually
-	UObject* CurrentCDO = ExportObject;
-	LoadClass->ClassDefaultObject = CurrentCDO;
-
-	// Make sure that we regenerate any parent classes first before attempting to build a child
-	TArray<UClass*> ClassChainOrdered;
-	{
-		// Just ordering the class hierarchy from root to leafs:
-		UClass* ClassChain = LoadClass->GetSuperClass();
-		while (ClassChain && ClassChain->ClassGeneratedBy)
+		// Make sure that we regenerate any parent classes first before attempting to build a child
+		TArray<UClass*> ClassChainOrdered;
 		{
-			// O(n) insert, but n is tiny because this is a class hierarchy...
-			ClassChainOrdered.Insert(ClassChain, 0);
-			ClassChain = ClassChain->GetSuperClass();
+			// Just ordering the class hierarchy from root to leafs:
+			UClass* ClassChain = LoadClass->GetSuperClass();
+			while (ClassChain && ClassChain->ClassGeneratedBy)
+			{
+				// O(n) insert, but n is tiny because this is a class hierarchy...
+				ClassChainOrdered.Insert(ClassChain, 0);
+				ClassChain = ClassChain->GetSuperClass();
+			}
 		}
-	}
- 	for( auto Class : ClassChainOrdered )
-	{
-		UObject* BlueprintObject = Class->ClassGeneratedBy;
-		if (BlueprintObject && BlueprintObject->HasAnyFlags(RF_BeingRegenerated))
+		for (auto Class : ClassChainOrdered)
 		{
-			// Always load the parent blueprint here in case there is a circular dependency. This will
-			// ensure that the blueprint is fully serialized before attempting to regenerate the class.
-			if (!BlueprintObject->HasAnyFlags(RF_LoadCompleted))
+			UObject* BlueprintObject = Class->ClassGeneratedBy;
+			if (BlueprintObject && BlueprintObject->HasAnyFlags(RF_BeingRegenerated))
 			{
-				BlueprintObject->SetFlags(RF_NeedLoad);
-				if (auto Linker = BlueprintObject->GetLinker())
+				// Always load the parent blueprint here in case there is a circular dependency. This will
+				// ensure that the blueprint is fully serialized before attempting to regenerate the class.
+				if (!BlueprintObject->HasAnyFlags(RF_LoadCompleted))
 				{
-					Linker->Preload(BlueprintObject);
+					BlueprintObject->SetFlags(RF_NeedLoad);
+					if (ULinkerLoad* Linker = BlueprintObject->GetLinker())
+					{
+						Linker->Preload(BlueprintObject);
+					}
 				}
+			
+				FPreloadMembersHelper::PreloadMembers(BlueprintObject);
+				// recurse into this function for this parent class; 
+				// 'ClassDefaultObject' should be the class's original ExportObject
+				RegenerateBlueprintClass(Class, BlueprintObject);
 			}
+		}
 
-			FPreloadMembersHelper::PreloadMembers(BlueprintObject);
-			BlueprintObject->RegenerateClass(Class, NULL, GObjLoaded);
+		// Preload the blueprint to make sure it has all the data the class needs for regeneration
+		if (LoadClass->ClassGeneratedBy->HasAnyFlags(RF_NeedLoad))
+		{
+			LoadClass->ClassGeneratedBy->GetLinker()->Preload(LoadClass->ClassGeneratedBy);
 		}
-	}
 
-	// Preload the blueprint to make sure it has all the data the class needs for regeneration
-	if( LoadClass->ClassGeneratedBy->HasAnyFlags(RF_NeedLoad) )
-	{
-		LoadClass->ClassGeneratedBy->GetLinker()->Preload(LoadClass->ClassGeneratedBy);
+		UClass* RegeneratedClass = LoadClass->ClassGeneratedBy->RegenerateClass(LoadClass, CurrentCDO, GObjLoaded);
+		if (RegeneratedClass)
+		{
+			LoadClass->ClassGeneratedBy->ClearFlags(RF_BeingRegenerated);
+			// Fix up the linker so that the RegeneratedClass is used
+			LoadClass->ClearFlags(RF_NeedLoad | RF_NeedPostLoad);
+		}
 	}
 
-	UClass* RegeneratedClass = LoadClass->ClassGeneratedBy->RegenerateClass(LoadClass, PreviousCDO, GObjLoaded);
-	if( RegeneratedClass )
+	bool const bSuccessfulRegeneration = !LoadClass->ClassGeneratedBy->HasAnyFlags(RF_BeingRegenerated);
+	// if this wasn't already flagged as regenerating when we first entered this 
+	// function, the clear it ourselves.
+	if (!bAlreadyRegenerating)
 	{
-		// Fix up the linker so that the RegeneratedClass is used
- 		LoadClass->ClearFlags(RF_NeedLoad|RF_NeedPostLoad);
-		return true;
+		LoadClass->ClassGeneratedBy->ClearFlags(RF_BeingRegenerated);
 	}
-	return false;
+
+	return bSuccessfulRegeneration;
 }
 
 /** 
