Settings: Removed file header comments; documentation cleanup pass

diff --git a/Engine/Source/Developer/Settings/Public/Interfaces/ISettingsCategory.h b/Engine/Source/Developer/Settings/Public/Interfaces/ISettingsCategory.h
--- a/Engine/Source/Developer/Settings/Public/Interfaces/ISettingsCategory.h
+++ b/Engine/Source/Developer/Settings/Public/Interfaces/ISettingsCategory.h
@@ -67,8 +67,6 @@ public:
 
 public:
 
-	/**
-	 * Virtual destructor.
-	 */
+	/** Virtual destructor. */
 	virtual ~ISettingsCategory( ) { }
 };
diff --git a/Engine/Source/Developer/Settings/Public/Interfaces/ISettingsContainer.h b/Engine/Source/Developer/Settings/Public/Interfaces/ISettingsContainer.h
--- a/Engine/Source/Developer/Settings/Public/Interfaces/ISettingsContainer.h
+++ b/Engine/Source/Developer/Settings/Public/Interfaces/ISettingsContainer.h
@@ -116,8 +116,6 @@ public:
 
 public:
 
-	/**
-	 * Virtual destructor.
-	 */
+	/** Virtual destructor. */
 	virtual ~ISettingsContainer( ) { }
 };
diff --git a/Engine/Source/Developer/Settings/Public/Interfaces/ISettingsModule.h b/Engine/Source/Developer/Settings/Public/Interfaces/ISettingsModule.h
--- a/Engine/Source/Developer/Settings/Public/Interfaces/ISettingsModule.h
+++ b/Engine/Source/Developer/Settings/Public/Interfaces/ISettingsModule.h
@@ -134,7 +134,6 @@ public:
 	 * Gets a pointer to the Settings module instance.
 	 *
 	 * @todo gmp: better implementation using dependency injection / abstract factory.
-	 *
 	 * @return A pointer to the Settings module, or nullptr if the module isn't loaded.
 	 */
 	static ISettingsModule* Get( )
@@ -149,8 +148,6 @@ public:
 
 public:
 
-	/**
-	 * Virtual destructor.
-	 */
+	/** Virtual destructor. */
 	virtual ~ISettingsModule( ) { }
 };
diff --git a/Engine/Source/Developer/Settings/Public/Interfaces/ISettingsSection.h b/Engine/Source/Developer/Settings/Public/Interfaces/ISettingsSection.h
--- a/Engine/Source/Developer/Settings/Public/Interfaces/ISettingsSection.h
+++ b/Engine/Source/Developer/Settings/Public/Interfaces/ISettingsSection.h
@@ -147,7 +147,7 @@ public:
 	 * Checks whether this section can have its settings saved.
 	 *
 	 * This method does not indicate whether saving would actually succeed.
-	 1* For example, saving may be aborted by a ModifiedDelegate handler.
+	 * For example, saving may be aborted by a ModifiedDelegate handler.
 	 *
 	 * @return true if saving is supported, false otherwise.
 	 */
@@ -255,8 +255,6 @@ public:
 
 public:
 
-	/**
-	 * Virtual destructor.
-	 */
+	/** Virtual destructor. */
 	virtual ~ISettingsSection( ) { }
 };
diff --git a/Engine/Source/Developer/Settings/Public/Interfaces/ISettingsViewer.h b/Engine/Source/Developer/Settings/Public/Interfaces/ISettingsViewer.h
--- a/Engine/Source/Developer/Settings/Public/Interfaces/ISettingsViewer.h
+++ b/Engine/Source/Developer/Settings/Public/Interfaces/ISettingsViewer.h
@@ -20,8 +20,6 @@ public:
 
 public:
 
-	/**
-	 * Virtual destructor.
-	 */
+	/** Virtual destructor. */
 	virtual ~ISettingsViewer( ) { }
 };
diff --git a/Engine/Source/Developer/Settings/Public/Settings.h b/Engine/Source/Developer/Settings/Public/Settings.h
--- a/Engine/Source/Developer/Settings/Public/Settings.h
+++ b/Engine/Source/Developer/Settings/Public/Settings.h
@@ -3,7 +3,7 @@
 #pragma once
 
 
-/* Dependencies
+/* Public dependencies
  *****************************************************************************/
 
 #include "CoreUObject.h"
@@ -12,7 +12,7 @@
 class SWidget;
 
 
-/* Interfaces
+/* Public includes
  *****************************************************************************/
 
 #include "ISettingsSection.h"
diff --git a/Engine/Source/Editor/UnrealEd/Classes/Settings/ContentBrowserSettings.h b/Engine/Source/Editor/UnrealEd/Classes/Settings/ContentBrowserSettings.h
--- a/Engine/Source/Editor/UnrealEd/Classes/Settings/ContentBrowserSettings.h
+++ b/Engine/Source/Editor/UnrealEd/Classes/Settings/ContentBrowserSettings.h
@@ -44,6 +44,7 @@ public:
 	{ 
 		bOverride ? OverrideDisplayEngineFolder = bInDisplayEngineFolder : DisplayEngineFolder = bInDisplayEngineFolder;
 	}
+
 	/** Gets whether we are allowed to display the engine folder or not, optional flag ignoring the override */
 	bool GetDisplayEngineFolder( bool bExcludeOverride = false ) const
 	{ 
@@ -67,6 +68,7 @@ public:
 	{ 
 		bOverride ? OverrideDisplayPluginFolders = bInDisplayPluginFolders : DisplayPluginFolders = bInDisplayPluginFolders;
 	}
+
 	/** Gets whether we are allowed to display the plugin folders or not, optional flag ignoring the override */
 	bool GetDisplayPluginFolders( bool bExcludeOverride = false ) const
 	{ 
@@ -83,12 +85,10 @@ public:
 
 protected:
 
-	// Begin UObject overrides
+	// UObject overrides
 
 	virtual void PostEditChangeProperty( struct FPropertyChangedEvent& PropertyChangedEvent ) override;
 
-	// End UObject overrides
-
 private:
 
 	/** Whether to display the engine folder in the assets view of the content browser. */
diff --git a/Engine/Source/Editor/UnrealEd/Classes/Settings/EditorExperimentalSettings.h b/Engine/Source/Editor/UnrealEd/Classes/Settings/EditorExperimentalSettings.h
--- a/Engine/Source/Editor/UnrealEd/Classes/Settings/EditorExperimentalSettings.h
+++ b/Engine/Source/Editor/UnrealEd/Classes/Settings/EditorExperimentalSettings.h
@@ -72,12 +72,10 @@ public:
 
 protected:
 
-	// Begin UObject overrides
+	// UObject overrides
 
 	virtual void PostEditChangeProperty( struct FPropertyChangedEvent& PropertyChangedEvent ) override;
 
-	// End UObject overrides
-
 private:
 
 	// Holds an event delegate that is executed when a setting has changed.
diff --git a/Engine/Source/Editor/UnrealEd/Classes/Settings/EditorLoadingSavingSettings.h b/Engine/Source/Editor/UnrealEd/Classes/Settings/EditorLoadingSavingSettings.h
--- a/Engine/Source/Editor/UnrealEd/Classes/Settings/EditorLoadingSavingSettings.h
+++ b/Engine/Source/Editor/UnrealEd/Classes/Settings/EditorLoadingSavingSettings.h
@@ -1,9 +1,5 @@
 // Copyright 1998-2014 Epic Games, Inc. All Rights Reserved.
 
-/*=============================================================================
-	EditorLoadingSavingSettings.h: Declares the UEditorLoadingSavingSettings class.
-=============================================================================*/
-
 #pragma once
 
 
@@ -106,12 +102,10 @@ public:
 
 protected:
 
-	// Begin UObject overrides
+	// UObject overrides
 
 	virtual void PostEditChangeProperty( struct FPropertyChangedEvent& PropertyChangedEvent ) override;
 
-	// End UObject overrides
-
 private:
 
 	// Holds an event delegate that is executed when a setting has changed.
diff --git a/Engine/Source/Editor/UnrealEd/Classes/Settings/EditorMiscSettings.h b/Engine/Source/Editor/UnrealEd/Classes/Settings/EditorMiscSettings.h
--- a/Engine/Source/Editor/UnrealEd/Classes/Settings/EditorMiscSettings.h
+++ b/Engine/Source/Editor/UnrealEd/Classes/Settings/EditorMiscSettings.h
@@ -1,9 +1,5 @@
 // Copyright 1998-2014 Epic Games, Inc. All Rights Reserved.
 
-/*=============================================================================
-	EditorMiscSettings.h: Declares the UEditorMiscSettings class.
-=============================================================================*/
-
 #pragma once
 
 
diff --git a/Engine/Source/Editor/UnrealEd/Classes/Settings/LevelEditorMiscSettings.h b/Engine/Source/Editor/UnrealEd/Classes/Settings/LevelEditorMiscSettings.h
--- a/Engine/Source/Editor/UnrealEd/Classes/Settings/LevelEditorMiscSettings.h
+++ b/Engine/Source/Editor/UnrealEd/Classes/Settings/LevelEditorMiscSettings.h
@@ -1,9 +1,5 @@
 // Copyright 1998-2014 Epic Games, Inc. All Rights Reserved.
 
-/*=============================================================================
-	LevelEditorMiscSettings.h: Declares the ULevelEditorMiscSettings class.
-=============================================================================*/
-
 #pragma once
 
 #include "LevelEditorMiscSettings.generated.h"
@@ -63,9 +59,7 @@ public:
 
 protected:
 
-	// Begin UObject overrides
+	// UObject overrides
 
 	virtual void PostEditChangeProperty( struct FPropertyChangedEvent& PropertyChangedEvent ) override;
-
-	// End UObject overrides
 };
diff --git a/Engine/Source/Editor/UnrealEd/Classes/Settings/LevelEditorPlaySettings.h b/Engine/Source/Editor/UnrealEd/Classes/Settings/LevelEditorPlaySettings.h
--- a/Engine/Source/Editor/UnrealEd/Classes/Settings/LevelEditorPlaySettings.h
+++ b/Engine/Source/Editor/UnrealEd/Classes/Settings/LevelEditorPlaySettings.h
@@ -1,9 +1,5 @@
 // Copyright 1998-2014 Epic Games, Inc. All Rights Reserved.
 
-/*=============================================================================
-	LevelEditorPlaySettings.h: Declares the ULevelEditorPlaySettings class.
-=============================================================================*/
-
 #pragma once
 
 
@@ -86,27 +82,19 @@ struct FPlayScreenResolution
 
 public:
 
-	/**
-	 * The description text for this screen resolution.
-	 */
+	/** The description text for this screen resolution. */
 	UPROPERTY(config)
 	/*FText*/FString Description;
 
-	/**
-	 * The screen resolution's width (in pixels).
-	 */
+	/** The screen resolution's width (in pixels). */
 	UPROPERTY(config)
 	int32 Width;
 
-	/**
-	 * The screen resolution's height (in pixels).
-	 */
+	/** The screen resolution's height (in pixels). */
 	UPROPERTY(config)
 	int32 Height;
 
-	/**
-	 * The screen resolution's aspect ratio (as a string).
-	 */
+	/** The screen resolution's aspect ratio (as a string). */
 	UPROPERTY(config)
 	FString AspectRatio;
 };
@@ -123,131 +111,91 @@ class UNREALED_API ULevelEditorPlaySettings
 
 public:
 
-	/**
-	 * The PlayerStart class used when spawning the player at the current camera location.
-	 */
+	/** The PlayerStart class used when spawning the player at the current camera location. */
 	UPROPERTY(globalconfig, EditAnywhere, Category=General, meta=(ToolTip=""))
 	FString PlayFromHerePlayerStartClassName;
 
 public:
 
-	/**
-	 * Should Play-in-Editor automatically give mouse control to the game on PIE start (default = false).
-	 */
+	/** Should Play-in-Editor automatically give mouse control to the game on PIE start (default = false). */
 	UPROPERTY(config, EditAnywhere, Category=PlayInEditor, meta=(ToolTip="Give the game mouse control when PIE starts or require a click in the viewport first"))
 	bool GameGetsMouseControl;
 
-	/**
-	 * Whether to show a label for mouse control gestures in the PIE view.
-	 */
+	/** Whether to show a label for mouse control gestures in the PIE view. */
 	UPROPERTY(config, EditAnywhere, Category=PlayInEditor)
 	bool ShowMouseControlLabel;
 
-	/**
-	 * Location on screen to anchor the mouse control label when in PIE mode.
-	 */
+	/** Location on screen to anchor the mouse control label when in PIE mode. */
 	UPROPERTY(config, EditAnywhere, Category=PlayInEditor)
 	TEnumAsByte<ELabelAnchorMode> MouseControlLabelPosition;
 
-	/**
-	 * Whether to automatically recompile blueprints on PIE
-	 */
+	/** Whether to automatically recompile blueprints on PIE */
 	UPROPERTY(config, EditAnywhere, Category=PlayInEditor, meta=(ToolTip="Automatically recompile blueprints used by the current level when initiating a Play In Editor session"))
 	bool AutoRecompileBlueprints;
 
-	/**
-	 * Whether to play sounds during PIE
-	 */
+	/** Whether to play sounds during PIE */
 	UPROPERTY(config, EditAnywhere, Category = PlayInEditor, meta = (ToolTip = "Whether to play sounds when in a Play In Editor session"))
 	bool EnableSound;
 
-	/**
-	 * True if Play In Editor should only load currently-visible levels in PIE.
-	 */
+	/** True if Play In Editor should only load currently-visible levels in PIE. */
 	UPROPERTY(config)
 	uint32 bOnlyLoadVisibleLevelsInPIE:1;
 
 public:
 
-	/**
-	 * The width of the new view port window in pixels (0 = use the desktop's screen resolution).
-	 */
+	/** The width of the new view port window in pixels (0 = use the desktop's screen resolution). */
 	UPROPERTY(config, EditAnywhere, Category=PlayInNewWindow)
 	int32 NewWindowWidth;
 
-	/**
-	 * The height of the new view port window in pixels (0 = use the desktop's screen resolution).
-	 */
+	/** The height of the new view port window in pixels (0 = use the desktop's screen resolution). */
 	UPROPERTY(config, EditAnywhere, Category=PlayInNewWindow)
 	int32 NewWindowHeight;
 
-	/**
-	 * The position of the new view port window on the screen in pixels.
-	 */
+	/** The position of the new view port window on the screen in pixels. */
 	UPROPERTY(config, EditAnywhere, Category=PlayInNewWindow)
 	FIntPoint NewWindowPosition;
 
-	/**
-	 * Whether the new window should be centered on the screen.
-	 */
+	/** Whether the new window should be centered on the screen. */
 	UPROPERTY(config, EditAnywhere, Category=PlayInNewWindow)
 	bool CenterNewWindow;
 
 public:
 
-	/**
-	 * The width of the standalone game window in pixels (0 = use the desktop's screen resolution).
-	 */
+	/** The width of the standalone game window in pixels (0 = use the desktop's screen resolution). */
 	UPROPERTY(config, EditAnywhere, Category=PlayInStandaloneGame)
 	int32 StandaloneWindowWidth;
 
-	/**
-	 * The height of the standalone game window in pixels (0 = use the desktop's screen resolution).
-	 */
+	/** The height of the standalone game window in pixels (0 = use the desktop's screen resolution). */
 	UPROPERTY(config, EditAnywhere, Category=PlayInStandaloneGame)
 	int32 StandaloneWindowHeight;
 
-	/**
-	 * Whether sound should be disabled when playing standalone games.
-	 */
+	/** Whether sound should be disabled when playing standalone games. */
 	UPROPERTY(config , EditAnywhere, Category=PlayInStandaloneGame, AdvancedDisplay)
 	uint32 DisableStandaloneSound:1;
 
-	/**
-	 * Extra parameters to be include as part of the command line for the standalone game.
-	 */
+	/** Extra parameters to be include as part of the command line for the standalone game. */
 	UPROPERTY(config , EditAnywhere, Category=PlayInStandaloneGame, AdvancedDisplay)
 	FString AdditionalLaunchParameters;
 
 public:
 
-	/**
-	 * NetMode to use for Play In Editor.
-	 */
+	/** NetMode to use for Play In Editor. */
 	UPROPERTY(config, EditAnywhere, Category=MultiplayerOptions)
 	TEnumAsByte<EPlayNetMode> PlayNetMode;
 
-	/**
-     * Spawn multiple player windows in a single instance of UE4. This will load much faster, but has potential to have more issues. 
-     */
+	/** Spawn multiple player windows in a single instance of UE4. This will load much faster, but has potential to have more issues.  */
 	UPROPERTY(config, EditAnywhere, Category=MultiplayerOptions)
 	bool RunUnderOneProcess;
 
-	/**
-	 * If checked, a separate dedicated server will be launched. Otherwise the first player window will act as a listen server that all other player windows connect to.
-	 */
+	/** If checked, a separate dedicated server will be launched. Otherwise the first player window will act as a listen server that all other player windows connect to. */
 	UPROPERTY(config, EditAnywhere, Category=MultiplayerOptions)
 	bool PlayNetDedicated;
 
-	/**
-	 * Number of player windows that should be opened. Listen server will count as a player, a dedicated server will not.
-	 */
+	/** Number of player windows that should be opened. Listen server will count as a player, a dedicated server will not. */
 	UPROPERTY(config, EditAnywhere, Category=MultiplayerOptions, meta=(ClampMin = "1", UIMin = "1", UIMax = "64"))
 	int32 PlayNumberOfClients;
 
-	/**
-	 * Window width to use when spawning additional clients.
-	 */
+	/** Window width to use when spawning additional clients. */
 	UPROPERTY(config, EditAnywhere, Category=MultiplayerOptions)
 	int32 ClientWindowWidth;
 
@@ -261,105 +209,73 @@ public:
 	UPROPERTY(config, EditAnywhere, Category=MultiplayerOptions)
 	bool RouteGamepadToSecondWindow;
 
-	/**
-	 * Window height to use when spawning additional clients.
-	 */
+	/** Window height to use when spawning additional clients. */
 	UPROPERTY(config, EditAnywhere, Category=MultiplayerOptions)
 	int32 ClientWindowHeight;
 
-	/**
-	 * Additional options that will be passed to the server as URL parameters.
-	 */
+	/** Additional options that will be passed to the server as URL parameters. */
 	UPROPERTY(config, EditAnywhere, Category=MultiplayerOptions)
 	FString AdditionalServerGameOptions;
 
-	/**
-	 * Additional command line options that will be passed to standalone game instances.
-	 */
+	/** Additional command line options that will be passed to standalone game instances. */
 	UPROPERTY(config, EditAnywhere, Category=MultiplayerOptions)
 	FString AdditionalLaunchOptions;
 
 public:
 
-	/**
-	 * The last used height for multiple instance windows (in pixels).
-	 */
+	/** The last used height for multiple instance windows (in pixels). */
 	UPROPERTY(config)
 	int32 MultipleInstanceLastHeight;
 
-	/**
-	 * The last used width for multiple instance windows (in pixels).
-	 */
+	/** The last used width for multiple instance windows (in pixels). */
 	UPROPERTY(config)
 	int32 MultipleInstanceLastWidth;
 
-	/**
-	 * The last known screen positions of multiple instance windows (in pixels).
-	 */
+	/** The last known screen positions of multiple instance windows (in pixels). */
 	UPROPERTY(config)
 	TArray<FIntPoint> MultipleInstancePositions;
 
 public:
 
-	/**
-	 * The name of the last platform that the user ran a play session on.
-	 */
+	/** The name of the last platform that the user ran a play session on. */
 	UPROPERTY(config)
 	FString LastExecutedLaunchDevice;
 
-	/**
-	 * The name of the last device that the user ran a play session on.
-	 */
+	/** The name of the last device that the user ran a play session on. */
 	UPROPERTY(config)
 	FString LastExecutedLaunchName;
 
-	/**
-	 * The last type of play-on session the user ran.
-	 */
+	/** The last type of play-on session the user ran. */
 	UPROPERTY(config)
 	TEnumAsByte<ELaunchModeType> LastExecutedLaunchModeType;
 
-	/**
-	 * The last type of play location the user ran.
-	 */
+	/** The last type of play location the user ran. */
 	UPROPERTY(config)
 	TEnumAsByte<EPlayModeLocations> LastExecutedPlayModeLocation;
 
-	/**
-	 * The last type of play session the user ran.
-	 */
+	/** The last type of play session the user ran. */
 	UPROPERTY(config)
 	TEnumAsByte<EPlayModeType> LastExecutedPlayModeType;
 
 public:
 
-	/**
-	 * Collection of common screen resolutions on mobile phones.
-	 */
+	/** Collection of common screen resolutions on mobile phones. */
 	UPROPERTY(config)
 	TArray<FPlayScreenResolution> LaptopScreenResolutions;
 
-	/**
-	 * Collection of common screen resolutions on desktop monitors.
-	 */
+	/** Collection of common screen resolutions on desktop monitors. */
 	UPROPERTY(config)
 	TArray<FPlayScreenResolution> MonitorScreenResolutions;
 
-	/**
-	 * Collection of common screen resolutions on mobile phones.
-	 */
+	/** Collection of common screen resolutions on mobile phones. */
 	UPROPERTY(config)
 	TArray<FPlayScreenResolution> PhoneScreenResolutions;
 
-	/**
-	 * Collection of common screen resolutions on tablet devices.
-	 */
+	/** Collection of common screen resolutions on tablet devices. */
 	UPROPERTY(config)
 	TArray<FPlayScreenResolution> TabletScreenResolutions;
 
-	/**
-	 * Collection of common screen resolutions on television screens.
-	 */
+	/** Collection of common screen resolutions on television screens. */
 	UPROPERTY(config)
 	TArray<FPlayScreenResolution> TelevisionScreenResolutions;
 };
diff --git a/Engine/Source/Editor/UnrealEd/Classes/Settings/LevelEditorViewportSettings.h b/Engine/Source/Editor/UnrealEd/Classes/Settings/LevelEditorViewportSettings.h
--- a/Engine/Source/Editor/UnrealEd/Classes/Settings/LevelEditorViewportSettings.h
+++ b/Engine/Source/Editor/UnrealEd/Classes/Settings/LevelEditorViewportSettings.h
@@ -1,9 +1,5 @@
 // Copyright 1998-2014 Epic Games, Inc. All Rights Reserved.
 
-/*=============================================================================
-	LevelEditorViewportSettings.h: Declares the ULevelEditorViewportSettings class.
-=============================================================================*/
-
 #pragma once
 
 #include "Viewports.h"
@@ -61,77 +57,57 @@ struct UNREALED_API FLevelEditorViewportInstanceSettings
 		, bShowStats(false)
 	{ }
 
-	/**
-	 * The viewport type
-	 */
+	/** The viewport type */
 	UPROPERTY(config)
 	TEnumAsByte<ELevelViewportType> ViewportType;
 
-	/* 
-	 * View mode to set when this viewport is of type LVT_Perspective 
-	 */
+	/* View mode to set when this viewport is of type LVT_Perspective. */
 	UPROPERTY(config)
 	TEnumAsByte<EViewModeIndex> PerspViewModeIndex;
 
-	/* 
-	 * View mode to set when this viewport is not of type LVT_Perspective 
-	 */
+	/* View mode to set when this viewport is not of type LVT_Perspective. */
 	UPROPERTY(config)
 	TEnumAsByte<EViewModeIndex> OrthoViewModeIndex;
 
 	/**
-	 * A set of flags that determines visibility for various scene elements (FEngineShowFlags), converted to string form
-	 * These have to be saved as strings since FEngineShowFlags is too complex for UHT to parse correctly
+	 * A set of flags that determines visibility for various scene elements (FEngineShowFlags), converted to string form.
+	 * These have to be saved as strings since FEngineShowFlags is too complex for UHT to parse correctly.
 	 */
 	UPROPERTY(config)
 	FString EditorShowFlagsString;
 
 	/**
-	 * A set of flags that determines visibility for various scene elements (FEngineShowFlags), converted to string form
-	 * These have to be saved as strings since FEngineShowFlags is too complex for UHT to parse correctly
+	 * A set of flags that determines visibility for various scene elements (FEngineShowFlags), converted to string form.
+	 * These have to be saved as strings since FEngineShowFlags is too complex for UHT to parse correctly.
 	 */
 	UPROPERTY(config)
 	FString GameShowFlagsString;
 
-	/**
-	 * The buffer visualization mode for the viewport
-	 */
+	/** The buffer visualization mode for the viewport. */
 	UPROPERTY(config)
 	FName BufferVisualizationMode;
 
-	/**
-	 * Setting to allow designers to override the automatic expose
-	 */
+	/** Setting to allow designers to override the automatic expose. */
 	UPROPERTY(config)
 	FExposureSettings ExposureSettings;
 
-	/*
-	 * Field of view angle for the viewport
-	 */
+	/* Field of view angle for the viewport. */
 	UPROPERTY(config)
 	float FOVAngle;
 
-	/*
-	 * Is this viewport updating in real-time?
-	 */
+	/* Whether this viewport is updating in real-time. */
 	UPROPERTY(config)
 	bool bIsRealtime;
 
-	/*
-	 * Should this viewport show an FPS count?
-	 */
+	/* Whether the FPS counter should be shown. */
 	UPROPERTY(config)
 	bool bShowFPS_DEPRECATED;
 
-	/*
-	 * Should this viewport show statistics?
-	 */
+	/* Whether viewport statistics should be shown. */
 	UPROPERTY(config)
 	bool bShowStats;
 
-	/*
-	 * Should this viewport have any stats enabled by default?
-	 */
+	/* Whether viewport statistics should be enabled by default. */
 	UPROPERTY(config)
 	TArray<FString> EnabledStats;
 };
@@ -145,15 +121,11 @@ struct UNREALED_API FLevelEditorViewportInstanceSettingsKeyValuePair
 {
 	GENERATED_USTRUCT_BODY()
 
-	/* 
-	 * Name identifying this config
-	 */
+	/*  Name identifying this config. */
 	UPROPERTY(config)
 	FString ConfigName;
 
-	/* 
-	 * Settings for this config
-	 */
+	/* Settings for this config. */
 	UPROPERTY(config)
 	FLevelEditorViewportInstanceSettings ConfigSettings;
 };
@@ -167,8 +139,11 @@ struct UNREALED_API FSnapToSurfaceSettings
 {
 	GENERATED_USTRUCT_BODY()
 
-	FSnapToSurfaceSettings() : bEnabled(false), SnapOffsetExtent(0.f), bSnapRotation(true)
-	{}
+	FSnapToSurfaceSettings()
+		: bEnabled(false)
+		, SnapOffsetExtent(0.f)
+		, bSnapRotation(true)
+	{ }
 
 	/** Whether snapping to surfaces in the world is enabled */
 	UPROPERTY(config)
@@ -193,21 +168,15 @@ class UNREALED_API ULevelEditorViewportSettings
 {
 	GENERATED_UCLASS_BODY()
 
-	/**
-	 * Enable the use of flight camera controls under various circumstances.
-	 */
+	/** Enable the use of flight camera controls under various circumstances. */
 	UPROPERTY(EditAnywhere, config, Category=Controls)
 	TEnumAsByte<enum EWASDType> FlightCameraControlType;
 
-	/**
-	 * If true, moves the canvas and shows the mouse.  If false, uses original camera movement.
-	 */
+	/** If true, moves the canvas and shows the mouse.  If false, uses original camera movement. */
 	UPROPERTY(EditAnywhere, config, Category=Controls, meta=(DisplayName = "Grab and Drag to Move Orthographic Cameras"), AdvancedDisplay)
 	uint32 bPanMovesCanvas:1;
 
-	/**
-	 * If checked, in orthographic view ports zooming will center on the mouse position.  If unchecked, the zoom is around the center of the viewport.
-	 */
+	/** If checked, in orthographic view ports zooming will center on the mouse position.  If unchecked, the zoom is around the center of the viewport. */
 	UPROPERTY(EditAnywhere, config, Category=Controls, meta=(DisplayName = "Orthographic Zoom to Cursor Position"))
 	uint32 bCenterZoomAroundCursor:1;
 
@@ -219,39 +188,27 @@ class UNREALED_API ULevelEditorViewportSettings
 	UPROPERTY(EditAnywhere, config, Category=Controls, meta=( DisplayName = "Clicking BSP Enables Brush" ))
 	uint32 bClickBSPSelectsBrush:1;
 
-	/**
-	 * How fast the perspective camera moves when flying through the world.
-	 */
+	/** How fast the perspective camera moves when flying through the world. */
 	UPROPERTY(config, meta=(UIMin = "1", UIMax = "8", ClampMin="1", ClampMax="8"))
 	int32 CameraSpeed;
 
-	/**
-	 * How fast the perspective camera moves through the world when using mouse scroll.
-	 */
+	/** How fast the perspective camera moves through the world when using mouse scroll. */
 	UPROPERTY(EditAnywhere, config, Category=Controls, meta=(UIMin = "1", UIMax = "8", ClampMin="1", ClampMax="8"))
 	int32 MouseScrollCameraSpeed;
 
-	/**
-	 * The sensitivity of mouse movement when rotating the camera.
-	 */
+	/** The sensitivity of mouse movement when rotating the camera. */
 	UPROPERTY(EditAnywhere, config, Category=Controls, meta=(ClampMin="0.0",ClampMax="1.0") )
 	float MouseSensitivty;
 	
-	/**
-	 * Whether or not to invert the direction of middle mouse panning in viewports
-	 */
+	/** Whether or not to invert the direction of middle mouse panning in viewports */
 	UPROPERTY(EditAnywhere, config, Category=Controls)
 	bool bInvertMiddleMousePan;
 
-	/**
-	 * Whether to use mouse position as direct widget position.
-	 */
+	/** Whether to use mouse position as direct widget position. */
 	UPROPERTY(EditAnywhere, config, Category=Controls, AdvancedDisplay)
 	uint32 bUseAbsoluteTranslation:1;
 
-	/**
-	 * If enabled, the viewport will stream in levels automatically when the camera is moved.
-	 */
+	/** If enabled, the viewport will stream in levels automatically when the camera is moved. */
 	UPROPERTY(EditAnywhere, config, Category=Controls, meta=(DisplayName = "Stream in Levels Automatically when Camera is Moved"), aDvancedDisplay)
 	bool bLevelStreamingVolumePrevis;
 
@@ -365,9 +322,7 @@ public:
 	UPROPERTY(EditAnywhere, config, Category=LookAndFeel, meta=(DisplayName = "Highlight Selected Objects with Brackets"))
 	uint32 bHighlightWithBrackets:1;
 
-	/**
-	 * If checked all orthographic view ports are linked to the same position and move together.
-	 */
+	/** If checked all orthographic view ports are linked to the same position and move together. */
 	UPROPERTY(EditAnywhere, config, Category=LookAndFeel, meta=(DisplayName = "Link Orthographic Viewport Movement"))
 	uint32 bUseLinkedOrthographicViewports:1;
 
@@ -414,9 +369,7 @@ public:
 	UPROPERTY(EditAnywhere, config, Category=LookAndFeel, meta=(ClampMin = "0.01", UIMin = "0.01", UIMax = "5"))
 	float BillboardScale;
 
-	/**
-	 * The size adjustment to apply to the translate/rotate/scale widgets (in Unreal units).
-	 */
+	/** The size adjustment to apply to the translate/rotate/scale widgets (in Unreal units). */
 	UPROPERTY(EditAnywhere, config, Category=LookAndFeel, meta=(ClampMin="-10",ClampMax="150") )
 	int32 TransformWidgetSizeAdjustment;
 
@@ -501,12 +454,10 @@ public:
 
 protected:
 
-	// Begin UObject overrides
+	// UObject overrides
 
 	virtual void PostEditChangeProperty( struct FPropertyChangedEvent& PropertyChangedEvent ) override;
 
-	// End UObject overrides
-
 private:
 
 	// Holds an event delegate that is executed when a setting has changed.
diff --git a/Engine/Source/Editor/UnrealEd/Classes/Settings/ProjectPackagingSettings.h b/Engine/Source/Editor/UnrealEd/Classes/Settings/ProjectPackagingSettings.h
--- a/Engine/Source/Editor/UnrealEd/Classes/Settings/ProjectPackagingSettings.h
+++ b/Engine/Source/Editor/UnrealEd/Classes/Settings/ProjectPackagingSettings.h
@@ -1,9 +1,5 @@
 // Copyright 1998-2014 Epic Games, Inc. All Rights Reserved.
 
-/*=============================================================================
-	ProjectPackagingSettings.h: Declares the UProjectPackagingSettings class.
-=============================================================================*/
-
 #pragma once
 
 #include "ProjectPackagingSettings.generated.h"
@@ -15,19 +11,13 @@
 UENUM()
 enum EProjectPackagingBuildConfigurations
 {
-	/**
-	 * Debug configuration.
-	 */
+	/** Debug configuration. */
 	PPBC_DebugGame UMETA(DisplayName="DebugGame"),
 
-	/**
-	 * Development configuration.
-	 */
+	/** Development configuration. */
 	PPBC_Development UMETA(DisplayName="Development"),
 
-	/**
-	 * Shipping configuration.
-	 */
+	/** Shipping configuration. */
 	PPBC_Shipping UMETA(DisplayName="Shipping")
 };
 
@@ -41,23 +31,13 @@ class UNREALED_API UProjectPackagingSettings
 {
 	GENERATED_UCLASS_BODY()
 
-	// Begin UObject Interface
-
-	virtual void PostEditChangeProperty( struct FPropertyChangedEvent& PropertyChangedEvent ) override;
-	
-	// End UObject Interface
-
 public:
 
-	/**
-	 * The build configuration for which the project is packaged.
-	 */
+	/** The build configuration for which the project is packaged. */
 	UPROPERTY(config, EditAnywhere, Category=Project)
 	TEnumAsByte<EProjectPackagingBuildConfigurations> BuildConfiguration;
 
-	/**
-	 * The directory to which the packaged project will be copied.
-	 */
+	/** The directory to which the packaged project will be copied. */
 	UPROPERTY(config, EditAnywhere, Category=Project)
 	FDirectoryPath StagingDirectory;
 
@@ -77,16 +57,12 @@ public:
 	UPROPERTY(config, EditAnywhere, Category=Project, AdvancedDisplay)
 	bool ForDistribution;
 
-	/**
-	 * If enabled, all content will be put into a single .pak file instead of many individual files (default = enabled).
-	 */
+	/** If enabled, all content will be put into a single .pak file instead of many individual files (default = enabled). */
 	UPROPERTY(config, EditAnywhere, Category=Packaging)
 	bool UsePakFile;
 
 
-	/**
-	* If enabled, on Android platforms, .pak files are placed inside the APK
-	*/
+	/** If enabled, on Android platforms, .pak files are placed inside the APK. */
 	UPROPERTY(config, EditAnywhere, Category = Packaging)
 	bool UseOBB_InAPK;
 
@@ -112,4 +88,10 @@ public:
 	 */
 	UPROPERTY(config, EditAnywhere, Category=Packaging, AdvancedDisplay, meta=(DisplayName="Additional Non-Asset Directories To Copy", RelativeToGameContentDir))
 	TArray<FDirectoryPath> DirectoriesToAlwaysStageAsNonUFS;
+
+public:
+
+	// UObject Interface
+
+	virtual void PostEditChangeProperty( struct FPropertyChangedEvent& PropertyChangedEvent ) override;
 };
diff --git a/Engine/Source/Editor/UnrealEd/Private/Settings/EditorLoadingSavingSettingsCustomization.h b/Engine/Source/Editor/UnrealEd/Private/Settings/EditorLoadingSavingSettingsCustomization.h
--- a/Engine/Source/Editor/UnrealEd/Private/Settings/EditorLoadingSavingSettingsCustomization.h
+++ b/Engine/Source/Editor/UnrealEd/Private/Settings/EditorLoadingSavingSettingsCustomization.h
@@ -1,9 +1,5 @@
 // Copyright 1998-2014 Epic Games, Inc. All Rights Reserved.
 
-/*=============================================================================
-	EditorLoadingSavingSettingsCustomization.h: Declares the FEditorLoadingSavingSettingsCustomization class.
-=============================================================================*/
-
 #pragma once
 
 #include "PropertyEditing.h"
@@ -20,17 +16,19 @@ class FEditorLoadingSavingSettingsCustomization
 	: public IDetailCustomization
 {
 public:
-	virtual ~FEditorLoadingSavingSettingsCustomization() {}
 
-	// Begin IDetailCustomization interface
+	/** Virtual destructor. */
+	virtual ~FEditorLoadingSavingSettingsCustomization( ) { }
+
+public:
+
+	// IDetailCustomization interface
 
 	virtual void CustomizeDetails( IDetailLayoutBuilder& LayoutBuilder ) override
 	{
 		CustomizeStartupCategory(LayoutBuilder);
 	}
 
-	// End IDetailCustomization interface
-
 public:
 
 	/**
@@ -51,7 +49,7 @@ protected:
 	 * This customization pulls in a setting from the game agnostic Editor settings, which
 	 * are stored in a different UObject, but which we would like to show in this section.
 	 *
-	 * @param LayoutBuilder - The layout builder.
+	 * @param LayoutBuilder The layout builder.
 	 */
 	void CustomizeStartupCategory( IDetailLayoutBuilder& LayoutBuilder )
 	{
diff --git a/Engine/Source/Editor/UnrealEd/Private/Settings/GameMapsSettingsCustomization.h b/Engine/Source/Editor/UnrealEd/Private/Settings/GameMapsSettingsCustomization.h
--- a/Engine/Source/Editor/UnrealEd/Private/Settings/GameMapsSettingsCustomization.h
+++ b/Engine/Source/Editor/UnrealEd/Private/Settings/GameMapsSettingsCustomization.h
@@ -1,9 +1,5 @@
 // Copyright 1998-2014 Epic Games, Inc. All Rights Reserved.
 
-/*=============================================================================
-	GameMapsSettingsCustomization.h: Declares the FGameMapsSettingsCustomization class.
-=============================================================================*/
-
 #pragma once
 
 #include "PropertyEditing.h"
@@ -20,9 +16,13 @@ class FGameMapsSettingsCustomization
 	: public IDetailCustomization
 {
 public:
-	virtual ~FGameMapsSettingsCustomization() {}
 
-	// Begin IDetailCustomization interface
+	/** Virtual destructor. */
+	virtual ~FGameMapsSettingsCustomization( ) { }
+
+public:
+
+	// IDetailCustomization interface
 
 	virtual void CustomizeDetails( IDetailLayoutBuilder& LayoutBuilder ) override
 	{
@@ -52,8 +52,6 @@ public:
 		}
 	}
 
-	// End IDetailCustomization interface
-
 public:
 
 	/**
@@ -71,9 +69,9 @@ protected:
 	/**
 	 * Customizes the property row for a map setting.
 	 *
-	 * @param LayoutBuilder - The layout builder.
-	 * @param CategoryBuilder - The builder for the detail category that the setting belongs to.
-	 * @param PropertyName - The name of the property that holds the map setting.
+	 * @param LayoutBuilder The layout builder.
+	 * @param CategoryBuilder The builder for the detail category that the setting belongs to.
+	 * @param PropertyName The name of the property that holds the map setting.
 	 */
 	void CustomizeMapSetting( IDetailLayoutBuilder& LayoutBuilder, IDetailCategoryBuilder& CategoryBuilder, const FName& PropertyName )
 	{
@@ -123,8 +121,7 @@ protected:
 	/**
 	 * Checks whether the specified map name is valid.
 	 *
-	 * @param MapName - The map name to validate.
-	 *
+	 * @param MapName The map name to validate.
 	 * @return true if the map name is valid, false otherwise.
 	 */
 	bool IsValidMapName( const FString& MapName ) const
@@ -140,8 +137,7 @@ protected:
 	/**
 	 * Creates a widget for the map picker.
 	 *
-	 * @param PropertyRow - The property row to create the widget for.
-	 *
+	 * @param PropertyRow The property row to create the widget for.
 	 * @return The widget.
 	 */
 	TSharedRef<SWidget> MakeMapMenu( const TSharedPtr<IPropertyHandle>& PropertyHandle )
diff --git a/Engine/Source/Editor/UnrealEd/Private/Settings/LevelEditorPlaySettingsCustomization.h b/Engine/Source/Editor/UnrealEd/Private/Settings/LevelEditorPlaySettingsCustomization.h
--- a/Engine/Source/Editor/UnrealEd/Private/Settings/LevelEditorPlaySettingsCustomization.h
+++ b/Engine/Source/Editor/UnrealEd/Private/Settings/LevelEditorPlaySettingsCustomization.h
@@ -1,9 +1,5 @@
 // Copyright 1998-2014 Epic Games, Inc. All Rights Reserved.
 
-/*=============================================================================
-	LevelEditorPlaySettingsCustomization.h: Declares the FLevelEditorPlaySettingsCustomization class.
-=============================================================================*/
-
 #pragma once
 
 #include "PropertyEditing.h"
@@ -26,10 +22,10 @@ public:
 	/**
 	 * Construct this widget
 	 *
-	 * @param InArgs - The declaration data for this widget.
-	 * @param LayoutBuilder - The layout builder to use for generating property widgets.
-	 * @param InWindowPositionProperty - The handle to the window position property.
-	 * @param InCenterWindowProperty - The handle to the center window property.
+	 * @param InArgs The declaration data for this widget.
+	 * @param LayoutBuilder The layout builder to use for generating property widgets.
+	 * @param InWindowPositionProperty The handle to the window position property.
+	 * @param InCenterWindowProperty The handle to the center window property.
 	 */
 	void Construct( const FArguments& InArgs, IDetailLayoutBuilder* LayoutBuilder, const TSharedRef<IPropertyHandle>& InWindowPositionProperty, const TSharedRef<IPropertyHandle>& InCenterWindowProperty )
 	{
@@ -132,10 +128,10 @@ public:
 	/**
 	 * Construct this widget
 	 *
-	 * @param InArgs - The declaration data for this widget.
-	 * @param LayoutBuilder - The layout builder to use for generating property widgets.
-	 * @param InWindowHeightProperty - The handle to the window height property.
-	 * @param InWindowWidthProperty - The handle to the window width property.
+	 * @param InArgs The declaration data for this widget.
+	 * @param LayoutBuilder The layout builder to use for generating property widgets.
+	 * @param InWindowHeightProperty The handle to the window height property.
+	 * @param InWindowWidthProperty The handle to the window width property.
 	 */
 	void Construct( const FArguments& InArgs, IDetailLayoutBuilder* LayoutBuilder, const TSharedRef<IPropertyHandle>& InWindowHeightProperty, const TSharedRef<IPropertyHandle>& InWindowWidthProperty )
 	{
@@ -288,9 +284,13 @@ class FLevelEditorPlaySettingsCustomization
 	: public IDetailCustomization
 {
 public:
-	virtual ~FLevelEditorPlaySettingsCustomization() {}
 
-	// Begin IDetailCustomization interface
+	/** Virtual destructor. */
+	virtual ~FLevelEditorPlaySettingsCustomization( ) { }
+
+public:
+
+	// IDetailCustomization interface
 
 	virtual void CustomizeDetails( IDetailLayoutBuilder& LayoutBuilder ) override
 	{
diff --git a/Engine/Source/Editor/UnrealEd/Private/Settings/ProjectPackagingSettingsCustomization.h b/Engine/Source/Editor/UnrealEd/Private/Settings/ProjectPackagingSettingsCustomization.h
--- a/Engine/Source/Editor/UnrealEd/Private/Settings/ProjectPackagingSettingsCustomization.h
+++ b/Engine/Source/Editor/UnrealEd/Private/Settings/ProjectPackagingSettingsCustomization.h
@@ -1,9 +1,5 @@
 // Copyright 1998-2014 Epic Games, Inc. All Rights Reserved.
 
-/*=============================================================================
-	ProjectPackagingSettingsCustomization.h: Declares the FProjectPackagingSettingsCustomization class.
-=============================================================================*/
-
 #pragma once
 
 #include "PropertyEditing.h"
@@ -22,15 +18,13 @@ class FProjectPackagingSettingsCustomization
 {
 public:
 
-	// Begin IDetailCustomization interface
+	// IDetailCustomization interface
 
 	virtual void CustomizeDetails( IDetailLayoutBuilder& LayoutBuilder ) override
 	{
 		CustomizeProjectCategory(LayoutBuilder);
 	}
 
-	// End IDetailCustomization interface
-
 public:
 
 	/**
@@ -48,7 +42,7 @@ protected:
 	/**
 	 * Customizes the Project property category.
 	 *
-	 * @param LayoutBuilder - The layout builder.
+	 * @param LayoutBuilder The layout builder.
 	 */
 	void CustomizeProjectCategory( IDetailLayoutBuilder& LayoutBuilder )
 	{
diff --git a/Engine/Source/Editor/UnrealEd/Private/Settings/SettingsClasses.cpp b/Engine/Source/Editor/UnrealEd/Private/Settings/SettingsClasses.cpp
--- a/Engine/Source/Editor/UnrealEd/Private/Settings/SettingsClasses.cpp
+++ b/Engine/Source/Editor/UnrealEd/Private/Settings/SettingsClasses.cpp
@@ -1,9 +1,5 @@
 ﻿// Copyright 1998-2014 Epic Games, Inc. All Rights Reserved.
 
-/*=============================================================================
-	Settings.cpp: Implements the constructors for the various settings classes.
-=============================================================================*/
-
 #include "UnrealEd.h"
 #include "ISourceControlModule.h"
 
@@ -22,7 +18,9 @@ void UContentBrowserSettings::PostEditChangeProperty( struct FPropertyChangedEve
 {
 	Super::PostEditChangeProperty(PropertyChangedEvent);
 
-	const FName Name = (PropertyChangedEvent.Property != nullptr) ? PropertyChangedEvent.Property->GetFName() : NAME_None;
+	const FName Name = (PropertyChangedEvent.Property != nullptr)
+		? PropertyChangedEvent.Property->GetFName()
+		: NAME_None;
 
 	if (!FUnrealEdMisc::Get().IsDeletePreferences())
 	{
@@ -39,12 +37,12 @@ void UContentBrowserSettings::PostEditChangeProperty( struct FPropertyChangedEve
 UDestructableMeshEditorSettings::UDestructableMeshEditorSettings( const class FPostConstructInitializeProperties& PCIP )
 	: Super(PCIP)
 {
-	AnimPreviewLightingDirection = FRotator(-45.0f,45.0f,0);
+	AnimPreviewLightingDirection = FRotator(-45.0f, 45.0f, 0);
 	AnimPreviewSkyColor = FColor(0, 0, 255);
 	AnimPreviewFloorColor = FColor(51, 51, 51);
-	AnimPreviewSkyBrightness = 0.2f*PI;
+	AnimPreviewSkyBrightness = 0.2f * PI;
 	AnimPreviewDirectionalColor = FColor(255, 255, 255);
-	AnimPreviewLightBrightness = 1.0f*PI;
+	AnimPreviewLightBrightness = 1.0f * PI;
 }
 
 
@@ -53,8 +51,7 @@ UDestructableMeshEditorSettings::UDestructableMeshEditorSettings( const class FP
 
 UEditorExperimentalSettings::UEditorExperimentalSettings( const class FPostConstructInitializeProperties& PCIP )
 	: Super(PCIP)
-{
-}
+{ }
 
 
 void UEditorExperimentalSettings::PostEditChangeProperty( struct FPropertyChangedEvent& PropertyChangedEvent )
@@ -137,7 +134,9 @@ void ULevelEditorMiscSettings::PostEditChangeProperty( struct FPropertyChangedEv
 {
 	Super::PostEditChangeProperty(PropertyChangedEvent);
 
-	const FName Name = (PropertyChangedEvent.Property != nullptr) ? PropertyChangedEvent.Property->GetFName() : NAME_None;
+	const FName Name = (PropertyChangedEvent.Property != nullptr)
+		? PropertyChangedEvent.Property->GetFName()
+		: NAME_None;
 
 	if (Name == FName(TEXT("bNavigationAutoUpdate")))
 	{
@@ -187,7 +186,9 @@ void ULevelEditorViewportSettings::PostEditChangeProperty( struct FPropertyChang
 {
 	Super::PostEditChangeProperty(PropertyChangedEvent);
 
-	const FName Name = (PropertyChangedEvent.Property != nullptr) ? PropertyChangedEvent.Property->GetFName() : NAME_None;
+	const FName Name = (PropertyChangedEvent.Property != nullptr)
+		? PropertyChangedEvent.Property->GetFName()
+		: NAME_None;
 
 	if (Name == FName(TEXT("bAllowTranslateRotateZWidget")))
 	{
@@ -202,7 +203,9 @@ void ULevelEditorViewportSettings::PostEditChangeProperty( struct FPropertyChang
 	}
 	else if (Name == FName(TEXT("bHighlightWithBrackets")))
 	{
-		GEngine->SetSelectedMaterialColor(bHighlightWithBrackets ? FLinearColor::Black : GetDefault<UEditorStyleSettings>()->SelectionColor);
+		GEngine->SetSelectedMaterialColor(bHighlightWithBrackets
+			? FLinearColor::Black
+			: GetDefault<UEditorStyleSettings>()->SelectionColor);
 	}
 	else if (Name == FName(TEXT("HoverHighlightIntensity")))
 	{
@@ -219,14 +222,15 @@ void ULevelEditorViewportSettings::PostEditChangeProperty( struct FPropertyChang
 	else if ((Name == FName(TEXT("UserDefinedPosGridSizes"))) || (Name == FName(TEXT("UserDefinedRotGridSizes"))) || (Name == FName(TEXT("ScalingGridSizes"))) || (Name == FName(TEXT("GridIntervals"))))
 	{
 		const float MinGridSize = (Name == FName(TEXT("GridIntervals"))) ? 4.0f : 0.0001f;
-		TArray<float>* ArrayRef = NULL;
-		int32* IndexRef = NULL;
+		TArray<float>* ArrayRef = nullptr;
+		int32* IndexRef = nullptr;
 
 		if (Name == FName(TEXT("ScalingGridSizes")))
 		{
 			ArrayRef = &(ScalingGridSizes);
 			IndexRef = &(CurrentScalingGridSize);
 		}
+
 		// Don't allow an empty array of grid sizes
 		if (ArrayRef->Num() == 0)
 		{
@@ -242,12 +246,12 @@ void ULevelEditorViewportSettings::PostEditChangeProperty( struct FPropertyChang
 			}
 		}
 	}
-	else if( Name == FName(TEXT("bUsePowerOf2SnapSize")) )
+	else if (Name == FName(TEXT("bUsePowerOf2SnapSize")))
 	{
 		const float BSPSnapSize = bUsePowerOf2SnapSize ? 128.0f : 100.0f;
-		UModel::SetGlobalBSPTexelScale( BSPSnapSize );
+		UModel::SetGlobalBSPTexelScale(BSPSnapSize);
 	}
-	else if( Name == FName(TEXT("BillboardScale")))
+	else if (Name == FName(TEXT("BillboardScale")))
 	{
 		UBillboardComponent::SetEditorScale(BillboardScale);
 		UArrowComponent::SetEditorScale(BillboardScale);
@@ -280,7 +284,9 @@ void UProjectPackagingSettings::PostEditChangeProperty( FPropertyChangedEvent& P
 {
 	Super::PostEditChangeProperty(PropertyChangedEvent);
 
-	const FName Name = (PropertyChangedEvent.Property != nullptr) ? PropertyChangedEvent.Property->GetFName() : NAME_None;
+	const FName Name = (PropertyChangedEvent.Property != nullptr)
+		? PropertyChangedEvent.Property->GetFName()
+		: NAME_None;
 
 	if (Name == FName((TEXT("DirectoriesToAlwaysCook"))))
 	{
