fix for device list in Launch On menu
#ue4
#editor

diff --git a/Engine/Source/Editor/UnrealEd/Private/Kismet2/DebuggerCommands.cpp b/Engine/Source/Editor/UnrealEd/Private/Kismet2/DebuggerCommands.cpp
--- a/Engine/Source/Editor/UnrealEd/Private/Kismet2/DebuggerCommands.cpp
+++ b/Engine/Source/Editor/UnrealEd/Private/Kismet2/DebuggerCommands.cpp
@@ -562,13 +562,12 @@ TSharedRef< SWidget > FPlayWorldCommands::GenerateLaunchMenuContent( TSharedRef<
 	const bool bShouldCloseWindowAfterMenuSelection = true;
 	FMenuBuilder MenuBuilder(bShouldCloseWindowAfterMenuSelection, InCommandList);
 
-/*	TArray<ITargetPlatform*> Platforms = GetTargetPlatformManager()->GetTargetPlatforms();
+	TArray<PlatformInfo::FVanillaPlatformEntry> VanillaPlatforms = PlatformInfo::BuildPlatformHierarchy(PlatformInfo::EPlatformFilter::All);
 
-	// Sort the platforms by display name so that they appear in the same order as the other menus (like "Package Project" and "Cook Project")
-	Platforms.Sort([](ITargetPlatform& One, ITargetPlatform& Two) -> bool
+	VanillaPlatforms.Sort([](const PlatformInfo::FVanillaPlatformEntry& One, const PlatformInfo::FVanillaPlatformEntry& Two) -> bool
 	{
-		return One.DisplayName().CompareTo(Two.DisplayName()) < 0;
-	});*/
+		return One.PlatformInfo->DisplayName.CompareTo(Two.PlatformInfo->DisplayName) < 0;
+	});
 
 	// shared devices section
 	TSharedPtr<ITargetDeviceServicesModule> TargetDeviceServicesModule = StaticCastSharedPtr<ITargetDeviceServicesModule>(FModuleManager::Get().LoadModule(TEXT("TargetDeviceServices")));
@@ -578,131 +577,161 @@ TSharedRef< SWidget > FPlayWorldCommands::GenerateLaunchMenuContent( TSharedRef<
 	PlatformsToMaybeInstallLinksFor.Add(TEXT("Android"));
 	PlatformsToMaybeInstallLinksFor.Add(TEXT("IOS"));
 	PlatformsToMaybeInstallLinksFor.Add(TEXT("Linux"));
-	TArray<FString> PlatformsWithNoDevices;
+	TArray<FString> PlatformsToCheckFlavorsFor;
+	PlatformsToCheckFlavorsFor.Add(TEXT("Android"));
+	PlatformsToCheckFlavorsFor.Add(TEXT("IOS"));
+	TArray<FName> PlatformsWithNoDevices;
 	TArray<PlatformInfo::FPlatformInfo> PlatformsToAddInstallLinksFor;
 
 	MenuBuilder.BeginSection("LevelEditorLaunchDevices", LOCTEXT("LaunchButtonDevicesSection", "Devices"));
 	{
-		for(const PlatformInfo::FPlatformInfo& PlatformInfo : PlatformInfo::EnumeratePlatformInfoArray())
+		for (const PlatformInfo::FVanillaPlatformEntry& VanillaPlatform : VanillaPlatforms)
 		{
 			// for the Editor we are only interested in launching standalone games
-			if (PlatformInfo.PlatformType != PlatformInfo::EPlatformType::Game || !PlatformInfo.bEnabledForUse || (!PlatformInfo.bEnabledInBinary && FRocketSupport::IsRocket()))
+			if (VanillaPlatform.PlatformInfo->PlatformType != PlatformInfo::EPlatformType::Game || !VanillaPlatform.PlatformInfo->bEnabledForUse || (!VanillaPlatform.PlatformInfo->bEnabledInBinary && FRocketSupport::IsRocket()))
 			{
 				continue;
 			}
 
-			if (PlatformInfo.SDKStatus == PlatformInfo::EPlatformSDKStatus::Installed)
+			if (VanillaPlatform.PlatformInfo->SDKStatus == PlatformInfo::EPlatformSDKStatus::Installed)
 			{
-				ITargetPlatform* Platform = GetTargetPlatformManager()->FindTargetPlatform(PlatformInfo.VanillaPlatformName.ToString());
-
-				// for each platform...
-				TArray<ITargetDeviceProxyPtr> DeviceProxies;
-				TargetDeviceServicesModule->GetDeviceProxyManager()->GetProxies(Platform->PlatformName(), false, DeviceProxies);
-
-				// if this platform had no devices, but we want to show an extra option if not installed right
-				if (DeviceProxies.Num() == 0)
+				TArray<FString> PlatformFlavors;
+				if (VanillaPlatform.PlatformFlavors.Num() && PlatformsToCheckFlavorsFor.Find(VanillaPlatform.PlatformInfo->TargetPlatformName.ToString()) != INDEX_NONE)
 				{
-					if (PlatformsToMaybeInstallLinksFor.Find(PlatformInfo.VanillaPlatformName.ToString()) != INDEX_NONE && PlatformsWithNoDevices.Find(PlatformInfo.VanillaPlatformName.ToString()) == INDEX_NONE)
+					for (const PlatformInfo::FPlatformInfo* SubPlatformInfo : VanillaPlatform.PlatformFlavors)
 					{
-						// add an entry with a no devices found
-						PlatformsWithNoDevices.Add(PlatformInfo.VanillaPlatformName.ToString());
-
-						// ... generate display label...
-						FFormatNamedArguments LabelArguments;
-						LabelArguments.Add(TEXT("DisplayName"), PlatformInfo.DisplayName);
-
-						FText Label = FText::Format(LOCTEXT("NoDeviceLabel", "{DisplayName} - No Devices Found"), LabelArguments);
-
-						// ... create an action...
-						FUIAction NoDeviceAction(
-							FExecuteAction::CreateStatic(&FInternalPlayWorldCommandCallbacks::HandleNoDeviceFoundActionExecute),
-							FCanExecuteAction::CreateStatic(&FInternalPlayWorldCommandCallbacks::HandleNoDeviceFoundActionCanExecute)
-							);
-
-						// ... generate tooltip text
-						FFormatNamedArguments TooltipArguments;
-						TooltipArguments.Add(TEXT("DisplayName"), PlatformInfo.DisplayName);
-						FText Tooltip = FText::Format(LOCTEXT("LaunchDeviceToolTipText", "Found no connected devices for {DisplayName}"), TooltipArguments);
-
-						// ... and add a menu entry
-						MenuBuilder.AddMenuEntry(
-							NoDeviceAction, 
-							ProjectTargetPlatformEditorModule.MakePlatformMenuItemWidget(PlatformInfo, true, Label), 
-							NAME_None, 
-							Tooltip, 
-							EUserInterfaceActionType::Check
-							);
+						PlatformFlavors.Add(SubPlatformInfo->TargetPlatformName.ToString());
 					}
 				}
 				else
 				{
-					// for each proxy...
-					for (auto DeviceProxyIt = DeviceProxies.CreateIterator(); DeviceProxyIt; ++DeviceProxyIt)
-					{
-						ITargetDeviceProxyPtr DeviceProxy = *DeviceProxyIt;
-
-						// ... create an action...
-						FUIAction LaunchDeviceAction(
-							FExecuteAction::CreateStatic(&FInternalPlayWorldCommandCallbacks::HandleLaunchOnDeviceActionExecute, DeviceProxy->GetDeviceId(), DeviceProxy->GetName()),
-							FCanExecuteAction::CreateStatic(&FInternalPlayWorldCommandCallbacks::HandleLaunchOnDeviceActionCanExecute, DeviceProxy->GetDeviceId()),
-							FIsActionChecked::CreateStatic(&FInternalPlayWorldCommandCallbacks::HandleLaunchOnDeviceActionIsChecked, DeviceProxy->GetDeviceId())
-							);
+					PlatformFlavors.Add(VanillaPlatform.PlatformInfo->TargetPlatformName.ToString());
+				}
+				for (int32 PlatformIndex = 0; PlatformIndex < PlatformFlavors.Num(); ++PlatformIndex)
+				{
+					ITargetPlatform* Platform = GetTargetPlatformManager()->FindTargetPlatform(PlatformFlavors[PlatformIndex]);
 
-						// ... generate display label...
-						FFormatNamedArguments LabelArguments;
-						LabelArguments.Add(TEXT("DeviceName"), FText::FromString(DeviceProxy->GetName()));
+					// for each platform...
+					TArray<ITargetDeviceProxyPtr> DeviceProxies;
+					TargetDeviceServicesModule->GetDeviceProxyManager()->GetProxies(Platform->PlatformName(), false, DeviceProxies);
 
-						if (!DeviceProxy->IsConnected())
-						{
-							LabelArguments.Add(TEXT("HostUser"), LOCTEXT("DisconnectedHint", " [Disconnected]"));
-						}
-						else if (DeviceProxy->GetHostUser() != FPlatformProcess::UserName(false))
-						{
-							LabelArguments.Add(TEXT("HostUser"), FText::FromString(DeviceProxy->GetHostUser()));
-						}
-						else
+					// if this platform had no devices, but we want to show an extra option if not installed right
+					if (DeviceProxies.Num() == 0)
+					{
+						if (PlatformsToMaybeInstallLinksFor.Find(PlatformFlavors[PlatformIndex]) != INDEX_NONE && PlatformsWithNoDevices.Find(VanillaPlatform.PlatformInfo->VanillaPlatformName) == INDEX_NONE)
 						{
-							LabelArguments.Add(TEXT("HostUser"), FText::GetEmpty());
+							// add an entry with a no devices found
+							PlatformsWithNoDevices.Add(VanillaPlatform.PlatformInfo->VanillaPlatformName);
 						}
-
-						FText Label = FText::Format(LOCTEXT("LaunchDeviceLabel", "{DeviceName}{HostUser}"), LabelArguments);
-
-						// ... generate tooltip text
-						FFormatNamedArguments TooltipArguments;
-						TooltipArguments.Add(TEXT("DeviceID"), FText::FromString(DeviceProxy->GetName()));
-						TooltipArguments.Add(TEXT("DisplayName"), PlatformInfo.DisplayName);
-						FText Tooltip = FText::Format(LOCTEXT("LaunchDeviceToolTipText", "Launch the game on this {DisplayName} device ({DeviceID})"), TooltipArguments);
-
-						FProjectStatus ProjectStatus;
-						if (IProjectManager::Get().QueryStatusForCurrentProject(ProjectStatus) && !ProjectStatus.IsTargetPlatformSupported(PlatformInfo.VanillaPlatformName)) 
+					}
+					else
+					{
+						// for each proxy...
+						for (auto DeviceProxyIt = DeviceProxies.CreateIterator(); DeviceProxyIt; ++DeviceProxyIt)
 						{
-							FText TooltipLine2 = FText::Format(LOCTEXT("LaunchDevicePlatformWarning", "{DisplayName} is not listed as a target platform for this project, so may not run as expected."), TooltipArguments);
-							Tooltip = FText::Format(FText::FromString(TEXT("{0}\n\n{1}")), Tooltip, TooltipLine2);
+							ITargetDeviceProxyPtr DeviceProxy = *DeviceProxyIt;
+
+							// ... create an action...
+							FUIAction LaunchDeviceAction(
+								FExecuteAction::CreateStatic(&FInternalPlayWorldCommandCallbacks::HandleLaunchOnDeviceActionExecute, DeviceProxy->GetDeviceId(), DeviceProxy->GetName()),
+								FCanExecuteAction::CreateStatic(&FInternalPlayWorldCommandCallbacks::HandleLaunchOnDeviceActionCanExecute, DeviceProxy->GetDeviceId()),
+								FIsActionChecked::CreateStatic(&FInternalPlayWorldCommandCallbacks::HandleLaunchOnDeviceActionIsChecked, DeviceProxy->GetDeviceId())
+								);
+
+							// ... generate display label...
+							FFormatNamedArguments LabelArguments;
+							LabelArguments.Add(TEXT("DeviceName"), FText::FromString(DeviceProxy->GetName()));
+
+							if (!DeviceProxy->IsConnected())
+							{
+								LabelArguments.Add(TEXT("HostUser"), LOCTEXT("DisconnectedHint", " [Disconnected]"));
+							}
+							else if (DeviceProxy->GetHostUser() != FPlatformProcess::UserName(false))
+							{
+								LabelArguments.Add(TEXT("HostUser"), FText::FromString(DeviceProxy->GetHostUser()));
+							}
+							else
+							{
+								LabelArguments.Add(TEXT("HostUser"), FText::GetEmpty());
+							}
+
+							FText Label = FText::Format(LOCTEXT("LaunchDeviceLabel", "{DeviceName}{HostUser}"), LabelArguments);
+
+							// ... generate tooltip text
+							FFormatNamedArguments TooltipArguments;
+							TooltipArguments.Add(TEXT("DeviceID"), FText::FromString(DeviceProxy->GetName()));
+							TooltipArguments.Add(TEXT("DisplayName"), VanillaPlatform.PlatformInfo->DisplayName);
+							FText Tooltip = FText::Format(LOCTEXT("LaunchDeviceToolTipText", "Launch the game on this {DisplayName} device ({DeviceID})"), TooltipArguments);
+
+							FProjectStatus ProjectStatus;
+							if (IProjectManager::Get().QueryStatusForCurrentProject(ProjectStatus) && !ProjectStatus.IsTargetPlatformSupported(VanillaPlatform.PlatformInfo->VanillaPlatformName)) 
+							{
+								FText TooltipLine2 = FText::Format(LOCTEXT("LaunchDevicePlatformWarning", "{DisplayName} is not listed as a target platform for this project, so may not run as expected."), TooltipArguments);
+								Tooltip = FText::Format(FText::FromString(TEXT("{0}\n\n{1}")), Tooltip, TooltipLine2);
+							}
+
+							// ... and add a menu entry
+							MenuBuilder.AddMenuEntry(
+								LaunchDeviceAction, 
+								ProjectTargetPlatformEditorModule.MakePlatformMenuItemWidget(Platform->GetPlatformInfo(), true, Label), 
+								NAME_None, 
+								Tooltip, 
+								EUserInterfaceActionType::Check
+								);
 						}
-
-						// ... and add a menu entry
-						MenuBuilder.AddMenuEntry(
-							LaunchDeviceAction, 
-							ProjectTargetPlatformEditorModule.MakePlatformMenuItemWidget(PlatformInfo, true, Label), 
-							NAME_None, 
-							Tooltip, 
-							EUserInterfaceActionType::Check
-							);
 					}
 				}
 			}
 			else
 			{
 				// if the platform wasn't installed, we'll add a menu item later (we never care about code in this case, since we don't compile)
-				if (PlatformsToMaybeInstallLinksFor.Find(PlatformInfo.VanillaPlatformName.ToString()) != INDEX_NONE)
+				if (PlatformsToMaybeInstallLinksFor.Find(VanillaPlatform.PlatformInfo->VanillaPlatformName.ToString()) != INDEX_NONE)
 				{
-					PlatformsToAddInstallLinksFor.Add(PlatformInfo);
+					PlatformsToAddInstallLinksFor.Add(*(VanillaPlatform.PlatformInfo));
 				}
 			}
 		}
 	}
 	MenuBuilder.EndSection();
 
+	if (PlatformsWithNoDevices.Num() > 0)
+	{
+		MenuBuilder.BeginSection("NoDevices");
+		{
+			for (int32 PlatformIndex = 0; PlatformIndex < PlatformsWithNoDevices.Num(); PlatformIndex++)
+			{
+				const PlatformInfo::FPlatformInfo* PlatformInfo = PlatformInfo::FindVanillaPlatformInfo(PlatformsWithNoDevices[PlatformIndex]);
+
+				// ... generate display label...
+				FFormatNamedArguments LabelArguments;
+				LabelArguments.Add(TEXT("DisplayName"), PlatformInfo->DisplayName);
+
+				FText Label = FText::Format(LOCTEXT("NoDeviceLabel", "{DisplayName} - No Devices Found"), LabelArguments);
+
+				// ... create an action...
+				FUIAction NoDeviceAction(
+					FExecuteAction::CreateStatic(&FInternalPlayWorldCommandCallbacks::HandleNoDeviceFoundActionExecute),
+					FCanExecuteAction::CreateStatic(&FInternalPlayWorldCommandCallbacks::HandleNoDeviceFoundActionCanExecute)
+					);
+
+				// ... generate tooltip text
+				FFormatNamedArguments TooltipArguments;
+				TooltipArguments.Add(TEXT("DisplayName"), PlatformInfo->DisplayName);
+				FText Tooltip = FText::Format(LOCTEXT("LaunchDeviceToolTipText", "Found no connected devices for {DisplayName}"), TooltipArguments);
+
+				// ... and add a menu entry
+				MenuBuilder.AddMenuEntry(
+					NoDeviceAction, 
+					ProjectTargetPlatformEditorModule.MakePlatformMenuItemWidget(*PlatformInfo, true, Label), 
+					NAME_None, 
+					Tooltip, 
+					EUserInterfaceActionType::Check
+					);
+			}
+		}
+		MenuBuilder.EndSection();
+	}
+
 	// tip section
 	MenuBuilder.BeginSection("LevelEditorLaunchHint");
 	{
