Removed code related to ttp237118 "ROCKET: HUGE: K2: Change 'Class' to be primary asset (seen in content browser etc) instead of Blueprint"
The "_C" postif should not be visible in UI and (and gameplay code).

diff --git a/Engine/Source/Developer/AssetTools/Private/AssetTools.cpp b/Engine/Source/Developer/AssetTools/Private/AssetTools.cpp
--- a/Engine/Source/Developer/AssetTools/Private/AssetTools.cpp
+++ b/Engine/Source/Developer/AssetTools/Private/AssetTools.cpp
@@ -31,7 +31,6 @@ FAssetTools::FAssetTools()
 	RegisterAssetTypeActions( MakeShareable(new FAssetTypeActions_BlendSpace) );
 	RegisterAssetTypeActions( MakeShareable(new FAssetTypeActions_BlendSpace1D) );
 	RegisterAssetTypeActions( MakeShareable(new FAssetTypeActions_Blueprint) );
-	RegisterAssetTypeActions( MakeShareable(new FAssetTypeActions_BlueprintGeneratedClass) );
 	RegisterAssetTypeActions( MakeShareable(new FAssetTypeActions_CameraAnim) );
 	RegisterAssetTypeActions( MakeShareable(new FAssetTypeActions_Curve) );
 	RegisterAssetTypeActions( MakeShareable(new FAssetTypeActions_CurveFloat) );
diff --git a/Engine/Source/Developer/AssetTools/Private/AssetToolsPrivatePCH.h b/Engine/Source/Developer/AssetTools/Private/AssetToolsPrivatePCH.h
--- a/Engine/Source/Developer/AssetTools/Private/AssetToolsPrivatePCH.h
+++ b/Engine/Source/Developer/AssetTools/Private/AssetToolsPrivatePCH.h
@@ -15,7 +15,6 @@ DECLARE_LOG_CATEGORY_EXTERN(LogAssetTools, Log, All);
 #include "AssetTypeActions_CSVAssetBase.h"
 
 #include "AssetTypeActions/AssetTypeActions_Blueprint.h"
-#include "AssetTypeActions/AssetTypeActions_BlueprintGeneratedClass.h"
 #include "AssetTypeActions/AssetTypeActions_Curve.h"
 #include "AssetTypeActions/AssetTypeActions_MaterialInterface.h"
 #include "AssetTypeActions/AssetTypeActions_SkeletalMesh.h"
diff --git a/Engine/Source/Developer/AssetTools/Private/AssetTypeActions/AssetTypeActions_BlueprintGeneratedClass.cpp b/Engine/Source/Developer/AssetTools/Private/AssetTypeActions/AssetTypeActions_BlueprintGeneratedClass.cpp
deleted file mode 100644
diff --git a/Engine/Source/Developer/AssetTools/Private/AssetTypeActions/AssetTypeActions_BlueprintGeneratedClass.h b/Engine/Source/Developer/AssetTools/Private/AssetTypeActions/AssetTypeActions_BlueprintGeneratedClass.h
deleted file mode 100644
diff --git a/Engine/Source/Editor/BlueprintGraph/Classes/EdGraphSchema_K2.h b/Engine/Source/Editor/BlueprintGraph/Classes/EdGraphSchema_K2.h
--- a/Engine/Source/Editor/BlueprintGraph/Classes/EdGraphSchema_K2.h
+++ b/Engine/Source/Editor/BlueprintGraph/Classes/EdGraphSchema_K2.h
@@ -418,32 +418,7 @@ public:
 		}
 		
 		/** Returns a succinct menu description of this type */
-		FString GetDescription() const
-		{
-			if ((PinType.PinCategory != FriendlyName) && !FriendlyName.IsEmpty())
-			{
-				return FriendlyName;
-			}
-			else if( PinType.PinSubCategoryObject.IsValid() )
-			{
-				FString DisplayName = PinType.PinSubCategoryObject->GetName();
-				//@todo:  Fix this once the XX_YYYY names in the schema are static!  This is mirrored to PC_Class
-				if( (PinType.PinCategory == TEXT("class")) && PinType.PinSubCategoryObject->IsA(UClass::StaticClass()) )
-				{
-					DisplayName = FString::Printf(TEXT("class'%s'"), *DisplayName);
-				}
-
-				return DisplayName;
-			}
-			else if( !PinType.PinCategory.IsEmpty() )
-			{
-				return PinType.PinCategory;
-			}
-			else
-			{
-				return TEXT("Error!");
-			}
-		}
+		FString GetDescription() const;
 
 		FString GetToolTip() const
 		{
diff --git a/Engine/Source/Editor/BlueprintGraph/Private/EdGraphSchema_K2.cpp b/Engine/Source/Editor/BlueprintGraph/Private/EdGraphSchema_K2.cpp
--- a/Engine/Source/Editor/BlueprintGraph/Private/EdGraphSchema_K2.cpp
+++ b/Engine/Source/Editor/BlueprintGraph/Private/EdGraphSchema_K2.cpp
@@ -217,6 +217,39 @@ UEdGraphSchema_K2::FPinTypeTreeInfo::FPinTypeTreeInfo(const FString& CategoryNam
 	bReadOnly = bInReadOnly;
 }
 
+FString UEdGraphSchema_K2::FPinTypeTreeInfo::GetDescription() const
+{
+	if ((PinType.PinCategory != FriendlyName) && !FriendlyName.IsEmpty())
+	{
+		return FriendlyName;
+	}
+	else if (PinType.PinSubCategoryObject.IsValid())
+	{
+		FString DisplayName = PinType.PinSubCategoryObject->GetName();
+		const auto SubCategoryClass = Cast<UClass>(PinType.PinSubCategoryObject.Get());
+		if (SubCategoryClass && !SubCategoryClass->HasAnyClassFlags(CLASS_Native))
+		{
+			DisplayName.RemoveFromEnd(TEXT("_C"));
+		}
+
+		//@todo:  Fix this once the XX_YYYY names in the schema are static!  This is mirrored to PC_Class
+		if ((PinType.PinCategory == TEXT("class")) && SubCategoryClass)
+		{
+			DisplayName = FString::Printf(TEXT("class'%s'"), *DisplayName);
+		}
+
+		return DisplayName;
+	}
+	else if (!PinType.PinCategory.IsEmpty())
+	{
+		return PinType.PinCategory;
+	}
+	else
+	{
+		return TEXT("Error!");
+	}
+}
+
 UEdGraphSchema_K2::UEdGraphSchema_K2(const class FPostConstructInitializeProperties& PCIP)
 	: Super(PCIP)
 {
diff --git a/Engine/Source/Editor/UnrealEd/Classes/Factories/BlueprintGeneratedClassFactory.h b/Engine/Source/Editor/UnrealEd/Classes/Factories/BlueprintGeneratedClassFactory.h
deleted file mode 100644
diff --git a/Engine/Source/Editor/UnrealEd/Private/Factories/EditorFactories.cpp b/Engine/Source/Editor/UnrealEd/Private/Factories/EditorFactories.cpp
--- a/Engine/Source/Editor/UnrealEd/Private/Factories/EditorFactories.cpp
+++ b/Engine/Source/Editor/UnrealEd/Private/Factories/EditorFactories.cpp
@@ -6044,90 +6044,6 @@ public:
 	}
 };
 
-/*------------------------------------------------------------------------------
-	UBlueprintGeneratedClassFactory implementation.
-------------------------------------------------------------------------------*/
-UBlueprintGeneratedClassFactory::UBlueprintGeneratedClassFactory(const class FPostConstructInitializeProperties& PCIP)
-	: Super(PCIP)
-{
-
-	// Look in the config file to determine what the default base class is, if any
-	FString ClassPath;
-	GConfig->GetString(TEXT("/Script/Engine.Engine"), TEXT("DefaultBlueprintBaseClassName"), /*out*/ ClassPath, GEngineIni);
-	UClass* DefaultParentClass = !ClassPath.IsEmpty() 
-		? LoadClass<UObject>(nullptr, *ClassPath, nullptr, LOAD_None, nullptr) 
-		: nullptr;
-	
-	if( !DefaultParentClass || !FKismetEditorUtilities::CanCreateBlueprintOfClass(DefaultParentClass) )
-	{
-		DefaultParentClass = AActor::StaticClass();
-	}
-
-	bCreateNew = false;
-	bEditAfterNew = true;
-	SupportedClass = UBlueprintGeneratedClass::StaticClass();
-	ParentClass = DefaultParentClass;
-}
-
-bool UBlueprintGeneratedClassFactory::ConfigureProperties()
-{
-	// Null the parent class to ensure one is selected
-	ParentClass = nullptr;
-
-	// Fill in options
-	FClassViewerInitializationOptions Options;
-	Options.Mode = EClassViewerMode::ClassPicker;
-	Options.DisplayMode = EClassViewerDisplayMode::TreeView;
-	Options.bShowObjectRootClass = true;
-
-	// Only want blueprint actor base classes.
-	Options.bIsBlueprintBaseOnly = true;
-
-	// This will allow unloaded blueprints to be shown.
-	Options.bShowUnloadedBlueprints = true;
-
-	// Prevent creating blueprints of classes that require special setup (they'll be allowed in the corresponding factories / via other means)
-	TSharedPtr<FBlueprintParentFilter> Filter = MakeShareable(new FBlueprintParentFilter);
-	Options.ClassFilter = Filter;
-	if (!IsMacroFactory())
-	{
-		Filter->DisallowedChildrenOfClasses.Add(ALevelScriptActor::StaticClass());
-		Filter->DisallowedChildrenOfClasses.Add(UAnimInstance::StaticClass());
-	}
-
-	// Filter out interfaces in all cases; they can never contain code, so it doesn't make sense to use them as a macro basis
-	Filter->DisallowedChildrenOfClasses.Add(UInterface::StaticClass());
-
-	const FText TitleText = LOCTEXT("CreateBlueprintOptions", "Pick Parent Class");
-	UClass* ChosenClass = nullptr;
-	const bool bPressedOk = SClassPickerDialog::PickClass(TitleText, Options, ChosenClass, UBlueprint::StaticClass());
-
-	if ( bPressedOk )
-	{
-		ParentClass = ChosenClass;
-
-		FEditorDelegates::OnFinishPickingBlueprintClass.Broadcast(ParentClass);
-	}
-
-	return bPressedOk;
-};
-
-UObject* UBlueprintGeneratedClassFactory::FactoryCreateNew(UClass* Class, UObject* InParent, FName Name, EObjectFlags Flags, UObject* Context, FFeedbackContext* Warn)
-{
-	// Make sure we are trying to factory a blueprint, then create and init one
-	check(Class->IsChildOf(UBlueprintGeneratedClass::StaticClass()));
-
-	if ((ParentClass == nullptr) || !FKismetEditorUtilities::CanCreateBlueprintOfClass(ParentClass))
-	{
-		FMessageDialog::Open( EAppMsgType::Ok, FText::FromString(FString::Printf(TEXT("Cannot create a blueprint based on the class '%s'."),
-			(ParentClass != nullptr) ? *(ParentClass->GetName()) : TEXT("(null)"))));
-		return nullptr;
-	}
-	else
-	{
-		return FKismetEditorUtilities::CreateBlueprint(ParentClass, InParent, Name, BPTYPE_Normal, UBlueprint::StaticClass(), UBlueprintGeneratedClass::StaticClass());
-	}
-}
 
 /*------------------------------------------------------------------------------
 	UBlueprintFactory implementation.
