UMG - UMG uses a (not the) transient package to perform copy and paste, in order to prevent saving this package, we now ignore packages with the RF_Transient flag.

diff --git a/Engine/Source/Editor/UnrealEd/Private/FileHelpers.cpp b/Engine/Source/Editor/UnrealEd/Private/FileHelpers.cpp
--- a/Engine/Source/Editor/UnrealEd/Private/FileHelpers.cpp
+++ b/Engine/Source/Editor/UnrealEd/Private/FileHelpers.cpp
@@ -2986,6 +2986,8 @@ void FEditorFileUtils::GetDirtyContentPackages(TArray<UPackage*>& OutDirtyPackag
 		bShouldIgnorePackage |= Package->GetOuter() != NULL;
 		// Don't try to save "Transient" package.
 		bShouldIgnorePackage |= Package == GetTransientPackage();
+		// Don't try to save packages with the RF_Transient flag.
+		bShouldIgnorePackage |= Package->HasAnyFlags(RF_Transient);
 		// Ignore PIE packages.
 		bShouldIgnorePackage |= (Package->PackageFlags & PKG_PlayInEditor) != 0;
 		// Ignore packages that haven't been modified.
