Implement FMacPlatformMisc::GetOSVersions & use it instead of Cocoa when checking the Mac OS X version. Will be esp. useful when we start needing to support more than one major version of OS X.

diff --git a/Engine/Source/Editor/UnrealEd/Private/Editor.cpp b/Engine/Source/Editor/UnrealEd/Private/Editor.cpp
--- a/Engine/Source/Editor/UnrealEd/Private/Editor.cpp
+++ b/Engine/Source/Editor/UnrealEd/Private/Editor.cpp
@@ -6488,15 +6488,16 @@ void UEditorEngine::UpdateAutoLoadProject()
 #if PLATFORM_MAC
 	if ( !GIsBuildMachine )
 	{
-		SCOPED_AUTORELEASE_POOL;
-		NSDictionary* SystemVersion = [NSDictionary dictionaryWithContentsOfFile: @"/System/Library/CoreServices/SystemVersion.plist"];
-		NSString* OSVersion = (NSString*)[SystemVersion objectForKey: @"ProductVersion"];
-		NSArray* Components = [OSVersion componentsSeparatedByString:@"."];
-		NSInteger ComponentValues[3] = {0};
-		for(uint32 i = 0; i < [Components count] && i < 3; i++)
-		{
-			NSString* Component = [Components objectAtIndex:i];
-			ComponentValues[i] = [Component integerValue];
+		FString OSVersion, OSSubVersion;
+		FPlatformMisc::GetOSVersions(OSVersion, OSSubVersion);
+		
+		TArray<FString> Components;
+		OSVersion.ParseIntoArray(&Components, TEXT("."), true);
+		uint8 ComponentValues[3] = {0};
+		
+		for(uint32 i = 0; i < Components.Num() && i < 3; i++)
+		{
+			TTypeFromString<uint8>::FromString(ComponentValues[i], *Components[i]);
 		}
 		
 		if(ComponentValues[0] < 10 || ComponentValues[1] < 9 || (ComponentValues[1] == 9 && ComponentValues[2] < 4))
diff --git a/Engine/Source/Runtime/Core/Private/Mac/MacPlatformMisc.cpp b/Engine/Source/Runtime/Core/Private/Mac/MacPlatformMisc.cpp
--- a/Engine/Source/Runtime/Core/Private/Mac/MacPlatformMisc.cpp
+++ b/Engine/Source/Runtime/Core/Private/Mac/MacPlatformMisc.cpp
@@ -888,6 +888,12 @@ void FMacPlatformMisc::NormalizePath(FString& InPath)
 	}
 }
 
+void FMacPlatformMisc::GetOSVersions( FString& out_OSVersionLabel, FString& out_OSSubVersionLabel )
+{
+	out_OSVersionLabel = GMacAppInfo.OSVersion;
+	out_OSSubVersionLabel = GMacAppInfo.OSBuild;
+}
+
 #include "ModuleManager.h"
 
 void FMacPlatformMisc::LoadPreInitModules()
diff --git a/Engine/Source/Runtime/Core/Public/Mac/MacPlatformMisc.h b/Engine/Source/Runtime/Core/Public/Mac/MacPlatformMisc.h
--- a/Engine/Source/Runtime/Core/Public/Mac/MacPlatformMisc.h
+++ b/Engine/Source/Runtime/Core/Public/Mac/MacPlatformMisc.h
@@ -69,6 +69,7 @@ struct CORE_API FMacPlatformMisc : public FGenericPlatformMisc
 	static int32 NumberOfCoresIncludingHyperthreads();
 	static void LoadPreInitModules();
 	static void NormalizePath(FString& InPath);
+	static void GetOSVersions( FString& out_OSVersionLabel, FString& out_OSSubVersionLabel );
 
 
 	/** 
