Fixed an issue with passing a gather text configuration file path featuring spaces when running the GatherText commandlet.

diff --git a/Engine/Source/Editor/UnrealEd/Private/Commandlets/GatherTextCommandlet.cpp b/Engine/Source/Editor/UnrealEd/Private/Commandlets/GatherTextCommandlet.cpp
--- a/Engine/Source/Editor/UnrealEd/Private/Commandlets/GatherTextCommandlet.cpp
+++ b/Engine/Source/Editor/UnrealEd/Private/Commandlets/GatherTextCommandlet.cpp
@@ -107,7 +107,7 @@ int32 UGatherTextCommandlet::Main( const FString& Params )
 		UE_LOG(LogGatherTextCommandlet, Log,TEXT("Executing %s: %s"), *SectionName, *CommandletClassName);
 		
 		
-		FString GeneratedCmdLine = FString::Printf(TEXT("-Config=%s -Section=%s"), *GatherTextConfigPath , *SectionName);
+		FString GeneratedCmdLine = FString::Printf(TEXT("-Config=\"%s\" -Section=%s"), *GatherTextConfigPath , *SectionName);
 
 		// Add all the command params with the exception of config
 		for(auto ParamIter = ParamVals.CreateConstIterator(); ParamIter; ++ParamIter)
