Check that Slate is initialised before trying to open the 'Update Mac OS X' dialog when running beneat the minimum version. When Slate is not initialised, just output the message as a warning to the log.

diff --git a/Engine/Source/Editor/UnrealEd/Private/Editor.cpp b/Engine/Source/Editor/UnrealEd/Private/Editor.cpp
--- a/Engine/Source/Editor/UnrealEd/Private/Editor.cpp
+++ b/Engine/Source/Editor/UnrealEd/Private/Editor.cpp
@@ -6502,11 +6502,18 @@ void UEditorEngine::UpdateAutoLoadProject()
 		
 		if(ComponentValues[0] < 10 || ComponentValues[1] < 9 || (ComponentValues[1] == 9 && ComponentValues[2] < 4))
 		{
-			FSuppressableWarningDialog::FSetupInfo Info( LOCTEXT("UpdateMacOSX_Body","Please update to the latest version of Mac OS X for best performance."), LOCTEXT("UpdateMacOSX_Title","Update Mac OS X"), TEXT("UpdateMacOSX"), GEditorGameAgnosticIni );
-			Info.ConfirmText = LOCTEXT( "OK", "OK");
-			Info.bDefaultToSupressInTheFuture = true;
-			FSuppressableWarningDialog OSUpdateWarning( Info );
-			OSUpdateWarning.ShowModal();
+			if(FSlateApplication::IsInitialized())
+			{
+				FSuppressableWarningDialog::FSetupInfo Info( LOCTEXT("UpdateMacOSX_Body","Please update to the latest version of Mac OS X for best performance."), LOCTEXT("UpdateMacOSX_Title","Update Mac OS X"), TEXT("UpdateMacOSX"), GEditorGameAgnosticIni );
+				Info.ConfirmText = LOCTEXT( "OK", "OK");
+				Info.bDefaultToSupressInTheFuture = true;
+				FSuppressableWarningDialog OSUpdateWarning( Info );
+				OSUpdateWarning.ShowModal();
+			}
+			else
+			{
+				UE_LOG(LogEditor, Warning, TEXT("Please update to the latest version of Mac OS X for best performance."));
+			}
 		}
 	}
 #endif
