Header cleanup: Engine module, components headers.

diff --git a/Engine/Source/Editor/UMGEditor/Private/WidgetBlueprintEditor.h b/Engine/Source/Editor/UMGEditor/Private/WidgetBlueprintEditor.h
--- a/Engine/Source/Editor/UMGEditor/Private/WidgetBlueprintEditor.h
+++ b/Engine/Source/Editor/UMGEditor/Private/WidgetBlueprintEditor.h
@@ -3,7 +3,7 @@
 #pragma once
 
 #include "Editor/Kismet/Public/BlueprintEditor.h"
-
+#include "ISequencer.h"
 #include "PreviewScene.h"
 
 class ISequencer;
diff --git a/Engine/Source/Editor/UnrealEd/Classes/Factories/TextureFactory.h b/Engine/Source/Editor/UnrealEd/Classes/Factories/TextureFactory.h
--- a/Engine/Source/Editor/UnrealEd/Classes/Factories/TextureFactory.h
+++ b/Engine/Source/Editor/UnrealEd/Classes/Factories/TextureFactory.h
@@ -1,7 +1,8 @@
 // Copyright 1998-2014 Epic Games, Inc. All Rights Reserved.
 
 #pragma once
-
+#include "Engine/Texture.h"
+#include "Factories/Factory.h"
 #include "TextureFactory.generated.h"
 
 UCLASS(customconstructor, collapsecategories, hidecategories=Object)
diff --git a/Engine/Source/Editor/UnrealEd/Classes/MaterialEditor/MaterialEditorInstanceConstant.h b/Engine/Source/Editor/UnrealEd/Classes/MaterialEditor/MaterialEditorInstanceConstant.h
--- a/Engine/Source/Editor/UnrealEd/Classes/MaterialEditor/MaterialEditorInstanceConstant.h
+++ b/Engine/Source/Editor/UnrealEd/Classes/MaterialEditor/MaterialEditorInstanceConstant.h
@@ -5,11 +5,9 @@
  */
 
 #pragma once
-
 #include "Editor/UnrealEdTypes.h"
-
 #include "Materials/MaterialInstanceBasePropertyOverrides.h"
-
+#include "StaticParameterSet.h"
 #include "MaterialEditorInstanceConstant.generated.h"
 
 USTRUCT()
diff --git a/Engine/Source/Editor/UnrealEd/Classes/MaterialGraph/MaterialGraphNode_Base.h b/Engine/Source/Editor/UnrealEd/Classes/MaterialGraph/MaterialGraphNode_Base.h
--- a/Engine/Source/Editor/UnrealEd/Classes/MaterialGraph/MaterialGraphNode_Base.h
+++ b/Engine/Source/Editor/UnrealEd/Classes/MaterialGraph/MaterialGraphNode_Base.h
@@ -1,7 +1,7 @@
 // Copyright 1998-2014 Epic Games, Inc. All Rights Reserved.
 
 #pragma once
-
+#include "EdGraph/EdGraphNode.h"
 #include "MaterialGraphNode_Base.generated.h"
 
 UCLASS(MinimalAPI)
@@ -29,7 +29,7 @@ class UMaterialGraphNode_Base : public UEdGraphNode
 	/** Get the Material Expression input index from an input pin */
 	virtual int32 GetInputIndex(const UEdGraphPin* InputPin) const {return -1;}
 	/** Get the Material value type of an input pin */
-	virtual uint32 GetInputType(const UEdGraphPin* InputPin) const {return MCT_Unknown;}
+	virtual uint32 GetInputType(const UEdGraphPin* InputPin) const;
 
 	/**
 	 * Handles inserting the node between the FromPin and what the FromPin was original connected to
diff --git a/Engine/Source/Editor/UnrealEd/Private/MaterialGraphNode_Base.cpp b/Engine/Source/Editor/UnrealEd/Private/MaterialGraphNode_Base.cpp
--- a/Engine/Source/Editor/UnrealEd/Private/MaterialGraphNode_Base.cpp
+++ b/Engine/Source/Editor/UnrealEd/Private/MaterialGraphNode_Base.cpp
@@ -5,6 +5,7 @@
 =============================================================================*/
 
 #include "UnrealEd.h"
+#include "MaterialShared.h"
 
 /////////////////////////////////////////////////////
 // UMaterialGraphNode_Base
@@ -320,3 +321,8 @@ FString UMaterialGraphNode_Base::GetDocumentationLink() const
 {
 	return TEXT("Shared/GraphNodes/Material");
 }
+
+uint32 UMaterialGraphNode_Base::GetInputType(const UEdGraphPin* InputPin) const
+{
+	return MCT_Unknown;
+}
diff --git a/Engine/Source/Runtime/Engine/Classes/AI/Navigation/NavLinkCustomComponent.h b/Engine/Source/Runtime/Engine/Classes/AI/Navigation/NavLinkCustomComponent.h
--- a/Engine/Source/Runtime/Engine/Classes/AI/Navigation/NavLinkCustomComponent.h
+++ b/Engine/Source/Runtime/Engine/Classes/AI/Navigation/NavLinkCustomComponent.h
@@ -4,6 +4,7 @@
 #include "AI/Navigation/NavLinkCustomInterface.h"
 #include "AI/Navigation/NavLinkDefinition.h"
 #include "AI/Navigation/NavRelevantComponent.h"
+#include "Engine/EngineTypes.h"
 #include "NavLinkCustomComponent.generated.h"
 
 /**
diff --git a/Engine/Source/Runtime/Engine/Classes/Animation/AnimInstance.h b/Engine/Source/Runtime/Engine/Classes/Animation/AnimInstance.h
--- a/Engine/Source/Runtime/Engine/Classes/Animation/AnimInstance.h
+++ b/Engine/Source/Runtime/Engine/Classes/Animation/AnimInstance.h
@@ -5,6 +5,7 @@
 #include "AnimationAsset.h"
 #include "AnimSequence.h"
 #include "Components/SkeletalMeshComponent.h"
+#include "Components/SkinnedMeshComponent.h"
 #include "AnimInstance.generated.h"
 
 struct FAnimMontageInstance;
diff --git a/Engine/Source/Runtime/Engine/Classes/Components/ActorComponent.h b/Engine/Source/Runtime/Engine/Classes/Components/ActorComponent.h
--- a/Engine/Source/Runtime/Engine/Classes/Components/ActorComponent.h
+++ b/Engine/Source/Runtime/Engine/Classes/Components/ActorComponent.h
@@ -1,12 +1,13 @@
 // Copyright 1998-2014 Epic Games, Inc. All Rights Reserved.
 
 #pragma once
-
 #include "Engine/EngineBaseTypes.h"
 #include "Engine/EngineTypes.h"
 #include "Interfaces/Interface_AssetUserData.h"
 #include "ActorComponent.generated.h"
 
+struct FReplicationFlags;
+
 UCLASS(DefaultToInstanced, abstract, hidecategories=(ComponentReplication))
 class ENGINE_API UActorComponent : public UObject, public IInterface_AssetUserData
 {
diff --git a/Engine/Source/Runtime/Engine/Classes/Components/AudioComponent.h b/Engine/Source/Runtime/Engine/Classes/Components/AudioComponent.h
--- a/Engine/Source/Runtime/Engine/Classes/Components/AudioComponent.h
+++ b/Engine/Source/Runtime/Engine/Classes/Components/AudioComponent.h
@@ -4,7 +4,6 @@
 #pragma once
 #include "Audio.h"
 #include "Sound/SoundAttenuation.h"
-#include "Engine/EngineTypes.h"
 
 #include "AudioComponent.generated.h"
 
diff --git a/Engine/Source/Runtime/Engine/Classes/Components/InputComponent.h b/Engine/Source/Runtime/Engine/Classes/Components/InputComponent.h
--- a/Engine/Source/Runtime/Engine/Classes/Components/InputComponent.h
+++ b/Engine/Source/Runtime/Engine/Classes/Components/InputComponent.h
@@ -1,8 +1,6 @@
 // Copyright 1998-2014 Epic Games, Inc. All Rights Reserved.
 
 #pragma once
-
-#include "Components/ActorComponent.h"
 #include "Runtime/InputCore/Classes/InputCoreTypes.h"
 #include "InputComponent.generated.h"
 
diff --git a/Engine/Source/Runtime/Engine/Classes/Components/LightComponentBase.h b/Engine/Source/Runtime/Engine/Classes/Components/LightComponentBase.h
--- a/Engine/Source/Runtime/Engine/Classes/Components/LightComponentBase.h
+++ b/Engine/Source/Runtime/Engine/Classes/Components/LightComponentBase.h
@@ -1,9 +1,6 @@
 // Copyright 1998-2014 Epic Games, Inc. All Rights Reserved.
 
 #pragma once
-
-#include "Engine/Texture2D.h"
-
 #include "LightComponentBase.generated.h"
 
 UCLASS(abstract, HideCategories=(Trigger,Activation,"Components|Activation",Physics), ShowCategories=(Mobility))
diff --git a/Engine/Source/Runtime/Engine/Classes/Components/MeshComponent.h b/Engine/Source/Runtime/Engine/Classes/Components/MeshComponent.h
--- a/Engine/Source/Runtime/Engine/Classes/Components/MeshComponent.h
+++ b/Engine/Source/Runtime/Engine/Classes/Components/MeshComponent.h
@@ -5,6 +5,8 @@
 #include "Components/PrimitiveComponent.h"
 #include "MeshComponent.generated.h"
 
+struct FMaterialRelevance;
+
 UCLASS(abstract)
 class ENGINE_API UMeshComponent : public UPrimitiveComponent
 {
diff --git a/Engine/Source/Runtime/Engine/Classes/Components/PrimitiveComponent.h b/Engine/Source/Runtime/Engine/Classes/Components/PrimitiveComponent.h
--- a/Engine/Source/Runtime/Engine/Classes/Components/PrimitiveComponent.h
+++ b/Engine/Source/Runtime/Engine/Classes/Components/PrimitiveComponent.h
@@ -2,21 +2,17 @@
 
 #pragma once
 
-#include "Runtime/Core/Public/Containers/StaticArray.h"
 #include "Runtime/InputCore/Classes/InputCoreTypes.h"
 #include "PhysicsEngine/BodyInstance.h"
 #include "Components/SceneComponent.h"
-#include "Components/LightComponent.h"
-#include "Materials/MaterialInterface.h"
 #include "SceneTypes.h"
-#include "CollisionQueryParams.h"
-#include "Engine/Scene.h"
 #include "Engine/EngineTypes.h"
 
 #include "PrimitiveComponent.generated.h"
 
 class FPrimitiveSceneProxy;
 class AController; 
+class UTexture;
 
 /** Information about a streaming texture that a primitive uses for rendering. */
 struct FStreamingTexturePrimitiveInfo
diff --git a/Engine/Source/Runtime/Engine/Classes/Components/SceneComponent.h b/Engine/Source/Runtime/Engine/Classes/Components/SceneComponent.h
--- a/Engine/Source/Runtime/Engine/Classes/Components/SceneComponent.h
+++ b/Engine/Source/Runtime/Engine/Classes/Components/SceneComponent.h
@@ -1,8 +1,8 @@
 // Copyright 1998-2014 Epic Games, Inc. All Rights Reserved.
 
 #pragma once
-
 #include "Components/ActorComponent.h"
+#include "Engine/EngineTypes.h"
 #include "SceneComponent.generated.h"
 
 /** Overlap info consisting of the primitive and the body that is overlapping */
diff --git a/Engine/Source/Runtime/Engine/Classes/Components/SkeletalMeshComponent.h b/Engine/Source/Runtime/Engine/Classes/Components/SkeletalMeshComponent.h
--- a/Engine/Source/Runtime/Engine/Classes/Components/SkeletalMeshComponent.h
+++ b/Engine/Source/Runtime/Engine/Classes/Components/SkeletalMeshComponent.h
@@ -5,6 +5,7 @@
 #include "Interfaces/Interface_CollisionDataProvider.h"
 #include "Components/SkinnedMeshComponent.h"
 #include "PhysicsEngine/PhysicsConstraintComponent.h"
+#include "SkeletalMeshTypes.h"
 #include "SkeletalMeshComponent.generated.h"
 
 class UAnimInstance;
diff --git a/Engine/Source/Runtime/Engine/Classes/Components/SkinnedMeshComponent.h b/Engine/Source/Runtime/Engine/Classes/Components/SkinnedMeshComponent.h
--- a/Engine/Source/Runtime/Engine/Classes/Components/SkinnedMeshComponent.h
+++ b/Engine/Source/Runtime/Engine/Classes/Components/SkinnedMeshComponent.h
@@ -8,7 +8,6 @@
  */
 
 #pragma once
-#include "SkeletalMeshTypes.h"
 #include "Components/MeshComponent.h"
 #include "SkinnedMeshComponent.generated.h"
 
@@ -16,7 +15,8 @@
 // Forward declarations
 //
 class FSkeletalMeshResource;
-
+struct FSkelMeshChunk;
+class FSkeletalMeshVertexBuffer;
 //
 // Bone Visibility.
 //
diff --git a/Engine/Source/Runtime/Engine/Classes/Components/StaticMeshComponent.h b/Engine/Source/Runtime/Engine/Classes/Components/StaticMeshComponent.h
--- a/Engine/Source/Runtime/Engine/Classes/Components/StaticMeshComponent.h
+++ b/Engine/Source/Runtime/Engine/Classes/Components/StaticMeshComponent.h
@@ -5,7 +5,6 @@
 #include "SceneTypes.h"
 #include "Components/MeshComponent.h"
 #include "Runtime/RenderCore/Public/PackedNormal.h"
-#include "Lightmass/LightmassPrimitiveSettingsObject.h"
 
 #include "StaticMeshComponent.generated.h"
 
diff --git a/Engine/Source/Runtime/Engine/Classes/Engine/InputKeyDelegateBinding.h b/Engine/Source/Runtime/Engine/Classes/Engine/InputKeyDelegateBinding.h
--- a/Engine/Source/Runtime/Engine/Classes/Engine/InputKeyDelegateBinding.h
+++ b/Engine/Source/Runtime/Engine/Classes/Engine/InputKeyDelegateBinding.h
@@ -1,6 +1,8 @@
 // Copyright 1998-2014 Epic Games, Inc. All Rights Reserved.
 
 #pragma once
+#include "Components/InputComponent.h"
+#include "Engine/InputDelegateBinding.h"
 #include "InputKeyDelegateBinding.generated.h"
 
 USTRUCT()
diff --git a/Engine/Source/Runtime/Engine/Classes/Engine/TimelineTemplate.h b/Engine/Source/Runtime/Engine/Classes/Engine/TimelineTemplate.h
--- a/Engine/Source/Runtime/Engine/Classes/Engine/TimelineTemplate.h
+++ b/Engine/Source/Runtime/Engine/Classes/Engine/TimelineTemplate.h
@@ -1,6 +1,7 @@
 // Copyright 1998-2014 Epic Games, Inc. All Rights Reserved.
 
 #pragma once
+#include "Components/TimelineComponent.h"
 #include "TimelineTemplate.generated.h"
 
 USTRUCT()
@@ -116,7 +117,7 @@ class UTimelineTemplate : public UObject
 
 	/** How we want the timeline to determine its own length (e.g. specified length, last keyframe) */
 	UPROPERTY(EditAnywhere, Category=TimelineTemplate)
-	TEnumAsByte<enum ETimelineLengthMode> LengthMode;
+	TEnumAsByte<ETimelineLengthMode> LengthMode;
 
 	/** If we want the timeline to auto-play */
 	UPROPERTY(EditAnywhere, Category=TimelineTemplate)
diff --git a/Engine/Source/Runtime/Engine/Classes/GameFramework/Character.h b/Engine/Source/Runtime/Engine/Classes/GameFramework/Character.h
--- a/Engine/Source/Runtime/Engine/Classes/GameFramework/Character.h
+++ b/Engine/Source/Runtime/Engine/Classes/GameFramework/Character.h
@@ -14,7 +14,10 @@ DECLARE_DYNAMIC_MULTICAST_DELEGATE_ThreeParams(FCharacterMovementUpdatedSignatur
 //
 // Forward declarations
 //
+class UAnimMontage;
+class UPrimitiveComponent;
 struct FAnimMontageInstance;
+
 /** Replicated data when playing a root motion montage. */
 USTRUCT()
 struct FRepRootMotionMontage
@@ -23,7 +26,7 @@ struct FRepRootMotionMontage
 
 	/** AnimMontage providing Root Motion */
 	UPROPERTY()
-	class UAnimMontage* AnimMontage;
+	UAnimMontage* AnimMontage;
 
 	/** Track position of Montage */
 	UPROPERTY()
diff --git a/Engine/Source/Runtime/Engine/Classes/GameFramework/MovementComponent.h b/Engine/Source/Runtime/Engine/Classes/GameFramework/MovementComponent.h
--- a/Engine/Source/Runtime/Engine/Classes/GameFramework/MovementComponent.h
+++ b/Engine/Source/Runtime/Engine/Classes/GameFramework/MovementComponent.h
@@ -6,6 +6,7 @@
 
 #pragma once
 #include "Components/ActorComponent.h"
+#include "Engine/EngineTypes.h"
 #include "MovementComponent.generated.h"
 
 struct FCollisionQueryParams;
diff --git a/Engine/Source/Runtime/Engine/Classes/Materials/MaterialExpression.h b/Engine/Source/Runtime/Engine/Classes/Materials/MaterialExpression.h
--- a/Engine/Source/Runtime/Engine/Classes/Materials/MaterialExpression.h
+++ b/Engine/Source/Runtime/Engine/Classes/Materials/MaterialExpression.h
@@ -4,6 +4,8 @@
 #pragma once
 #include "MaterialExpression.generated.h"
 
+class FMaterialExpressionKey;
+
 //@warning: FExpressionInput is mirrored in MaterialShared.h and manually "subclassed" in Material.h (FMaterialInput)
 #if !CPP      //noexport struct
 USTRUCT(noexport)
diff --git a/Engine/Source/Runtime/Engine/Classes/Materials/MaterialInstanceBasePropertyOverrides.h b/Engine/Source/Runtime/Engine/Classes/Materials/MaterialInstanceBasePropertyOverrides.h
--- a/Engine/Source/Runtime/Engine/Classes/Materials/MaterialInstanceBasePropertyOverrides.h
+++ b/Engine/Source/Runtime/Engine/Classes/Materials/MaterialInstanceBasePropertyOverrides.h
@@ -3,6 +3,8 @@
 #pragma once
 #include "MaterialInstanceBasePropertyOverrides.generated.h"
 
+class UMaterialInstance;
+
 /** Properties from the base material that can be overridden in material instances. */
 USTRUCT()
 struct ENGINE_API FMaterialInstanceBasePropertyOverrides
diff --git a/Engine/Source/Runtime/Engine/Classes/Particles/Event/ParticleModuleEventReceiverBase.h b/Engine/Source/Runtime/Engine/Classes/Particles/Event/ParticleModuleEventReceiverBase.h
--- a/Engine/Source/Runtime/Engine/Classes/Particles/Event/ParticleModuleEventReceiverBase.h
+++ b/Engine/Source/Runtime/Engine/Classes/Particles/Event/ParticleModuleEventReceiverBase.h
@@ -13,7 +13,7 @@ class UParticleModuleEventReceiverBase : public UParticleModuleEventBase
 
 	/** The type of event that will generate the kill. */
 	UPROPERTY(EditAnywhere, Category=Source)
-	TEnumAsByte<enum EParticleEventType> EventGeneratorType;
+	TEnumAsByte<EParticleEventType> EventGeneratorType;
 
 	/** The name of the emitter of interest for generating the event. */
 	UPROPERTY(EditAnywhere, Category=Source)
diff --git a/Engine/Source/Runtime/Engine/Classes/Particles/ParticleEmitter.h b/Engine/Source/Runtime/Engine/Classes/Particles/ParticleEmitter.h
--- a/Engine/Source/Runtime/Engine/Classes/Particles/ParticleEmitter.h
+++ b/Engine/Source/Runtime/Engine/Classes/Particles/ParticleEmitter.h
@@ -6,6 +6,7 @@
 //=============================================================================
 
 #pragma once
+#include "Components/SceneComponent.h"
 #include "ParticleEmitterInstances.h"
 #include "ParticleEmitter.generated.h"
 
@@ -145,7 +146,7 @@ class UParticleEmitter : public UObject
 
 	/** If detail mode is >= system detail mode, primitive won't be rendered. */
 	UPROPERTY(EditAnywhere, Category=Particle)
-	TEnumAsByte<enum EDetailMode> DetailMode;
+	TEnumAsByte<EDetailMode> DetailMode;
 
 #if WITH_EDITORONLY_DATA
 	/** This value indicates the emitter should be drawn 'collapsed' in Cascade */
diff --git a/Engine/Source/Runtime/Engine/Classes/Particles/ParticleSystemComponent.h b/Engine/Source/Runtime/Engine/Classes/Particles/ParticleSystemComponent.h
--- a/Engine/Source/Runtime/Engine/Classes/Particles/ParticleSystemComponent.h
+++ b/Engine/Source/Runtime/Engine/Classes/Particles/ParticleSystemComponent.h
@@ -3,6 +3,7 @@
 
 #pragma once
 #include "Components/PrimitiveComponent.h"
+#include "Materials/MaterialInterface.h"
 #include "Particles/Emitter.h"
 #include "Particles/ParticleSystem.h"
 #include "ParticleSystemComponent.generated.h"
diff --git a/Engine/Source/Runtime/Engine/Public/ParticleHelper.h b/Engine/Source/Runtime/Engine/Public/ParticleHelper.h
--- a/Engine/Source/Runtime/Engine/Public/ParticleHelper.h
+++ b/Engine/Source/Runtime/Engine/Public/ParticleHelper.h
@@ -10,6 +10,7 @@
 #include "ParticleBeamTrailVertexFactory.h"
 #include "MeshParticleVertexFactory.h"
 #include "MaterialShared.h"
+#include "Materials/MaterialInterface.h"
 
 #define _ENABLE_PARTICLE_LOD_INGAME_
 
diff --git a/Engine/Source/Runtime/Engine/Public/SkeletalMeshTypes.h b/Engine/Source/Runtime/Engine/Public/SkeletalMeshTypes.h
--- a/Engine/Source/Runtime/Engine/Public/SkeletalMeshTypes.h
+++ b/Engine/Source/Runtime/Engine/Public/SkeletalMeshTypes.h
@@ -14,6 +14,7 @@
 #include "GPUSkinPublicDefs.h"
 #include "Components/PrimitiveComponent.h"
 #include "PrimitiveSceneProxy.h"
+#include "Materials/MaterialInterface.h"
 
 typedef uint16 FBoneIndexType;
 
