Fixed check() in OptionallyPreserveNonUniformScale()

Fixes case where scale delta is the same as remaining scale, causing no max component to be determined. check() still remains, as its there to test the assumption that at least one axis should be active.
Also removed excess semicolon.

diff --git a/Engine/Source/Editor/UnrealEd/Private/LevelEditorViewport.cpp b/Engine/Source/Editor/UnrealEd/Private/LevelEditorViewport.cpp
--- a/Engine/Source/Editor/UnrealEd/Private/LevelEditorViewport.cpp
+++ b/Engine/Source/Editor/UnrealEd/Private/LevelEditorViewport.cpp
@@ -3365,7 +3365,7 @@ static bool OptionallyPreserveNonUniformScale(const FVector& InCurrentScale, con
 	if(ViewportSettings->SnapScaleEnabled && ViewportSettings->PreserveNonUniformScale)
 	{
 		// when using 'auto-precision', we take the max component & snap its scale, then proportionally scale the other components
-		float MaxComponentSum = 0.0f;
+		float MaxComponentSum = -1.0f;
 		int32 MaxAxisIndex = -1;
 		for( int Axis = 0; Axis < 3; ++Axis )
 		{
@@ -3382,7 +3382,7 @@ static bool OptionallyPreserveNonUniformScale(const FVector& InCurrentScale, con
 
 		check(MaxAxisIndex != -1);
 
-		float AbsoluteScaleValue = FMath::GridSnap( InCurrentScale[MaxAxisIndex] + InOutScaleDelta[MaxAxisIndex], GEditor->GetScaleGridSize() );;
+		float AbsoluteScaleValue = FMath::GridSnap( InCurrentScale[MaxAxisIndex] + InOutScaleDelta[MaxAxisIndex], GEditor->GetScaleGridSize() );
 		float ScaleRatioMax = InCurrentScale[MaxAxisIndex] == 0.0f ? 1.0f : AbsoluteScaleValue / InCurrentScale[MaxAxisIndex];
 		for( int Axis = 0; Axis < 3; ++Axis )
 		{
