UE4 - Propagate OpenGL for mobile preview

diff --git a/Engine/Source/Editor/UnrealEd/Private/PlayLevel.cpp b/Engine/Source/Editor/UnrealEd/Private/PlayLevel.cpp
--- a/Engine/Source/Editor/UnrealEd/Private/PlayLevel.cpp
+++ b/Engine/Source/Editor/UnrealEd/Private/PlayLevel.cpp
@@ -1040,6 +1040,10 @@ void UEditorEngine::PlayStandaloneLocalPc(FString MapNameOverride, FIntPoint* Wi
 	// apply additional settings
 	if (bPlayUsingMobilePreview)
 	{
+		if (IsOpenGLPlatform(GRHIShaderPlatform))
+		{
+			AdditionalParameters += TEXT(" -opengl");
+		}
 		AdditionalParameters += TEXT(" -featureleveles2 -faketouches");
 	}
 
