Warn Mac users that attempt to run on Macs with OpenGL 3.3 GPUs as they aren't fully supported also warn Mac users whose GPUs are below the recommended minimums that they may encounter low performance. Added warnings for running with less than the recommended amount of memory or on a dual-core rather than quad-core processor as they also make a big difference.

diff --git a/Engine/Source/Editor/UnrealEd/Private/Editor.cpp b/Engine/Source/Editor/UnrealEd/Private/Editor.cpp
--- a/Engine/Source/Editor/UnrealEd/Private/Editor.cpp
+++ b/Engine/Source/Editor/UnrealEd/Private/Editor.cpp
@@ -6515,6 +6515,134 @@ void UEditorEngine::UpdateAutoLoadProject()
 				UE_LOG(LogEditor, Warning, TEXT("Please update to the latest version of Mac OS X for best performance."));
 			}
 		}
+		
+		NSOpenGLContext* Context = [NSOpenGLContext currentContext];
+		const bool bTemporaryContext = (!Context);
+		if(bTemporaryContext)
+		{
+			NSOpenGLPixelFormatAttribute Attribs[] = {NSOpenGLPFAOpenGLProfile, (NSOpenGLPixelFormatAttribute)NSOpenGLProfileVersion3_2Core, (NSOpenGLPixelFormatAttribute)0};
+			NSOpenGLPixelFormat* PixelFormat = [[NSOpenGLPixelFormat alloc] initWithAttributes:Attribs];
+			if(PixelFormat)
+			{
+				Context = [[NSOpenGLContext alloc] initWithFormat:PixelFormat shareContext:nil];
+				[PixelFormat release];
+				if(Context)
+				{
+					[Context makeCurrentContext];
+				}
+			}
+		}
+		
+		GLint RendererID = 0;
+		if(Context)
+		{
+			[Context getValues:&RendererID forParameter:NSOpenGLCPCurrentRendererID];
+		}
+		
+		if(bTemporaryContext && Context)
+		{
+			[Context release];
+		}
+		
+		// This is quite a coarse test, alerting users who are running on GPUs which run on older drivers only - this is what we want to tell people running the only barely working OpenGL 3.3 cards.
+		// They can run, but they are largely on their own in terms of rendering performance and bugs as Apple/vendors won't fix OpenGL drivers for them.
+		const bool bOldGpuDriver = ((RendererID & 0x00022000 && RendererID < kCGLRendererGeForceID) || (RendererID & 0x00021000 && RendererID < kCGLRendererATIRadeonX3000ID) || (RendererID & 0x00024000 && RendererID < kCGLRendererIntelHD4000ID));
+		if(bOldGpuDriver)
+		{
+			if(FSlateApplication::IsInitialized())
+			{
+				FSuppressableWarningDialog::FSetupInfo Info( LOCTEXT("UnsupportedGPUWarning_Body","The current graphics card does not meet the minimum specification, for best performance an NVIDIA GeForce 470 GTX or AMD Radeon 6870 HD series card or higher is recommended. Rendering performance and compatibility are not guaranteed with this graphics card."), LOCTEXT("UnsupportedGPUWarning_Title","Unsupported Graphics Card"), TEXT("UnsupportedGPUWarning"), GEditorGameAgnosticIni );
+				Info.ConfirmText = LOCTEXT( "OK", "OK");
+				Info.bDefaultToSupressInTheFuture = true;
+				FSuppressableWarningDialog OSUpdateWarning( Info );
+				OSUpdateWarning.ShowModal();
+			}
+			else
+			{
+				UE_LOG(LogEditor, Warning, TEXT("The current graphics card does not meet the minimum specification, for best performance an NVIDIA GeForce 470 GTX or AMD Radeon 6870 HD series card or higher is recommended. Rendering performance and compatibility are not guaranteed with this graphics card."));
+			}
+		}
+		
+		// The editor is much more demanding than people realise, so warn them when they are beneath the official recommended specs.
+		// Any GPU slower than the Nvidia 470 GTX or AMD 6870 will receive a warning so that it is clear that the performance will not be stellar.
+		// This will affect all MacBooks, MacBook Pros and Mac Minis, none of which have shipped with a powerful GPU.
+		// Low-end or older iMacs and some older Mac Pros will also see this warning, though Mac Pro owners can upgrade their GPUs to something better if they choose.
+		// Newer, high-end iMacs and the new Mac Pros, or old Mac Pros with top-end/aftermarket GPUs shouldn't see this at all.
+
+		bool bSlowGpu = false;
+		// The 6970 and the 5870 were the only Apple supplied cards that are fast enough for the Editor supported by the kCGLRendererATIRadeonX3000ID driver
+		// All the kCGLRendererATIRadeonX4000ID supported cards should be more than sufficiently fast AFAICT
+		if(RendererID & 0x00021000 && RendererID == kCGLRendererATIRadeonX3000ID) // AMD
+		{
+			if(!GRHIAdapterName.Contains(TEXT("6970")) && !GRHIAdapterName.Contains(TEXT("5870")))
+			{
+				bSlowGpu = true;
+			}
+		}
+		else if(RendererID & 0x00022000 && RendererID == kCGLRendererGeForceID) // Nvidia
+		{
+			// Most of the 600 & 700 series Nvidia or later cards Apple use in iMacs should be fast enough
+			// but the ones in MacBook Pros are not.
+			if(GRHIAdapterName.Contains(TEXT("640")) || GRHIAdapterName.Contains(TEXT("650")) || GRHIAdapterName.Contains(TEXT("660")) || GRHIAdapterName.Contains(TEXT("750")) || GRHIAdapterName.Contains(TEXT("755")))
+			{
+				bSlowGpu = true;
+			}
+		}
+		else if(RendererID & 0x00024000 && RendererID == kCGLRendererIntelHD5000ID) // Intel
+		{
+			// Only the Iris Pro is even approaching fast enough - but even that is slower than our recommended specs by a wide margin
+			bSlowGpu = true;
+		}
+		
+		if(bSlowGpu)
+		{
+			if(FSlateApplication::IsInitialized())
+			{
+				FSuppressableWarningDialog::FSetupInfo Info( LOCTEXT("SlowGPUWarning_Body","The current graphics card is slower than the recommanded specification of an NVIDIA GeForce 470 GTX or AMD Radeon 6870 HD series card or higher, performance may be low."), LOCTEXT("SlowGPUWarning_Title","Slow Graphics Card"), TEXT("SlowGPUWarning"), GEditorGameAgnosticIni );
+				Info.ConfirmText = LOCTEXT( "OK", "OK");
+				Info.bDefaultToSupressInTheFuture = true;
+				FSuppressableWarningDialog OSUpdateWarning( Info );
+				OSUpdateWarning.ShowModal();
+			}
+			else
+			{
+				UE_LOG(LogEditor, Warning, TEXT("The current graphics card is slower than the recommanded specification of an NVIDIA GeForce 470 GTX or AMD Radeon 6870 HD series card or higher, performance may be low."));
+			}
+		}
+		
+		// Warn about low-memory configurations as they harm performance in the Editor
+		if(FPlatformMemory::GetPhysicalGBRam() < 8)
+		{
+			if(FSlateApplication::IsInitialized())
+			{
+				FSuppressableWarningDialog::FSetupInfo Info( LOCTEXT("LowRAMWarning_Body","For best performance install at least 8GB of RAM."), LOCTEXT("LowRAMWarning_Title","Low RAM"), TEXT("LowRAMWarning"), GEditorGameAgnosticIni );
+				Info.ConfirmText = LOCTEXT( "OK", "OK");
+				Info.bDefaultToSupressInTheFuture = true;
+				FSuppressableWarningDialog OSUpdateWarning( Info );
+				OSUpdateWarning.ShowModal();
+			}
+			else
+			{
+				UE_LOG(LogEditor, Warning, TEXT("For best performance install at least 8GB of RAM."));
+			}
+		}
+		
+		// And also warn about machines with fewer than 4 cores as they will also struggle
+		if(FPlatformMisc::NumberOfCores() < 4)
+		{
+			if(FSlateApplication::IsInitialized())
+			{
+				FSuppressableWarningDialog::FSetupInfo Info( LOCTEXT("SlowCPUWarning_Body","For best performance a Quad-core Intel or AMD processor, 2.5 GHz or faster is recommended."), LOCTEXT("SlowCPUWarning_Title","CPU Performance Warning"), TEXT("SlowCPUWarning"), GEditorGameAgnosticIni );
+				Info.ConfirmText = LOCTEXT( "OK", "OK");
+				Info.bDefaultToSupressInTheFuture = true;
+				FSuppressableWarningDialog OSUpdateWarning( Info );
+				OSUpdateWarning.ShowModal();
+			}
+			else
+			{
+				UE_LOG(LogEditor, Warning, TEXT("For best performance a Quad-core Intel or AMD processor, 2.5 GHz or faster is recommended."));
+			}
+		}
 	}
 #endif
 	
