Fixed vertex snapping draw helpers drawing when snapping to surfaces during copy/paste

diff --git a/Engine/Source/Editor/UnrealEd/Private/Editor/ActorPositioning.cpp b/Engine/Source/Editor/UnrealEd/Private/Editor/ActorPositioning.cpp
--- a/Engine/Source/Editor/UnrealEd/Private/Editor/ActorPositioning.cpp
+++ b/Engine/Source/Editor/UnrealEd/Private/Editor/ActorPositioning.cpp
@@ -203,9 +203,10 @@ FTransform FActorPositioning::GetSnappedSurfaceAlignedTransform(FLevelEditorView
 {
 	FSnappingUtils::SnapPointToGrid(InSurfaceLocation, FVector(0.f));
 
+	bool bDrawVertexHelpers = false;
 	// Secondly, attempt vertex snapping
 	FVector AlignToNormal;
-	if (!InViewportClient || !FSnappingUtils::SnapLocationToNearestVertex( InSurfaceLocation, InViewportClient->GetDropPreviewLocation(), InViewportClient, AlignToNormal ))
+	if (!InViewportClient || !FSnappingUtils::SnapLocationToNearestVertex( InSurfaceLocation, InViewportClient->GetDropPreviewLocation(), InViewportClient, AlignToNormal, bDrawVertexHelpers ))
 	{
 		AlignToNormal = SurfaceNormal;
 	}
diff --git a/Engine/Source/Editor/UnrealEd/Private/SnappingUtils.cpp b/Engine/Source/Editor/UnrealEd/Private/SnappingUtils.cpp
--- a/Engine/Source/Editor/UnrealEd/Private/SnappingUtils.cpp
+++ b/Engine/Source/Editor/UnrealEd/Private/SnappingUtils.cpp
@@ -81,7 +81,7 @@ public:
 	 * @param OutVertexNormal	The normal at the closest vertex
 	 * @return true if anything was snapped
 	 */
-	bool SnapLocationToNearestVertex( FVector& Location, const FVector2D& MouseLocation, FLevelEditorViewportClient* ViewportClient, FVector& OutVertexNormal );
+	bool SnapLocationToNearestVertex( FVector& Location, const FVector2D& MouseLocation, FLevelEditorViewportClient* ViewportClient, FVector& OutVertexNormal, bool bDrawVertHelpers );
 
 	bool SnapToBSPVertex( FVector& Location, FVector GridBase, FRotator& Rotation );
 
@@ -369,12 +369,12 @@ void FEditorViewportSnapping::DrawSnappingHelpers(const FSceneView* View,FPrimit
 	VertexSnappingImpl.DrawSnappingHelpers( View, PDI );
 }
 
-bool FEditorViewportSnapping::SnapLocationToNearestVertex( FVector& Location, const FVector2D& MouseLocation, FLevelEditorViewportClient* ViewportClient, FVector& OutVertexNormal )
+bool FEditorViewportSnapping::SnapLocationToNearestVertex( FVector& Location, const FVector2D& MouseLocation, FLevelEditorViewportClient* ViewportClient, FVector& OutVertexNormal, bool bDrawVertHelpers )
 {
 	bool bSnapped = false;
 	if( IsSnapToVertexEnabled() )
 	{
-		bSnapped = VertexSnappingImpl.SnapLocationToNearestVertex( Location, MouseLocation, ViewportClient, OutVertexNormal );
+		bSnapped = VertexSnappingImpl.SnapLocationToNearestVertex( Location, MouseLocation, ViewportClient, OutVertexNormal, bDrawVertHelpers );
 	}
 	else
 	{
@@ -445,9 +445,9 @@ bool FSnappingUtils::SnapDragLocationToNearestVertex(const FVector& BaseLocation
 	return EditorViewportSnapper->SnapDragLocationToNearestVertex(BaseLocation, DragDelta, ViewportClient);
 }
 
-bool FSnappingUtils::SnapLocationToNearestVertex(FVector& Location, const FVector2D& MouseLocation, FLevelEditorViewportClient* ViewportClient, FVector& OutVertexNormal)
+bool FSnappingUtils::SnapLocationToNearestVertex(FVector& Location, const FVector2D& MouseLocation, FLevelEditorViewportClient* ViewportClient, FVector& OutVertexNormal, bool bDrawVertHelpers )
 {
-	return EditorViewportSnapper->SnapLocationToNearestVertex(Location, MouseLocation, ViewportClient, OutVertexNormal);
+	return EditorViewportSnapper->SnapLocationToNearestVertex(Location, MouseLocation, ViewportClient, OutVertexNormal, bDrawVertHelpers );
 }
 
 void FSnappingUtils::SnapScale(FVector& Point, const FVector& GridBase)
diff --git a/Engine/Source/Editor/UnrealEd/Private/VertexSnapping.cpp b/Engine/Source/Editor/UnrealEd/Private/VertexSnapping.cpp
--- a/Engine/Source/Editor/UnrealEd/Private/VertexSnapping.cpp
+++ b/Engine/Source/Editor/UnrealEd/Private/VertexSnapping.cpp
@@ -360,13 +360,14 @@ void FVertexSnappingImpl::DrawSnappingHelpers(const FSceneView* View,FPrimitiveD
 
 struct FVertexSnappingArgs
 {
-	FVertexSnappingArgs( const FPlane& InActorPlane, const FVector& InCurrentLocation, FLevelEditorViewportClient* InViewportClient, const FSceneView* InView, const FVector2D& InMousePosition, EAxisList::Type InCurrentAxis )
+	FVertexSnappingArgs( const FPlane& InActorPlane, const FVector& InCurrentLocation, FLevelEditorViewportClient* InViewportClient, const FSceneView* InView, const FVector2D& InMousePosition, EAxisList::Type InCurrentAxis, bool bInDrawVertexHelpers )
 		: ActorPlane( InActorPlane )
 		, CurrentLocation( InCurrentLocation )
 		, ViewportClient( InViewportClient )
 		, MousePosition( InMousePosition )
 		, SceneView( InView )
 		, CurrentAxis( InCurrentAxis )
+		, bDrawVertexHelpers( bInDrawVertexHelpers )
 	{}
 
 	/** Plane that the actor is on.  For checking distances and culling vertices behind the plane */
@@ -381,6 +382,8 @@ struct FVertexSnappingArgs
 	const FSceneView* SceneView;
 	/** The current axis being dragged */
 	EAxisList::Type CurrentAxis;
+	/** Whether or not to draw vertex helpers */
+	bool bDrawVertexHelpers;
 };
 
 bool FVertexSnappingImpl::GetClosestVertexOnComponent( const FSnapActor& SnapActor, UPrimitiveComponent* Component, const FVertexSnappingArgs& InArgs, FSnappingVertex& OutClosestLocation )
@@ -546,14 +549,21 @@ FSnappingVertex FVertexSnappingImpl::GetClosestVertex( const TArray<FSnapActor>&
 		}
 	}
 
-	if( ActorVertsToDraw.IsValid() )
+	if( InArgs.bDrawVertexHelpers )
 	{
-		ActorVertsToFade.Add( ActorVertsToDraw, FApp::GetCurrentTime() );
-	}
-
+		if(ActorVertsToDraw.IsValid())
+		{
+			ActorVertsToFade.Add(ActorVertsToDraw, FApp::GetCurrentTime());
+		}
 
-	ActorVertsToFade.Remove( ClosestActor );
-	ActorVertsToDraw = ClosestActor;
+		ActorVertsToFade.Remove(ClosestActor);
+		ActorVertsToDraw = ClosestActor;
+	}
+	else
+	{
+		ActorVertsToDraw = nullptr;
+		ActorVertsToFade.Empty();
+	}
 
 
 	return ClosestLocation;
@@ -601,7 +611,7 @@ void FVertexSnappingImpl::GetActorsInsideBox( const FBox& Box, UWorld* World, TA
 }
 
 
-bool FVertexSnappingImpl::SnapLocationToNearestVertex( FVector& Location, const FVector2D& MouseLocation, FLevelEditorViewportClient* ViewportClient, FVector& OutVertexNormal )
+bool FVertexSnappingImpl::SnapLocationToNearestVertex( FVector& Location, const FVector2D& MouseLocation, FLevelEditorViewportClient* ViewportClient, FVector& OutVertexNormal, bool bDrawVertexHelpers )
 {
 	bool bSnapped = false;
 
@@ -634,7 +644,8 @@ bool FVertexSnappingImpl::SnapLocationToNearestVertex( FVector& Location, const
 		ViewportClient,
 		View,
 		Cursor.GetCursorPos(),
-		EAxisList::Screen
+		EAxisList::Screen,
+		bDrawVertexHelpers
 	);
 	
 	// Snap to the nearest vertex
@@ -741,7 +752,8 @@ bool FVertexSnappingImpl::SnapDraggedActorsToNearestVertex( FVector& DragDelta,
 					ViewportClient,
 					View,
 					MousePosition,
-					CurrentAxis
+					CurrentAxis,
+					true
 				);
 
 				// Snap the drag delta
@@ -872,7 +884,8 @@ bool FVertexSnappingImpl::SnapDragLocationToNearestVertex( const FVector& BaseLo
 			ViewportClient,
 			View,
 			MousePosition,
-			CurrentAxis
+			CurrentAxis,
+			true
 		);
 
 		SnapDragDelta( Args, BaseLocation, AllowedSnappingBox, NoActorsToIgnore, DragDelta );
diff --git a/Engine/Source/Editor/UnrealEd/Private/VertexSnapping.h b/Engine/Source/Editor/UnrealEd/Private/VertexSnapping.h
--- a/Engine/Source/Editor/UnrealEd/Private/VertexSnapping.h
+++ b/Engine/Source/Editor/UnrealEd/Private/VertexSnapping.h
@@ -111,7 +111,7 @@ public:
 	 * @param OutVertexNormal	The normal at the closest vertex
 	 * @return true if anything was snapped
 	 */
-	bool SnapLocationToNearestVertex( FVector& Location, const FVector2D& MouseLocation, FLevelEditorViewportClient* ViewportClient, FVector& OutVertexNormal );
+	bool SnapLocationToNearestVertex( FVector& Location, const FVector2D& MouseLocation, FLevelEditorViewportClient* ViewportClient, FVector& OutVertexNormal, bool bDrawVertexHelpers );
 private:
 	/**
 	 * Gets the closest vertex on a primitive component
diff --git a/Engine/Source/Editor/UnrealEd/Public/SnappingUtils.h b/Engine/Source/Editor/UnrealEd/Public/SnappingUtils.h
--- a/Engine/Source/Editor/UnrealEd/Public/SnappingUtils.h
+++ b/Engine/Source/Editor/UnrealEd/Public/SnappingUtils.h
@@ -55,7 +55,7 @@ public:
 	 * @param OutVertexNormal	The normal at the closest vertex
 	 * @return true if anything was snapped
 	 */
-	static bool SnapLocationToNearestVertex( FVector& Location, const FVector2D& MouseLocation, FLevelEditorViewportClient* ViewportClient, FVector& OutVertexNormal );
+	static bool SnapLocationToNearestVertex( FVector& Location, const FVector2D& MouseLocation, FLevelEditorViewportClient* ViewportClient, FVector& OutVertexNormal, bool bDrawVertHelpers );
 
 	/**
 	 * Snaps a scale value to the scale grid
