343038	Assert in UFactory::StaticImportObject ( !InFactory->SupportedClass || Class->IsChildOf(InFactory->SupportedClass) )

diff --git a/Engine/Source/Editor/UnrealEd/Private/Factories/Factory.cpp b/Engine/Source/Editor/UnrealEd/Private/Factories/Factory.cpp
--- a/Engine/Source/Editor/UnrealEd/Private/Factories/Factory.cpp
+++ b/Engine/Source/Editor/UnrealEd/Private/Factories/Factory.cpp
@@ -192,8 +192,11 @@ UObject* UFactory::StaticImportObject
 	if( InFactory )
 	{
 		// Use just the specified factory.
-		check( !InFactory->SupportedClass || Class->IsChildOf(InFactory->SupportedClass) );
-		Factories.Add( InFactory );
+		if (ensureMsgf( !InFactory->SupportedClass || Class->IsChildOf(InFactory->SupportedClass), 
+			TEXT("Factory is (%s), SupportedClass is (%s) and Class name is (%s)"), *InFactory->GetName(), (InFactory->SupportedClass)? *InFactory->SupportedClass->GetName() : TEXT("None"), *Class->GetName() ))
+		{
+			Factories.Add( InFactory );
+		}
 	}
 	else
 	{
