Fixed crash when adding brush to a world with a disabled navigation system

diff --git a/Engine/Source/Editor/UnrealEd/Private/EditorServer.cpp b/Engine/Source/Editor/UnrealEd/Private/EditorServer.cpp
--- a/Engine/Source/Editor/UnrealEd/Private/EditorServer.cpp
+++ b/Engine/Source/Editor/UnrealEd/Private/EditorServer.cpp
@@ -697,7 +697,7 @@ bool UEditorEngine::Exec_Brush( UWorld* InWorld, const TCHAR* Str, FOutputDevice
 			GEditor->RebuildStaticNavigableGeometry(Level);
 		}
 
-		if (IsLoading() == false) 
+		if (IsLoading() == false && InWorld->GetNavigationSystem()) 
 		{
 			InWorld->GetNavigationSystem()->Build();
 		}
