Cleaning up redirectors before removing sublevel from the world
Fixes: CRASH: UE4Editor_UnrealEd!EditorLevelUtils::RemoveLevelFromWorld() UE4Editor_WorldBrowser!FStreamingLevelCollectionModel::UnloadLevels()
TTP# 343101

diff --git a/Engine/Source/Editor/UnrealEd/Private/EditorLevelUtils.cpp b/Engine/Source/Editor/UnrealEd/Private/EditorLevelUtils.cpp
--- a/Engine/Source/Editor/UnrealEd/Private/EditorLevelUtils.cpp
+++ b/Engine/Source/Editor/UnrealEd/Private/EditorLevelUtils.cpp
@@ -389,9 +389,6 @@ namespace EditorLevelUtils
 
 		GEditor->CloseEditedWorldAssets(CastChecked<UWorld>(InLevel->GetOuter()));
 
-		// Need to clean refs to a removed level 
-		GEditor->ResetTransaction( LOCTEXT("RemoveLevelTransReset", "Removing Levels from World") );
-				
 		UWorld* OwningWorld = InLevel->OwningWorld;
 		const FName LevelPackageName		= InLevel->GetOutermost()->GetFName();
 		const bool bRemovingCurrentLevel	= InLevel && InLevel->IsCurrentLevel();
@@ -423,13 +420,10 @@ namespace EditorLevelUtils
 			// refresh editor windows
 			FEditorDelegates::RefreshAllBrowsers.Broadcast();
 			
-			// Update selection for any selected actors that were in the level and are no longer valid
-			GEditor->NoteSelectionChange();
-
-			// Collect garbage to clear out the destroyed level
-			CollectGarbage( GARBAGE_COLLECTION_KEEPFLAGS );
+			// Reset transaction buffer and run GC to clear out the destroyed level
+			GEditor->Cleanse(true, false, LOCTEXT("RemoveLevelTransReset", "Removing Levels from World"));
 
-			// Ensure that level package were removed
+			// Ensure that level package was removed
 			UPackage* LevelPackage = FindObjectFast<UPackage>(NULL, LevelPackageName);
 			if (LevelPackage != nullptr)
 			{
