Fix for packages to always load not being cooked when generated chunk manifests

diff --git a/Engine/Source/Editor/UnrealEd/Private/Commandlets/ChunkManifestGenerator.cpp b/Engine/Source/Editor/UnrealEd/Private/Commandlets/ChunkManifestGenerator.cpp
--- a/Engine/Source/Editor/UnrealEd/Private/Commandlets/ChunkManifestGenerator.cpp
+++ b/Engine/Source/Editor/UnrealEd/Private/Commandlets/ChunkManifestGenerator.cpp
@@ -269,22 +269,6 @@ void FChunkManifestGenerator::AddPackageToChunkManifest(UPackage* Package, const
 void FChunkManifestGenerator::PrepareToLoadNewPackage(const FString& Filename)
 {
 	AssetsLoadedWithLastPackage.Empty();
-
-	// To make the above work, we need to do a full GC before we load the package
-	// to make sure we get all the assets loaded (or at least created) with a map
-	TScopedPointer< FArchive > FileReader(IFileManager::Get().CreateFileReader(*Filename));
-	if (FileReader)
-	{
-		// Read package file summary from the file
-		FPackageFileSummary PackageSummary;
-		(*FileReader) << PackageSummary;
-		if (PackageSummary.PackageFlags & PKG_ContainsMap)
-		{
-			// FULL GC
-			UE_LOG(LogChunkManifestGenerator, Display, TEXT("Pre map load Full GC..."));
-			CollectGarbage(RF_Native);
-		}
-	}
 }
 
 void FChunkManifestGenerator::CleanManifestDirectories()
