Add a make new blueprint button on class dropdowns for classes that are Blueprintable
Make buttons next to the Game Mode customization class drop downs consistent with normal class dropdowns

diff --git a/Engine/Source/Editor/DetailCustomizations/Private/StringClassReferenceCustomization.cpp b/Engine/Source/Editor/DetailCustomizations/Private/StringClassReferenceCustomization.cpp
--- a/Engine/Source/Editor/DetailCustomizations/Private/StringClassReferenceCustomization.cpp
+++ b/Engine/Source/Editor/DetailCustomizations/Private/StringClassReferenceCustomization.cpp
@@ -2,6 +2,7 @@
 
 #include "DetailCustomizationsPrivatePCH.h"
 #include "StringClassReferenceCustomization.h"
+#include "EditorClassUtils.h"
 
 void FStringClassReferenceCustomization::CustomizeHeader(TSharedRef<IPropertyHandle> InStructPropertyHandle, FDetailWidgetRow& HeaderRow, IPropertyTypeCustomizationUtils& StructCustomizationUtils)
 {
@@ -17,8 +18,8 @@ void FStringClassReferenceCustomization::CustomizeHeader(TSharedRef<IPropertyHan
 	const bool bAllowNone = !(StructPropertyHandle->GetMetaDataProperty()->PropertyFlags & CPF_NoClear);
 
 	check(!MetaClassName.IsEmpty());
-	const UClass* const MetaClass = StringToClass(MetaClassName);
-	const UClass* const RequiredInterface = StringToClass(RequiredInterfaceName);
+	const UClass* const MetaClass = FEditorClassUtils::GetClassFromString(MetaClassName);
+	const UClass* const RequiredInterface = FEditorClassUtils::GetClassFromString(RequiredInterfaceName);
 
 	HeaderRow
 	.NameContent()
@@ -54,7 +55,7 @@ const UClass* FStringClassReferenceCustomization::OnGetClass() const
 	const UClass* Class = CachedClassPtr.Get();
 	if(!Class || Class->GetPathName() != ClassName)
 	{
-		Class = StringToClass(ClassName);
+		Class = FEditorClassUtils::GetClassFromString(ClassName);
 		CachedClassPtr = Class;
 	}
 	return Class;
@@ -66,19 +67,4 @@ void FStringClassReferenceCustomization::OnSetClass(const UClass* NewClass)
 	{
 		CachedClassPtr = NewClass;
 	}
-}
-
-const UClass* FStringClassReferenceCustomization::StringToClass(const FString& ClassName)
-{
-	if(ClassName.IsEmpty() || ClassName == "None")
-	{
-		return nullptr;
-	}
-
-	const UClass* Class = FindObject<UClass>(ANY_PACKAGE, *ClassName);
-	if(!Class)
-	{
-		Class = LoadObject<UClass>(nullptr, *ClassName);
-	}
-	return Class;
-}
+}
\ No newline at end of file
diff --git a/Engine/Source/Editor/EditorStyle/Private/SlateEditorStyle.cpp b/Engine/Source/Editor/EditorStyle/Private/SlateEditorStyle.cpp
--- a/Engine/Source/Editor/EditorStyle/Private/SlateEditorStyle.cpp
+++ b/Engine/Source/Editor/EditorStyle/Private/SlateEditorStyle.cpp
@@ -2561,6 +2561,7 @@ void FSlateEditorStyle::FStyle::SetupPropertyEditorStyles()
 		Set( "PropertyWindow.CategoryBackground", new BOX_BRUSH( "/PropertyView/CategoryBackground", FMargin(4.f/16.f) ) );
 		Set( "PropertyWindow.CategoryForeground", FLinearColor::Black );
 		Set( "PropertyWindow.Button_Browse", new IMAGE_BRUSH( "Icons/lens_12x", Icon12x12 ) );
+		Set( "PropertyWindow.Button_CreateNewBlueprint", new IMAGE_BRUSH( "Icons/PlusSymbol_12x", Icon12x12 ) );
 		Set( "PropertyWindow.Button_Use", new IMAGE_BRUSH( "Icons/assign_12x", Icon12x12 ) );
 		Set( "PropertyWindow.Button_Delete", new IMAGE_BRUSH("Icons/Cross_12x", Icon12x12 ) );
 		Set( "PropertyWindow.Button_Clear", new IMAGE_BRUSH( "Icons/Cross_12x", Icon12x12 ) );
diff --git a/Engine/Source/Editor/PropertyEditor/Private/Presentation/PropertyEditor/PropertyEditor.cpp b/Engine/Source/Editor/PropertyEditor/Private/Presentation/PropertyEditor/PropertyEditor.cpp
--- a/Engine/Source/Editor/PropertyEditor/Private/Presentation/PropertyEditor/PropertyEditor.cpp
+++ b/Engine/Source/Editor/PropertyEditor/Private/Presentation/PropertyEditor/PropertyEditor.cpp
@@ -11,6 +11,9 @@
 
 #include "DelegateFilter.h"
 
+#include "EditorClassUtils.h"
+#include "KismetEditorUtilities.h"
+
 #define LOCTEXT_NAMESPACE "PropertyEditor"
 
 TSharedRef< FPropertyEditor > FPropertyEditor::Create( const TSharedRef< class FPropertyNode >& InPropertyNode, const TSharedRef<class IPropertyUtilities >& InPropertyUtilities )
@@ -208,6 +211,25 @@ void FPropertyEditor::OnClearItem()
 	PropertyHandle->SetValueFromFormattedString( None );
 }
 
+void FPropertyEditor::MakeNewBlueprint()
+{
+	UProperty* NodeProperty = PropertyNode->GetProperty();
+	UClassProperty* ClassProp = Cast<UClassProperty>(NodeProperty);
+	UClass* Class = (ClassProp ? ClassProp->MetaClass : FEditorClassUtils::GetClassFromString(NodeProperty->GetMetaData("MetaClass")));
+
+	if (Class)
+	{
+		UBlueprint* Blueprint = FKismetEditorUtilities::CreateBlueprintFromClass(LOCTEXT("CreateNewBlueprint", "Create New Blueprint"), Class, FString::Printf(TEXT("New%s"),*Class->GetName()));
+
+		if(Blueprint != NULL && Blueprint->GeneratedClass)
+		{
+			PropertyHandle->SetValueFromFormattedString(Blueprint->GeneratedClass->GetPathName());
+
+			FAssetEditorManager::Get().OpenEditorForAsset(Blueprint);
+		}
+	}
+}
+
 void FPropertyEditor::InsertItem()
 {
 	// This action must be deferred until next tick so that we avoid accessing invalid data before we have a chance to tick
diff --git a/Engine/Source/Editor/PropertyEditor/Private/Presentation/PropertyEditor/PropertyEditor.h b/Engine/Source/Editor/PropertyEditor/Private/Presentation/PropertyEditor/PropertyEditor.h
--- a/Engine/Source/Editor/PropertyEditor/Private/Presentation/PropertyEditor/PropertyEditor.h
+++ b/Engine/Source/Editor/PropertyEditor/Private/Presentation/PropertyEditor/PropertyEditor.h
@@ -80,6 +80,7 @@ public:
 	void InsertItem();
 	void DeleteItem();
 	void DuplicateItem();
+	void MakeNewBlueprint();
 	void BrowseTo();
 	void EmptyArray();
 	void Find();
diff --git a/Engine/Source/Editor/PropertyEditor/Private/PropertyCustomizationHelpers.cpp b/Engine/Source/Editor/PropertyEditor/Private/PropertyCustomizationHelpers.cpp
--- a/Engine/Source/Editor/PropertyEditor/Private/PropertyCustomizationHelpers.cpp
+++ b/Engine/Source/Editor/PropertyEditor/Private/PropertyCustomizationHelpers.cpp
@@ -152,6 +152,18 @@ namespace PropertyCustomizationHelpers
 			.IsFocusable( false );
 	}
 
+	TSharedRef<SWidget> MakeNewBlueprintButton( FSimpleDelegate OnFindClicked, TAttribute<FText> OptionalToolTipText, TAttribute<bool> IsEnabled )
+	{
+		return
+			SNew( SPropertyEditorButton )
+			.Text( LOCTEXT( "NewBlueprintButtonLabel", "New Blueprint") )
+			.ToolTipText( OptionalToolTipText.Get().IsEmpty() ? LOCTEXT( "NewBlueprintButtonToolTipText", "Create New Blueprint") : OptionalToolTipText )
+			.Image( FEditorStyle::GetBrush("PropertyWindow.Button_CreateNewBlueprint") )
+			.OnClickAction( OnFindClicked )
+			.IsEnabled(IsEnabled)
+			.IsFocusable( false );
+	}
+
 	TSharedRef<SWidget> MakeInsertDeleteDuplicateButton( FExecuteAction OnInsertClicked, FExecuteAction OnDeleteClicked, FExecuteAction OnDuplicateClicked )
 	{	
 		FMenuBuilder MenuContentBuilder( true, NULL );
diff --git a/Engine/Source/Editor/PropertyEditor/Private/PropertyEditorHelpers.cpp b/Engine/Source/Editor/PropertyEditor/Private/PropertyEditorHelpers.cpp
--- a/Engine/Source/Editor/PropertyEditor/Private/PropertyEditorHelpers.cpp
+++ b/Engine/Source/Editor/PropertyEditor/Private/PropertyEditorHelpers.cpp
@@ -25,6 +25,9 @@
 #include "SPropertyEditorClass.h"
 #include "IDocumentation.h"
 
+#include "Kismet2/KismetEditorUtilities.h"
+#include "EditorClassUtils.h"
+
 #define LOCTEXT_NAMESPACE "PropertyEditor"
 
 void SPropertyNameWidget::Construct( const FArguments& InArgs, TSharedPtr<FPropertyEditor> InPropertyEditor )
@@ -577,7 +580,25 @@ namespace PropertyEditorHelpers
 		//////////////////////////////
 		// Handle a class property.
 
-		if( NodeProperty->IsA<UClassProperty>() || IsStringClassReference(NodeProperty) || NodeProperty->IsA<UAssetClassProperty>() )
+		UClassProperty* ClassProp = Cast<UClassProperty>(NodeProperty);
+		if( ClassProp || IsStringClassReference(NodeProperty))
+		{
+			OutRequiredButtons.Add( EPropertyButton::Use );			
+			OutRequiredButtons.Add( EPropertyButton::Browse );
+
+			UClass* Class = (ClassProp ? ClassProp->MetaClass : FEditorClassUtils::GetClassFromString(NodeProperty->GetMetaData("MetaClass")));
+
+			if (Class && FKismetEditorUtilities::CanCreateBlueprintOfClass(Class))
+			{
+				OutRequiredButtons.Add( EPropertyButton::NewBlueprint );
+			}
+
+			if( !(NodeProperty->PropertyFlags & CPF_NoClear) )
+			{
+				OutRequiredButtons.Add( EPropertyButton::Clear );
+			}
+		}
+		else if (NodeProperty->IsA<UAssetClassProperty>() )
 		{
 			OutRequiredButtons.Add( EPropertyButton::Use );
 			
@@ -780,6 +801,10 @@ namespace PropertyEditorHelpers
 			NewButton = PropertyCustomizationHelpers::MakeInteractiveActorPicker( FOnGetAllowedClasses::CreateSP( PropertyEditor, &FPropertyEditor::OnGetClassesForAssetPicker ), FOnShouldFilterActor(), FOnActorSelected::CreateSP( PropertyEditor, &FPropertyEditor::OnActorSelected ) );
 			break;
 
+		case EPropertyButton::NewBlueprint:
+			NewButton = PropertyCustomizationHelpers::MakeNewBlueprintButton( FSimpleDelegate::CreateSP( PropertyEditor, &FPropertyEditor::MakeNewBlueprint )  );
+			break;
+
 		default:
 			checkf( 0, TEXT( "Unknown button type" ) );
 			break;
diff --git a/Engine/Source/Editor/PropertyEditor/Private/PropertyEditorHelpers.h b/Engine/Source/Editor/PropertyEditor/Private/PropertyEditorHelpers.h
--- a/Engine/Source/Editor/PropertyEditor/Private/PropertyEditorHelpers.h
+++ b/Engine/Source/Editor/PropertyEditor/Private/PropertyEditorHelpers.h
@@ -24,6 +24,7 @@ namespace EPropertyButton
 		PickActorInteractive,
 		Clear,
 		Use,
+		NewBlueprint,
 	};
 }
 
diff --git a/Engine/Source/Editor/PropertyEditor/Public/PropertyCustomizationHelpers.h b/Engine/Source/Editor/PropertyEditor/Public/PropertyCustomizationHelpers.h
--- a/Engine/Source/Editor/PropertyEditor/Public/PropertyCustomizationHelpers.h
+++ b/Engine/Source/Editor/PropertyEditor/Public/PropertyCustomizationHelpers.h
@@ -21,6 +21,7 @@ namespace PropertyCustomizationHelpers
 	PROPERTYEDITOR_API TSharedRef<SWidget> MakeInsertDeleteDuplicateButton( FExecuteAction OnInsertClicked, FExecuteAction OnDeleteClicked, FExecuteAction OnDuplicateClicked );
 	PROPERTYEDITOR_API TSharedRef<SWidget> MakeDeleteButton( FSimpleDelegate OnDeleteClicked, TAttribute<FText> OptionalToolTipText = FText(), TAttribute<bool> IsEnabled = true );
 	PROPERTYEDITOR_API TSharedRef<SWidget> MakeClearButton( FSimpleDelegate OnClearClicked, TAttribute<FText> OptionalToolTipText = FText(), TAttribute<bool> IsEnabled = true );
+	PROPERTYEDITOR_API TSharedRef<SWidget> MakeNewBlueprintButton( FSimpleDelegate OnFindClicked, TAttribute<FText> OptionalToolTipText = FText(), TAttribute<bool> IsEnabled = true );
 	PROPERTYEDITOR_API TSharedRef<SWidget> MakeUseSelectedButton( FSimpleDelegate OnUseSelectedClicked, TAttribute<FText> OptionalToolTipText = FText(), TAttribute<bool> IsEnabled = true );
 	PROPERTYEDITOR_API TSharedRef<SWidget> MakeBrowseButton( FSimpleDelegate OnClearClicked, TAttribute<FText> OptionalToolTipText = FText(), TAttribute<bool> IsEnabled = true );
 	PROPERTYEDITOR_API TSharedRef<SWidget> MakeAssetPickerAnchorButton( FOnGetAllowedClasses OnGetAllowedClasses, FOnAssetSelected OnAssetSelectedFromPicker );
diff --git a/Engine/Source/Editor/UnrealEd/Private/EditorClassUtils.cpp b/Engine/Source/Editor/UnrealEd/Private/EditorClassUtils.cpp
--- a/Engine/Source/Editor/UnrealEd/Private/EditorClassUtils.cpp
+++ b/Engine/Source/Editor/UnrealEd/Private/EditorClassUtils.cpp
@@ -116,4 +116,19 @@ TSharedRef<SWidget> FEditorClassUtils::GetSourceLink(const UClass* Class, const
 	}
 
 	return SourceHyperlink;
-}
\ No newline at end of file
+}
+
+UClass* FEditorClassUtils::GetClassFromString(const FString& ClassName)
+{
+	if(ClassName.IsEmpty() || ClassName == "None")
+	{
+		return nullptr;
+	}
+
+	UClass* Class = FindObject<UClass>(ANY_PACKAGE, *ClassName);
+	if(!Class)
+	{
+		Class = LoadObject<UClass>(nullptr, *ClassName);
+	}
+	return Class;
+}
diff --git a/Engine/Source/Editor/UnrealEd/Public/EditorClassUtils.h b/Engine/Source/Editor/UnrealEd/Public/EditorClassUtils.h
--- a/Engine/Source/Editor/UnrealEd/Public/EditorClassUtils.h
+++ b/Engine/Source/Editor/UnrealEd/Public/EditorClassUtils.h
@@ -57,4 +57,12 @@ namespace FEditorClassUtils
 	 * @return					Shared pointer to the constructed tooltip
 	 */
 	UNREALED_API TSharedRef<SWidget> GetSourceLink(const UClass* Class, const TWeakObjectPtr<UObject> ObjectWeakPtr);
+
+	/**
+	 * Fetches a UClass from the string name of the class
+	 *
+	 * @param	ClassName		Name of the class we want the UClass for
+	 * @return					UClass pointer if it exists
+	 */
+	UNREALED_API UClass* GetClassFromString(const FString& ClassName);
 };
\ No newline at end of file
diff --git a/Engine/Source/Editor/UnrealEd/Public/GameModeInfoCustomizer.h b/Engine/Source/Editor/UnrealEd/Public/GameModeInfoCustomizer.h
--- a/Engine/Source/Editor/UnrealEd/Public/GameModeInfoCustomizer.h
+++ b/Engine/Source/Editor/UnrealEd/Public/GameModeInfoCustomizer.h
@@ -43,10 +43,12 @@ public:
 	{
 		// Find the metaclass of this property
 		UClassProperty* ClassProp = FindFieldChecked<UClassProperty>(AGameMode::StaticClass(), DefaultClassPropertyName);
+
 		UClass* MetaClass = ClassProp->MetaClass;
 		const bool bAllowNone = !(ClassProp->PropertyFlags & CPF_NoClear);
 
-		TAttribute<bool> CanBrowseAtrribute = TAttribute<bool>::Create( TAttribute<bool>::FGetter::CreateSP( this, &FGameModeInfoCustomizer::CanBrowseDefaultClass, DefaultClassPropertyName) ) ;
+		// This is inconsistent with all the other browsers, so disabling it for now
+		//TAttribute<bool> CanBrowseAtrribute = TAttribute<bool>::Create( TAttribute<bool>::FGetter::CreateSP( this, &FGameModeInfoCustomizer::CanBrowseDefaultClass, DefaultClassPropertyName) ) ;
 
 		// Add a row for choosing a new default class
 		Group.AddWidgetRow()
@@ -71,7 +73,17 @@ public:
 			+SHorizontalBox::Slot()
 			.AutoWidth()
 			[
-				PropertyCustomizationHelpers::MakeBrowseButton(FSimpleDelegate::CreateSP(this, &FGameModeInfoCustomizer::OnBrowseDefaultClassClicked, DefaultClassPropertyName), FText(), CanBrowseAtrribute)
+				PropertyCustomizationHelpers::MakeUseSelectedButton(FSimpleDelegate::CreateSP(this, &FGameModeInfoCustomizer::OnMakeSelectedDefaultClassClicked, DefaultClassPropertyName))
+			]
+			+SHorizontalBox::Slot()
+			.AutoWidth()
+			[
+				PropertyCustomizationHelpers::MakeBrowseButton(FSimpleDelegate::CreateSP(this, &FGameModeInfoCustomizer::OnBrowseDefaultClassClicked, DefaultClassPropertyName))
+			]
+			+SHorizontalBox::Slot()
+			.AutoWidth()
+			[
+				PropertyCustomizationHelpers::MakeNewBlueprintButton(FSimpleDelegate::CreateSP(this, &FGameModeInfoCustomizer::OnMakeNewDefaultClassClicked, DefaultClassPropertyName))
 			]
 		];
 	}
@@ -87,7 +99,8 @@ public:
 		// SEe if we are allowed to choose 'no' GameMode
 		const bool bAllowNone = !(DefaultGameModeClassHandle->GetProperty()->PropertyFlags & CPF_NoClear);
 
-		TAttribute<bool> CanBrowseAtrribute = TAttribute<bool>(this, &FGameModeInfoCustomizer::CanBrowseGameMode);
+		// This is inconsistent with other property 
+		//TAttribute<bool> CanBrowseAtrribute = TAttribute<bool>(this, &FGameModeInfoCustomizer::CanBrowseGameMode);
 
 		DefaultGameModeRow
 		.ShowPropertyButtons(false)
@@ -112,24 +125,18 @@ public:
 
 			+SHorizontalBox::Slot()
 			.AutoWidth()
-			.Padding( 2.0f, 1.0f, 2.0f, 0.0f )
 			[
-				SNew(SButton)
-				.HAlign(HAlign_Center)
-				.VAlign(VAlign_Center)
-				.ContentPadding(FMargin(2.0f,0.0f))
-				.OnClicked(this, &FGameModeInfoCustomizer::OnClickNewGameMode, &LayoutBuilder)
-				.ToolTipText(LOCTEXT("NewGameMode_ToolTip", "Create a new Game Mode"))
-				[
-					SNew( STextBlock )
-					.Text(LOCTEXT("NewGameMode", "New.."))
-					.Font( IDetailLayoutBuilder::GetDetailFont() )
-				]
+				PropertyCustomizationHelpers::MakeUseSelectedButton(FSimpleDelegate::CreateSP(this, &FGameModeInfoCustomizer::OnUseSelectedGameModeClicked, &LayoutBuilder))
+			]
+			+SHorizontalBox::Slot()
+			.AutoWidth()
+			[
+				PropertyCustomizationHelpers::MakeBrowseButton(FSimpleDelegate::CreateSP(this, &FGameModeInfoCustomizer::OnBrowseGameModeClicked))
 			]
-			+ SHorizontalBox::Slot()
+			+SHorizontalBox::Slot()
 			.AutoWidth()
 			[
-				PropertyCustomizationHelpers::MakeBrowseButton(FSimpleDelegate::CreateSP(this, &FGameModeInfoCustomizer::OnBrowseGameModeClicked), FText(), CanBrowseAtrribute)
+				PropertyCustomizationHelpers::MakeNewBlueprintButton(FSimpleDelegate::CreateSP(this, &FGameModeInfoCustomizer::OnClickNewGameMode, &LayoutBuilder))
 			]
 		];
 
@@ -224,6 +231,30 @@ public:
 		SyncBrowserToClass(OnGetDefaultClass(ClassPropertyName));
 	}
 
+	void OnMakeNewDefaultClassClicked(FName ClassPropertyName)
+	{
+		UClassProperty* ClassProp = FindFieldChecked<UClassProperty>(AGameMode::StaticClass(), ClassPropertyName);
+
+		UBlueprint* Blueprint = FKismetEditorUtilities::CreateBlueprintFromClass(LOCTEXT("CreateNewBlueprint", "Create New Blueprint"), ClassProp->MetaClass, FString::Printf(TEXT("New%s"),*ClassProp->MetaClass->GetName()));
+
+		if(Blueprint != NULL && Blueprint->GeneratedClass)
+		{
+			OnSetDefaultClass(Blueprint->GeneratedClass, ClassPropertyName);
+
+			FAssetEditorManager::Get().OpenEditorForAsset(Blueprint);
+		}
+	}
+
+	void OnMakeSelectedDefaultClassClicked(FName ClassPropertyName)
+	{
+		UClassProperty* ClassProp = FindFieldChecked<UClassProperty>(AGameMode::StaticClass(), ClassPropertyName);
+		const UClass* SelectedClass = GEditor->GetFirstSelectedClass(ClassProp->MetaClass);
+		if (SelectedClass)
+		{
+			OnSetDefaultClass(SelectedClass, ClassPropertyName);
+		}
+	}
+
 	bool CanBrowseGameMode() const
 	{
 		return CanSyncToClass(GetCurrentGameModeClass());
@@ -253,7 +284,21 @@ public:
 		}
 	}
 
-	FReply OnClickNewGameMode(IDetailLayoutBuilder* DetailLayout)
+	void OnUseSelectedGameModeClicked(IDetailLayoutBuilder* DetailLayout)
+	{
+		const UClass* SelectedClass = GEditor->GetFirstSelectedClass(AGameMode::StaticClass());
+		if (SelectedClass)
+		{
+			DefaultGameModeClassHandle->SetValueFromFormattedString(SelectedClass->GetPathName());
+		}
+
+		if (DetailLayout)
+		{
+			DetailLayout->ForceRefreshDetails();
+		}
+	}
+
+	void OnClickNewGameMode(IDetailLayoutBuilder* DetailLayout)
 	{
 		// Create a new GameMode BP
 		UBlueprint* Blueprint = FKismetEditorUtilities::CreateBlueprintFromClass(LOCTEXT("CreateNewGameMode", "Create New GameMode"), AGameMode::StaticClass(), TEXT("NewGameMode"));
@@ -267,8 +312,6 @@ public:
 		{
 			DetailLayout->ForceRefreshDetails();
 		}
-
-		return FReply::Handled();
 	}
 
 	/** Are we allowed to modify the currently selected GameMode */
