EditorAutomationTests.cpp
This would have never gone to the "else if" because the MakePathRelativeTo code only returns false if the location is on a different drive letter.
This is now fixed.  I've setup it up where it checks the location before assuming the asset is there.

diff --git a/Engine/Source/Editor/UnrealEd/Private/Tests/EditorAutomationTests.cpp b/Engine/Source/Editor/UnrealEd/Private/Tests/EditorAutomationTests.cpp
--- a/Engine/Source/Editor/UnrealEd/Private/Tests/EditorAutomationTests.cpp
+++ b/Engine/Source/Editor/UnrealEd/Private/Tests/EditorAutomationTests.cpp
@@ -21,6 +21,7 @@
 #include "SlateWordWrapper.h"
 #include "AutomationCommon.h"
 #include "AutomationEditorCommon.h"
+#include "AutomationTest.h"
 
 
 //DEFINE_LOG_CATEGORY_STATIC(LogEditorAutomationTests, Log, All);
@@ -1416,30 +1417,69 @@ void FOpenAssetEditors::GetTests(TArray<FString>& OutBeautifiedNames, TArray <FS
 		}
 	}
 
-
+	//Location of the engine folder
+	FString EngineContentFolderLocation = FPaths::ConvertRelativePathToFull(*FPaths::EngineContentDir());
+	//Put the Engine Content Folder Location into an array.  
+	TArray<FString> EngineContentFolderLocationArray;
+	EngineContentFolderLocation.ParseIntoArray(&EngineContentFolderLocationArray, TEXT("/"), true);
+	
 	for (int32 i = 0; i < AssetNames.Num(); ++i)
 	{
+		//true means that the life is located in the Engine/Content folder.
+		bool bFileIsLocatedInTheEngineContentFolder = true;
+		
+		//Get the location of the asset that is being opened.
 		FString Filename = FPaths::ConvertRelativePathToFull(AssetNames[i]);
 
-		if (FPaths::MakePathRelativeTo( Filename, *FPaths::EngineContentDir() ))
-		{
-			FString ShortName = FPaths::GetBaseFilename(Filename);
-			FString PathName =  FPaths::GetPath(Filename);
-			OutBeautifiedNames.Add(ShortName);
-			FString AssetName = FString::Printf(TEXT("/Engine/%s/%s.%s"), *PathName, *ShortName, *ShortName);
-			OutTestCommands.Add(AssetName);
-		}
-		else if (FPaths::MakePathRelativeTo( Filename, *FPaths::GameContentDir() ))
+		//Put File location into an array.  
+		TArray<FString> FilenameArray;
+		Filename.ParseIntoArray(&FilenameArray, TEXT("/"), true);
+
+		//Loop through the location array's and compare each element.  
+		//The loop runs while the index is less than the number of elements in the EngineContentFolderLocation array.
+		//If the elements are the same when the counter is up then it is assumed that the asset file is in the engine content folder. 
+		//Otherwise we'll it's in the game content folder.
+		for (int32 index = 0; index < EngineContentFolderLocationArray.Num(); index++)
 		{
-			FString ShortName = FPaths::GetBaseFilename(Filename);
-			FString PathName =  FPaths::GetPath(Filename);
-			OutBeautifiedNames.Add(ShortName);
-			FString AssetName = FString::Printf(TEXT("/Game/%s/%s.%s"), *PathName, *ShortName, *ShortName);
-			OutTestCommands.Add(AssetName);
+			if ((FilenameArray[index] != EngineContentFolderLocationArray[index]))
+			{
+				//If true it will proceed to add the asset to the Open Asset Editor test list.
+				//This will be false if the asset is on a different drive.
+				if (FPaths::MakePathRelativeTo(Filename, *FPaths::GameContentDir()))
+				{
+					FString ShortName = FPaths::GetBaseFilename(Filename);
+					FString PathName = FPaths::GetPath(Filename);
+					OutBeautifiedNames.Add(ShortName);
+					FString AssetName = FString::Printf(TEXT("/Game/%s/%s.%s"), *PathName, *ShortName, *ShortName);
+					OutTestCommands.Add(AssetName);
+					bFileIsLocatedInTheEngineContentFolder = false;
+					break;
+				}
+				else
+				{
+					UE_LOG(LogEditorAutomationTests, Error, TEXT("Invalid asset path: %s."), *Filename);
+					bFileIsLocatedInTheEngineContentFolder = false;
+					break;
+				}
+			}
 		}
-		else
+		//If true then the asset is in the Engine/Content folder and not in the Game/Content folder.
+		if (bFileIsLocatedInTheEngineContentFolder)
 		{
-			UE_LOG(LogEditorAutomationTests, Error, TEXT("Invalid asset path: %s."), *Filename);
+			//If true it will proceed to add the asset to the Open Asset Editor test list.
+			//This will be false if the asset is on a different drive.
+			if (FPaths::MakePathRelativeTo(Filename, *FPaths::EngineContentDir()))
+			{
+				FString ShortName = FPaths::GetBaseFilename(Filename);
+				FString PathName = FPaths::GetPath(Filename);
+				OutBeautifiedNames.Add(ShortName);
+				FString AssetName = FString::Printf(TEXT("/Engine/%s/%s.%s"), *PathName, *ShortName, *ShortName);
+				OutTestCommands.Add(AssetName);
+			}
+			else
+			{
+				UE_LOG(LogEditorAutomationTests, Error, TEXT("Invalid asset path: %s."), *Filename);
+			}
 		}
 	}
 }
