Removing PackagesToAlwaysCook from WrangleContentCommandlet as it's not used anymore.

diff --git a/Engine/Source/Editor/UnrealEd/Private/Commandlets/ContentCommandlets.cpp b/Engine/Source/Editor/UnrealEd/Private/Commandlets/ContentCommandlets.cpp
--- a/Engine/Source/Editor/UnrealEd/Private/Commandlets/ContentCommandlets.cpp
+++ b/Engine/Source/Editor/UnrealEd/Private/Commandlets/ContentCommandlets.cpp
@@ -964,7 +964,6 @@ int32 UWrangleContentCommandlet::Main( const FString& Params )
 
 		// get a list of packages to load
 		const FConfigSection* PackagesToFullyLoadSection = GConfig->GetSectionPrivate( *SectionToUse, 0, 1, *WrangleContentIniName );
-		const FConfigSection* PackagesToAlwaysCook = GConfig->GetSectionPrivate( TEXT("/Script/Engine.PackagesToAlwaysCook"), 0, 1, GEngineIni );
 		const FConfigSection* StartupPackages = GConfig->GetSectionPrivate( TEXT("/Script/Engine.StartupPackages"), 0, 1, GEngineIni );
 
 		// we expect either the .ini to exist, or -allmaps to be specified
@@ -987,18 +986,6 @@ int32 UWrangleContentCommandlet::Main( const FString& Params )
 			PackagesToFullyLoad = *PackagesToFullyLoadSection;
 		}
 
-		// move any always cook packages to list of packages to load
-		if (PackagesToAlwaysCook)
-		{
-			for (FConfigSectionMap::TConstIterator It(*PackagesToAlwaysCook); It; ++It)
-			{
-				if (It.Key() == TEXT("Package") || It.Key() == TEXT("SeekFreePackage"))
-				{
-					PackagesToFullyLoad.Add(*It.Key().ToString(), *It.Value());
-				}
-			}
-		}
-
 		// make sure all possible script/startup packages are loaded
 		FStartupPackages::LoadAll();
 
