Added simulate and toggle between SIE/PIE analytics event

diff --git a/Engine/Source/Editor/UnrealEd/Classes/Editor/EditorEngine.h b/Engine/Source/Editor/UnrealEd/Classes/Editor/EditorEngine.h
--- a/Engine/Source/Editor/UnrealEd/Classes/Editor/EditorEngine.h
+++ b/Engine/Source/Editor/UnrealEd/Classes/Editor/EditorEngine.h
@@ -2354,7 +2354,7 @@ private:
 	/**
 	 * Toggles PIE to SIE or vice-versa
 	 */
-	void ToggleBetweenPIEandSIE();
+	void ToggleBetweenPIEandSIE( bool bNewSession = false);
 
 	/**
 	 * Hack to switch worlds for the PIE window before and after a slate event
diff --git a/Engine/Source/Editor/UnrealEd/Private/Kismet2/DebuggerCommands.cpp b/Engine/Source/Editor/UnrealEd/Private/Kismet2/DebuggerCommands.cpp
--- a/Engine/Source/Editor/UnrealEd/Private/Kismet2/DebuggerCommands.cpp
+++ b/Engine/Source/Editor/UnrealEd/Private/Kismet2/DebuggerCommands.cpp
@@ -923,6 +923,11 @@ void FInternalPlayWorldCommandCallbacks::Simulate_Clicked()
 		// Start a new simulation session!
 		if( !HasPlayWorld() )
 		{
+			if( FEngineAnalytics::IsAvailable() )
+			{
+				FEngineAnalytics::GetProvider().RecordEvent( TEXT("Editor.Usage.SimulateInEditor") );
+			}
+
 			GUnrealEd->RequestPlaySession(false, ActiveLevelViewport, true/*bSimulateInEditor*/, NULL, NULL, -1, false );
 		}
 		else
diff --git a/Engine/Source/Editor/UnrealEd/Private/PlayLevel.cpp b/Engine/Source/Editor/UnrealEd/Private/PlayLevel.cpp
--- a/Engine/Source/Editor/UnrealEd/Private/PlayLevel.cpp
+++ b/Engine/Source/Editor/UnrealEd/Private/PlayLevel.cpp
@@ -1916,7 +1916,7 @@ void UEditorEngine::PlayInEditor( UWorld* InWorld, bool bInSimulateInEditor )
 
 		if (bInSimulateInEditor)
 		{
-			ToggleBetweenPIEandSIE();
+			ToggleBetweenPIEandSIE( true );
 		}
 	}
 	else
@@ -2567,7 +2567,7 @@ FViewport* UEditorEngine::GetPIEViewport()
 	return NULL;
 }
 
-void UEditorEngine::ToggleBetweenPIEandSIE()
+void UEditorEngine::ToggleBetweenPIEandSIE( bool bNewSession )
 {
 	bIsToggleBetweenPIEandSIEQueued = false;
 
@@ -2589,6 +2589,13 @@ void UEditorEngine::ToggleBetweenPIEandSIE()
 		return;
 	}
 
+	if( FEngineAnalytics::IsAvailable() && !bNewSession )
+	{
+		FString ToggleType = bIsSimulatingInEditor ? TEXT("SIEtoPIE") : TEXT("PIEtoSIE");
+
+		FEngineAnalytics::GetProvider().RecordEvent(TEXT("Editor.Usage.PIE"), TEXT("ToggleBetweenPIEandSIE"), ToggleType );
+	}
+
 	FSlatePlayInEditorInfo & SlatePlayInEditorSession = *SlateInfoPtr;
 
 	// This is only supported inside SLevelEditor viewports currently
