Make DDC commandlet cache platform data.
Segway to another change which stops caching of platform data on postload of object.

diff --git a/Engine/Source/Editor/UnrealEd/Private/Commandlets/DerivedDataCacheCommandlet.cpp b/Engine/Source/Editor/UnrealEd/Private/Commandlets/DerivedDataCacheCommandlet.cpp
--- a/Engine/Source/Editor/UnrealEd/Private/Commandlets/DerivedDataCacheCommandlet.cpp
+++ b/Engine/Source/Editor/UnrealEd/Private/Commandlets/DerivedDataCacheCommandlet.cpp
@@ -192,6 +192,15 @@ int32 UDerivedDataCacheCommandlet::Main( const FString& Params )
 			}
 			else
 			{
+				// cache all the resources for this platform
+				for ( TObjectIterator<UObject> It; It; ++It )
+				{
+					for ( auto Platform : Platforms )
+					{
+						It->BeginCacheForCookedPlatformData( Platform );
+					}
+				}
+
 				bLastPackageWasMap = Package->ContainsMap();
 				NumProcessedSinceLastGC++;
 			}
