Lowered the work scale amout for auto-gpu scalibility detection

diff --git a/Engine/Source/Editor/UnrealEd/Classes/Editor/EditorGameAgnosticSettings.h b/Engine/Source/Editor/UnrealEd/Classes/Editor/EditorGameAgnosticSettings.h
--- a/Engine/Source/Editor/UnrealEd/Classes/Editor/EditorGameAgnosticSettings.h
+++ b/Engine/Source/Editor/UnrealEd/Classes/Editor/EditorGameAgnosticSettings.h
@@ -37,6 +37,9 @@ class UEditorGameAgnosticSettings : public UObject
 	UPROPERTY(config)
 	TArray<FGuid> InProgressSurveys;
 
+	UPROPERTY(config)
+	float AutoScalabilityWorkScaleAmount;
+
 	/** Engine scalability benchmark results */
 	Scalability::FQualityLevels EngineBenchmarkResult;
 
diff --git a/Engine/Source/Editor/UnrealEd/Private/Editor/EditorGameAgnosticSettings.cpp b/Engine/Source/Editor/UnrealEd/Private/Editor/EditorGameAgnosticSettings.cpp
--- a/Engine/Source/Editor/UnrealEd/Private/Editor/EditorGameAgnosticSettings.cpp
+++ b/Engine/Source/Editor/UnrealEd/Private/Editor/EditorGameAgnosticSettings.cpp
@@ -6,6 +6,7 @@ UEditorGameAgnosticSettings::UEditorGameAgnosticSettings(const class FPostConstr
 	: Super(PCIP)
 {
 	bCopyStarterContentPreference = true;
+	AutoScalabilityWorkScaleAmount = 1;
 }
 
 void UEditorGameAgnosticSettings::PostEditChangeProperty( struct FPropertyChangedEvent& PropertyChangedEvent)
@@ -59,9 +60,8 @@ void UEditorGameAgnosticSettings::AutoApplyScalabilityBenchmark()
 	GWarn->StatusUpdate(0, 1, NSLOCTEXT("UnrealEd", "RunningEngineBenchmark", "Running engine benchmark..."));
 	GWarn->PushStatus();
 
-	uint32 WorkScale = 5.0f;
 
-	Scalability::FQualityLevels Temporary = Scalability::BenchmarkQualityLevels(WorkScale);
+	Scalability::FQualityLevels Temporary = Scalability::BenchmarkQualityLevels(AutoScalabilityWorkScaleAmount);
 
 	GConfig->SetBool(Section, TEXT("Valid"), true, GEditorGameAgnosticIni);
 	GConfig->SetInt(Section, TEXT("ResolutionQuality"), Temporary.ResolutionQuality, GEditorGameAgnosticIni);
