Fix manually auto-applying scalability not saving that data

diff --git a/Engine/Source/Editor/UnrealEd/Private/SScalabilitySettings.cpp b/Engine/Source/Editor/UnrealEd/Private/SScalabilitySettings.cpp
--- a/Engine/Source/Editor/UnrealEd/Private/SScalabilitySettings.cpp
+++ b/Engine/Source/Editor/UnrealEd/Private/SScalabilitySettings.cpp
@@ -103,9 +103,11 @@ FReply SScalabilitySettings::OnHeaderClicked(int32 InQualityLevel)
 
 FReply SScalabilitySettings::OnAutoClicked()
 {
+	GEditor->AccessGameAgnosticSettings().AutoApplyScalabilityBenchmark();
+	GEditor->AccessGameAgnosticSettings().LoadScalabilityBenchmark();
+
 	CachedQualityLevels = GEditor->GetGameAgnosticSettings().EngineBenchmarkResult;
-	SetQualityLevels(CachedQualityLevels);
-	Scalability::SaveState(GEditorGameAgnosticIni);
+
 	GEditor->RedrawAllViewports();
 	return FReply::Handled();
 }
