Ensure that array properties with external BP dependencies are fully loaded prior to generating a BP class. This fixes an invalid compile error that may result during async loading.

diff --git a/Engine/Source/Editor/UnrealEd/Private/Kismet2/BlueprintEditorUtils.cpp b/Engine/Source/Editor/UnrealEd/Private/Kismet2/BlueprintEditorUtils.cpp
--- a/Engine/Source/Editor/UnrealEd/Private/Kismet2/BlueprintEditorUtils.cpp
+++ b/Engine/Source/Editor/UnrealEd/Private/Kismet2/BlueprintEditorUtils.cpp
@@ -777,6 +777,19 @@ struct FRegenerationHelper
 									MacroSources.Add(MacroSource);
 								}
 							}
+
+							// If a variable node has an external dependency, then its BP class will differ from ours. For array properties,
+							// the external BP class (and thus the array property itself) will have been loaded/processed via the above
+							// ProcessHierarchy() call. However, the array's 'Inner' property may not have been preloaded as part of that path.
+							// Thus, we handle that here in order to ensure that all 'Inner' fields are also valid before class regeneration.
+							if (auto VariableNode = Cast<UK2Node_Variable>(Node))
+							{
+								UArrayProperty* ArrayProperty = Cast<UArrayProperty>(VariableNode->VariableReference.ResolveMember<UProperty>(Node));
+								if (ArrayProperty != nullptr && ArrayProperty->Inner != nullptr && ArrayProperty->Inner->HasAnyFlags(RF_NeedLoad|RF_WasLoaded))
+								{
+									ForcedLoad(ArrayProperty->Inner);
+								}
+							}
 						}
 
 						LocalDependentStructures.Empty(LocalDependentStructures.Max());
