Disallow vertex snapping when used with modifier keys not bound to vertex snap

diff --git a/Engine/Source/Editor/UnrealEd/Private/SnappingUtils.cpp b/Engine/Source/Editor/UnrealEd/Private/SnappingUtils.cpp
--- a/Engine/Source/Editor/UnrealEd/Private/SnappingUtils.cpp
+++ b/Engine/Source/Editor/UnrealEd/Private/SnappingUtils.cpp
@@ -131,10 +131,9 @@ bool FEditorViewportSnapping::IsSnapToVertexEnabled()
 
 		const FInputGesture& Gesture = *Commands.HoldToEnableVertexSnapping->GetActiveGesture();
 
-		// Note the gesture is allowed to pass if modifier keys are pressed but the gesture doesnt require them to be pressed.  This is so multiple actions can be combined with vertex snap
-		return (!Gesture.bCtrl	|| GCurrentLevelEditingViewportClient->IsCtrlPressed() ) 
-			&& (!Gesture.bAlt	|| GCurrentLevelEditingViewportClient->IsAltPressed() ) 
-			&& (!Gesture.bShift || GCurrentLevelEditingViewportClient->IsShiftPressed() ) 
+		return (Gesture.bCtrl == GCurrentLevelEditingViewportClient->IsCtrlPressed() ) 
+			&& (Gesture.bAlt ==  GCurrentLevelEditingViewportClient->IsAltPressed() ) 
+			&& (Gesture.bShift == GCurrentLevelEditingViewportClient->IsShiftPressed() ) 
 			&& GCurrentLevelEditingViewportClient->Viewport->KeyState(Gesture.Key) == true;
 	}
 	else
