#ue4 Disabling clearing of the undo buffer after PIE

diff --git a/Engine/Source/Editor/UnrealEd/Private/EditorServer.cpp b/Engine/Source/Editor/UnrealEd/Private/EditorServer.cpp
--- a/Engine/Source/Editor/UnrealEd/Private/EditorServer.cpp
+++ b/Engine/Source/Editor/UnrealEd/Private/EditorServer.cpp
@@ -988,10 +988,15 @@ bool UEditorEngine::Exec_Brush( UWorld* InWorld, const TCHAR* Str, FOutputDevice
 
 int32 UEditorEngine::BeginTransaction(const TCHAR* TransactionContext, const FText& Description, UObject* PrimaryObject)
 {
-	// generate transaction context
-	int32 index = Trans->Begin( TransactionContext, Description );
-	Trans->SetPrimaryUndoObject(PrimaryObject);
-	return index;
+	int32 Index = INDEX_NONE;
+
+	if (!bIsSimulatingInEditor)
+	{
+		// generate transaction context
+		Index = Trans->Begin(TransactionContext, Description);
+		Trans->SetPrimaryUndoObject(PrimaryObject);
+	}
+	return Index;
 }
 
 int32 UEditorEngine::BeginTransaction(const FText& Description)
@@ -1001,7 +1006,13 @@ int32 UEditorEngine::BeginTransaction(const FText& Description)
 
 int32 UEditorEngine::EndTransaction()
 {
-	return Trans->End();
+	int32 Index = INDEX_NONE;
+	if (!bIsSimulatingInEditor)
+	{
+		Index = Trans->End();
+	}
+
+	return Index;
 }
 
 void UEditorEngine::ResetTransaction(const FText& Reason)
diff --git a/Engine/Source/Editor/UnrealEd/Private/PlayLevel.cpp b/Engine/Source/Editor/UnrealEd/Private/PlayLevel.cpp
--- a/Engine/Source/Editor/UnrealEd/Private/PlayLevel.cpp
+++ b/Engine/Source/Editor/UnrealEd/Private/PlayLevel.cpp
@@ -167,8 +167,6 @@ void UEditorEngine::EndPlayMap()
 		// Clear out viewport index
 		PlayInEditorViewportIndex = -1; 
 
-		// The undo system may have a reference to a SIE/PIE object that is about to be destroyed, so clear the transactions
-		ResetTransaction( NSLOCTEXT("UnrealEd", "EndPlayMap", "End Play Map") );
 
 		// We could have been toggling back and forth between simulate and pie before ending the play map
 		// Make sure the property windows are cleared of any pie actors
diff --git a/Engine/Source/Editor/UnrealEd/Private/ScopedTransaction.cpp b/Engine/Source/Editor/UnrealEd/Private/ScopedTransaction.cpp
--- a/Engine/Source/Editor/UnrealEd/Private/ScopedTransaction.cpp
+++ b/Engine/Source/Editor/UnrealEd/Private/ScopedTransaction.cpp
@@ -16,7 +16,7 @@ FScopedTransaction::FScopedTransaction(const TCHAR* TransactionContext, const FT
 
 void FScopedTransaction::Construct (const TCHAR* TransactionContext, const FText& SessionName, UObject* PrimaryObject, const bool bShouldActuallyTransact)
 {
-	if( bShouldActuallyTransact && GEditor && GEditor->Trans )
+	if( bShouldActuallyTransact && GEditor && GEditor->Trans && !GEditor->bIsSimulatingInEditor)
 	{
 		FSlateApplication::Get().OnLogSlateEvent(EEventLog::BeginTransaction, SessionName );
 		Index = GEditor->BeginTransaction( TransactionContext, SessionName, PrimaryObject );
