Another band-aid, this time for crashing when undo-ing the creation of a sound cue node.

diff --git a/Engine/Source/Editor/UnrealEd/Private/SoundCueGraphNode.cpp b/Engine/Source/Editor/UnrealEd/Private/SoundCueGraphNode.cpp
--- a/Engine/Source/Editor/UnrealEd/Private/SoundCueGraphNode.cpp
+++ b/Engine/Source/Editor/UnrealEd/Private/SoundCueGraphNode.cpp
@@ -127,8 +127,15 @@ int32 USoundCueGraphNode::EstimateNodeWidth() const
 
 bool USoundCueGraphNode::CanAddInputPin() const
 {
-	// Check if adding another input would exceed max child nodes.
-	return SoundNode->ChildNodes.Num() < SoundNode->GetMaxChildNodes();
+	if(SoundNode)
+	{
+		// Check if adding another input would exceed max child nodes.
+		return SoundNode->ChildNodes.Num() < SoundNode->GetMaxChildNodes();
+	}
+	else
+	{
+		return false;
+	}
 }
 
 FText USoundCueGraphNode::GetNodeTitle(ENodeTitleType::Type TitleType) const
