TTP 342929	Export FBX: Do not group Mesh and Skeleton

diff --git a/Engine/Source/Editor/UnrealEd/Private/Fbx/FbxAnimationExport.cpp b/Engine/Source/Editor/UnrealEd/Private/Fbx/FbxAnimationExport.cpp
--- a/Engine/Source/Editor/UnrealEd/Private/Fbx/FbxAnimationExport.cpp
+++ b/Engine/Source/Editor/UnrealEd/Private/Fbx/FbxAnimationExport.cpp
@@ -196,15 +196,12 @@ void FFbxExporter::ExportAnimSequence( const UAnimSequence* AnimSeq, USkeletalMe
  		return;
 	}
 
-	FbxString NodeName("BaseNode");
-
-	FbxNode* BaseNode = FbxNode::Create(Scene, NodeName);
-	Scene->GetRootNode()->AddChild(BaseNode);
 
+	FbxNode* RootNode = Scene->GetRootNode();
 	// Create the Skeleton
 	TArray<FbxNode*> BoneNodes;
 	FbxNode* SkeletonRootNode = CreateSkeleton(*SkelMesh, BoneNodes);
-	BaseNode->AddChild(SkeletonRootNode);
+	RootNode->AddChild(SkeletonRootNode);
 
 
 	// Export the anim sequence
@@ -232,7 +229,7 @@ void FFbxExporter::ExportAnimSequence( const UAnimSequence* AnimSeq, USkeletalMe
 		FbxNode* MeshRootNode = CreateMesh(*SkelMesh, *MeshName);
 		if(MeshRootNode)
 		{
-			BaseNode->AddChild(MeshRootNode);
+			RootNode->AddChild(MeshRootNode);
 		}
 
 		if(SkeletonRootNode && MeshRootNode)
