Back out changelist 2251857
- I see it's importing fine correctly unless you have Use Time 0 as ref pose. So am going to investigate that.

diff --git a/Engine/Source/Editor/UnrealEd/Private/Fbx/FbxAnimationExport.cpp b/Engine/Source/Editor/UnrealEd/Private/Fbx/FbxAnimationExport.cpp
--- a/Engine/Source/Editor/UnrealEd/Private/Fbx/FbxAnimationExport.cpp
+++ b/Engine/Source/Editor/UnrealEd/Private/Fbx/FbxAnimationExport.cpp
@@ -196,12 +196,15 @@ void FFbxExporter::ExportAnimSequence( const UAnimSequence* AnimSeq, USkeletalMe
  		return;
 	}
 
+	FbxString NodeName("BaseNode");
+
+	FbxNode* BaseNode = FbxNode::Create(Scene, NodeName);
+	Scene->GetRootNode()->AddChild(BaseNode);
 
-	FbxNode* RootNode = Scene->GetRootNode();
 	// Create the Skeleton
 	TArray<FbxNode*> BoneNodes;
 	FbxNode* SkeletonRootNode = CreateSkeleton(*SkelMesh, BoneNodes);
-	RootNode->AddChild(SkeletonRootNode);
+	BaseNode->AddChild(SkeletonRootNode);
 
 
 	// Export the anim sequence
@@ -229,7 +232,7 @@ void FFbxExporter::ExportAnimSequence( const UAnimSequence* AnimSeq, USkeletalMe
 		FbxNode* MeshRootNode = CreateMesh(*SkelMesh, *MeshName);
 		if(MeshRootNode)
 		{
-			RootNode->AddChild(MeshRootNode);
+			BaseNode->AddChild(MeshRootNode);
 		}
 
 		if(SkeletonRootNode && MeshRootNode)
