Ensure that additional channels for multichannel sounds can be imported correctly

TTP#343648 - AUDIO: Surround sound import not taking

Sounds were imported correctly but compressed data could remain in memory if the sound had been played before importing.

diff --git a/Engine/Source/Editor/UnrealEd/Private/Factories/EditorFactories.cpp b/Engine/Source/Editor/UnrealEd/Private/Factories/EditorFactories.cpp
--- a/Engine/Source/Editor/UnrealEd/Private/Factories/EditorFactories.cpp
+++ b/Engine/Source/Editor/UnrealEd/Private/Factories/EditorFactories.cpp
@@ -2017,6 +2017,9 @@ UObject* USoundSurroundFactory::FactoryCreateBinary
 			}
 		}
 
+		// Clear resources so that if it's already been played, it will reload the wave data
+		Sound->FreeResources();
+
 		// Presize the offsets array, in case the sound was new or the original sound data was stripped by cooking.
 		if ( Sound->ChannelOffsets.Num() != SPEAKER_Count )
 		{
