Header cleanup: Engine module, public headers.

diff --git a/Engine/Source/Editor/UnrealEd/Private/Lightmass/LightmassRender.cpp b/Engine/Source/Editor/UnrealEd/Private/Lightmass/LightmassRender.cpp
--- a/Engine/Source/Editor/UnrealEd/Private/Lightmass/LightmassRender.cpp
+++ b/Engine/Source/Editor/UnrealEd/Private/Lightmass/LightmassRender.cpp
@@ -10,6 +10,7 @@
 #include "Lightmass.h"
 #include "LightmassRender.h"
 #include "MaterialCompiler.h"
+#include "LightMap.h"
 
 DEFINE_LOG_CATEGORY_STATIC(LogLightmassRender, Error, All);
 
diff --git a/Engine/Source/Editor/UnrealEd/Private/StaticLightingSystem/StaticLightingDebug.cpp b/Engine/Source/Editor/UnrealEd/Private/StaticLightingSystem/StaticLightingDebug.cpp
--- a/Engine/Source/Editor/UnrealEd/Private/StaticLightingSystem/StaticLightingDebug.cpp
+++ b/Engine/Source/Editor/UnrealEd/Private/StaticLightingSystem/StaticLightingDebug.cpp
@@ -7,6 +7,7 @@
 #include "UnrealEd.h"
 #include "StaticMeshResources.h"
 #include "StaticLightingPrivate.h"
+#include "LightMap.h"
 
 /** Information about the texel that is selected */
 FSelectedLightmapSample GCurrentSelectedLightmapSample;
diff --git a/Engine/Source/Editor/UnrealEd/Private/StaticLightingSystem/StaticLightingPrivate.h b/Engine/Source/Editor/UnrealEd/Private/StaticLightingSystem/StaticLightingPrivate.h
--- a/Engine/Source/Editor/UnrealEd/Private/StaticLightingSystem/StaticLightingPrivate.h
+++ b/Engine/Source/Editor/UnrealEd/Private/StaticLightingSystem/StaticLightingPrivate.h
@@ -12,6 +12,8 @@
 #include "StaticLighting.h"
 #include "LightingBuildOptions.h"
 
+struct FSelectedLightmapSample;
+
 DECLARE_LOG_CATEGORY_EXTERN(LogStaticLightingSystem, Log, All);
 
 /** Encapsulation of all Lightmass statistics */
diff --git a/Engine/Source/Editor/UnrealEd/Private/StaticLightingSystem/StaticLightingSystem.cpp b/Engine/Source/Editor/UnrealEd/Private/StaticLightingSystem/StaticLightingSystem.cpp
--- a/Engine/Source/Editor/UnrealEd/Private/StaticLightingSystem/StaticLightingSystem.cpp
+++ b/Engine/Source/Editor/UnrealEd/Private/StaticLightingSystem/StaticLightingSystem.cpp
@@ -18,6 +18,9 @@ FSwarmDebugOptions GSwarmDebugOptions;
 #include "LevelUtils.h"
 #include "CrashTracker.h"
 #include "EngineModule.h"
+#include "LightMap.h"
+#include "ShadowMap.h"
+#include "RendererInterface.h"
 
 DEFINE_LOG_CATEGORY(LogStaticLightingSystem);
 
diff --git a/Engine/Source/Editor/UnrealEd/Private/StaticLightingSystem/StaticLightingTextureMapping.cpp b/Engine/Source/Editor/UnrealEd/Private/StaticLightingSystem/StaticLightingTextureMapping.cpp
--- a/Engine/Source/Editor/UnrealEd/Private/StaticLightingSystem/StaticLightingTextureMapping.cpp
+++ b/Engine/Source/Editor/UnrealEd/Private/StaticLightingSystem/StaticLightingTextureMapping.cpp
@@ -6,6 +6,7 @@
 
 #include "UnrealEd.h"
 #include "StaticLightingPrivate.h"
+#include "LightMap.h"
 
 // Doxygen cannot not parse this definition since the declaration is in Engine, not UnrealEd
 #if !UE_BUILD_DOCS
diff --git a/Engine/Source/Runtime/Engine/Private/LightMap.cpp b/Engine/Source/Runtime/Engine/Private/LightMap.cpp
--- a/Engine/Source/Runtime/Engine/Private/LightMap.cpp
+++ b/Engine/Source/Runtime/Engine/Private/LightMap.cpp
@@ -7,7 +7,7 @@
 #include "EnginePrivate.h"
 #include "TargetPlatform.h"
 #include "StaticLighting.h"
-
+#include "LightMap.h"
 
 DEFINE_LOG_CATEGORY_STATIC(LogLightMap, Log, All);
 
diff --git a/Engine/Source/Runtime/Engine/Private/ModelLight.cpp b/Engine/Source/Runtime/Engine/Private/ModelLight.cpp
--- a/Engine/Source/Runtime/Engine/Private/ModelLight.cpp
+++ b/Engine/Source/Runtime/Engine/Private/ModelLight.cpp
@@ -9,6 +9,8 @@
 #include "ModelLight.h"
 #include "Collision.h"
 #include "Model.h"
+#include "LightMap.h"
+#include "ShadowMap.h"
 
 //
 //	Static vars
diff --git a/Engine/Source/Runtime/Engine/Private/StaticMeshLight.cpp b/Engine/Source/Runtime/Engine/Private/StaticMeshLight.cpp
--- a/Engine/Source/Runtime/Engine/Private/StaticMeshLight.cpp
+++ b/Engine/Source/Runtime/Engine/Private/StaticMeshLight.cpp
@@ -8,6 +8,8 @@
 #include "StaticMeshResources.h"
 #include "Raster.h"
 #include "StaticMeshLight.h"
+#include "LightMap.h"
+#include "ShadowMap.h"
 
 /**
  * Creates a static lighting vertex to represent the given static mesh vertex.
diff --git a/Engine/Source/Runtime/Engine/Public/StaticLighting.h b/Engine/Source/Runtime/Engine/Public/StaticLighting.h
--- a/Engine/Source/Runtime/Engine/Public/StaticLighting.h
+++ b/Engine/Source/Runtime/Engine/Public/StaticLighting.h
@@ -6,14 +6,10 @@
 
 #pragma once
 
-#include "VertexFactory.h"
-#include "LightMap.h"
-#include "ShadowMap.h"
 
 // Forward declarations.
 class FStaticLightingTextureMapping;
 class FStaticLightingMapping;
-class FLightMapData2D;
 
 /** The vertex data used to build static lighting. */
 struct FStaticLightingVertex
