Export FBX - Fixed timing issue when it doesn't export time correctly

diff --git a/Engine/Source/Editor/UnrealEd/Private/Fbx/FbxAnimationExport.cpp b/Engine/Source/Editor/UnrealEd/Private/Fbx/FbxAnimationExport.cpp
--- a/Engine/Source/Editor/UnrealEd/Private/Fbx/FbxAnimationExport.cpp
+++ b/Engine/Source/Editor/UnrealEd/Private/Fbx/FbxAnimationExport.cpp
@@ -56,6 +56,34 @@ void FFbxExporter::ExportAnimSequenceToFbx(const UAnimSequence* AnimSeq,
 {
 	USkeleton * Skeleton = AnimSeq->GetSkeleton();
 
+	if (AnimSeq->SequenceLength == 0.f)
+	{
+		// something is wrong
+		return;	
+	}
+
+	const float FrameRate			=  AnimSeq->NumFrames / AnimSeq->SequenceLength;
+
+	// set time correctly
+	FbxTime ExportedStartTime, ExportedStopTime;
+	if ( FMath::IsNearlyEqual(FrameRate, DEFAULT_SAMPLERATE, 1.f) )
+	{
+		ExportedStartTime.SetGlobalTimeMode(FbxTime::eFrames30);
+		ExportedStopTime.SetGlobalTimeMode(FbxTime::eFrames30);
+	}
+	else
+	{
+		ExportedStartTime.SetGlobalTimeMode(FbxTime::eCustom, FrameRate);
+		ExportedStopTime.SetGlobalTimeMode(FbxTime::eCustom, FrameRate);
+	}
+
+	ExportedStartTime.SetSecondDouble(0.f);
+	ExportedStopTime.SetSecondDouble(AnimSeq->SequenceLength);
+
+	FbxTimeSpan ExportedTimeSpan;
+	ExportedTimeSpan.Set(ExportedStartTime, ExportedStopTime);
+	AnimStack->SetLocalTimeSpan(ExportedTimeSpan);
+	
 	// Add the animation data to the bone nodes
 	for(int32 BoneIndex = 0; BoneIndex < BoneNodes.Num(); ++BoneIndex)
 	{
