Automation: Milestone 2 -  Editor build promotion tests

diff --git a/Engine/Config/BaseEngine.ini b/Engine/Config/BaseEngine.ini
--- a/Engine/Config/BaseEngine.ini
+++ b/Engine/Config/BaseEngine.ini
@@ -2024,6 +2024,8 @@ TestAllBlueprints=false
 
 [/Script/Engine.AutomationTestSettings]
 +EditorTestModules=StaticMeshEditor
++EditorTestModules=LandscapeEditor
++EditorTestModules=GameProjectGeneration
 +TestLevelFolders=TestMaps
 
 [AutomationTesting.StaticMeshEditorTest]
diff --git a/Engine/Source/Developer/SourceControl/Private/Tests/SourceControlTests.cpp b/Engine/Source/Developer/SourceControl/Private/Tests/SourceControlTests.cpp
--- a/Engine/Source/Developer/SourceControl/Private/Tests/SourceControlTests.cpp
+++ b/Engine/Source/Developer/SourceControl/Private/Tests/SourceControlTests.cpp
@@ -11,64 +11,7 @@
 
 #if WITH_EDITOR
 
-/**
- * Helper class for receiving the results of async source control operations
- */
-class FAsyncCommandHelper
-{
-public:
-	FAsyncCommandHelper( const FString& InParameter = FString() )
-		: Parameter(InParameter)
-		, bDispatched(false)
-		, bDone(false)
-		, bSuccessful(false)
-	{
-	}
-
-	void SourceControlOperationComplete( const FSourceControlOperationRef& Operation, ECommandResult::Type InResult )
-	{
-		bDone = true;
-		bSuccessful = InResult == ECommandResult::Succeeded;
-	}
-
-	const FString& GetParameter() const
-	{
-		return Parameter;
-	}
-
-	bool IsDispatched() const
-	{
-		return bDispatched;
-	}
-
-	void SetDispatched()
-	{
-		bDispatched = true;
-	}
-
-	bool IsDone() const
-	{
-		return bDone;
-	}
-
-	bool IsSuccessful() const
-	{
-		return bSuccessful;
-	}
-
-private:
-	/** Parameter we perform this operation with, if any */
-	FString Parameter;
-
-	/** Whether the async operation been issued */
-	bool bDispatched;
-
-	/** Whether the async operation has completed */
-	bool bDone;
-
-	/** Whether the operation was successful */
-	bool bSuccessful;
-};
+#include "Tests/SourceControlAutomationCommon.h"
 
 static void GetProviders(TArray<FString>& OutBeautifiedNames, TArray<FString>& OutTestCommands)
 {
@@ -156,14 +99,14 @@ bool FSetProviderTest::RunTest(const FString& Parameters)
 	return true;
 }
 
-DEFINE_LATENT_AUTOMATION_COMMAND_ONE_PARAMETER(FConnectLatentCommand, FAsyncCommandHelper, AsyncHelper);
+DEFINE_LATENT_AUTOMATION_COMMAND_ONE_PARAMETER(FConnectLatentCommand, SourceControlAutomationCommon::FAsyncCommandHelper, AsyncHelper);
 
 bool FConnectLatentCommand::Update()
 {
 	// attempt a login and wait for the result
 	if(!AsyncHelper.IsDispatched())
 	{
-		if(ISourceControlModule::Get().GetProvider().Login( FString(), EConcurrency::Asynchronous, FSourceControlOperationComplete::CreateRaw( &AsyncHelper, &FAsyncCommandHelper::SourceControlOperationComplete ) ) != ECommandResult::Succeeded)
+		if(ISourceControlModule::Get().GetProvider().Login( FString(), EConcurrency::Asynchronous, FSourceControlOperationComplete::CreateRaw( &AsyncHelper, &SourceControlAutomationCommon::FAsyncCommandHelper::SourceControlOperationComplete ) ) != ECommandResult::Succeeded)
 		{
 			return false;
 		}
@@ -184,12 +127,12 @@ bool FConnectTest::RunTest(const FString& Parameters)
 {
 	// parameter is the provider we want to use
 	ADD_LATENT_AUTOMATION_COMMAND(FSetProviderLatentCommand(FName(*Parameters)));
-	ADD_LATENT_AUTOMATION_COMMAND(FConnectLatentCommand(FAsyncCommandHelper()));
+	ADD_LATENT_AUTOMATION_COMMAND(FConnectLatentCommand(SourceControlAutomationCommon::FAsyncCommandHelper()));
 
 	return true;
 }
 
-DEFINE_LATENT_AUTOMATION_COMMAND_ONE_PARAMETER(FRevertLatentCommand, FAsyncCommandHelper, AsyncHelper);
+DEFINE_LATENT_AUTOMATION_COMMAND_ONE_PARAMETER(FRevertLatentCommand, SourceControlAutomationCommon::FAsyncCommandHelper, AsyncHelper);
 
 bool FRevertLatentCommand::Update()
 {
@@ -199,7 +142,7 @@ bool FRevertLatentCommand::Update()
 			ISourceControlOperation::Create<FRevert>(), 
 			SourceControlHelpers::PackageFilename(AsyncHelper.GetParameter()),
 			EConcurrency::Asynchronous, 
-			FSourceControlOperationComplete::CreateRaw( &AsyncHelper, &FAsyncCommandHelper::SourceControlOperationComplete ) 
+			FSourceControlOperationComplete::CreateRaw( &AsyncHelper, &SourceControlAutomationCommon::FAsyncCommandHelper::SourceControlOperationComplete ) 
 			) != ECommandResult::Succeeded)
 		{
 			return true;
@@ -231,7 +174,7 @@ bool FRevertLatentCommand::Update()
 	return AsyncHelper.IsDone();
 }
 
-DEFINE_LATENT_AUTOMATION_COMMAND_ONE_PARAMETER(FCheckOutLatentCommand, FAsyncCommandHelper, AsyncHelper);
+DEFINE_LATENT_AUTOMATION_COMMAND_ONE_PARAMETER(FCheckOutLatentCommand, SourceControlAutomationCommon::FAsyncCommandHelper, AsyncHelper);
 
 bool FCheckOutLatentCommand::Update()
 {
@@ -241,7 +184,7 @@ bool FCheckOutLatentCommand::Update()
 			ISourceControlOperation::Create<FCheckOut>(), 
 			SourceControlHelpers::PackageFilename(AsyncHelper.GetParameter()),
 			EConcurrency::Asynchronous, 
-			FSourceControlOperationComplete::CreateRaw( &AsyncHelper, &FAsyncCommandHelper::SourceControlOperationComplete ) 
+			FSourceControlOperationComplete::CreateRaw( &AsyncHelper, &SourceControlAutomationCommon::FAsyncCommandHelper::SourceControlOperationComplete ) 
 			) != ECommandResult::Succeeded)
 		{
 			return true;
@@ -290,9 +233,9 @@ bool FCheckOutTest::RunTest(const FString& Parameters)
 	bool bWasReadOnly = IFileManager::Get().IsReadOnly(*SourceControlHelpers::PackageFilename(ParamArray[1]));
 
 	ADD_LATENT_AUTOMATION_COMMAND(FSetProviderLatentCommand(FName(*ParamArray[0])));
-	ADD_LATENT_AUTOMATION_COMMAND(FConnectLatentCommand(FAsyncCommandHelper()));
-	ADD_LATENT_AUTOMATION_COMMAND(FCheckOutLatentCommand(FAsyncCommandHelper(ParamArray[1])));
-	ADD_LATENT_AUTOMATION_COMMAND(FRevertLatentCommand(FAsyncCommandHelper(ParamArray[1])));
+	ADD_LATENT_AUTOMATION_COMMAND(FConnectLatentCommand(SourceControlAutomationCommon::FAsyncCommandHelper()));
+	ADD_LATENT_AUTOMATION_COMMAND(FCheckOutLatentCommand(SourceControlAutomationCommon::FAsyncCommandHelper(ParamArray[1])));
+	ADD_LATENT_AUTOMATION_COMMAND(FRevertLatentCommand(SourceControlAutomationCommon::FAsyncCommandHelper(ParamArray[1])));
 
 	ADD_LATENT_AUTOMATION_COMMAND(FSetReadOnlyFlag(FReadOnlyState(ParamArray[1], bWasReadOnly)));
 
@@ -382,7 +325,7 @@ bool FDeletePackageLatentCommand::Update()
 	return true;
 }
 
-DEFINE_LATENT_AUTOMATION_COMMAND_ONE_PARAMETER(FMarkForAddLatentCommand, FAsyncCommandHelper, AsyncHelper);
+DEFINE_LATENT_AUTOMATION_COMMAND_ONE_PARAMETER(FMarkForAddLatentCommand, SourceControlAutomationCommon::FAsyncCommandHelper, AsyncHelper);
 
 bool FMarkForAddLatentCommand::Update()
 {
@@ -392,7 +335,7 @@ bool FMarkForAddLatentCommand::Update()
 			ISourceControlOperation::Create<FMarkForAdd>(), 
 			SourceControlHelpers::PackageFilename(AsyncHelper.GetParameter()),
 			EConcurrency::Asynchronous, 
-			FSourceControlOperationComplete::CreateRaw( &AsyncHelper, &FAsyncCommandHelper::SourceControlOperationComplete ) 
+			FSourceControlOperationComplete::CreateRaw( &AsyncHelper, &SourceControlAutomationCommon::FAsyncCommandHelper::SourceControlOperationComplete ) 
 			) != ECommandResult::Succeeded)
 		{
 			return true;
@@ -438,14 +381,14 @@ bool FMarkForAddTest::RunTest(const FString& Parameters)
 	ensure(ParamArray.Num() == 2);
 
 	ADD_LATENT_AUTOMATION_COMMAND(FSetProviderLatentCommand(FName(*ParamArray[0])));
-	ADD_LATENT_AUTOMATION_COMMAND(FConnectLatentCommand(FAsyncCommandHelper()));
+	ADD_LATENT_AUTOMATION_COMMAND(FConnectLatentCommand(SourceControlAutomationCommon::FAsyncCommandHelper()));
 
 	struct Local
 	{
 		static void AddDependentCommands(const FString& InParameter)
 		{
-			ADD_LATENT_AUTOMATION_COMMAND(FMarkForAddLatentCommand(FAsyncCommandHelper(InParameter)));
-			ADD_LATENT_AUTOMATION_COMMAND(FRevertLatentCommand(FAsyncCommandHelper(InParameter)));
+			ADD_LATENT_AUTOMATION_COMMAND(FMarkForAddLatentCommand(SourceControlAutomationCommon::FAsyncCommandHelper(InParameter)));
+			ADD_LATENT_AUTOMATION_COMMAND(FRevertLatentCommand(SourceControlAutomationCommon::FAsyncCommandHelper(InParameter)));
 			ADD_LATENT_AUTOMATION_COMMAND(FDeletePackageLatentCommand(InParameter));
 		}
 	};
@@ -455,7 +398,7 @@ bool FMarkForAddTest::RunTest(const FString& Parameters)
 	return true;
 }
 
-DEFINE_LATENT_AUTOMATION_COMMAND_ONE_PARAMETER(FDeleteLatentCommand, FAsyncCommandHelper, AsyncHelper);
+DEFINE_LATENT_AUTOMATION_COMMAND_ONE_PARAMETER(FDeleteLatentCommand, SourceControlAutomationCommon::FAsyncCommandHelper, AsyncHelper);
 
 bool FDeleteLatentCommand::Update()
 {
@@ -465,7 +408,7 @@ bool FDeleteLatentCommand::Update()
 			ISourceControlOperation::Create<FDelete>(), 
 			SourceControlHelpers::PackageFilename(AsyncHelper.GetParameter()),
 			EConcurrency::Asynchronous, 
-			FSourceControlOperationComplete::CreateRaw( &AsyncHelper, &FAsyncCommandHelper::SourceControlOperationComplete ) 
+			FSourceControlOperationComplete::CreateRaw( &AsyncHelper, &SourceControlAutomationCommon::FAsyncCommandHelper::SourceControlOperationComplete ) 
 			) != ECommandResult::Succeeded)
 		{
 			return true;
@@ -516,16 +459,16 @@ bool FDeleteTest::RunTest(const FString& Parameters)
 	FString AbsoluteFilename = SourceControlHelpers::PackageFilename(ParamArray[1]);
 
 	ADD_LATENT_AUTOMATION_COMMAND(FSetProviderLatentCommand(FName(*ParamArray[0])));
-	ADD_LATENT_AUTOMATION_COMMAND(FConnectLatentCommand(FAsyncCommandHelper()));
-	ADD_LATENT_AUTOMATION_COMMAND(FDeleteLatentCommand(FAsyncCommandHelper(AbsoluteFilename)));
-	ADD_LATENT_AUTOMATION_COMMAND(FRevertLatentCommand(FAsyncCommandHelper(AbsoluteFilename)));
+	ADD_LATENT_AUTOMATION_COMMAND(FConnectLatentCommand(SourceControlAutomationCommon::FAsyncCommandHelper()));
+	ADD_LATENT_AUTOMATION_COMMAND(FDeleteLatentCommand(SourceControlAutomationCommon::FAsyncCommandHelper(AbsoluteFilename)));
+	ADD_LATENT_AUTOMATION_COMMAND(FRevertLatentCommand(SourceControlAutomationCommon::FAsyncCommandHelper(AbsoluteFilename)));
 
 	ADD_LATENT_AUTOMATION_COMMAND(FSetReadOnlyFlag(FReadOnlyState(ParamArray[1], bWasReadOnly)));
 
 	return true;
 }
 
-DEFINE_LATENT_AUTOMATION_COMMAND_ONE_PARAMETER(FCheckInLatentCommand, FAsyncCommandHelper, AsyncHelper);
+DEFINE_LATENT_AUTOMATION_COMMAND_ONE_PARAMETER(FCheckInLatentCommand, SourceControlAutomationCommon::FAsyncCommandHelper, AsyncHelper);
 
 bool FCheckInLatentCommand::Update()
 {
@@ -538,7 +481,7 @@ bool FCheckInLatentCommand::Update()
 			CheckInOperation, 
 			SourceControlHelpers::PackageFilename(AsyncHelper.GetParameter()),
 			EConcurrency::Asynchronous, 
-			FSourceControlOperationComplete::CreateRaw( &AsyncHelper, &FAsyncCommandHelper::SourceControlOperationComplete ) 
+			FSourceControlOperationComplete::CreateRaw( &AsyncHelper, &SourceControlAutomationCommon::FAsyncCommandHelper::SourceControlOperationComplete ) 
 			) != ECommandResult::Succeeded)
 		{
 			return true;
@@ -616,17 +559,17 @@ bool FCheckInTest::RunTest(const FString& Parameters)
 
 	// parameter is the provider we want to use
 	ADD_LATENT_AUTOMATION_COMMAND(FSetProviderLatentCommand(FName(*ParamArray[0])));
-	ADD_LATENT_AUTOMATION_COMMAND(FConnectLatentCommand(FAsyncCommandHelper()));
-	ADD_LATENT_AUTOMATION_COMMAND(FCheckOutLatentCommand(FAsyncCommandHelper(ParamArray[1])));
+	ADD_LATENT_AUTOMATION_COMMAND(FConnectLatentCommand(SourceControlAutomationCommon::FAsyncCommandHelper()));
+	ADD_LATENT_AUTOMATION_COMMAND(FCheckOutLatentCommand(SourceControlAutomationCommon::FAsyncCommandHelper(ParamArray[1])));
 	ADD_LATENT_AUTOMATION_COMMAND(FEditTextureLatentCommand(ParamArray[1]));
-	ADD_LATENT_AUTOMATION_COMMAND(FCheckInLatentCommand(FAsyncCommandHelper(ParamArray[1])));
+	ADD_LATENT_AUTOMATION_COMMAND(FCheckInLatentCommand(SourceControlAutomationCommon::FAsyncCommandHelper(ParamArray[1])));
 
 	ADD_LATENT_AUTOMATION_COMMAND(FSetReadOnlyFlag(FReadOnlyState(ParamArray[1], bWasReadOnly)));
 
 	return true;
 }
 
-DEFINE_LATENT_AUTOMATION_COMMAND_ONE_PARAMETER(FSyncLatentCommand, FAsyncCommandHelper, AsyncHelper);
+DEFINE_LATENT_AUTOMATION_COMMAND_ONE_PARAMETER(FSyncLatentCommand, SourceControlAutomationCommon::FAsyncCommandHelper, AsyncHelper);
 
 bool FSyncLatentCommand::Update()
 {
@@ -636,7 +579,7 @@ bool FSyncLatentCommand::Update()
 			ISourceControlOperation::Create<FSync>(), 
 			SourceControlHelpers::PackageFilename(AsyncHelper.GetParameter()),
 			EConcurrency::Asynchronous, 
-			FSourceControlOperationComplete::CreateRaw( &AsyncHelper, &FAsyncCommandHelper::SourceControlOperationComplete ) 
+			FSourceControlOperationComplete::CreateRaw( &AsyncHelper, &SourceControlAutomationCommon::FAsyncCommandHelper::SourceControlOperationComplete ) 
 			) != ECommandResult::Succeeded)
 		{
 			return true;
@@ -682,8 +625,8 @@ bool FSyncTest::RunTest(const FString& Parameters)
 	ensure(ParamArray.Num() == 2);
 
 	ADD_LATENT_AUTOMATION_COMMAND(FSetProviderLatentCommand(FName(*ParamArray[0])));
-	ADD_LATENT_AUTOMATION_COMMAND(FConnectLatentCommand(FAsyncCommandHelper()));
-	ADD_LATENT_AUTOMATION_COMMAND(FSyncLatentCommand(FAsyncCommandHelper(ParamArray[1])));
+	ADD_LATENT_AUTOMATION_COMMAND(FConnectLatentCommand(SourceControlAutomationCommon::FAsyncCommandHelper()));
+	ADD_LATENT_AUTOMATION_COMMAND(FSyncLatentCommand(SourceControlAutomationCommon::FAsyncCommandHelper(ParamArray[1])));
 
 	return true;
 }
@@ -705,14 +648,14 @@ bool FRevertTest::RunTest(const FString& Parameters)
 	ensure(ParamArray.Num() == 2);
 
 	ADD_LATENT_AUTOMATION_COMMAND(FSetProviderLatentCommand(FName(*ParamArray[0])));
-	ADD_LATENT_AUTOMATION_COMMAND(FConnectLatentCommand(FAsyncCommandHelper()));
+	ADD_LATENT_AUTOMATION_COMMAND(FConnectLatentCommand(SourceControlAutomationCommon::FAsyncCommandHelper()));
 
 	struct Local
 	{
 		static void AddDepedentCommands(const FString& InParameter)
 		{
-			ADD_LATENT_AUTOMATION_COMMAND(FMarkForAddLatentCommand(FAsyncCommandHelper(InParameter)));
-			ADD_LATENT_AUTOMATION_COMMAND(FRevertLatentCommand(FAsyncCommandHelper(InParameter)));
+			ADD_LATENT_AUTOMATION_COMMAND(FMarkForAddLatentCommand(SourceControlAutomationCommon::FAsyncCommandHelper(InParameter)));
+			ADD_LATENT_AUTOMATION_COMMAND(FRevertLatentCommand(SourceControlAutomationCommon::FAsyncCommandHelper(InParameter)));
 			ADD_LATENT_AUTOMATION_COMMAND(FDeletePackageLatentCommand(InParameter));
 		}
 	};
@@ -722,7 +665,7 @@ bool FRevertTest::RunTest(const FString& Parameters)
 	return true;
 }
 
-DEFINE_LATENT_AUTOMATION_COMMAND_ONE_PARAMETER(FUpdateStatusLatentCommand, FAsyncCommandHelper, AsyncHelper);
+DEFINE_LATENT_AUTOMATION_COMMAND_ONE_PARAMETER(FUpdateStatusLatentCommand, SourceControlAutomationCommon::FAsyncCommandHelper, AsyncHelper);
 
 bool FUpdateStatusLatentCommand::Update()
 {
@@ -736,7 +679,7 @@ bool FUpdateStatusLatentCommand::Update()
 			UpdateStatusOperation, 
 			SourceControlHelpers::PackageFilename(AsyncHelper.GetParameter()),
 			EConcurrency::Asynchronous, 
-			FSourceControlOperationComplete::CreateRaw( &AsyncHelper, &FAsyncCommandHelper::SourceControlOperationComplete ) 
+			FSourceControlOperationComplete::CreateRaw( &AsyncHelper, &SourceControlAutomationCommon::FAsyncCommandHelper::SourceControlOperationComplete ) 
 			) != ECommandResult::Succeeded)
 		{
 			return true;
@@ -800,9 +743,9 @@ bool FUpdateStatusTest::RunTest(const FString& Parameters)
 	ensure(ParamArray.Num() == 2);
 
 	ADD_LATENT_AUTOMATION_COMMAND(FSetProviderLatentCommand(FName(*ParamArray[0])));
-	ADD_LATENT_AUTOMATION_COMMAND(FConnectLatentCommand(FAsyncCommandHelper()));
-	ADD_LATENT_AUTOMATION_COMMAND(FCheckOutLatentCommand(FAsyncCommandHelper(ParamArray[1])));
-	ADD_LATENT_AUTOMATION_COMMAND(FUpdateStatusLatentCommand(FAsyncCommandHelper(ParamArray[1])));
+	ADD_LATENT_AUTOMATION_COMMAND(FConnectLatentCommand(SourceControlAutomationCommon::FAsyncCommandHelper()));
+	ADD_LATENT_AUTOMATION_COMMAND(FCheckOutLatentCommand(SourceControlAutomationCommon::FAsyncCommandHelper(ParamArray[1])));
+	ADD_LATENT_AUTOMATION_COMMAND(FUpdateStatusLatentCommand(SourceControlAutomationCommon::FAsyncCommandHelper(ParamArray[1])));
 	ADD_LATENT_AUTOMATION_COMMAND(FGetStateLatentCommand(ParamArray[1]));
 	ADD_LATENT_AUTOMATION_COMMAND(FRevertLatentCommand(ParamArray[1]));
 
@@ -887,7 +830,7 @@ bool FGetLabelTest::RunTest(const FString& Parameters)
 	FilesToGet.Add(FPaths::ConvertRelativePathToFull(TEXT("../../../Engine/Source/Developer/SourceControl/SourceControl.Build.cs")));
 
 	ADD_LATENT_AUTOMATION_COMMAND(FSetProviderLatentCommand(FName(*ParamArray[0])));
-	ADD_LATENT_AUTOMATION_COMMAND(FConnectLatentCommand(FAsyncCommandHelper()));
+	ADD_LATENT_AUTOMATION_COMMAND(FConnectLatentCommand(SourceControlAutomationCommon::FAsyncCommandHelper()));
 	ADD_LATENT_AUTOMATION_COMMAND(FGetLabelLatentCommand(FLabelAndFilenames(ParamArray[1], FilesToGet)));
 
 	return true;
@@ -933,7 +876,7 @@ bool FSyncLabelTest::RunTest(const FString& Parameters)
 	FilesToGet.Add(FPaths::ConvertRelativePathToFull(TEXT("../../../Engine/Source/Developer/SourceControl/Public/ISourceControlModule.h")));
 
 	ADD_LATENT_AUTOMATION_COMMAND(FSetProviderLatentCommand(FName(*ParamArray[0])));
-	ADD_LATENT_AUTOMATION_COMMAND(FConnectLatentCommand(FAsyncCommandHelper()));
+	ADD_LATENT_AUTOMATION_COMMAND(FConnectLatentCommand(SourceControlAutomationCommon::FAsyncCommandHelper()));
 	ADD_LATENT_AUTOMATION_COMMAND(FSyncLabelLatentCommand(FLabelAndFilenames(ParamArray[1], FilesToGet)));
 
 	return true;
@@ -995,8 +938,8 @@ bool FGetRevisionTest::RunTest(const FString& Parameters)
 	ensure(ParamArray.Num() == 2);
 
 	ADD_LATENT_AUTOMATION_COMMAND(FSetProviderLatentCommand(FName(*ParamArray[0])));
-	ADD_LATENT_AUTOMATION_COMMAND(FConnectLatentCommand(FAsyncCommandHelper()));
-	ADD_LATENT_AUTOMATION_COMMAND(FUpdateStatusLatentCommand(FAsyncCommandHelper(ParamArray[1])));
+	ADD_LATENT_AUTOMATION_COMMAND(FConnectLatentCommand(SourceControlAutomationCommon::FAsyncCommandHelper()));
+	ADD_LATENT_AUTOMATION_COMMAND(FUpdateStatusLatentCommand(SourceControlAutomationCommon::FAsyncCommandHelper(ParamArray[1])));
 	ADD_LATENT_AUTOMATION_COMMAND(FGetRevisionLatentCommand(ParamArray[1]));
 
 	return true;
diff --git a/Engine/Source/Developer/SourceControl/Public/Tests/SourceControlAutomationCommon.h b/Engine/Source/Developer/SourceControl/Public/Tests/SourceControlAutomationCommon.h
new file mode 100644
--- /dev/null
+++ b/Engine/Source/Developer/SourceControl/Public/Tests/SourceControlAutomationCommon.h
@@ -0,0 +1,65 @@
+// Copyright 1998-2014 Epic Games, Inc. All Rights Reserved.
+
+#pragma once
+
+namespace SourceControlAutomationCommon
+{
+	/**
+	* Helper class for receiving the results of async source control operations
+	*/
+	class FAsyncCommandHelper
+	{
+	public:
+		FAsyncCommandHelper(const FString& InParameter = FString())
+			: Parameter(InParameter)
+			, bDispatched(false)
+			, bDone(false)
+			, bSuccessful(false)
+		{
+		}
+
+		void SourceControlOperationComplete(const FSourceControlOperationRef& Operation, ECommandResult::Type InResult)
+		{
+			bDone = true;
+			bSuccessful = InResult == ECommandResult::Succeeded;
+		}
+
+		const FString& GetParameter() const
+		{
+			return Parameter;
+		}
+
+		bool IsDispatched() const
+		{
+			return bDispatched;
+		}
+
+		void SetDispatched()
+		{
+			bDispatched = true;
+		}
+
+		bool IsDone() const
+		{
+			return bDone;
+		}
+
+		bool IsSuccessful() const
+		{
+			return bSuccessful;
+		}
+
+	private:
+		/** Parameter we perform this operation with, if any */
+		FString Parameter;
+
+		/** Whether the async operation been issued */
+		bool bDispatched;
+
+		/** Whether the async operation has completed */
+		bool bDone;
+
+		/** Whether the operation was successful */
+		bool bSuccessful;
+	};
+}
diff --git a/Engine/Source/Editor/GameProjectGeneration/Private/Tests/GameProjectAutomationTests.cpp b/Engine/Source/Editor/GameProjectGeneration/Private/Tests/GameProjectAutomationTests.cpp
new file mode 100644
--- /dev/null
+++ b/Engine/Source/Editor/GameProjectGeneration/Private/Tests/GameProjectAutomationTests.cpp
@@ -0,0 +1,201 @@
+// Copyright 1998-2014 Epic Games, Inc. All Rights Reserved.
+
+#include "../GameProjectGenerationPrivatePCH.h"
+
+#include "DesktopPlatformModule.h"
+#include "TargetPlatform.h"
+#include "../GameProjectUtils.h"
+#include "Tests/AutomationTestSettings.h"
+#include "AutomationEditorCommon.h"
+
+DEFINE_LOG_CATEGORY_STATIC(LogGameProjectGenerationTests, Log, All);
+
+namespace GameProjectAutomationUtils
+{
+	/**
+	* Generates the desired project file name
+	*/
+	static FString GetDesiredProjectFilename()
+	{
+		UAutomationTestSettings const* AutomationTestSettings = GetDefault<UAutomationTestSettings>();
+		check(AutomationTestSettings);
+
+		FString ProjectName;
+		const FString ProjectNameOverride = AutomationTestSettings->BuildPromotionTest.NewProjectSettings.NewProjectNameOverride;
+		if (ProjectNameOverride.Len())
+		{
+			ProjectName = ProjectNameOverride;
+		}
+		else
+		{
+			ProjectName = TEXT("NewTestProject");
+		}
+
+		FString ProjectPath;
+		const FString ProjectPathOverride = AutomationTestSettings->BuildPromotionTest.NewProjectSettings.NewProjectFolderOverride.Path;
+		if (ProjectPathOverride.Len())
+		{
+			ProjectPath = ProjectPathOverride;
+		}
+		else
+		{
+			ProjectPath = IFileManager::Get().ConvertToAbsolutePathForExternalAppForWrite(*FDesktopPlatformModule::Get()->GetDefaultProjectCreationPath());
+		}
+
+
+		const FString Filename = ProjectName + TEXT(".") + IProjectManager::GetProjectFileExtension();
+		FString ProjectFilename = FPaths::Combine(*ProjectPath, *ProjectName, *Filename);
+		FPaths::MakePlatformFilename(ProjectFilename);
+
+		return ProjectFilename;
+	}
+}
+
+/**
+* Automation test to clean up old test project files
+*/
+IMPLEMENT_SIMPLE_AUTOMATION_TEST(FBuildPromotionNewProjectCleanupTest, "Promotion.Project Promotion Pass.Step 1 Blank Project Creation.Cleanup Potential Project Location", EAutomationTestFlags::ATF_Editor);
+bool FBuildPromotionNewProjectCleanupTest::RunTest(const FString& Parameters)
+{
+	FString DesiredProjectFilename = GameProjectAutomationUtils::GetDesiredProjectFilename();
+
+	if (FPaths::FileExists(DesiredProjectFilename))
+	{
+		UE_LOG(LogGameProjectGenerationTests, Display, TEXT("Found an old project file at %s"), *DesiredProjectFilename);
+		if (FPaths::IsProjectFilePathSet())
+		{
+			FString CurrentProjectPath = FPaths::GetProjectFilePath();
+			if (CurrentProjectPath == DesiredProjectFilename)
+			{
+				UE_LOG(LogGameProjectGenerationTests, Warning, TEXT("Can not clean up the target project location because it is the current active project."));
+			}
+			else
+			{
+				UE_LOG(LogGameProjectGenerationTests, Display, TEXT("Removing files from old project path: %s"), *FPaths::GetPath(DesiredProjectFilename));
+				bool bEnsureExists = false;
+				bool bDeleteEntireTree = true;
+				IFileManager::Get().DeleteDirectory(*FPaths::GetPath(DesiredProjectFilename), bEnsureExists, bDeleteEntireTree);
+			}
+		}
+	}
+	else
+	{
+		UE_LOG(LogGameProjectGenerationTests, Display, TEXT("Target project location is clear"));
+	}
+
+	return true;
+}
+
+/**
+* Automation test to create a new project
+*/
+IMPLEMENT_SIMPLE_AUTOMATION_TEST(FBuildPromotionNewProjectCreateTest, "Promotion.Project Promotion Pass.Step 1 Blank Project Creation.Create Project", EAutomationTestFlags::ATF_Editor);
+bool FBuildPromotionNewProjectCreateTest::RunTest(const FString& Parameters)
+{
+	FString DesiredProjectFilename = GameProjectAutomationUtils::GetDesiredProjectFilename();
+
+	if (FPaths::FileExists(DesiredProjectFilename))
+	{
+		UE_LOG(LogGameProjectGenerationTests, Warning, TEXT("A project already exists at the target location: %s"), *DesiredProjectFilename);
+		const FString OldProjectFolder = FPaths::GetPath(DesiredProjectFilename);
+		const FString OldProjectName = FPaths::GetBaseFilename(DesiredProjectFilename);
+		const FString RootFolder = FPaths::GetPath(OldProjectFolder);
+
+		//Add a number to the end
+		for (uint32 i = 2;; ++i)
+		{
+			const FString PossibleProjectName = FString::Printf(TEXT("%s%i"), *OldProjectName, i);
+			const FString PossibleProjectFilename = FString::Printf(TEXT("%s/%s/%s.%s"), *RootFolder, *PossibleProjectName, *PossibleProjectName, *IProjectManager::GetProjectFileExtension());
+			if (!FPaths::FileExists(PossibleProjectFilename))
+			{
+				DesiredProjectFilename = PossibleProjectFilename;
+				UE_LOG(LogGameProjectGenerationTests, Warning, TEXT("Changing the target project name to: %s"), *FPaths::GetBaseFilename(DesiredProjectFilename));
+				break;
+			}
+		}
+	}
+
+	FText FailReason;
+	if (GameProjectUtils::CreateProject(DesiredProjectFilename, TEXT(""), false, true, FailReason))
+	{
+		UE_LOG(LogGameProjectGenerationTests, Display, TEXT("Generated a new project: %s"), *DesiredProjectFilename);
+		UE_LOG(LogGameProjectGenerationTests, Display, TEXT("Test successful!"));
+		UE_LOG(LogGameProjectGenerationTests, Display, TEXT("\nPlease switch to the new project and continue to Step 2."));
+	}
+	else
+	{
+		UE_LOG(LogGameProjectGenerationTests, Error, TEXT("Could not generate new project: %s"), *FailReason.ToString());
+		UE_LOG(LogGameProjectGenerationTests, Display, TEXT("Test failed!"));
+	}
+
+	return true;
+}
+
+/**
+* Automation test to create a simple level and save it
+*/
+IMPLEMENT_SIMPLE_AUTOMATION_TEST(FBuildPromotionNewProjectMapTest, "Promotion.Project Promotion Pass.Step 2 Basic Level Creation.Create Basic Level", EAutomationTestFlags::ATF_Editor);
+bool FBuildPromotionNewProjectMapTest::RunTest(const FString& Parameters)
+{
+	//New level
+	UWorld* CurrentWorld = AutomationEditorCommonUtils::CreateNewMap();
+	if (CurrentWorld)
+	{
+		UE_LOG(LogGameProjectGenerationTests, Display, TEXT("Adding Level Geometry"));
+	}
+	else
+	{
+		UE_LOG(LogGameProjectGenerationTests, Error, TEXT("Failed to create an empty level"));
+	}
+
+	
+
+	//Add some bsp and a player start
+	GEditor->Exec(CurrentWorld, TEXT("BRUSH Scale 1 1 1"));
+	for (int32 i = 0; i < GEditor->LevelViewportClients.Num(); i++)
+	{
+		FLevelEditorViewportClient* ViewportClient = GEditor->LevelViewportClients[i];
+		if (!ViewportClient->IsOrtho())
+		{
+			ViewportClient->SetViewLocation(FVector(176, 2625, 2075));
+			ViewportClient->SetViewRotation(FRotator(319, 269, 1));
+		}
+	}
+	ULevel* CurrentLevel = CurrentWorld->GetCurrentLevel();
+
+	//Cube Additive Brush
+	UCubeBuilder* CubeAdditiveBrushBuilder = Cast<UCubeBuilder>(GEditor->FindBrushBuilder(UCubeBuilder::StaticClass()));
+	CubeAdditiveBrushBuilder->X = 4096.0f;
+	CubeAdditiveBrushBuilder->Y = 4096.0f;
+	CubeAdditiveBrushBuilder->Z = 128.0f;
+	CubeAdditiveBrushBuilder->Build(CurrentWorld);
+	GEditor->Exec(CurrentWorld, TEXT("BRUSH MOVETO X=0 Y=0 Z=0"));
+	GEditor->Exec(CurrentWorld, TEXT("BRUSH ADD"));
+
+	//Add a playerstart
+	const FTransform Transform(FRotator(-16384, 0, 0), FVector(0.f, 1750.f, 166.f));
+	AActor* PlayerStart = GEditor->AddActor(CurrentWorld->GetCurrentLevel(), APlayerStart::StaticClass(), Transform);
+	if (PlayerStart)
+	{
+		UE_LOG(LogGameProjectGenerationTests, Display, TEXT("Added a player start"));
+	}
+	else
+	{
+		UE_LOG(LogGameProjectGenerationTests, Error, TEXT("Failed to add a player start"));
+	}
+
+	// Save the map
+	FEditorFileUtils::SaveLevel(CurrentLevel, TEXT("/Game/Maps/NewProjectTest"));
+	UE_LOG(LogGameProjectGenerationTests, Display, TEXT("Saved map"));
+
+	if (ExecutionInfo.Errors.Num() > 0)
+	{
+		UE_LOG(LogGameProjectGenerationTests, Display, TEXT("Test failed!"));
+	}
+	else
+	{
+		UE_LOG(LogGameProjectGenerationTests, Display, TEXT("Test successful!"));
+	}
+
+	return true;
+}
\ No newline at end of file
diff --git a/Engine/Source/Editor/LandscapeEditor/Private/Tests/LandscapeAutomationTests.cpp b/Engine/Source/Editor/LandscapeEditor/Private/Tests/LandscapeAutomationTests.cpp
new file mode 100644
--- /dev/null
+++ b/Engine/Source/Editor/LandscapeEditor/Private/Tests/LandscapeAutomationTests.cpp
@@ -0,0 +1,172 @@
+// Copyright 1998-2014 Epic Games, Inc. All Rights Reserved.
+
+#include "../LandscapeEditorPrivatePCH.h"
+
+#include "AutomationCommon.h"
+#include "AutomationEditorCommon.h"
+
+#include "LandscapeEdMode.h"
+#include "../LandscapeEditorDetailCustomization_NewLandscape.h"
+#include "ScopedTransaction.h"
+#include "Landscape/Landscape.h"
+
+#define LOCTEXT_NAMESPACE "LandscapeEditor.NewLandscape"
+
+DEFINE_LOG_CATEGORY_STATIC(LogLandscapeAutomationTests, Log, All);
+
+/**
+* Landscape test helper functions
+*/
+namespace LandscapeTestUtils
+{
+	/**
+	* Finds the viewport to use for the landscape tool
+	*/
+	static FLevelEditorViewportClient* FindSelectedViewport()
+	{
+		FLevelEditorViewportClient* SelectedViewport = NULL;
+
+		for (int32 i = 0; i < GEditor->LevelViewportClients.Num(); i++)
+		{
+			FLevelEditorViewportClient* ViewportClient = GEditor->LevelViewportClients[i];
+			if (!ViewportClient->IsOrtho())
+			{
+				SelectedViewport = ViewportClient;
+			}
+		}
+
+		return SelectedViewport;
+	}
+}
+
+/**
+* Latent command to create a new landscape
+*/
+DEFINE_LATENT_AUTOMATION_COMMAND(FCreateLandscapeCommand);
+bool FCreateLandscapeCommand::Update()
+{
+	//Switch to the Landscape tool
+	GLevelEditorModeTools().ActivateMode(FBuiltinEditorModes::EM_Landscape);
+	FEdModeLandscape* LandscapeEdMode = (FEdModeLandscape*)GLevelEditorModeTools().GetActiveMode(FBuiltinEditorModes::EM_Landscape);
+
+	//Modify the "Section size"
+	LandscapeEdMode->UISettings->NewLandscape_QuadsPerSection = 7;
+	LandscapeEdMode->UISettings->NewLandscape_ClampSize();
+
+	//Create the landscape
+	TSharedPtr<FLandscapeEditorDetailCustomization_NewLandscape> Customization_NewLandscape = MakeShareable(new FLandscapeEditorDetailCustomization_NewLandscape);
+	Customization_NewLandscape->OnCreateButtonClicked();
+
+	if (LandscapeEdMode->CurrentToolTarget.LandscapeInfo.IsValid())
+	{
+		UE_LOG(LogLandscapeAutomationTests, Display, TEXT("Created a new landscape"));
+	}
+	else
+	{
+		UE_LOG(LogLandscapeAutomationTests, Error, TEXT("Failed to create a new landscape"));
+	}
+
+	return true;
+}
+
+/**
+* Latent command to start using the sculpting tool
+*/
+DEFINE_LATENT_AUTOMATION_COMMAND(FBeginModifyLandscapeCommand);
+bool FBeginModifyLandscapeCommand::Update()
+{
+	//Find the landscape
+	FEdModeLandscape* LandscapeEdMode = (FEdModeLandscape*)GLevelEditorModeTools().GetActiveMode(FBuiltinEditorModes::EM_Landscape);
+
+	//Find a location on the edge of the landscape along the x axis so the default camera can see it in the distance.
+	FVector LandscapeSizePerComponent = LandscapeEdMode->UISettings->NewLandscape_QuadsPerSection * LandscapeEdMode->UISettings->NewLandscape_SectionsPerComponent * LandscapeEdMode->UISettings->NewLandscape_Scale;
+	FVector TargetLoctaion(0);
+	TargetLoctaion.X = -LandscapeSizePerComponent.X * (LandscapeEdMode->UISettings->NewLandscape_ComponentCount.X / 2);
+
+	ALandscapeProxy* Proxy = LandscapeEdMode->CurrentToolTarget.LandscapeInfo.Get()->GetCurrentLevelLandscapeProxy(true);
+	if (Proxy)
+	{
+		TargetLoctaion = Proxy->LandscapeActorToWorld().InverseTransformPosition(TargetLoctaion);
+	}
+
+	//Begin using the sculpting tool
+	FLevelEditorViewportClient* SelectedViewport = LandscapeTestUtils::FindSelectedViewport();
+	LandscapeEdMode->CurrentTool->BeginTool(SelectedViewport, LandscapeEdMode->CurrentToolTarget, TargetLoctaion);
+	SelectedViewport->Invalidate();
+
+	UE_LOG(LogLandscapeAutomationTests, Display, TEXT("Modified the landscape using the sculpt tool"));
+
+	return true;
+}
+
+/**
+*  Latent command stop using the sculpting tool
+*/
+DEFINE_LATENT_AUTOMATION_COMMAND(FEndModifyLandscapeCommand);
+bool FEndModifyLandscapeCommand::Update()
+{
+	//Find the landscape
+	FEdModeLandscape* LandscapeEdMode = (FEdModeLandscape*)GLevelEditorModeTools().GetActiveMode(FBuiltinEditorModes::EM_Landscape);
+
+	//End using the sculpting tool
+	FLevelEditorViewportClient* SelectedViewport = LandscapeTestUtils::FindSelectedViewport();
+	LandscapeEdMode->CurrentTool->EndTool(SelectedViewport);
+	return true;
+}
+
+/**
+* Landscape creation / edit test
+*/
+IMPLEMENT_SIMPLE_AUTOMATION_TEST(FLandscapeEditorTest, "Promotion.Editor Promotion Pass.Step 1 Main Editor Test.New Landscape", EAutomationTestFlags::ATF_Editor);
+bool FLandscapeEditorTest::RunTest(const FString& Parameters)
+{
+	//New level
+	UWorld* NewMap = AutomationEditorCommonUtils::CreateNewMap();
+	if (NewMap)
+	{
+		UE_LOG(LogLandscapeAutomationTests, Display, TEXT("Created an empty level"));
+	}
+	else
+	{
+		UE_LOG(LogLandscapeAutomationTests, Error, TEXT("Failed to create an empty level"));
+		return false;
+	}
+
+	ADD_LATENT_AUTOMATION_COMMAND(FCreateLandscapeCommand());
+
+	//For some reason the heightmap component takes a few ticks to register with the nav system.  We crash if we try to modify the heightmap before then.
+	ADD_LATENT_AUTOMATION_COMMAND(FWaitLatentCommand(1.0f));
+	ADD_LATENT_AUTOMATION_COMMAND(FBeginModifyLandscapeCommand());
+	ADD_LATENT_AUTOMATION_COMMAND(FEndModifyLandscapeCommand());
+
+	//Wait for the landscape to update before the screenshot
+	ADD_LATENT_AUTOMATION_COMMAND(FWaitLatentCommand(1.0f));
+
+	//Update the screenshot name, then take a screenshot.
+	if (FAutomationTestFramework::GetInstance().IsScreenshotAllowed())
+	{
+		WindowScreenshotParameters ScreenshotParameters;
+
+		//Find the main editor window
+		TArray<TSharedRef<SWindow> > AllWindows;
+		FSlateApplication::Get().GetAllVisibleWindowsOrdered(AllWindows);
+		if (AllWindows.Num() == 0)
+		{
+			UE_LOG(LogLandscapeAutomationTests, Error, TEXT("ERROR: Could not find the main editor window."));
+			return true;
+		}
+
+		//Create the path
+		ScreenshotParameters.CurrentWindow = AllWindows[0];
+		const FString TestName = TEXT("NewLandscapeTest");
+		FString PathName = FPaths::AutomationDir() + TestName / FPlatformProperties::PlatformName();
+		FPaths::MakePathRelativeTo(PathName, *FPaths::RootDir());
+		ScreenshotParameters.ScreenshotName = FString::Printf(TEXT("%s/%d.png"), *PathName, GEngineVersion.GetChangelist());
+
+		//Take a screenshot
+		ADD_LATENT_AUTOMATION_COMMAND(FTakeEditorScreenshotCommand(ScreenshotParameters));
+	}
+	return true;
+}
+
+#undef LOCTEXT_NAMESPACE
diff --git a/Engine/Source/Editor/UnrealEd/Private/Kismet2/KismetDebugUtilities.cpp b/Engine/Source/Editor/UnrealEd/Private/Kismet2/KismetDebugUtilities.cpp
--- a/Engine/Source/Editor/UnrealEd/Private/Kismet2/KismetDebugUtilities.cpp
+++ b/Engine/Source/Editor/UnrealEd/Private/Kismet2/KismetDebugUtilities.cpp
@@ -391,7 +391,7 @@ void FKismetDebugUtilities::AttemptToBreakExecution(UBlueprint* BlueprintObj, co
 		UE_LOG(LogBlueprintDebug, Warning, TEXT("Tried to break execution in an unknown spot at object %s:%04X"), *StackFrame.Node->GetFullName(), StackFrame.Code - StackFrame.Node->Script.GetTypedData());
 	}
 
-	if (GUnrealEd->PlayWorld != NULL)
+	if ((GUnrealEd->PlayWorld != NULL) && !GIsAutomationTesting)
 	{
 		// Pause the simulation
 		GUnrealEd->PlayWorld->bDebugPauseExecution = true;
diff --git a/Engine/Source/Editor/UnrealEd/Private/StaticLightingSystem/StaticLightingSystem.cpp b/Engine/Source/Editor/UnrealEd/Private/StaticLightingSystem/StaticLightingSystem.cpp
--- a/Engine/Source/Editor/UnrealEd/Private/StaticLightingSystem/StaticLightingSystem.cpp
+++ b/Engine/Source/Editor/UnrealEd/Private/StaticLightingSystem/StaticLightingSystem.cpp
@@ -2054,11 +2054,10 @@ bool FStaticLightingSystem::CanAutoApplyLighting() const
 	const bool bInterpEditMode = GLevelEditorModeTools().IsModeActive( FBuiltinEditorModes::EM_InterpEdit );
 	const bool bPlayWorldValid = GUnrealEd->PlayWorld != nullptr;
 	const bool bAnyMenusVisible = FSlateApplication::Get().AnyMenusVisible();
-	const bool bAutomationTesting = GIsAutomationTesting;
 	const bool bIsInteratcting = false;// FSlateApplication::Get().GetMouseCaptor().IsValid() || GUnrealEd->IsUserInteracting();
 	const bool bHasGameOrProjectLoaded = FApp::HasGameName();
 
-	return ( bAutoApplyEnabled && !bSlowTask && !bInterpEditMode && !bPlayWorldValid && !bAnyMenusVisible && !bAutomationTesting && !bIsInteratcting && !GIsDemoMode && bHasGameOrProjectLoaded );
+	return ( bAutoApplyEnabled && !bSlowTask && !bInterpEditMode && !bPlayWorldValid && !bAnyMenusVisible && !bIsInteratcting && !GIsDemoMode && bHasGameOrProjectLoaded );
 }
 
 /**
diff --git a/Engine/Source/Editor/UnrealEd/Private/Tests/AutomationEditorCommon.cpp b/Engine/Source/Editor/UnrealEd/Private/Tests/AutomationEditorCommon.cpp
--- a/Engine/Source/Editor/UnrealEd/Private/Tests/AutomationEditorCommon.cpp
+++ b/Engine/Source/Editor/UnrealEd/Private/Tests/AutomationEditorCommon.cpp
@@ -3,6 +3,275 @@
 #include "UnrealEd.h"
 #include "AutomationEditorCommon.h"
 
+DEFINE_LOG_CATEGORY_STATIC(LogAutomationEditorCommon, Log, All);
+
+namespace AutomationEditorCommonUtils
+{
+	/**
+	* Converts a package path to an asset path
+	*
+	* @param PackagePath - The package path to convert
+	*/
+	static FString ConvertPackagePathToAssetPath(const FString& PackagePath)
+	{
+		const FString Filename = FPaths::ConvertRelativePathToFull(PackagePath);
+		FString EngineFileName = Filename;
+		FString GameFileName = Filename;
+		if (FPaths::MakePathRelativeTo(EngineFileName, *FPaths::EngineContentDir()) && !FPaths::IsRelative(EngineFileName))
+		{
+			const FString ShortName = FPaths::GetBaseFilename(EngineFileName);
+			const FString PathName = FPaths::GetPath(EngineFileName);
+			const FString AssetName = FString::Printf(TEXT("/Engine/%s/%s.%s"), *PathName, *ShortName, *ShortName);
+			return AssetName;
+		}
+		else if (FPaths::MakePathRelativeTo(GameFileName, *FPaths::GameContentDir()) && !FPaths::IsRelative(GameFileName))
+		{
+			const FString ShortName = FPaths::GetBaseFilename(GameFileName);
+			const FString PathName = FPaths::GetPath(GameFileName);
+			const FString AssetName = FString::Printf(TEXT("/Game/%s/%s.%s"), *PathName, *ShortName, *ShortName);
+			return AssetName;
+		}
+		else
+		{
+			UE_LOG(LogAutomationEditorCommon, Error, TEXT("PackagePath (%s) is invalid for the current project"), *PackagePath);
+			return TEXT("");
+		}
+	}
+
+	/**
+	* Imports an object using a given factory
+	*
+	* @param ImportFactory - The factory to use to import the object
+	* @param ObjectName - The name of the object to create
+	* @param PackagePath - The full path of the package file to create
+	* @param ImportPath - The path to the object to import
+	*/
+	static UObject* ImportAssetUsingFactory(UFactory* ImportFactory, const FString& ObjectName, const FString& PackagePath, const FString& ImportPath)
+	{
+		UObject* ImportedAsset = NULL;
+
+		UPackage* Pkg = CreatePackage(NULL, *PackagePath);
+		if (Pkg)
+		{
+			// Make sure the destination package is loaded
+			Pkg->FullyLoad();
+
+			UClass* ImportAssetType = ImportFactory->ResolveSupportedClass();
+			bool bDummy = false;
+
+			//If we are a texture factory suppress some warning dialog that we don't want
+			if (ImportFactory->IsA(UTextureFactory::StaticClass()))
+			{
+				UTextureFactory::SuppressImportResolutionWarningDialog();
+				UTextureFactory::SuppressImportOverwriteDialog();
+			}
+
+			ImportedAsset = UFactory::StaticImportObject(ImportAssetType, Pkg, FName(*ObjectName), RF_Public | RF_Standalone, bDummy, *ImportPath, NULL, ImportFactory, NULL, GWarn, 0);
+
+			if (ImportedAsset)
+			{
+				UE_LOG(LogAutomationEditorCommon, Display, TEXT("Imported %s"), *ImportPath);
+			}
+			else
+			{
+				UE_LOG(LogAutomationEditorCommon, Error, TEXT("Failed to import asset using factory %s!"), *ImportFactory->GetName());
+			}
+		}
+		else
+		{
+			UE_LOG(LogAutomationEditorCommon, Error, TEXT("Failed to create a package!"));
+		}
+
+		return ImportedAsset;
+	}
+
+	/**
+	* Nulls out references to a given object
+	*
+	* @param InObject - Object to null references to
+	*/
+	static void NullReferencesToObject(UObject* InObject)
+	{
+		TArray<UObject*> ReplaceableObjects;
+		TMap<UObject*, UObject*> ReplacementMap;
+		ReplacementMap.Add(InObject, NULL);
+		ReplacementMap.GenerateKeyArray(ReplaceableObjects);
+
+		// Find all the properties (and their corresponding objects) that refer to any of the objects to be replaced
+		TMap< UObject*, TArray<UProperty*> > ReferencingPropertiesMap;
+		for (FObjectIterator ObjIter; ObjIter; ++ObjIter)
+		{
+			UObject* CurObject = *ObjIter;
+
+			// Find the referencers of the objects to be replaced
+			FFindReferencersArchive FindRefsArchive(CurObject, ReplaceableObjects);
+
+			// Inform the object referencing any of the objects to be replaced about the properties that are being forcefully
+			// changed, and store both the object doing the referencing as well as the properties that were changed in a map (so that
+			// we can correctly call PostEditChange later)
+			TMap<UObject*, int32> CurNumReferencesMap;
+			TMultiMap<UObject*, UProperty*> CurReferencingPropertiesMMap;
+			if (FindRefsArchive.GetReferenceCounts(CurNumReferencesMap, CurReferencingPropertiesMMap) > 0)
+			{
+				TArray<UProperty*> CurReferencedProperties;
+				CurReferencingPropertiesMMap.GenerateValueArray(CurReferencedProperties);
+				ReferencingPropertiesMap.Add(CurObject, CurReferencedProperties);
+				for (TArray<UProperty*>::TConstIterator RefPropIter(CurReferencedProperties); RefPropIter; ++RefPropIter)
+				{
+					CurObject->PreEditChange(*RefPropIter);
+				}
+			}
+
+		}
+
+		// Iterate over the map of referencing objects/changed properties, forcefully replacing the references and then
+		// alerting the referencing objects the change has completed via PostEditChange
+		int32 NumObjsReplaced = 0;
+		for (TMap< UObject*, TArray<UProperty*> >::TConstIterator MapIter(ReferencingPropertiesMap); MapIter; ++MapIter)
+		{
+			++NumObjsReplaced;
+
+			UObject* CurReplaceObj = MapIter.Key();
+			const TArray<UProperty*>& RefPropArray = MapIter.Value();
+
+			FArchiveReplaceObjectRef<UObject> ReplaceAr(CurReplaceObj, ReplacementMap, false, true, false);
+
+			for (TArray<UProperty*>::TConstIterator RefPropIter(RefPropArray); RefPropIter; ++RefPropIter)
+			{
+				FPropertyChangedEvent PropertyEvent(*RefPropIter);
+				CurReplaceObj->PostEditChangeProperty(PropertyEvent);
+			}
+
+			if (!CurReplaceObj->HasAnyFlags(RF_Transient) && CurReplaceObj->GetOutermost() != GetTransientPackage())
+			{
+				if (!CurReplaceObj->RootPackageHasAnyFlags(PKG_CompiledIn))
+				{
+					CurReplaceObj->MarkPackageDirty();
+				}
+			}
+		}
+	}
+
+	/**
+	* gets a factory class based off an asset file extension
+	*
+	* @param AssetExtension - The file extension to use to find a supporting UFactory
+	*/
+	static UClass* GetFactoryClassForType(const FString& AssetExtension)
+	{
+		// First instantiate one factory for each file extension encountered that supports the extension
+		for (TObjectIterator<UClass> ClassIt; ClassIt; ++ClassIt)
+		{
+			if ((*ClassIt)->IsChildOf(UFactory::StaticClass()) && !((*ClassIt)->HasAnyClassFlags(CLASS_Abstract)))
+			{
+				UFactory* Factory = Cast<UFactory>((*ClassIt)->GetDefaultObject());
+				if (Factory->bEditorImport && Factory->ValidForCurrentGame())
+				{
+					TArray<FString> FactoryExtensions;
+					Factory->GetSupportedFileExtensions(FactoryExtensions);
+
+					// Case insensitive string compare with supported formats of this factory
+					if (FactoryExtensions.Contains(AssetExtension))
+					{
+						return *ClassIt;
+					}
+				}
+			}
+		}
+
+		return NULL;
+	}
+
+	/**
+	* Applies settings to an object by finding UProperties by name and calling ImportText
+	*
+	* @param InObject - The object to search for matching properties
+	* @param PropertyChain - The list UProperty names recursively to search through
+	* @param Value - The value to import on the found property
+	*/
+	static void ApplyCustomFactorySetting(UObject* InObject, TArray<FString>& PropertyChain, const FString& Value)
+	{
+		const FString PropertyName = PropertyChain[0];
+		PropertyChain.RemoveAt(0);
+
+		UProperty* TargetProperty = FindField<UProperty>(InObject->GetClass(), *PropertyName);
+		if (TargetProperty)
+		{
+			if (PropertyChain.Num() == 0)
+			{
+				TargetProperty->ImportText(*Value, TargetProperty->ContainerPtrToValuePtr<uint8>(InObject), 0, InObject);
+			}
+			else
+			{
+				UStructProperty* StructProperty = Cast<UStructProperty>(TargetProperty);
+				UObjectProperty* ObjectProperty = Cast<UObjectProperty>(TargetProperty);
+
+				UObject* SubObject = NULL;
+				bool bValidPropertyType = true;
+
+				if (StructProperty)
+				{
+					SubObject = StructProperty->Struct;
+				}
+				else if (ObjectProperty)
+				{
+					SubObject = ObjectProperty->GetObjectPropertyValue(ObjectProperty->ContainerPtrToValuePtr<UObject>(InObject));
+				}
+				else
+				{
+					//Unknown nested object type
+					bValidPropertyType = false;
+					UE_LOG(LogAutomationEditorCommon, Error, TEXT("ERROR: Unknown nested object type for property: %s"), *PropertyName);
+				}
+
+				if (SubObject)
+				{
+					ApplyCustomFactorySetting(SubObject, PropertyChain, Value);
+				}
+				else if (bValidPropertyType)
+				{
+					UE_LOG(LogAutomationEditorCommon, Error, TEXT("Error accessing null property: %s"), *PropertyName);
+				}
+			}
+		}
+		else
+		{
+			UE_LOG(LogAutomationEditorCommon, Error, TEXT("ERROR: Could not find factory property: %s"), *PropertyName);
+		}
+	}
+
+	/**
+	* Applies the custom factory settings
+	*
+	* @param InFactory - The factory to apply custom settings to
+	* @param FactorySettings - An array of custom settings to apply to the factory
+	*/
+	static void ApplyCustomFactorySettings(UFactory* InFactory, const TArray<FImportFactorySettingValues>& FactorySettings)
+	{
+		bool bCallConfigureProperties = true;
+
+		for (int32 i = 0; i < FactorySettings.Num(); ++i)
+		{
+			if (FactorySettings[i].SettingName.Len() > 0 && FactorySettings[i].Value.Len() > 0)
+			{
+				//Check if we are setting an FBX import type override.  If we are, we don't want to call ConfigureProperties because that enables bDetectImportTypeOnImport
+				if (FactorySettings[i].SettingName.Contains(TEXT("MeshTypeToImport")))
+				{
+					bCallConfigureProperties = false;
+				}
+
+				TArray<FString> PropertyChain;
+				FactorySettings[i].SettingName.ParseIntoArray(&PropertyChain, TEXT("."), false);
+				ApplyCustomFactorySetting(InFactory, PropertyChain, FactorySettings[i].Value);
+			}
+		}
+
+		if (bCallConfigureProperties)
+		{
+			InFactory->ConfigureProperties();
+		}
+	}
+}
 
 ///////////////////////////////////////////////////////////////////////
 // Common Latent commands
diff --git a/Engine/Source/Editor/UnrealEd/Private/Tests/EditorAssetAutomationTests.cpp b/Engine/Source/Editor/UnrealEd/Private/Tests/EditorAssetAutomationTests.cpp
--- a/Engine/Source/Editor/UnrealEd/Private/Tests/EditorAssetAutomationTests.cpp
+++ b/Engine/Source/Editor/UnrealEd/Private/Tests/EditorAssetAutomationTests.cpp
@@ -29,76 +29,6 @@
 
 DEFINE_LOG_CATEGORY_STATIC(LogEditorAssetAutomationTests, Log, All);
 
-namespace AssetAutomationCommon
-{
-	/**
-	* Nulls out references to a given object
-	*
-	* @param InObject - Object to null references to
-	*/
-	static void NullReferencesToObject(UObject* InObject)
-	{
-		TArray<UObject*> ReplaceableObjects;
-		TMap<UObject*, UObject*> ReplacementMap;
-		ReplacementMap.Add(InObject, NULL);
-		ReplacementMap.GenerateKeyArray(ReplaceableObjects);
-
-		// Find all the properties (and their corresponding objects) that refer to any of the objects to be replaced
-		TMap< UObject*, TArray<UProperty*> > ReferencingPropertiesMap;
-		for (FObjectIterator ObjIter; ObjIter; ++ObjIter)
-		{
-			UObject* CurObject = *ObjIter;
-
-			// Find the referencers of the objects to be replaced
-			FFindReferencersArchive FindRefsArchive(CurObject, ReplaceableObjects);
-
-			// Inform the object referencing any of the objects to be replaced about the properties that are being forcefully
-			// changed, and store both the object doing the referencing as well as the properties that were changed in a map (so that
-			// we can correctly call PostEditChange later)
-			TMap<UObject*, int32> CurNumReferencesMap;
-			TMultiMap<UObject*, UProperty*> CurReferencingPropertiesMMap;
-			if (FindRefsArchive.GetReferenceCounts(CurNumReferencesMap, CurReferencingPropertiesMMap) > 0)
-			{
-				TArray<UProperty*> CurReferencedProperties;
-				CurReferencingPropertiesMMap.GenerateValueArray(CurReferencedProperties);
-				ReferencingPropertiesMap.Add(CurObject, CurReferencedProperties);
-				for (TArray<UProperty*>::TConstIterator RefPropIter(CurReferencedProperties); RefPropIter; ++RefPropIter)
-				{
-					CurObject->PreEditChange(*RefPropIter);
-				}
-			}
-
-		}
-
-		// Iterate over the map of referencing objects/changed properties, forcefully replacing the references and then
-		// alerting the referencing objects the change has completed via PostEditChange
-		int32 NumObjsReplaced = 0;
-		for (TMap< UObject*, TArray<UProperty*> >::TConstIterator MapIter(ReferencingPropertiesMap); MapIter; ++MapIter)
-		{
-			++NumObjsReplaced;
-
-			UObject* CurReplaceObj = MapIter.Key();
-			const TArray<UProperty*>& RefPropArray = MapIter.Value();
-
-			FArchiveReplaceObjectRef<UObject> ReplaceAr(CurReplaceObj, ReplacementMap, false, true, false);
-
-			for (TArray<UProperty*>::TConstIterator RefPropIter(RefPropArray); RefPropIter; ++RefPropIter)
-			{
-				FPropertyChangedEvent PropertyEvent(*RefPropIter);
-				CurReplaceObj->PostEditChangeProperty(PropertyEvent);
-			}
-
-			if (!CurReplaceObj->HasAnyFlags(RF_Transient) && CurReplaceObj->GetOutermost() != GetTransientPackage())
-			{
-				if (!CurReplaceObj->RootPackageHasAnyFlags(PKG_CompiledIn))
-				{
-					CurReplaceObj->MarkPackageDirty();
-				}
-			}
-		}
-	}
-}
-
 /**
 * Container for items related to the create asset test
 */
@@ -328,7 +258,7 @@ namespace CreateAssetHelper
 				if (!bSuccessful)
 				{
 					//Clear references to the object so we can delete it
-					AssetAutomationCommon::NullReferencesToObject(CreatedAsset);
+					AutomationEditorCommonUtils::NullReferencesToObject(CreatedAsset);
 
 					bSuccessful = ObjectTools::DeleteSingleObject(CreatedAsset, false);
 				}
@@ -687,113 +617,6 @@ namespace ImportExportAssetHelper
 		};
 	}
 
-	/**
-	* gets a factory class based off an asset file extension
-	*
-	* @param AssetExtension - The file extension to use to find a supporting UFactory
-	*/
-	static UClass* GetFactoryClassForType(const FString& AssetExtension)
-	{
-		// First instantiate one factory for each file extension encountered that supports the extension
-		for (TObjectIterator<UClass> ClassIt; ClassIt; ++ClassIt)
-		{
-			if ((*ClassIt)->IsChildOf(UFactory::StaticClass()) && !((*ClassIt)->HasAnyClassFlags(CLASS_Abstract)))
-			{
-				UFactory* Factory = Cast<UFactory>((*ClassIt)->GetDefaultObject());
-				if (Factory->bEditorImport && Factory->ValidForCurrentGame())
-				{
-					TArray<FString> FactoryExtensions;
-					Factory->GetSupportedFileExtensions(FactoryExtensions);
-
-					// Case insensitive string compare with supported formats of this factory
-					if (FactoryExtensions.Contains(AssetExtension))
-					{
-						return *ClassIt;
-					}
-				}
-			}
-		}
-
-		return NULL;
-	}
-
-	/**
-	* Applies settings to an object by finding UProperties by name and calling ImportText
-	*
-	* @param InObject - The object to search for matching properties
-	* @param PropertyChain - The list UProperty names recursively to search through
-	* @param Value - The value to import on the found property
-	*/
-	static void ApplyCustomFactorySetting(UObject* InObject, TArray<FString>& PropertyChain, const FString& Value)
-	{
-		const FString PropertyName = PropertyChain[0];
-		PropertyChain.RemoveAt(0);
-
-		UProperty* TargetProperty = FindField<UProperty>(InObject->GetClass(), *PropertyName);
-		if (TargetProperty)
-		{
-			if (PropertyChain.Num() == 0)
-			{
-				TargetProperty->ImportText(*Value, TargetProperty->ContainerPtrToValuePtr<uint8>(InObject), 0, InObject);
-			}
-			else
-			{
-				UStructProperty* StructProperty = Cast<UStructProperty>(TargetProperty);
-				UObjectProperty* ObjectProperty = Cast<UObjectProperty>(TargetProperty);
-
-				UObject* SubObject = NULL;
-				bool bValidPropertyType = true;
-
-				if (StructProperty)
-				{
-					SubObject = StructProperty->Struct;
-				}
-				else if (ObjectProperty)
-				{
-					SubObject = ObjectProperty->GetObjectPropertyValue(ObjectProperty->ContainerPtrToValuePtr<UObject>(InObject));
-				}
-				else
-				{
-					//Unknown nested object type
-					bValidPropertyType = false;
-					UE_LOG(LogEditorAssetAutomationTests, Error, TEXT("ERROR: Unknown nested object type for property: %s"), *PropertyName);
-				}
-
-				if (SubObject)
-				{
-					ApplyCustomFactorySetting(SubObject, PropertyChain, Value);
-				}
-				else if (bValidPropertyType)
-				{
-					UE_LOG(LogEditorAssetAutomationTests, Error, TEXT("Error accessing null property: %s"), *PropertyName);
-				}
-			}
-		}
-		else
-		{
-			UE_LOG(LogEditorAssetAutomationTests, Error, TEXT("ERROR: Could not find factory property: %s"), *PropertyName);
-		}
-	}
-
-	/**
-	* Applies the custom factory settings
-	*
-	* @param InFactory - The factory to apply custom settings to
-	* @param FactorySettings - An array of custom settings to apply to the factory
-	*/
-	static void ApplyCustomFactorySettings(UFactory* InFactory, const TArray<FImportFactorySettingValues>& FactorySettings)
-	{
-		for (int32 i = 0; i < FactorySettings.Num(); ++i)
-		{
-			if (FactorySettings[i].SettingName.Len() > 0 && FactorySettings[i].Value.Len() > 0)
-			{
-				TArray<FString> PropertyChain;
-				FactorySettings[i].SettingName.ParseIntoArray(&PropertyChain, TEXT("."), false);
-				ApplyCustomFactorySetting(InFactory, PropertyChain, FactorySettings[i].Value);
-			}
-		}
-	}
-
 	/**
 	* Import test report for a single asset
 	*/
@@ -961,7 +784,7 @@ namespace ImportExportAssetHelper
 
 			//Get the factory
 			const FString FileExtension = FPaths::GetExtension(ImportPath);
-			UClass* FactoryClass = GetFactoryClassForType(FileExtension);
+			UClass* FactoryClass = AutomationEditorCommonUtils::GetFactoryClassForType(FileExtension);
 
 			if (FactoryClass)
 			{
@@ -969,37 +792,18 @@ namespace ImportExportAssetHelper
 
 				UFactory* ImportFactory = ConstructObject<UFactory>(FactoryClass);
 
-				ImportFactory->ConfigureProperties();
-
 				//Apply any custom settings to the factory
-				ApplyCustomFactorySettings(ImportFactory, FactorySettings);
+				AutomationEditorCommonUtils::ApplyCustomFactorySettings(ImportFactory, FactorySettings);
 
 				FString Name = ObjectTools::SanitizeObjectName(FPaths::GetBaseFilename(ImportPath));
 				FString PackageName = FString::Printf(TEXT("/Game/Automation_Imports/%s"), *Name);
 
-				UPackage* Pkg = CreatePackage(NULL, *PackageName);
-				if (Pkg)
-				{
-					// Make sure the destination package is loaded
-					Pkg->FullyLoad();
-
-					UClass* ImportAssetType = ImportFactory->ResolveSupportedClass();
-					bool bDummy = false;
-					ImportedAsset = UFactory::StaticImportObject(ImportAssetType, Pkg, FName(*Name), RF_Public | RF_Standalone, bDummy, *ImportPath, NULL, ImportFactory, NULL, GWarn, 0);
+				ImportedAsset = AutomationEditorCommonUtils::ImportAssetUsingFactory(ImportFactory, Name, PackageName, ImportPath);
 
-					if (ImportedAsset)
-					{
-						TestReport.bImportSuccessful = true;
-						State = EState::OpenEditor;
-					}
-					else
-					{
-						UE_LOG(LogEditorAssetAutomationTests, Error, TEXT("Failed to import asset %s using factory %s!"), *Name, *ImportFactory->GetName());
-					}
-				}
-				else
+				if (ImportedAsset)
 				{
-					UE_LOG(LogEditorAssetAutomationTests, Error, TEXT("Failed to create a package for %s!"), *Name);
+					TestReport.bImportSuccessful = true;
+					State = EState::OpenEditor;
 				}
 
 				GWarn->EndSlowTask();
@@ -1155,7 +959,7 @@ namespace ImportExportAssetHelper
 			GUnrealEd->ResetTransaction(FText::FromString(TEXT("FAssetEditorTest")));
 
 			//Clear references to the object so we can delete it
-			AssetAutomationCommon::NullReferencesToObject(ImportedAsset);
+			AutomationEditorCommonUtils::NullReferencesToObject(ImportedAsset);
 
 			//Delete the object
 			TArray<UObject*> ObjList;
diff --git a/Engine/Source/Editor/UnrealEd/Private/Tests/EditorAutomationTests.cpp b/Engine/Source/Editor/UnrealEd/Private/Tests/EditorAutomationTests.cpp
--- a/Engine/Source/Editor/UnrealEd/Private/Tests/EditorAutomationTests.cpp
+++ b/Engine/Source/Editor/UnrealEd/Private/Tests/EditorAutomationTests.cpp
@@ -1017,7 +1017,7 @@ IMPLEMENT_SIMPLE_AUTOMATION_TEST(FBSPValidation, "QA.BSP Validation", EAutomatio
 
 bool FBSPValidation::RunTest(const FString& Parameters)
 {
-	UWorld* World = GEditor->NewMap();
+	UWorld* World = AutomationEditorCommonUtils::CreateNewMap();
 	GEditor->Exec( World, TEXT("BRUSH Scale 1 1 1"));
 
 	for( int32 i = 0; i < GEditor->LevelViewportClients.Num(); i++ )
@@ -1166,7 +1166,7 @@ IMPLEMENT_SIMPLE_AUTOMATION_TEST(FStaticMeshValidation, "QA.Mesh Factory Validat
 
 bool FStaticMeshValidation::RunTest(const FString& Parameters)
 {
-	UWorld* World = GEditor->NewMap();
+	UWorld* World = AutomationEditorCommonUtils::CreateNewMap();
 
 	//Adjust camera in viewports
 	for( int32 i = 0; i < GEditor->LevelViewportClients.Num(); i++ )
@@ -1269,7 +1269,7 @@ ABrush* ConvertTestFindNewBrush (const TArray<ABrush*> &PreviousBrushes)
 
 bool FConvertToValidation::RunTest(const FString& Parameters)
 {
-	UWorld* World = GEditor->NewMap();
+	UWorld* World = AutomationEditorCommonUtils::CreateNewMap();
 
 	//Set the Test Name which is used later for getting the directory to store the screenshots.
 	const FString BaseFileName = TEXT("ConvertMeshTest");
@@ -1422,7 +1422,7 @@ void FOpenAssetEditors::GetTests(TArray<FString>& OutBeautifiedNames, TArray <FS
 	//Put the Engine Content Folder Location into an array.  
 	TArray<FString> EngineContentFolderLocationArray;
 	EngineContentFolderLocation.ParseIntoArray(&EngineContentFolderLocationArray, TEXT("/"), true);
-	
+
 	for (int32 i = 0; i < AssetNames.Num(); ++i)
 	{
 		//true means that the life is located in the Engine/Content folder.
@@ -1583,7 +1583,7 @@ bool FLightPlacement::RunTest(const FString& Parameters)
 	const FString BaseFileName = TEXT("PointLightPlacementTest");
 
 	//Open a new blank map.
-	UWorld* World = GEditor->NewMap();
+	UWorld* World = AutomationEditorCommonUtils::CreateNewMap();
 
 	//Move the perspective viewport view to show the test.
 	for (int32 i = 0; i < GEditor->LevelViewportClients.Num(); i++)
diff --git a/Engine/Source/Editor/UnrealEd/Private/Tests/EditorBuildPromotionTests.cpp b/Engine/Source/Editor/UnrealEd/Private/Tests/EditorBuildPromotionTests.cpp
new file mode 100644
--- /dev/null
+++ b/Engine/Source/Editor/UnrealEd/Private/Tests/EditorBuildPromotionTests.cpp
@@ -0,0 +1,4474 @@
+// Copyright 1998-2014 Epic Games, Inc. All Rights Reserved.
+
+#include "UnrealEd.h"
+
+//Automation
+#include "Tests/AutomationTestSettings.h"
+#include "AutomationCommon.h"
+#include "AutomationEditorCommon.h"
+
+//Assets
+#include "AssetRegistryModule.h"
+#include "IAssetRegistry.h"
+#include "ComponentAssetBroker.h"
+
+//Materials
+#include "AssetEditorManager.h"
+#include "IMaterialEditor.h"
+#include "Private/MaterialEditor.h"
+#include "Materials/MaterialExpressionTextureSample.h"
+#include "Materials/MaterialExpressionConstant3Vector.h"
+
+//Particle system
+#include "Particles/ParticleSystem.h"
+#include "Particles/ParticleEmitter.h"
+#include "Particles/ParticleLODLevel.h"
+#include "Particles/Size/ParticleModuleSize.h"
+
+//Blueprints
+#include "EdGraphSchema_K2_Actions.h"
+#include "BlueprintEditor.h"
+#include "BlueprintEditorModes.h"
+#include "CompilerResultsLog.h"
+#include "EdGraphUtilities.h"
+#include "KismetDebugUtilities.h"
+#include "Engine/Breakpoint.h"
+#include "Engine/LevelScriptBlueprint.h"
+
+//Utils
+#include "EditorBuildUtils.h"
+#include "EditorLevelUtils.h"
+
+//Source Control
+#include "Developer/SourceControl/Public/ISourceControlModule.h"
+#include "Developer/SourceControl/Public/ISourceControlProvider.h"
+#include "Developer/SourceControl/Public/SourceControlOperations.h"
+#include "Developer/SourceControl/Public/Tests/SourceControlAutomationCommon.h"
+
+//Lighting
+#include "LightingBuildOptions.h"
+#include "../Private/StaticLightingSystem/StaticLightingPrivate.h"
+
+#define LOCTEXT_NAMESPACE "EditorBuildPromotionTests"
+
+DEFINE_LOG_CATEGORY_STATIC(LogEditorBuildPromotionTests, Log, All);
+
+/**
+* Helper functions used by the build promotion automation test
+*/
+namespace EditorBuildPromotionTestUtils
+{
+	/**
+	* Constants
+	*/
+	static const FString BlueprintNameString = TEXT("BuildPromotionBlueprint");
+	static const FName BlueprintStringVariableName(TEXT("MyStringVariable"));
+
+	/**
+	* Gets the base path for this asset
+	*/
+	static FString GetGamePath()
+	{
+		return TEXT("/Game/BuildPromotionTest");
+	}
+
+	/**
+	* Gets the full path to the folder on disk
+	*/
+	static FString GetFullPath()
+	{
+		return FPackageName::FilenameToLongPackageName(FPaths::GameContentDir() + TEXT("BuildPromotionTest"));
+	}
+
+	/**
+	* Helper class to track once a certain time has passed
+	*/
+	class FDelayHelper
+	{
+	public:
+		/**
+		* Constructor
+		*/
+		FDelayHelper() :
+			bIsRunning(false),
+			StartTime(0.f),
+			Duration(0.f)
+		{
+		}
+
+		/**
+		* Returns if the delay is still running
+		*/
+		bool IsRunning()
+		{
+			return bIsRunning;
+		}
+
+		/**
+		* Sets the helper state to not running
+		*/
+		void Reset()
+		{
+			bIsRunning = false;
+		}
+
+		/**
+		* Starts the delay timer
+		*
+		* @param InDuration - How long to delay for in seconds
+		*/
+		void Start(double InDuration)
+		{
+			bIsRunning = true;
+			StartTime = FPlatformTime::Seconds();
+			Duration = InDuration;
+		}
+
+		/**
+		* Returns true if the desired amount of time has passed
+		*/
+		bool IsComplete()
+		{
+			if (IsRunning())
+			{
+				const double CurrentTime = FPlatformTime::Seconds();
+				return CurrentTime - StartTime >= Duration;
+			}
+			else
+			{
+				return false;
+			}
+		}
+
+	private:
+		/** If true, this delay timer is active */
+		bool bIsRunning;
+		/** The time the delay started */
+		double StartTime;
+		/** How long the timer is for */
+		double Duration;
+	};
+
+	/**
+	* Sets an editor keyboard shortcut
+	*
+	* @param CommandContext - The context of the command
+	* @param Command - The command name to set
+	* @param NewGesture - The new input gesture to assign
+	*/
+	static void SetEditorKeybinding(const FString& CommandContext, const FString& Command, const FInputGesture& NewGesture)
+	{
+		TSharedPtr<FUICommandInfo> UICommand = FInputBindingManager::Get().FindCommandInContext(*CommandContext, *Command);
+		if (UICommand.IsValid())
+		{
+			UICommand->SetActiveGesture(NewGesture);
+			FInputBindingManager::Get().NotifyActiveGestureChanged(*UICommand.Get());
+			FInputBindingManager::Get().SaveInputBindings();
+		}
+		else
+		{
+			UE_LOG(LogEditorBuildPromotionTests, Error, TEXT("Could not find keybinding for %s using context %s"), *Command, *CommandContext);
+		}
+	}
+
+	/**
+	* Gets an editor keyboard shortcut
+	*
+	* @param CommandContext - The context of the command
+	* @param Command - The command name to get
+	* @param NewGesture - The current input gesture that is assigned
+	*/
+	static bool GetEditorKeybinding(const FString& CommandContext, const FString& Command, FInputGesture& CurrentGesture)
+	{
+		TSharedPtr<FUICommandInfo> UICommand = FInputBindingManager::Get().FindCommandInContext(*CommandContext, *Command);
+		if (UICommand.IsValid())
+		{
+			TSharedRef<const FInputGesture> ActiveGesture = UICommand->GetActiveGesture();
+			CurrentGesture = ActiveGesture.Get();
+		}
+		else
+		{
+			UE_LOG(LogEditorBuildPromotionTests, Error, TEXT("Could not find keybinding for %s using context %s"), *Command, *CommandContext);
+		}
+		return true;
+	}
+
+	/**
+	* Exports the current editor keybindings
+	*
+	* @param TargetFilename - The name of the file to export to
+	*/
+	static void ExportKeybindings(const FString& TargetFilename)
+	{
+		FInputBindingManager::Get().SaveInputBindings();
+		GConfig->Flush(false, GEditorKeyBindingsIni);
+		IFileManager::Get().Copy(*TargetFilename, *GEditorKeyBindingsIni);
+	}
+
+	/**
+	* Imports new editor keybindings  (Unused)
+	*
+	* @param TargetFilename - The name of the file to import from
+	*/
+	static void ImportKeybindings(const FString& TargetFilename)
+	{
+		GConfig->Flush(true, GEditorKeyBindingsIni);
+		IFileManager::Get().Copy(*GEditorKeyBindingsIni, *TargetFilename);
+		GConfig->LoadFile(GEditorKeyBindingsIni);
+	}
+
+	/**
+	* Exports the current editor settings
+	*
+	* @param TargetFilename - The name of the file to export to
+	*/
+	static void ExportEditorSettings(const FString& TargetFilename)
+	{
+		FInputBindingManager::Get().SaveInputBindings();
+		GConfig->Flush(false, GEditorUserSettingsIni);
+		IFileManager::Get().Copy(*TargetFilename, *GEditorUserSettingsIni);
+	}
+
+	/**
+	* Imports new editor settings (unused)
+	*
+	* @param TargetFilename - The name of the file to export to
+	*/
+	static void ImportEditorSettings(const FString& TargetFilename)
+	{
+		GConfig->Flush(true, GEditorUserSettingsIni);
+		IFileManager::Get().Copy(*GEditorUserSettingsIni, *TargetFilename);
+		GConfig->LoadFile(GEditorUserSettingsIni);
+	}
+
+	/**
+	* Creates a material from an existing texture
+	*
+	* @param InTexture - The texture to use as the diffuse for the new material
+	*/
+	static UMaterial* CreateMaterialFromTexture(UTexture* InTexture)
+	{
+		// Create the factory used to generate the asset
+		UMaterialFactoryNew* Factory = ConstructObject<UMaterialFactoryNew>(UMaterialFactoryNew::StaticClass());
+		Factory->InitialTexture = InTexture;
+
+
+		const FString AssetName = FString::Printf(TEXT("%s_Mat"), *InTexture->GetName());
+		const FString PackageName = GetGamePath() + TEXT("/") + AssetName;
+		UPackage* AssetPackage = CreatePackage(NULL, *PackageName);
+		EObjectFlags Flags = RF_Public | RF_Standalone;
+
+		UObject* CreatedAsset = Factory->FactoryCreateNew(UMaterial::StaticClass(), AssetPackage, FName(*AssetName), Flags, NULL, GWarn);
+
+		if (CreatedAsset)
+		{
+			// Notify the asset registry
+			FAssetRegistryModule::AssetCreated(CreatedAsset);
+
+			// Mark the package dirty...
+			AssetPackage->MarkPackageDirty();
+		}
+
+		return Cast<UMaterial>(CreatedAsset);
+	}
+
+	/**
+	* Assigns a normal map to a material
+	*
+	* @param InNormalTexture - The texture to use as the normal map
+	* @param InMaterial - The material to modify
+	*/
+	static bool AssignNormalToMaterial(UTexture* InNormalTexture, UMaterial* InMaterial)
+	{
+		IAssetEditorInstance* OpenEditor = FAssetEditorManager::Get().FindEditorForAsset(InMaterial, true);
+		IMaterialEditor* CurrentMaterialEditor = (IMaterialEditor*)OpenEditor;
+
+		//Create the texture sample and auto assign the normal texture
+		UMaterialExpressionTextureSample* NewTextureSampleExpression = Cast<UMaterialExpressionTextureSample>(CurrentMaterialEditor->CreateNewMaterialExpression(UMaterialExpressionTextureSample::StaticClass(), FVector2D(100.f, 200.f), true, true));
+		if (NewTextureSampleExpression)
+		{
+			UE_LOG(LogEditorBuildPromotionTests, Display, TEXT("Created a new texture sample expression"));
+			NewTextureSampleExpression->Texture = InNormalTexture;
+			NewTextureSampleExpression->SamplerType = EMaterialSamplerType::SAMPLERTYPE_Normal;
+			UE_LOG(LogEditorBuildPromotionTests, Display, TEXT("Assigned the normal map texture to the new node"));
+			UMaterial* EditorMaterial = Cast<UMaterial>(CurrentMaterialEditor->GetMaterialInterface());
+			UMaterialGraph* MaterialGraph = EditorMaterial->MaterialGraph;
+			EditorMaterial->Normal.Connect(0, NewTextureSampleExpression);
+			UE_LOG(LogEditorBuildPromotionTests, Display, TEXT("Connected the new node to the normal pin"));
+			MaterialGraph->LinkGraphNodesFromMaterial();
+
+			return true;
+		}
+		else
+		{
+			UE_LOG(LogEditorBuildPromotionTests, Error, TEXT("Could not add a texture sample to %s"), *InMaterial->GetName());
+			return false;
+		}
+	}
+
+	/**
+	* Finds a visible widget by type.     SLOW!!!!!
+	*
+	* @param InParent - We search this widget and its children for a matching widget (recursive)
+	* @param InWidgetType - The widget type we are searching for
+	*/
+	static TSharedPtr<SWidget> FindFirstWidgetByClass(TSharedRef<SWidget> InParent, const FName& InWidgetType)
+	{
+		if (InParent->GetType() == InWidgetType)
+		{
+			return InParent;
+		}
+		FChildren* Children = InParent->GetChildren();
+		for (int32 i = 0; i < Children->Num(); ++i)
+		{
+			TSharedRef<SWidget> ChildWidget = Children->GetChildAt(i);
+			TSharedPtr<SWidget> FoundWidget = FindFirstWidgetByClass(ChildWidget, InWidgetType);
+			if (FoundWidget.IsValid())
+			{
+				return FoundWidget;
+			}
+		}
+		return NULL;
+	}
+
+	/**
+	* Sends a UI command to the active top level window after focusing on a widget of a given type
+	*
+	* @param InGesture - The gesture to send to the window
+	* @param WidgetTypeToFocus - The widget type to find and focus on
+	*/
+	static void SendCommandToCurrentEditor(const FInputGesture& InGesture, const FName& WidgetTypeToFocus)
+	{
+		//Focus the asset Editor / Graph 
+		TSharedRef<SWindow> EditorWindow = FSlateApplication::Get().GetActiveTopLevelWindow().ToSharedRef();
+		FSlateApplication::Get().ProcessWindowActivatedEvent(FWindowActivateEvent(FWindowActivateEvent::EA_Activate, EditorWindow));
+		TSharedPtr<SWidget> FocusWidget = FindFirstWidgetByClass(EditorWindow, WidgetTypeToFocus);
+		if (FocusWidget.IsValid())
+		{
+			FSlateApplication::Get().SetKeyboardFocus(FocusWidget.ToSharedRef(), EKeyboardFocusCause::SetDirectly);
+
+			//Send the command
+			FModifierKeysState ModifierKeys(InGesture.bShift, false, InGesture.bCtrl, false, InGesture.bAlt, false, InGesture.bCmd, false, false);
+			FKeyboardEvent KeyEvent(InGesture.Key, ModifierKeys, false, 0);
+			FSlateApplication::Get().ProcessKeyDownEvent(KeyEvent);
+			FSlateApplication::Get().ProcessKeyUpEvent(KeyEvent);
+		}
+		else
+		{
+			UE_LOG(LogEditorBuildPromotionTests, Error, TEXT("Could not find widget %s to send UI command"), *WidgetTypeToFocus.ToString());
+		}
+	}
+
+	/**
+	* Gets the current input gesture or sets a new one if it doesn't exist
+	*
+	* @param Context - The context of the UI Command
+	* @param Command - The name of the UI command
+	*/
+	static FInputGesture GetOrSetUICommand(const FString& Context, const FString& Command)
+	{
+		FInputGesture CurrentGesture;
+		GetEditorKeybinding(Context, Command, CurrentGesture);
+
+		//If there is no current keybinding, set one
+		if (!CurrentGesture.Key.IsValid())
+		{
+			FInputGesture NewGesture(EKeys::J, EModifierKey::Control);
+			SetEditorKeybinding(Context, Command, NewGesture);
+			CurrentGesture = NewGesture;
+		}
+
+		return CurrentGesture;
+	}
+
+	/**
+	* Sends the MaterialEditor->Apply UI command
+	*/
+	static void SendUpdateMaterialCommand()
+	{
+		const FString Context = TEXT("MaterialEditor");
+		const FString Command = TEXT("Apply");
+		FInputGesture CurrentApplyGesture = GetOrSetUICommand(Context, Command);
+
+		const FName FocusWidgetType(TEXT("SGraphEditor"));
+		SendCommandToCurrentEditor(CurrentApplyGesture, FocusWidgetType);
+	}
+
+	/**
+	* Sends the AssetEditor->SaveAsset UI command
+	*/
+	static void SendSaveCascadeCommand()
+	{
+		const FString Context = TEXT("AssetEditor");
+		const FString Command = TEXT("SaveAsset");
+
+		FInputGesture CurrentSaveGesture = GetOrSetUICommand(Context, Command);
+
+		const FName FocusWidgetType(TEXT("SCascadeEmitterCanvas"));
+		SendCommandToCurrentEditor(CurrentSaveGesture, FocusWidgetType);
+	}
+
+	/**
+	* Sends the AssetEditor->SaveAsset UI command
+	*/
+	static void SendBlueprintResetViewCommand()
+	{
+		const FString Context = TEXT("BlueprintEditor");
+		const FString Command = TEXT("ResetCamera");
+
+		FInputGesture CurrentSaveGesture = GetOrSetUICommand(Context, Command);
+
+		const FName FocusWidgetType(TEXT("SSCSEditorViewport"));
+		SendCommandToCurrentEditor(CurrentSaveGesture, FocusWidgetType);
+	}
+
+	/**
+	* Compiles the blueprint
+	*
+	* @param InBlueprint - The blueprint to compile
+	*/
+	static void CompileBlueprint(UBlueprint* InBlueprint)
+	{
+		FBlueprintEditorUtils::RefreshAllNodes(InBlueprint);
+
+		bool bIsRegeneratingOnLoad = false;
+		bool bSkipGarbageCollection = true;
+
+		FKismetEditorUtilities::CompileBlueprint(InBlueprint, bIsRegeneratingOnLoad, bSkipGarbageCollection);
+		if (InBlueprint->Status == EBlueprintStatus::BS_UpToDate)
+		{
+			UE_LOG(LogEditorBuildPromotionTests, Display, TEXT("Blueprint compiled successfully (%s)"), *InBlueprint->GetName());
+		}
+		else if (InBlueprint->Status == EBlueprintStatus::BS_UpToDateWithWarnings)
+		{
+			UE_LOG(LogEditorBuildPromotionTests, Display, TEXT("Blueprint compiled successfully with warnings(%s)"), *InBlueprint->GetName());
+		}
+		else if (InBlueprint->Status == EBlueprintStatus::BS_Error)
+		{
+			UE_LOG(LogEditorBuildPromotionTests, Display, TEXT("Blueprint failed to compile (%s)"), *InBlueprint->GetName());
+		}
+		else
+		{
+			UE_LOG(LogEditorBuildPromotionTests, Error, TEXT("Blueprint is in an unexpected state after compiling (%s)"), *InBlueprint->GetName());
+		}
+	}
+
+	/**
+	* Creates a blueprint component based off the supplied asset
+	*
+	* @param InBlueprint - The blueprint to modify
+	* @param InAsset - The asset to use for the component
+	*/
+	static USCS_Node* CreateBlueprintComponent(UBlueprint* InBlueprint, UObject* InAsset)
+	{
+		IAssetEditorInstance* OpenEditor = FAssetEditorManager::Get().FindEditorForAsset(InBlueprint, true);
+		FBlueprintEditor* CurrentBlueprintEditor = (FBlueprintEditor*)OpenEditor;
+		TSubclassOf<UActorComponent> ComponentClass = FComponentAssetBrokerage::GetPrimaryComponentForAsset(InAsset->GetClass());
+		
+		USCS_Node* NewNode = InBlueprint->SimpleConstructionScript->CreateNode(ComponentClass);
+		// Assign the asset to the template
+		FComponentAssetBrokerage::AssignAssetToComponent(NewNode->ComponentTemplate, InAsset);
+
+		// Add node to the SCS
+		TArray<USCS_Node*> AllNodes = InBlueprint->SimpleConstructionScript->GetAllNodes();
+		USCS_Node* RootNode = AllNodes.Num() > 0 ? AllNodes[0] : NULL;
+		if (!RootNode || (RootNode == InBlueprint->SimpleConstructionScript->GetDefaultSceneRootNode()))
+		{
+			//New Root
+			InBlueprint->SimpleConstructionScript->AddNode(NewNode);
+		}
+		else
+		{
+			//Add as a child
+			RootNode->AddChildNode(NewNode);
+		}
+
+		// Recompile skeleton because of the new component we added
+		FKismetEditorUtilities::GenerateBlueprintSkeleton(InBlueprint, true);
+
+		CurrentBlueprintEditor->UpdateSCSPreview(true);
+
+		return NewNode;
+	}
+
+	/**
+	* Sets a new component as the root
+	*
+	* @param InBlueprint - The blueprint to modify
+	* @param NewRoot - The new root
+	*/
+	static void SetComponentAsRoot(UBlueprint* InBlueprint, USCS_Node* NewRoot)
+	{
+		//Get all the construction script nodes
+		TArray<USCS_Node*> AllNodes = InBlueprint->SimpleConstructionScript->GetAllNodes();
+
+		USCS_Node* OldRootNode = AllNodes[0];
+		USCS_Node* OldParent = NULL;
+
+		//Find old parent
+		for (int32 NodeIndex = 0; NodeIndex < AllNodes.Num(); ++NodeIndex)
+		{
+			if (AllNodes[NodeIndex]->ChildNodes.Contains(NewRoot))
+			{
+				OldParent = AllNodes[NodeIndex];
+				break;
+			}
+		}
+
+		//Remove the new root from its old parent and 
+		OldParent->ChildNodes.Remove(NewRoot);
+		NewRoot->Modify();
+		NewRoot->AttachToName = NAME_None;
+
+		//Remove the old root, add the new root, and attach the old root as a child
+		InBlueprint->SimpleConstructionScript->RemoveNode(OldRootNode);
+		InBlueprint->SimpleConstructionScript->AddNode(NewRoot);
+		NewRoot->AddChildNode(OldRootNode);
+	}
+
+	/**
+	* Removes a blueprint component from the simple construction script
+	*
+	* @param InBlueprint - The blueprint to modify
+	* @param InNode - The node of the component to remove
+	*/
+	static void RemoveBlueprintComponent(UBlueprint* InBlueprint, USCS_Node* InNode)
+	{
+		if (InNode != NULL)
+		{
+			// Remove node from SCS tree
+			InNode->GetSCS()->RemoveNodeAndPromoteChildren(InNode);
+
+			// Clear the delegate
+			InNode->SetOnNameChanged(FSCSNodeNameChanged());
+		}
+	}
+
+	/**
+	* Creates a new graph node from a given template
+	*
+	* @param NodeTemplate - The template to use for the node
+	* @param InGraph - The graph to create the new node in
+	* @param GraphLocation - The location to place the node
+	* @param ConnectPin - The pin to connect the node to
+	*/
+	static UEdGraphNode* CreateNewGraphNodeFromTemplate(UK2Node* NodeTemplate, UEdGraph* InGraph, const FVector2D& GraphLocation, UEdGraphPin* ConnectPin = NULL)
+	{
+		const FString EmptyString(TEXT(""));
+		TSharedPtr<FEdGraphSchemaAction_K2NewNode> Action = TSharedPtr<FEdGraphSchemaAction_K2NewNode>(new FEdGraphSchemaAction_K2NewNode(EmptyString, FText::FromString(EmptyString), EmptyString, 0));
+		Action->NodeTemplate = NodeTemplate;
+
+		return Action->PerformAction(InGraph, ConnectPin, GraphLocation, false);
+	}
+
+	/**
+	* Creates an AddComponent action to the blueprint graph
+	*
+	* @param InBlueprint - The blueprint to modify
+	* @param InGraph - The blueprint graph to use
+	* @param InAsset - The asset to use
+	*/
+	static UEdGraphNode* CreateAddComponentActionNode(UBlueprint* InBlueprint, UEdGraph* InGraph, UObject* InAsset)
+	{
+		UEdGraph* TempOuter = NewObject<UEdGraph>((UObject*)InBlueprint);
+		TempOuter->SetFlags(RF_Transient);
+
+		const FScopedTransaction PropertyChanged(LOCTEXT("AddedGraphNode", "Added a graph node"));
+		InGraph->Modify();
+
+		// Make an add component node
+		UK2Node_CallFunction* CallFuncNode = NewObject<UK2Node_AddComponent>(TempOuter);
+		UFunction* AddComponentFn = FindFieldChecked<UFunction>(AActor::StaticClass(), UK2Node_AddComponent::GetAddComponentFunctionName());
+		CallFuncNode->FunctionReference.SetFromField<UFunction>(AddComponentFn, FBlueprintEditorUtils::IsActorBased(InBlueprint));
+
+		UEdGraphNode* NewNode = CreateNewGraphNodeFromTemplate(CallFuncNode, InGraph, FVector2D(200, 0));
+
+		TSubclassOf<UActorComponent> ComponentClass = InAsset ? FComponentAssetBrokerage::GetPrimaryComponentForAsset(InAsset->GetClass()) : NULL;
+		if ((NewNode != NULL) && (InBlueprint != NULL))
+		{
+			UK2Node_AddComponent* AddCompNode = CastChecked<UK2Node_AddComponent>(NewNode);
+
+			ensure(NULL != Cast<UBlueprintGeneratedClass>(InBlueprint->GeneratedClass));
+			// Then create a new template object, and add to array in
+			UActorComponent* NewTemplate = ConstructObject<UActorComponent>(ComponentClass, InBlueprint->GeneratedClass);
+			NewTemplate->SetFlags(RF_ArchetypeObject);
+			InBlueprint->ComponentTemplates.Add(NewTemplate);
+
+			// Set the name of the template as the default for the TemplateName param
+			UEdGraphPin* TemplateNamePin = AddCompNode->GetTemplateNamePinChecked();
+			if (TemplateNamePin)
+			{
+				TemplateNamePin->DefaultValue = NewTemplate->GetName();
+			}
+
+			// Set the return type to be the type of the template
+			UEdGraphPin* ReturnPin = AddCompNode->GetReturnValuePin();
+			if (ReturnPin)
+			{
+				ReturnPin->PinType.PinSubCategoryObject = *ComponentClass;
+			}
+
+			// Set the asset
+			if (InAsset != NULL)
+			{
+				FComponentAssetBrokerage::AssignAssetToComponent(NewTemplate, InAsset);
+			}
+
+			AddCompNode->ReconstructNode();
+		}
+
+		FBlueprintEditorUtils::MarkBlueprintAsModified(InBlueprint);
+		return NewNode;
+	}
+
+	/**
+	* Creates a SetStaticMesh action in the blueprint graph
+	*
+	* @param InBlueprint - The blueprint to modify
+	* @param InGraph - The blueprint graph to use
+	*/
+	static UEdGraphNode* AddSetStaticMeshNode(UBlueprint* InBlueprint, UEdGraph* InGraph)
+	{
+		UEdGraph* TempOuter = NewObject<UEdGraph>((UObject*)InBlueprint);
+		TempOuter->SetFlags(RF_Transient);
+
+		// Make a call function template
+		UK2Node_CallFunction* CallFuncNode = NewObject<UK2Node_CallFunction>(TempOuter);
+		static FName PrintStringFunctionName(TEXT("SetStaticMesh"));
+		UFunction* DelayFn = FindFieldChecked<UFunction>(UStaticMeshComponent::StaticClass(), PrintStringFunctionName);
+		CallFuncNode->FunctionReference.SetFromField<UFunction>(DelayFn, false);
+
+		return CreateNewGraphNodeFromTemplate(CallFuncNode, InGraph, FVector2D(850, 0));
+	}
+
+	/**
+	* Connects two nodes using the supplied pin names
+	*
+	* @param NodeA - The first node to connect
+	* @param PinAName - The name of the pin on the first node
+	* @param NodeB - The second node to connect
+	* @param PinBName - The name of the pin on the second node
+	*/
+	static void ConnectGraphNodes(UEdGraphNode* NodeA, const FString& PinAName, UEdGraphNode* NodeB, const FString& PinBName)
+	{
+		const FScopedTransaction PropertyChanged(LOCTEXT("ConnectedNode", "Connected graph nodes"));
+		NodeA->GetGraph()->Modify();
+
+		UEdGraphPin* PinA = NodeA->FindPin(PinAName);
+		UEdGraphPin* PinB = NodeB->FindPin(PinBName);
+
+		if (PinA && PinB)
+		{
+			PinA->MakeLinkTo(PinB);
+		}
+		else
+		{
+			UE_LOG(LogEditorBuildPromotionTests, Error, TEXT("Could not connect pins %s and %s "), *PinAName, *PinBName);
+		}
+	}
+
+	/**
+	* Promotes a pin to a variable
+	*
+	* @param InBlueprint - The blueprint to modify
+	* @param Node - The node that owns the pin
+	* @param PinName - The name of the pin to promote
+	*/
+	static void PromotePinToVariable(UBlueprint* InBlueprint, UEdGraphNode* Node, const FString& PinName)
+	{
+		IAssetEditorInstance* OpenEditor = FAssetEditorManager::Get().FindEditorForAsset(InBlueprint, true);
+		FBlueprintEditor* CurrentBlueprintEditor = (FBlueprintEditor*)OpenEditor;
+
+		UEdGraphPin* PinToPromote = Node->FindPin(PinName);
+		CurrentBlueprintEditor->DoPromoteToVariable(InBlueprint, PinToPromote);
+	}
+
+	/**
+	* Creates a ReceiveBeginPlay event node
+	*
+	* @param InBlueprint - The blueprint to modify
+	* @param InGraph - The graph to use for the new node
+	*/
+	static UEdGraphNode* CreatePostBeginPlayEvent(UBlueprint* InBlueprint, UEdGraph* InGraph)
+	{
+		UEdGraph* TempOuter = NewObject<UEdGraph>((UObject*)InBlueprint);
+		TempOuter->SetFlags(RF_Transient);
+
+		// Make an add component node
+		UK2Node_Event* NewEventNode = NewObject<UK2Node_Event>(TempOuter);
+		NewEventNode->EventSignatureName = "ReceiveBeginPlay";
+		NewEventNode->EventSignatureClass = AActor::StaticClass();
+		NewEventNode->bOverrideFunction = true;
+
+		//Check for existing events
+		UK2Node_Event* ExistingEvent = FBlueprintEditorUtils::FindOverrideForFunction(InBlueprint, NewEventNode->EventSignatureClass, NewEventNode->EventSignatureName);
+
+		if (!ExistingEvent)
+		{
+			return CreateNewGraphNodeFromTemplate(NewEventNode, InGraph, FVector2D(200, 0));
+		}
+		return ExistingEvent;
+	}
+
+	/**
+	* Creates a custom event node
+	*
+	* @param InBlueprint - The blueprint to modify
+	* @param InGraph - The graph to use for the new node
+	* @param EventName - The name of the event
+	*/
+	static UEdGraphNode* CreateCustomEvent(UBlueprint* InBlueprint, UEdGraph* InGraph, const FString& EventName)
+	{
+		UEdGraph* TempOuter = NewObject<UEdGraph>((UObject*)InBlueprint);
+		TempOuter->SetFlags(RF_Transient);
+
+		// Make an add component node
+		UK2Node_CustomEvent* NewEventNode = NewObject<UK2Node_CustomEvent>(TempOuter);
+		NewEventNode->CustomFunctionName = "EventName";
+
+		return CreateNewGraphNodeFromTemplate(NewEventNode, InGraph, FVector2D(1200, 0));
+	}
+
+	/**
+	* Creates a node template for a UKismetSystemLibrary function
+	*
+	* @param NodeOuter - The outer to use for the template
+	* @param InGraph - The function to use for the node
+	*/
+	static UK2Node* CreateKismetFunctionTemplate(UObject* NodeOuter, const FName& FunctionName)
+	{
+		// Make a call function template
+		UK2Node_CallFunction* CallFuncNode = NewObject<UK2Node_CallFunction>(NodeOuter);
+		UFunction* Function = FindFieldChecked<UFunction>(UKismetSystemLibrary::StaticClass(), FunctionName);
+		CallFuncNode->FunctionReference.SetFromField<UFunction>(Function, false);
+		return CallFuncNode;
+	}
+
+	/**
+	* Creates a delay node
+	*
+	* @param InBlueprint - The blueprint to modify
+	* @param InGraph - The graph to use for the new node
+	* @param ConnectPin - The pin to connect the new node to
+	*/
+	static UEdGraphNode* AddDelayNode(UBlueprint* InBlueprint, UEdGraph* InGraph, UEdGraphPin* ConnectPin = NULL)
+	{
+		UEdGraph* TempOuter = NewObject<UEdGraph>((UObject*)InBlueprint);
+		TempOuter->SetFlags(RF_Transient);
+
+		const FScopedTransaction PropertyChanged(LOCTEXT("AddedGraphNode", "Added a graph node"));
+		InGraph->Modify();
+
+		// Make a call function template
+		static FName DelayFunctionName(TEXT("Delay"));
+		UK2Node* CallFuncNode = CreateKismetFunctionTemplate(TempOuter, DelayFunctionName);
+
+		//Create the node
+		return CreateNewGraphNodeFromTemplate(CallFuncNode, InGraph, FVector2D(400, 0), ConnectPin);
+	}
+
+	/**
+	* Creates a PrintString node
+	*
+	* @param InBlueprint - The blueprint to modify
+	* @param InGraph - The graph to use for the new node
+	* @param ConnectPin - The pin to connect the new node to
+	*/
+	static UEdGraphNode* AddPrintStringNode(UBlueprint* InBlueprint, UEdGraph* InGraph, UEdGraphPin* ConnectPin = NULL)
+	{
+		UEdGraph* TempOuter = NewObject<UEdGraph>((UObject*)InBlueprint);
+		TempOuter->SetFlags(RF_Transient);
+
+		// Make a call function template
+		static FName PrintStringFunctionName(TEXT("PrintString"));
+		UK2Node* CallFuncNode = CreateKismetFunctionTemplate(TempOuter, PrintStringFunctionName);
+
+		return CreateNewGraphNodeFromTemplate(CallFuncNode, InGraph, FVector2D(680, 0), ConnectPin);
+	}
+
+	/**
+	* Creates a call function node
+	*
+	* @param InBlueprint - The blueprint to modify
+	* @param InGraph - The graph to use for the new node
+	* @param FunctionName - The name of the function to call
+	* @param ConnectPin - The pin to connect the new node to
+	*/
+	static UEdGraphNode* AddCallFunctionGraphNode(UBlueprint* InBlueprint, UEdGraph* InGraph, const FName& FunctionName, UEdGraphPin* ConnectPin = NULL)
+	{
+		UEdGraph* TempOuter = NewObject<UEdGraph>((UObject*)InBlueprint);
+		TempOuter->SetFlags(RF_Transient);
+
+		// Make a call function template
+		UK2Node_CallFunction* CallFuncNode = NewObject<UK2Node_CallFunction>(TempOuter);
+		CallFuncNode->FunctionReference.SetSelfMember(FunctionName);
+
+		return CreateNewGraphNodeFromTemplate(CallFuncNode, InGraph, FVector2D(1200, 0), ConnectPin);
+	}
+
+	/**
+	* Creates Get or Set node
+	*
+	* @param InBlueprint - The blueprint to modify
+	* @param InGraph - The graph to use for the new node
+	* @param VarName - The name of the variable to use
+	* @param bGet - If true, create a Get node.  If false, create a Set node.
+	* @param XOffset - How far to offset the node in the graph
+	*/
+	static UEdGraphNode* AddGetSetNode(UBlueprint* InBlueprint, UEdGraph* InGraph, const FString& VarName, bool bGet, float XOffset = 0.f)
+	{
+		const FScopedTransaction PropertyChanged(LOCTEXT("AddedGraphNode", "Added a graph node"));
+		InGraph->Modify();
+
+		FEdGraphSchemaAction_K2NewNode NodeInfo;
+		// Create get or set node, depending on whether we clicked on an input or output pin
+		UK2Node_Variable* TemplateNode = NULL;
+		if (bGet)
+		{
+			TemplateNode = NewObject<UK2Node_VariableGet>();
+		}
+		else
+		{
+			TemplateNode = NewObject<UK2Node_VariableSet>();
+		}
+
+		TemplateNode->VariableReference.SetSelfMember(FName(*VarName));
+		NodeInfo.NodeTemplate = TemplateNode;
+
+		return NodeInfo.PerformAction(InGraph, NULL, FVector2D(XOffset,130), true);
+	}
+
+	/**
+	* Sets the default value for a pin
+	*
+	* @param Node - The node that owns the pin to set
+	* @param PinName - The name of the pin
+	* @param PinValue - The new default value
+	*/
+	static void SetPinDefaultValue(UEdGraphNode* Node, const FString& PinName, const FString& PinValue)
+	{
+		UEdGraphPin* Pin = Node->FindPin(PinName);
+		Pin->DefaultValue = PinValue;
+	}
+
+	/**
+	* Sets the default object for a pin
+	*
+	* @param Node - The node that owns the pin to set
+	* @param PinName - The name of the pin
+	* @param PinObject - The new default object
+	*/
+	static void SetPinDefaultObject(UEdGraphNode* Node, const FString& PinName, UObject* PinObject)
+	{
+		UEdGraphPin* Pin = Node->FindPin(PinName);
+		Pin->DefaultObject = PinObject;
+	}
+
+	/**
+	* Adds a string member variable to a blueprint
+	*
+	* @param InBlueprint - The blueprint to modify
+	* @param VariableName - The name of the new string variable
+	*/
+	static void AddStringMemberValue(UBlueprint* InBlueprint, const FName& VariableName)
+	{
+		const UEdGraphSchema_K2* K2Schema = GetDefault<UEdGraphSchema_K2>();
+		FEdGraphPinType StringPinType(K2Schema->PC_String, TEXT(""), NULL, false, false);
+		FBlueprintEditorUtils::AddMemberVariable(InBlueprint, VariableName, StringPinType);
+	}
+
+	/**
+	* Creates a new function graph
+	*
+	* @param InBlueprint - The blueprint to modify
+	* @param FunctionName - The function name to use for the new graph
+	*/
+	static UEdGraph* CreateNewFunctionGraph(UBlueprint* InBlueprint, const FName& FunctionName)
+	{
+		UEdGraph* NewGraph = FBlueprintEditorUtils::CreateNewGraph(InBlueprint, FunctionName, UEdGraph::StaticClass(), UEdGraphSchema_K2::StaticClass());
+		FBlueprintEditorUtils::AddFunctionGraph<UClass>(InBlueprint, NewGraph, /*bIsUserCreated=*/ true, NULL);
+		return NewGraph;
+	}
+
+	/**
+	* Starts a lighting build
+	*/
+	static void BuildLighting()
+	{
+		FLightingBuildOptions LightingBuildOptions;
+
+		// Retrieve settings from ini.
+		GConfig->GetBool(TEXT("LightingBuildOptions"), TEXT("OnlyBuildSelected"), LightingBuildOptions.bOnlyBuildSelected, GEditorUserSettingsIni);
+		GConfig->GetBool(TEXT("LightingBuildOptions"), TEXT("OnlyBuildCurrentLevel"), LightingBuildOptions.bOnlyBuildCurrentLevel, GEditorUserSettingsIni);
+		GConfig->GetBool(TEXT("LightingBuildOptions"), TEXT("OnlyBuildSelectedLevels"), LightingBuildOptions.bOnlyBuildSelectedLevels, GEditorUserSettingsIni);
+		GConfig->GetBool(TEXT("LightingBuildOptions"), TEXT("OnlyBuildVisibility"), LightingBuildOptions.bOnlyBuildVisibility, GEditorUserSettingsIni);
+		GConfig->GetBool(TEXT("LightingBuildOptions"), TEXT("UseErrorColoring"), LightingBuildOptions.bUseErrorColoring, GEditorUserSettingsIni);
+		GConfig->GetBool(TEXT("LightingBuildOptions"), TEXT("ShowLightingBuildInfo"), LightingBuildOptions.bShowLightingBuildInfo, GEditorUserSettingsIni);
+		int32 QualityLevel;
+		GConfig->GetInt(TEXT("LightingBuildOptions"), TEXT("QualityLevel"), QualityLevel, GEditorUserSettingsIni);
+		QualityLevel = FMath::Clamp<int32>(QualityLevel, Quality_Preview, Quality_Production);
+		LightingBuildOptions.QualityLevel = (ELightingBuildQuality)QualityLevel;
+
+		GUnrealEd->BuildLighting(LightingBuildOptions);
+	}
+
+	/**
+	* Sets an object property value by name
+	*
+	* @param TargetObject - The object to modify
+	* @param InVariableName - The name of the property
+	*/
+	static void SetPropertyByName(UObject* TargetObject, const FString& InVariableName, const FString& NewValueString)
+	{
+		UProperty* FoundProperty = FindField<UProperty>(TargetObject->GetClass(), *InVariableName);
+		if (FoundProperty)
+		{
+			const FScopedTransaction PropertyChanged(LOCTEXT("PropertyChanged", "Object Property Change"));
+
+			TargetObject->Modify();
+
+			TargetObject->PreEditChange(FoundProperty);
+			FoundProperty->ImportText(*NewValueString, FoundProperty->ContainerPtrToValuePtr<uint8>(TargetObject), 0, TargetObject);
+			FPropertyChangedEvent PropertyChangedEvent(FoundProperty, false, EPropertyChangeType::ValueSet);
+			TargetObject->PostEditChangeProperty(PropertyChangedEvent);
+		}
+	}
+
+	/**
+	* Gets an object property value by name
+	*
+	* @param TargetObject - The object to modify
+	* @param InVariableName - The name of the property
+	*/
+	static FString GetPropertyByName(UObject* TargetObject, const FString& InVariableName)
+	{
+		UProperty* FoundProperty = FindField<UProperty>(TargetObject->GetClass(), *InVariableName);
+		if (FoundProperty)
+		{
+			FString ValueString;
+			const uint8* PropertyAddr = FoundProperty->ContainerPtrToValuePtr<uint8>(TargetObject);
+			FoundProperty->ExportTextItem(ValueString, PropertyAddr, NULL, NULL, PPF_None);
+			return ValueString;
+		}
+		return TEXT("");
+	}
+
+	/**
+	* Starts a PIE session
+	*/
+	static void StartPIE(bool bSimulateInEditor)
+	{
+		FLevelEditorModule& LevelEditorModule = FModuleManager::Get().GetModuleChecked<FLevelEditorModule>(TEXT("LevelEditor"));
+		TSharedPtr<class ILevelViewport> ActiveLevelViewport = LevelEditorModule.GetFirstActiveViewport();
+
+		GUnrealEd->RequestPlaySession(false, ActiveLevelViewport, bSimulateInEditor, NULL, NULL, -1, false);
+	}
+
+	/**
+	* Ends a PIE session
+	*/
+	static void EndPIE()
+	{
+		GUnrealEd->RequestEndPlayMap();
+	}
+
+	/**
+	* Takes an automation screenshot
+	*
+	* @param ScreenshotName - The sub name to use for the screenshot
+	*/
+	static void TakeScreenshot(const FString& ScreenshotName, bool bUseTopWindow = false)
+	{
+		//Update the screenshot name, then take a screenshot.
+		if (FAutomationTestFramework::GetInstance().IsScreenshotAllowed())
+		{
+			TSharedPtr<SWindow> Window;
+
+			if (bUseTopWindow)
+			{
+				Window = FSlateApplication::Get().GetActiveTopLevelWindow();
+			}
+			else
+			{
+				//Find the main editor window
+				TArray<TSharedRef<SWindow> > AllWindows;
+				FSlateApplication::Get().GetAllVisibleWindowsOrdered(AllWindows);
+				if (AllWindows.Num() == 0)
+				{
+					UE_LOG(LogEditorAutomationTests, Error, TEXT("ERROR: Could not find the main editor window."));
+					return;
+				}
+
+				Window = AllWindows[0];
+			}
+
+			if (Window.IsValid())
+			{
+				FString ScreenshotFileName;
+				const FString TestName = FString::Printf(TEXT("EditorBuildPromotion/%s"), *ScreenshotName);
+				AutomationCommon::GetScreenshotPath(TestName, ScreenshotFileName, false);
+
+				TSharedRef<SWidget> WindowRef = Window.ToSharedRef();
+
+				TArray<FColor> OutImageData;
+				FIntVector OutImageSize;
+				FSlateApplication::Get().TakeScreenshot(WindowRef, OutImageData, OutImageSize);
+				FAutomationTestFramework::GetInstance().OnScreenshotCaptured().ExecuteIfBound(OutImageSize.X, OutImageSize.Y, OutImageData, ScreenshotFileName);
+			}
+			else
+			{
+				UE_LOG(LogEditorBuildPromotionTests, Error, TEXT("Failed to find editor window for screenshot (%s)"), *ScreenshotName);
+			}
+		}
+	}
+
+	/**
+	* Gets the asset data from a package path
+	*
+	* @param PackagePath - The package path used to look up the asset data
+	*/
+	static FAssetData GetAssetDataFromPackagePath(const FString& PackagePath)
+	{
+		FString AssetPath = AutomationEditorCommonUtils::ConvertPackagePathToAssetPath(PackagePath);
+		if (AssetPath.Len() > 0)
+		{
+			IAssetRegistry& AssetRegistry = FModuleManager::LoadModuleChecked<FAssetRegistryModule>(TEXT("AssetRegistry")).Get();
+			return AssetRegistry.GetAssetByObjectPath(*AssetPath);
+		}
+
+		return FAssetData();
+	}
+
+	/**
+	* Adds a default mesh to the level and applies a material
+	*
+	* @param Location - The location to place the actor
+	* @param Material - The material to apply
+	*/
+	static AActor* AddDefaultMeshToLevelWithMaterial(const FVector& Location, UMaterialInterface* Material)
+	{
+		UAutomationTestSettings const* AutomationTestSettings = GetDefault<UAutomationTestSettings>();
+		check(AutomationTestSettings);
+
+		// Default static mesh
+		FString AssetPackagePath = AutomationTestSettings->BuildPromotionTest.DefaultStaticMeshAsset.FilePath;
+		if (AssetPackagePath.Len() > 0)
+		{
+			FAssetData AssetData = EditorBuildPromotionTestUtils::GetAssetDataFromPackagePath(AssetPackagePath);
+			UStaticMesh* DefaultMesh = Cast<UStaticMesh>(AssetData.GetAsset());
+			if (DefaultMesh)
+			{
+				AStaticMeshActor* PlacedMesh = Cast<AStaticMeshActor>(FActorFactoryAssetProxy::AddActorForAsset(DefaultMesh));
+				PlacedMesh->SetActorLocation(Location);
+
+				PlacedMesh->StaticMeshComponent->SetMaterial(0, Material);
+
+				return PlacedMesh;
+			}
+			else
+			{
+				UE_LOG(LogEditorBuildPromotionTests, Error, TEXT("DefaultStaticMeshAsset is invalid.  Unable to add material to level"));
+			}
+		}
+		else
+		{
+			UE_LOG(LogEditorBuildPromotionTests, Warning, TEXT("Can't add material to level because no DefaultMeshAsset is assigned"));
+		}
+
+		return NULL;
+	}
+
+	/**
+	* Imports an asset using the supplied factory and file
+	*
+	* @param ImportFactory - The factory to use to import the asset
+	* @param ImportPath - The file path of the file to use
+	*/
+	static UObject* ImportAsset(UFactory* ImportFactory, const FString& ImportPath)
+	{
+		FString Name = ObjectTools::SanitizeObjectName(FPaths::GetBaseFilename(ImportPath));
+		FString PackageName = FString::Printf(TEXT("%s/%s"), *GetGamePath(), *Name);
+
+		UObject* ImportedAsset = AutomationEditorCommonUtils::ImportAssetUsingFactory(ImportFactory,Name,PackageName,ImportPath);
+
+		return ImportedAsset;
+	}
+
+	static void PlaceImportedAsset(UObject* InObject, FVector& PlaceLocation)
+	{
+		if (InObject)
+		{
+			AActor* PlacedActor = NULL;
+			UTexture* TextureObject = Cast<UTexture>(InObject);
+			if (TextureObject)
+			{
+				//Don't add if we are a normal map
+				if (TextureObject->LODGroup == TEXTUREGROUP_WorldNormalMap)
+				{
+					return;
+				}
+				else
+				{
+					UMaterial* NewMaterial = EditorBuildPromotionTestUtils::CreateMaterialFromTexture(TextureObject);
+					PlacedActor = EditorBuildPromotionTestUtils::AddDefaultMeshToLevelWithMaterial(PlaceLocation, NewMaterial);
+				}
+			}
+			else
+			{
+				PlacedActor = FActorFactoryAssetProxy::AddActorForAsset(InObject);
+			}
+
+			if (PlacedActor)
+			{
+				PlacedActor->SetActorLocation(PlaceLocation);
+				PlaceLocation.Y += 100;
+				UE_LOG(LogEditorBuildPromotionTests, Display, TEXT("Placed %s in the level"), *InObject->GetName());
+			}
+			else
+			{
+				UE_LOG(LogEditorBuildPromotionTests, Error, TEXT("Failed to place %s in the level"), *InObject->GetName());
+			}
+		}
+	}
+
+	/**
+	* Imports an asset using the supplied factory, class, and name
+	*
+	* @param CreateFactory - The factory to use to create the asset
+	* @param AssetClass - The class of the new asset
+	* @param AssetName - The name to use for the new asset
+	*/
+	static UObject* CreateAsset(UFactory* CreateFactory, UClass* AssetClass, const FString& AssetName)
+	{
+		FString PackageName = FString::Printf(TEXT("%s/%s"), *GetGamePath(), *AssetName);
+		UPackage* AssetPackage = CreatePackage(NULL, *PackageName);
+		EObjectFlags Flags = RF_Public | RF_Standalone;
+
+		UObject* CreatedAsset = CreateFactory->FactoryCreateNew(AssetClass, AssetPackage, FName(*AssetName), Flags, NULL, GWarn);
+
+		if (CreatedAsset)
+		{
+			// Notify the asset registry
+			FAssetRegistryModule::AssetCreated(CreatedAsset);
+
+			// Mark the package dirty...
+			AssetPackage->MarkPackageDirty();
+
+			UE_LOG(LogEditorBuildPromotionTests, Display, TEXT("Created asset %s (%s)"), *AssetName, *AssetClass->GetName());
+		}
+		else
+		{
+			UE_LOG(LogEditorBuildPromotionTests, Error, TEXT("Unable to create asset of type %s"), *AssetClass->GetName());
+		}
+
+		return CreatedAsset;
+	}
+
+	/**
+	* Checks that the level geometry changed after adding a brush and checks that undo / redo works
+	*
+	* @param CurrentLevel - The current level
+	* @param VertsBefore - The number of verts before the brush was added
+	* @param VertsAfter - The number of verts after the brush was added
+	* @param BrushType - The type of brush to use in the log
+	* @param bAdditive - True if this brush was added in additive mode
+	*/
+	static void TestGeometryUndoRedo(ULevel* CurrentLevel, int32 VertsBefore, int32 VertsAfter, const FString& BrushType, bool bAdditive)
+	{
+		if (VertsAfter > VertsBefore)
+		{
+			UE_LOG(LogEditorBuildPromotionTests, Display, TEXT("Placed a %s in %s mode"), *BrushType, bAdditive ? TEXT("additive") : TEXT("subtractive"));
+		}
+		else
+		{
+			UE_LOG(LogEditorBuildPromotionTests, Error, TEXT("Failed to place a %s in %s mode"), *BrushType, bAdditive ? TEXT("additive") : TEXT("subtractive"));
+		}
+
+		//Undo
+		GEditor->UndoTransaction();
+		if (CurrentLevel->Model->Verts.Num() == VertsBefore)
+		{
+			UE_LOG(LogEditorBuildPromotionTests, Display, TEXT("Undo successful for %s"), *BrushType);
+		}
+		else
+		{
+			UE_LOG(LogEditorBuildPromotionTests, Error, TEXT("Undo failed for %s"), *BrushType);
+		}
+
+		//Redo
+		GEditor->RedoTransaction();
+		if (CurrentLevel->Model->Verts.Num() == VertsAfter)
+		{
+			UE_LOG(LogEditorBuildPromotionTests, Display, TEXT("Redo successful for %s"), *BrushType);
+		}
+		else
+		{
+			UE_LOG(LogEditorBuildPromotionTests, Error, TEXT("Redo failed for %s"), *BrushType);
+		}
+	}
+
+	/**
+	* Saves all assets in a given folder
+	*
+	* @param InFolder - The folder that contains the assets to save
+	*/
+	static void SaveAllAssetsInFolder(const FString& InFolder)
+	{
+		// Load the asset registry module
+		FAssetRegistryModule& AssetRegistryModule = FModuleManager::Get().LoadModuleChecked<FAssetRegistryModule>(TEXT("AssetRegistry"));
+
+		// Form a filter from the paths
+		FARFilter Filter;
+		Filter.bRecursivePaths = true;
+		new (Filter.PackagePaths) FName(*InFolder);
+		
+
+		// Query for a list of assets in the selected paths
+		TArray<FAssetData> AssetList;
+		AssetRegistryModule.Get().GetAssets(Filter, AssetList);
+
+		// Form a list of unique package names from the assets
+		TSet<FName> UniquePackageNames;
+		for (int32 AssetIdx = 0; AssetIdx < AssetList.Num(); ++AssetIdx)
+		{
+			UniquePackageNames.Add(AssetList[AssetIdx].PackageName);
+		}
+
+		// Add all unique package names to the output
+		TArray<FString> PackageNames;
+		for (auto PackageIt = UniquePackageNames.CreateConstIterator(); PackageIt; ++PackageIt)
+		{
+			PackageNames.Add((*PackageIt).ToString());
+		}
+
+		// Form a list of packages from the assets
+		TArray<UPackage*> Packages;
+		for (int32 PackageIdx = 0; PackageIdx < PackageNames.Num(); ++PackageIdx)
+		{
+			UPackage* Package = FindPackage(NULL, *PackageNames[PackageIdx]);
+
+			// Only save loaded and dirty packages
+			if (Package != NULL && Package->IsDirty())
+			{
+				Packages.Add(Package);
+			}
+		}
+
+		// Save all packages that were found
+		if (Packages.Num())
+		{
+			const bool bCheckDirty = false;
+			const bool bPromptToSave = false;
+			FEditorFileUtils::PromptForCheckoutAndSave(Packages, bCheckDirty, bPromptToSave);
+		}
+	}
+
+	/**
+	* Cleans up objects created by the main build promotion test
+	*/
+	static void PerformCleanup()
+	{
+		//Make sure we don't have any level references
+		AutomationEditorCommonUtils::CreateNewMap();
+
+		// Deselect all
+		GEditor->SelectNone(false, true);
+
+		// Clear the transaction buffer so we aren't referencing the objects
+		GUnrealEd->ResetTransaction(FText::FromString(TEXT("FAssetEditorTest")));
+
+		//remove all assets in the build package
+		// Load the asset registry module
+		IAssetRegistry& AssetRegistry = FModuleManager::LoadModuleChecked<FAssetRegistryModule>(TEXT("AssetRegistry")).Get();
+
+		// Form a filter from the paths
+		FARFilter Filter;
+		Filter.bRecursivePaths = true;
+		new (Filter.PackagePaths) FName(*EditorBuildPromotionTestUtils::GetGamePath());
+
+		// Query for a list of assets in the selected paths
+		TArray<FAssetData> AssetList;
+		AssetRegistry.GetAssets(Filter, AssetList);
+
+		// Clear and try to delete all assets
+		for (int32 AssetIdx = 0; AssetIdx < AssetList.Num(); ++AssetIdx)
+		{
+			UE_LOG(LogEditorBuildPromotionTests, Display, TEXT("Removing asset: %s"), *AssetList[AssetIdx].AssetName.ToString());
+			if (AssetList[AssetIdx].IsAssetLoaded())
+			{
+				UObject* LoadedAsset = AssetList[AssetIdx].GetAsset();
+				AssetRegistry.AssetDeleted(LoadedAsset);
+
+				bool bSuccessful = ObjectTools::DeleteSingleObject(LoadedAsset, false);
+
+				//If we failed to delete this object manually clear any references and try again
+				if (!bSuccessful)
+				{
+					//Clear references to the object so we can delete it
+					AutomationEditorCommonUtils::NullReferencesToObject(LoadedAsset);
+
+					bSuccessful = ObjectTools::DeleteSingleObject(LoadedAsset, false);
+				}
+			}
+		}
+
+		UE_LOG(LogEditorBuildPromotionTests, Display, TEXT("Clearing Path: %s"), *EditorBuildPromotionTestUtils::GetGamePath());
+		AssetRegistry.RemovePath(EditorBuildPromotionTestUtils::GetGamePath());
+
+		//Remove the directory
+		bool bEnsureExists = false;
+		bool bDeleteEntireTree = true;
+		FString PackageDirectory = FPaths::GameContentDir() / TEXT("BuildPromotionTest");
+		IFileManager::Get().DeleteDirectory(*PackageDirectory, bEnsureExists, bDeleteEntireTree);
+		UE_LOG(LogEditorBuildPromotionTests, Display, TEXT("Deleting Folder: %s"), *PackageDirectory);
+
+		//Remove the map
+		FString MapFilePath = FPaths::GameContentDir() / TEXT("Maps/EditorBuildPromotionTest.umap");
+		IFileManager::Get().Delete(*MapFilePath, false, true, true);
+		UE_LOG(LogEditorBuildPromotionTests, Display, TEXT("Deleting Map: %s"), *MapFilePath);
+	}
+
+	/**
+	* Gets the number of actors in the persistent level for a given UWorld
+	*
+	* @param InWorld - The world to check for actors
+	* @param ActorType - The type of actors to look for
+	*/
+	static int32 GetNumActors(UWorld* InWorld,const UClass* ActorType)
+	{
+		int32 NumActors = 0;
+		ULevel* PersistentLevel = InWorld->PersistentLevel;
+		for (int32 i = 0; i < PersistentLevel->Actors.Num(); ++i)
+		{
+			if (PersistentLevel->Actors[i]->IsA(ActorType))
+			{
+				NumActors++;
+			}
+		}
+		return NumActors;
+	}
+}
+
+
+
+/**
+* Container for items related to the create asset test
+*/
+namespace BuildPromotionTestHelper
+{
+	/** The possible states of the FOpenAssetInfo class */
+	namespace EState
+	{
+		enum Type
+		{
+			OpenEditor,
+			WaitForEditor,
+			ChangeProperty
+		};
+	}
+
+	/** Stores info on an asset that we are opening */
+	struct FOpenAssetInfo
+	{
+		/** The asset we are opening */
+		UObject* Asset;
+		/** The asset data  */
+		FAssetData AssetData;
+		/** The name of the property we are going to change */
+		FString PropertyName;
+		/** The new value to assign to the property */
+		FString PropertyValue;
+
+		/**
+		* Constructor
+		*/
+		FOpenAssetInfo(UObject* InAsset, const FAssetData& InAssetData, const FString& InPropertyName, const FString& InPropertyValue) :
+			Asset(InAsset),
+			AssetData(InAssetData),
+			PropertyName(InPropertyName),
+			PropertyValue(InPropertyValue)
+		{}
+	};
+
+	/** Helper class to open, modify, and add an asset to the level */
+	class FOpenAssetHelper
+	{
+	public:
+		/**
+		* Constructor
+		*/
+		FOpenAssetHelper(TArray<FOpenAssetInfo> InAssets, FAutomationTestExecutionInfo* InTestExecutionInfo) :
+			Assets(InAssets),
+			CurrentStage(EState::OpenEditor),
+			AssetIndex(0),
+			TestExecutionInfo(InTestExecutionInfo)
+		{
+		}
+
+		/**
+		* Updates the current stage
+		*/
+		bool Update()
+		{
+			if (AssetIndex < Assets.Num())
+			{
+				switch (CurrentStage)
+				{
+				case EState::OpenEditor:
+					OpenAssetEditor();
+					break;
+				case EState::WaitForEditor:
+					WaitForEditor();
+					break;
+				case EState::ChangeProperty:
+					ChangeProperty();
+					break;
+				}
+				return false;
+			}
+			else
+			{
+				return true;
+			}
+		}
+	private:
+
+		/**
+		* Opens the asset editor
+		*/
+		void OpenAssetEditor()
+		{
+			ErrorStartCount = TestExecutionInfo->Errors.Num();
+			WarningStartCount = TestExecutionInfo->Warnings.Num();
+			LogStartCount = TestExecutionInfo->LogItems.Num();
+
+			UObject* CurrentAsset = Assets[AssetIndex].Asset;
+			if (FAssetEditorManager::Get().OpenEditorForAsset(CurrentAsset))
+			{
+				CurrentStage = EState::WaitForEditor;
+			}
+			else
+			{
+				UE_LOG(LogEditorBuildPromotionTests, Error, TEXT("Failed to open the asset editor!"));
+
+				//Move on to the next asset
+				NextAsset();
+			}
+		}
+
+		/**
+		* Waits for the asset editor to open
+		*/
+		void WaitForEditor()
+		{
+			UObject* CurrentAsset = Assets[AssetIndex].Asset;
+			TSharedPtr<SWindow> ActiveWindow = FSlateApplication::Get().GetActiveTopLevelWindow();
+			if (ActiveWindow.IsValid())
+			{
+				FString ActiveWindowTitle = ActiveWindow->GetTitle().ToString();
+
+				//Check that we have the right window (Tutorial may have opened on top of the editor)
+				if (!ActiveWindowTitle.StartsWith(CurrentAsset->GetName()))
+				{
+					//Bring the asset editor to the front
+					FAssetEditorManager::Get().FindEditorForAsset(CurrentAsset, true);
+				}
+
+				UE_LOG(LogEditorAssetAutomationTests, Display, TEXT("Opened asset (%s)"), *CurrentAsset->GetClass()->GetName());
+
+				CurrentStage = EState::ChangeProperty;
+			}
+			else
+			{
+				WaitingForEditorCount++;
+				if (WaitingForEditorCount > MaxWaitForEditorTicks)
+				{
+
+					UE_LOG(LogEditorBuildPromotionTests, Error, TEXT("Timed out waiting for editor window"));
+
+					//Move on the next asset
+					NextAsset();
+				}
+			}
+		}
+
+		/**
+		* Modifies a property on the current asset
+		*/
+		void ChangeProperty()
+		{
+			UObject* CurrentAsset = Assets[AssetIndex].Asset;
+			FString PropertyName = Assets[AssetIndex].PropertyName;
+			FString NewPropertyValue = Assets[AssetIndex].PropertyValue;
+
+			FString OldPropertyValue = EditorBuildPromotionTestUtils::GetPropertyByName(CurrentAsset, PropertyName);
+			EditorBuildPromotionTestUtils::SetPropertyByName(CurrentAsset, PropertyName, NewPropertyValue);
+			UE_LOG(LogEditorAssetAutomationTests, Display, TEXT("Modified asset.  %s = %s"), *PropertyName, *NewPropertyValue);
+
+			//Get the property again and use that to compare the redo action.  Parsing the new value may change the formatting a bit. ie) 100 becomes 100.0000
+			FString ParsedNewValue = EditorBuildPromotionTestUtils::GetPropertyByName(CurrentAsset, PropertyName);
+
+			GEditor->UndoTransaction();
+			FString CurrentValue = EditorBuildPromotionTestUtils::GetPropertyByName(CurrentAsset, PropertyName);
+			if (CurrentValue == OldPropertyValue)
+			{
+				UE_LOG(LogEditorAssetAutomationTests, Display, TEXT("Undo %s change successful"), *PropertyName);
+			}
+			else
+			{
+				UE_LOG(LogEditorBuildPromotionTests, Error, TEXT("Undo didn't change property back to original value"));
+			}
+
+			GEditor->RedoTransaction();
+			CurrentValue = EditorBuildPromotionTestUtils::GetPropertyByName(CurrentAsset, PropertyName);
+			if (CurrentValue == ParsedNewValue)
+			{
+				UE_LOG(LogEditorAssetAutomationTests, Display, TEXT("Redo %s change successful"), *PropertyName);
+			}
+			else
+			{
+				UE_LOG(LogEditorBuildPromotionTests, Error, TEXT("Redo didnt' change property back to the modified value"));
+			}
+
+			//Apply if this is a material.  (Editor won't close unless we do)
+			if (CurrentAsset->IsA(UMaterialInterface::StaticClass()))
+			{
+				EditorBuildPromotionTestUtils::SendUpdateMaterialCommand();
+			}
+
+			
+
+			//Save
+			const FString PackagePath = Assets[AssetIndex].AssetData.PackageName.ToString();
+			if (CurrentAsset && PackagePath.Len() > 0)
+			{
+
+				UPackage* AssetPackage = FindPackage(NULL, *PackagePath);
+				if (AssetPackage)
+				{
+					AssetPackage->SetDirtyFlag(true);
+					FString PackageFileName = FPackageName::LongPackageNameToFilename(PackagePath, FPackageName::GetAssetPackageExtension());
+					FPlatformFileManager::Get().GetPlatformFile().SetReadOnly(*PackageFileName, false);
+					if (UPackage::SavePackage(AssetPackage, NULL, RF_Standalone, *PackageFileName, GError, nullptr, false, true, SAVE_NoError))
+					{
+						UE_LOG(LogEditorAssetAutomationTests, Display, TEXT("Saved asset"));
+					}
+					else
+					{
+						UE_LOG(LogEditorAssetAutomationTests, Display, TEXT("Unable to save asset"));
+					}
+				}
+			}
+
+			//close editor
+			FAssetEditorManager::Get().CloseAllAssetEditors();
+
+			//Add to level
+			FVector PlaceLocation(-1090,970,160);
+			PlaceLocation.Y += AssetIndex * 150.f;
+			UMaterialInterface* MaterialAsset = Cast<UMaterialInterface>(CurrentAsset);
+			UTexture* TextureAsset = Cast<UTexture>(CurrentAsset);
+			if (MaterialAsset)
+			{
+				AActor* PlacedActor = EditorBuildPromotionTestUtils::AddDefaultMeshToLevelWithMaterial(PlaceLocation, MaterialAsset);
+			}
+			else if (TextureAsset)
+			{
+				UMaterial* NewMaterial = EditorBuildPromotionTestUtils::CreateMaterialFromTexture(TextureAsset);
+				AActor* PlacedActor = EditorBuildPromotionTestUtils::AddDefaultMeshToLevelWithMaterial(PlaceLocation, NewMaterial);
+			}
+			else
+			{
+				AActor* PlacedActor = FActorFactoryAssetProxy::AddActorForAsset(CurrentAsset);
+				if (PlacedActor)
+				{
+					UE_LOG(LogEditorAssetAutomationTests, Display, TEXT("Placed %s in the level"), *CurrentAsset->GetName());
+					PlacedActor->SetActorLocation(PlaceLocation);
+				}
+				else
+				{
+					UE_LOG(LogEditorBuildPromotionTests, Warning, TEXT("Unable to place asset"));
+				}
+			}
+
+			NextAsset();
+		}
+
+		/**
+		* Switches to the next asset in the list
+		*/
+		void NextAsset()
+		{
+			const FString LogTag = Assets[AssetIndex].Asset->GetName();
+			for (int32 ErrorIndex = ErrorStartCount; ErrorIndex < TestExecutionInfo->Errors.Num(); ++ErrorIndex)
+			{
+				TestExecutionInfo->Errors[ErrorIndex] = FString::Printf(TEXT("%s: %s"), *LogTag, *TestExecutionInfo->Errors[ErrorIndex]);
+			}
+			for (int32 WarningIndex = WarningStartCount; WarningIndex < TestExecutionInfo->Warnings.Num(); ++WarningIndex)
+			{
+				TestExecutionInfo->Warnings[WarningIndex] = FString::Printf(TEXT("%s: %s"), *LogTag, *TestExecutionInfo->Warnings[WarningIndex]);
+			}
+			for (int32 LogIndex = LogStartCount; LogIndex < TestExecutionInfo->LogItems.Num(); ++LogIndex)
+			{
+				TestExecutionInfo->LogItems[LogIndex] = FString::Printf(TEXT("%s: %s"), *LogTag, *TestExecutionInfo->LogItems[LogIndex]);
+			}
+
+			AssetIndex++;
+			CurrentStage = EState::OpenEditor;
+			WaitingForEditorCount = 0;
+		}
+
+		/** The asset list */
+		TArray<FOpenAssetInfo> Assets;
+		/** The current stage */
+		EState::Type CurrentStage;
+		/** The index of the current asset */
+		int32 AssetIndex;
+		/** How many ticks we have waited for the asset editor to open */
+		int32 WaitingForEditorCount;
+		/** The maximum number of ticks to wait for the editor*/
+		int32 MaxWaitForEditorTicks;
+
+		/** Pointer to the execution info to prefix logs*/
+		FAutomationTestExecutionInfo* TestExecutionInfo;
+		/** Log group counters */
+		int32 ErrorStartCount;
+		int32 WarningStartCount;
+		int32 LogStartCount;
+	};
+
+	/**
+	* The main build promotion test class
+	*/
+	struct FBuildPromotionTest
+	{
+		/** Pointer to the execution info of this test */
+		FAutomationTestExecutionInfo* TestExecutionInfo;
+
+		/** The number of existing errors, warnings, and logs when the command started */
+		int32 LastErrorCount;
+		int32 LastWarningCount;
+		int32 LastLogCount;
+
+		/** Function definition for the test stage functions */
+		typedef bool(BuildPromotionTestHelper::FBuildPromotionTest::*TestStageFn)();
+
+		/** List of test stage functions */
+		TArray<TestStageFn> TestStages;
+		TArray<FString> StageNames;
+
+		/** The index of the test stage we are on */
+		int32 CurrentStage;
+
+		/** Pointer to the active editor world */
+		UWorld* CurrentWorld;
+
+		/** Point light created by the test */
+		APointLight* PointLight;
+
+		/** If true, we will revert the auto apply lighting setting when the lighting build finishes */
+		bool bDisableAutoApplyLighting;
+
+		/** Items that were imported during the Import Workflow stage */
+		UTexture* DiffuseTexture;
+		UTexture* NormalTexture;
+		UStaticMesh* WindowMesh;
+		UStaticMesh* ReimportMesh;
+		USkeletalMesh* BlendShape;
+		USkeletalMesh* MorphAndMorphAnim;
+		USkeletalMesh* SkeletalMesh_Test;
+		UAnimSequence* AnimationTest;
+		USoundWave* FemaleVoice;
+		USoundWave* SurroundTest;
+
+		/** Pointer to the material we are editing for the source control test stage */
+		UMaterial* SCTestMat;
+		FString ChosenMaterialColor;
+
+		/** Material created from the "Creating a Material" stage */
+		UMaterial* CreatedMaterial;
+
+		/** Particle System created from the "Creating a Particle System" stage */
+		UParticleSystem* CreatedPS;
+
+		/** Helper for opening, modifying, and placing assets */
+		FOpenAssetHelper* OpenAssetHelper;
+
+		/** Meshes to use for the blueprint */
+		UStaticMesh* FirstBlueprintMesh;
+		UStaticMesh* SecondBlueprintMesh;
+
+		/** Objects created by the Blueprint stages */
+		UBlueprint* BlueprintObject;
+		UPackage* BlueprintPackage;
+		UEdGraph* CustomGraph;
+		USCS_Node* MeshNode;
+		USCS_Node* OtherMeshNode;
+		USCS_Node* PSNode;
+		UEdGraphNode* AddMeshNode;
+		UEdGraphNode* PostBeginPlayEventNode;
+		UEdGraphNode* DelayNode;
+		UEdGraphNode* SetNode;
+		UEdGraphNode* GetNode;
+		UEdGraphNode* PrintNode;
+		UEdGraphNode* SetStaticMeshNode;
+		UEdGraphNode* CustomEventNode;
+		UEdGraphNode* AddParticleSystemNode;
+		UEdGraphNode* CallFunctionNode;
+		AActor* PlacedBlueprint;
+
+		/** Source control async helper */
+		SourceControlAutomationCommon::FAsyncCommandHelper AsyncHelper;
+
+		/** Delay helper */
+		EditorBuildPromotionTestUtils::FDelayHelper DelayHelper;
+
+		/** List of skipped tests */
+		TArray<FString> SkippedTests;
+
+		/** summary logs to display at the end */
+		TArray<FString> SummaryLines;
+
+		/** Keeps track of errors generated by building the map */
+		int32 BuildStartErrorCount;
+
+		int32 SectionSuccessCount;
+		int32 SectionTestCount;
+
+		int32 SectionErrorCount;
+		int32 SectionWarningCount;
+		int32 SectionLogCount;
+
+#define ADD_TEST_STAGE(FuncName,StageName) \
+	TestStages.Add(&BuildPromotionTestHelper::FBuildPromotionTest::FuncName); \
+	StageNames.Add(StageName); 
+
+		/**
+		* Constructor
+		*/
+		FBuildPromotionTest(FAutomationTestExecutionInfo* InExecutionInfo) :
+			CurrentStage(0)
+		{
+			FMemory::Memzero(this, sizeof(FBuildPromotionTest));
+
+			TestExecutionInfo = InExecutionInfo;
+
+			// 2) Geometry
+			ADD_TEST_STAGE(Geometry_LevelCreationAndSetup,	TEXT("Level Creation and Setup"));
+			ADD_TEST_STAGE(Geometry_AddLevelGeometry,		TEXT("Level Geometry"));
+			ADD_TEST_STAGE(EndSection,						TEXT("Geometry Workflow"));
+
+			// 3) Lighting
+			ADD_TEST_STAGE(Lighting_PlaceAPointLight,		TEXT("Place a Point Light"));
+			ADD_TEST_STAGE(Lighting_PointLightProperties,	TEXT("Point Light Properties"));
+			ADD_TEST_STAGE(Lighting_DuplicatePointLight,	TEXT("Duplicate Point Light"));
+			ADD_TEST_STAGE(Lighting_BuildLighting_Part1,	TEXT("Build Lighting"));
+			ADD_TEST_STAGE(Lighting_BuildLighting_Part2,	TEXT("Build Lighting"));
+			ADD_TEST_STAGE(EndSection, TEXT("Lighting Workflow"));
+
+			// 4) Importing Workflow
+			ADD_TEST_STAGE(Workflow_ImportWorkflow, TEXT("")); // Not using a subsection name here because it would be redundant
+			ADD_TEST_STAGE(EndSection, TEXT("Importing Workflow"));
+
+			// 5) Content Browser
+			ADD_TEST_STAGE(ContentBrowser_SourceControl_Part1,			TEXT("Source Control"));
+			ADD_TEST_STAGE(ContentBrowser_SourceControl_Part2,			TEXT("Source Control"));
+			ADD_TEST_STAGE(ContentBrowser_SourceControl_Part3,			TEXT("Source Control"));
+			ADD_TEST_STAGE(ContentBrowser_SourceControl_Part4,			TEXT("Source Control"));
+			ADD_TEST_STAGE(ContentBrowser_SourceControl_Part5,			TEXT("Source Control"));
+			ADD_TEST_STAGE(ContentBrowser_OpenAssets_Part1,				TEXT("Open Asset Types"));
+			ADD_TEST_STAGE(ContentBrowser_OpenAssets_Part2,				TEXT("Open Asset Types"));
+			ADD_TEST_STAGE(ContentBrowser_ReimportAsset,				TEXT("Re-import Assets"));
+			ADD_TEST_STAGE(ContentBrowser_CreateAMaterial_Part1,		TEXT("Creating a Material"));
+			ADD_TEST_STAGE(ContentBrowser_CreateAMaterial_Part2,		TEXT("Creating a Material"));
+			ADD_TEST_STAGE(ContentBrowser_CreateAMaterial_Part3,		TEXT("Creating a Material"));
+			ADD_TEST_STAGE(ContentBrowser_CreateAParticleSystem_Part1,	TEXT("Creating a Particle System"));
+			ADD_TEST_STAGE(ContentBrowser_CreateAParticleSystem_Part2,	TEXT("Creating a Particle System"));
+			ADD_TEST_STAGE(EndSection, TEXT("Content Browser"));
+
+			// 6) Blueprints
+			ADD_TEST_STAGE(Blueprint_CreateNewBlueprint_Part1,	TEXT("Create a new Blueprint"));
+			ADD_TEST_STAGE(Blueprint_CreateNewBlueprint_Part2,	TEXT("Create a new Blueprint"));
+			ADD_TEST_STAGE(Blueprint_DataOnlyBlueprint_Part1,	TEXT("Data-only Blueprint"));
+			ADD_TEST_STAGE(Blueprint_DataOnlyBlueprint_Part2,	TEXT("Data-only Blueprint"));
+			ADD_TEST_STAGE(Blueprint_DataOnlyBlueprint_Part3,	TEXT("Data-only Blueprint"));
+			ADD_TEST_STAGE(Blueprint_ComponentsMode_Part1,		TEXT("Components Mode"));
+			ADD_TEST_STAGE(Blueprint_ComponentsMode_Part2,		TEXT("Components Mode"));
+			ADD_TEST_STAGE(Blueprint_ConstructionScript,		TEXT("Construction Script"));
+			ADD_TEST_STAGE(Blueprint_PromoteVariable_Part1,		TEXT("Variable from Component Mode"));
+			ADD_TEST_STAGE(Blueprint_PromoteVariable_Part2,		TEXT("Variable from Component Mode"));
+			ADD_TEST_STAGE(Blueprint_PromoteVariable_Part3,		TEXT("Variable from Component Mode"));
+			ADD_TEST_STAGE(Blueprint_EventGraph,				TEXT("Event Graph"));
+			ADD_TEST_STAGE(Blueprint_CustomVariable,			TEXT("Custom Variables"));
+			ADD_TEST_STAGE(Blueprint_UsingVariables,			TEXT("Using Variables"));
+			ADD_TEST_STAGE(Blueprint_RenameCustomEvent,			TEXT("Renaming Custom Event"));
+			ADD_TEST_STAGE(Blueprint_NewFunctions,				TEXT("New Function"));
+			ADD_TEST_STAGE(Blueprint_CompleteBlueprint,			TEXT("Completing the Blueprint"));
+			ADD_TEST_STAGE(Blueprint_Placement_Part1,			TEXT("Blueprint Placement"));
+			ADD_TEST_STAGE(Blueprint_Placement_Part2,			TEXT("Blueprint Placement"));
+			ADD_TEST_STAGE(Blueprint_Placement_Part3,			TEXT("Blueprint Placement"));
+			ADD_TEST_STAGE(Blueprint_Placement_Part4,			TEXT("Blueprint Placement"));
+			ADD_TEST_STAGE(Blueprint_SetBreakpoint_Part1,		TEXT("Set Breakpoints"));
+			ADD_TEST_STAGE(Blueprint_SetBreakpoint_Part2,		TEXT("Set Breakpoints"));
+			ADD_TEST_STAGE(Blueprint_SetBreakpoint_Part3,		TEXT("Set Breakpoints"));
+			ADD_TEST_STAGE(Blueprint_LevelScript_Part1,			TEXT("Level Script"));
+			ADD_TEST_STAGE(Blueprint_LevelScript_Part2,			TEXT("Level Script"));
+			ADD_TEST_STAGE(Blueprint_LevelScript_Part3,			TEXT("Level Script"));
+			ADD_TEST_STAGE(Blueprint_LevelScript_Part4,			TEXT("Level Script"));
+			ADD_TEST_STAGE(Blueprint_LevelScript_Part5,			TEXT("Level Script"));
+			ADD_TEST_STAGE(EndSection, TEXT("Blueprint"));
+
+			// 8) Build, play, save
+			ADD_TEST_STAGE(Building_BuildLevel_Part1, TEXT("")); // Not using a subsection name here because it would be redundant
+			ADD_TEST_STAGE(Building_BuildLevel_Part2, TEXT(""));
+			ADD_TEST_STAGE(Building_BuildLevel_Part3, TEXT(""));
+			ADD_TEST_STAGE(Building_BuildLevel_Part4, TEXT(""));
+			ADD_TEST_STAGE(EndSection, TEXT("Building and Saving"));
+
+			ADD_TEST_STAGE(LogSummary, TEXT(""));
+		}
+#undef ADD_TEST_STAGE
+
+		/**
+		* Runs the current test stage
+		*/
+		bool Update()
+		{
+			bool bStageComplete = (this->*TestStages[CurrentStage])();
+			if (bStageComplete)
+			{
+				//Only add headers if the next section has a different name, or we are the last one
+				if ( (CurrentStage + 1 < TestStages.Num() && StageNames[CurrentStage] != StageNames[CurrentStage+1]) || (CurrentStage + 1 == TestStages.Num()) )
+				{
+					TagPreviousLogs(StageNames[CurrentStage]);
+				}
+				
+				CurrentStage++;
+			}
+			return CurrentStage >= TestStages.Num();
+		}
+
+	private:
+
+		/**
+		* Handles tagging all logs that have been created since the last time this was called
+		*
+		* @param NewLogTag - The tag to prefix the logs with
+		*/
+		void TagPreviousLogs(const FString& NewLogTag)
+		{
+			if (TestExecutionInfo && NewLogTag.Len() > 0)
+			{
+				for (int32 ErrorIndex = LastErrorCount; ErrorIndex < TestExecutionInfo->Errors.Num(); ++ErrorIndex)
+				{
+					TestExecutionInfo->Errors[ErrorIndex] = FString::Printf(TEXT("%s: %s"), *NewLogTag, *TestExecutionInfo->Errors[ErrorIndex]);
+				}
+				for (int32 WarningIndex = LastWarningCount; WarningIndex < TestExecutionInfo->Warnings.Num(); ++WarningIndex)
+				{
+					TestExecutionInfo->Warnings[WarningIndex] = FString::Printf(TEXT("%s: %s"), *NewLogTag, *TestExecutionInfo->Warnings[WarningIndex]);
+				}
+				for (int32 LogIndex = LastLogCount; LogIndex < TestExecutionInfo->LogItems.Num(); ++LogIndex)
+				{
+					TestExecutionInfo->LogItems[LogIndex] = FString::Printf(TEXT("%s: %s"), *NewLogTag, *TestExecutionInfo->LogItems[LogIndex]);
+				}
+			}
+
+			//This sub tests was a success if we had no new errors
+			if (LastErrorCount == TestExecutionInfo->Errors.Num())
+			{
+				SectionSuccessCount++;
+			}
+			SectionTestCount++;
+
+			LastErrorCount = TestExecutionInfo->Errors.Num();
+			LastWarningCount = TestExecutionInfo->Warnings.Num();
+			LastLogCount = TestExecutionInfo->LogItems.Num();
+			
+		}
+
+		/**
+		* Handle adding headers and success summary
+		*/
+		bool EndSection()
+		{
+			const FString SectionName = StageNames[CurrentStage];
+			if (TestExecutionInfo->Errors.Num() > SectionErrorCount)
+			{
+				TestExecutionInfo->Errors.Insert(FString::Printf(TEXT("\n\n--%s Pass--"), *SectionName), SectionErrorCount);
+			}
+			if (TestExecutionInfo->Warnings.Num() > SectionWarningCount)
+			{
+				TestExecutionInfo->Warnings.Insert(FString::Printf(TEXT("\n\n--%s Pass--"), *SectionName), SectionWarningCount);
+			}
+			if (TestExecutionInfo->LogItems.Num() > SectionLogCount)
+			{
+				TestExecutionInfo->LogItems.Insert(FString::Printf(TEXT("\n\n--%s Pass--"), *SectionName), SectionLogCount);
+				UE_LOG(LogEditorBuildPromotionTests, Display, TEXT("%i of %i tests successful"), SectionSuccessCount, SectionTestCount);
+				SummaryLines.Add(FString::Printf(TEXT("%i of %i %s tests were successful"), SectionSuccessCount, SectionTestCount, *SectionName));
+			}
+
+			//Reset section test counts
+			SectionSuccessCount = -1;
+			SectionTestCount = -1;
+
+			//Reset section log counts
+			SectionErrorCount = TestExecutionInfo->Errors.Num();
+			SectionWarningCount = TestExecutionInfo->Warnings.Num();
+			SectionLogCount = TestExecutionInfo->LogItems.Num();
+
+			//Reset test log counts so we don't prefix the pass or success logs added here
+			LastErrorCount = TestExecutionInfo->Errors.Num();
+			LastWarningCount = TestExecutionInfo->Warnings.Num();
+			LastLogCount = TestExecutionInfo->LogItems.Num();
+
+			return true;
+		}
+
+		/**
+		* Adds a summary to the end of the log
+		*/
+		bool LogSummary()
+		{
+			//Log out the summary lines
+			if (SummaryLines.Num() > 0)
+			{
+				UE_LOG(LogEditorBuildPromotionTests, Display, TEXT("\nSummary:"));
+				for (int32 i = 0; i < SummaryLines.Num(); ++i)
+				{
+					UE_LOG(LogEditorBuildPromotionTests, Display, TEXT("    %s"), *SummaryLines[i]);
+				}
+			}
+			
+			//Log out skipped tests
+			if (SkippedTests.Num() > 0)
+			{
+				UE_LOG(LogEditorBuildPromotionTests, Display, TEXT("\nSkipped the following tests:"));
+				for (int32 i = 0; i < SkippedTests.Num(); ++i)
+				{
+					UE_LOG(LogEditorBuildPromotionTests, Display, TEXT("    %s"), *SkippedTests[i]);
+				}
+			}
+
+			UE_LOG(LogEditorBuildPromotionTests, Display, TEXT("\nPlease restart the editor and continue to Step 2"));
+			return true;
+		}
+
+		/**
+		* Geometry Test Stage: Level Creation and Setup
+		*   Create a new map and add a light
+		*/
+		bool Geometry_LevelCreationAndSetup()
+		{
+			//Create a new level
+			CurrentWorld = AutomationEditorCommonUtils::CreateNewMap();
+			UE_LOG(LogEditorBuildPromotionTests, Display, TEXT("Created an empty level"));
+
+			//Add a directional light
+			const FTransform Transform(FRotator(-45.f,5.f,0),FVector(0.0f, 0.0f, 400.0f));
+			ADirectionalLight* DirectionalLight = Cast<ADirectionalLight>(GEditor->AddActor(CurrentWorld->GetCurrentLevel(), ADirectionalLight::StaticClass(), Transform));
+
+			if (DirectionalLight)
+			{
+				UE_LOG(LogEditorBuildPromotionTests, Display, TEXT("Placed a directional light"));
+			}
+			else
+			{
+				UE_LOG(LogEditorBuildPromotionTests, Error, TEXT("Failed to place directional light"));
+			}
+
+			return true;
+		}
+
+		/**
+		* Geometry Test Stage: Adding, Subtracting and Editing Geometry
+		*    Adds different world geometry to the map and adds a player start
+		*/
+		bool Geometry_AddLevelGeometry()
+		{
+			GEditor->Exec(CurrentWorld, TEXT("BRUSH Scale 1 1 1"));
+
+			for (int32 i = 0; i < GEditor->LevelViewportClients.Num(); i++)
+			{
+				FLevelEditorViewportClient* ViewportClient = GEditor->LevelViewportClients[i];
+				if (!ViewportClient->IsOrtho())
+				{
+					ViewportClient->SetViewLocation(FVector(176, 2625, 2075));
+					ViewportClient->SetViewRotation(FRotator(319, 269, 1));
+				}
+			}
+			ULevel* CurrentLevel = CurrentWorld->GetCurrentLevel();
+			int32 VertsBefore = CurrentLevel->Model->Verts.Num();
+
+			//Cube Additive Brush
+			UCubeBuilder* CubeAdditiveBrushBuilder = Cast<UCubeBuilder>(GEditor->FindBrushBuilder(UCubeBuilder::StaticClass()));
+			CubeAdditiveBrushBuilder->X = 8192.0f;
+			CubeAdditiveBrushBuilder->Y = 8192.0f;
+			CubeAdditiveBrushBuilder->Z = 128.0f;
+			CubeAdditiveBrushBuilder->Build(CurrentWorld);
+			GEditor->Exec(CurrentWorld, TEXT("BRUSH MOVETO X=0 Y=0 Z=0"));
+			GEditor->Exec(CurrentWorld, TEXT("BRUSH ADD"));
+			int32 VertsAfter = CurrentLevel->Model->Verts.Num();
+
+			//Log results and test Undo/Redo
+			EditorBuildPromotionTestUtils::TestGeometryUndoRedo(CurrentLevel, VertsBefore, VertsAfter, TEXT("cube"), true);
+
+			//Cone Additive Brush
+			VertsBefore = CurrentLevel->Model->Verts.Num();
+			UConeBuilder* ConeAdditiveBrushBuilder = Cast<UConeBuilder>(GEditor->FindBrushBuilder(UConeBuilder::StaticClass()));
+			ConeAdditiveBrushBuilder->Z = 1024.0f;
+			ConeAdditiveBrushBuilder->CapZ = 256.0f;
+			ConeAdditiveBrushBuilder->OuterRadius = 512.0f;
+			ConeAdditiveBrushBuilder->InnerRadius = 384;
+			ConeAdditiveBrushBuilder->Sides = 32;
+			ConeAdditiveBrushBuilder->Build(CurrentWorld);
+			GEditor->Exec(CurrentWorld, TEXT("BRUSH MOVETO X=-1525 Y=-1777 Z=64"));
+			GEditor->Exec(CurrentWorld, TEXT("BRUSH ADD"));
+			VertsAfter = CurrentLevel->Model->Verts.Num();
+
+			//Log results and test Undo/Redo
+			EditorBuildPromotionTestUtils::TestGeometryUndoRedo(CurrentLevel, VertsBefore, VertsAfter, TEXT("cone"), true);
+
+			//Sphere Additive Brush
+			VertsBefore = CurrentLevel->Model->Verts.Num();
+			UTetrahedronBuilder* TetraAdditiveBrushBuilder = Cast<UTetrahedronBuilder>(GEditor->FindBrushBuilder(UTetrahedronBuilder::StaticClass()));
+			TetraAdditiveBrushBuilder->Radius = 512.0f;
+			TetraAdditiveBrushBuilder->SphereExtrapolation = 3;
+			TetraAdditiveBrushBuilder->Build(CurrentWorld);
+			GEditor->Exec(CurrentWorld, TEXT("BRUSH MOVETO X=-88 Y=-1777 Z=535"));
+			GEditor->Exec(CurrentWorld, TEXT("BRUSH ADD"));
+			VertsAfter = CurrentLevel->Model->Verts.Num();
+
+			//Log results and test Undo/Redo
+			EditorBuildPromotionTestUtils::TestGeometryUndoRedo(CurrentLevel, VertsBefore, VertsAfter, TEXT("tetrahedron"), true);
+
+			//Cylinder Additive Brush
+			VertsBefore = CurrentLevel->Model->Verts.Num();
+			UCylinderBuilder* CylinderAdditiveBrushBuilder = Cast<UCylinderBuilder>(GEditor->FindBrushBuilder(UCylinderBuilder::StaticClass()));
+			CylinderAdditiveBrushBuilder->Z = 1024.0f;
+			CylinderAdditiveBrushBuilder->OuterRadius = 512.0f;
+			CylinderAdditiveBrushBuilder->InnerRadius = 384.0f;
+			CylinderAdditiveBrushBuilder->Sides = 16;
+			CylinderAdditiveBrushBuilder->Build(CurrentWorld);
+			GEditor->Exec(CurrentWorld, TEXT("BRUSH MOVETO X=1338 Y=-1776 Z=535"));
+			GEditor->Exec(CurrentWorld, TEXT("BRUSH ADD"));
+			VertsAfter = CurrentLevel->Model->Verts.Num();
+
+			//Log results and test Undo/Redo
+			EditorBuildPromotionTestUtils::TestGeometryUndoRedo(CurrentLevel, VertsBefore, VertsAfter, TEXT("cylinder"), true);
+
+			//Cylinder Additive Brush
+			VertsBefore = CurrentLevel->Model->Verts.Num();
+			USheetBuilder* SheetAdditiveBrushBuilder = Cast<USheetBuilder>(GEditor->FindBrushBuilder(USheetBuilder::StaticClass()));
+			SheetAdditiveBrushBuilder->X = 512.0f;
+			SheetAdditiveBrushBuilder->Y = 512.0f;
+			SheetAdditiveBrushBuilder->XSegments = 1;
+			SheetAdditiveBrushBuilder->YSegments = 1;
+			SheetAdditiveBrushBuilder->Axis = AX_YAxis;
+			SheetAdditiveBrushBuilder->Build(CurrentWorld);
+			GEditor->Exec(CurrentWorld, TEXT("BRUSH MOVETO X=-760 Y=-346 Z=535"));
+			GEditor->Exec(CurrentWorld, TEXT("BRUSH ADD"));
+			VertsAfter = CurrentLevel->Model->Verts.Num();
+
+			//Log results and test Undo/Redo
+			EditorBuildPromotionTestUtils::TestGeometryUndoRedo(CurrentLevel, VertsBefore, VertsAfter, TEXT("sheet"), true);
+
+			//Volume Additive Brush
+			VertsBefore = CurrentLevel->Model->Verts.Num();
+			UVolumetricBuilder* VolumetricAdditiveBrushBuilder = Cast<UVolumetricBuilder>(GEditor->FindBrushBuilder(UVolumetricBuilder::StaticClass()));
+			VolumetricAdditiveBrushBuilder->Z = 512.0f;
+			VolumetricAdditiveBrushBuilder->Radius = 128.0f;
+			VolumetricAdditiveBrushBuilder->NumSheets = 3;
+			VolumetricAdditiveBrushBuilder->Build(CurrentWorld);
+			GEditor->Exec(CurrentWorld, TEXT("BRUSH MOVETO X=445 Y=-345 Z=535"));
+			GEditor->Exec(CurrentWorld, TEXT("BRUSH ADD"));
+			VertsAfter = CurrentLevel->Model->Verts.Num();
+
+			//Log results and test Undo/Redo
+			EditorBuildPromotionTestUtils::TestGeometryUndoRedo(CurrentLevel, VertsBefore, VertsAfter, TEXT("volume"), true);
+
+			//Linear Stair Additive Brush
+			VertsBefore = CurrentLevel->Model->Verts.Num();
+			ULinearStairBuilder* LinearStairAdditiveBrushBuilder = Cast<ULinearStairBuilder>(GEditor->FindBrushBuilder(ULinearStairBuilder::StaticClass()));
+			LinearStairAdditiveBrushBuilder->StepLength = 64.0f;
+			LinearStairAdditiveBrushBuilder->StepHeight = 16.0f;
+			LinearStairAdditiveBrushBuilder->StepWidth = 256.0f;
+			LinearStairAdditiveBrushBuilder->NumSteps = 8;
+			LinearStairAdditiveBrushBuilder->AddToFirstStep = 0;
+			LinearStairAdditiveBrushBuilder->Build(CurrentWorld);
+			GEditor->Exec(CurrentWorld, TEXT("BRUSH MOVETO X=1464 Y=-345 Z=-61"));
+			GEditor->Exec(CurrentWorld, TEXT("BRUSH ADD"));
+			VertsAfter = CurrentLevel->Model->Verts.Num();
+
+			//Log results and test Undo/Redo
+			EditorBuildPromotionTestUtils::TestGeometryUndoRedo(CurrentLevel, VertsBefore, VertsAfter, TEXT("linear stair"), true);
+
+			//Curved Stair Additive Brush
+			VertsBefore = CurrentLevel->Model->Verts.Num();
+			UCurvedStairBuilder* CurvedStairAdditiveBrushBuilder = Cast<UCurvedStairBuilder>(GEditor->FindBrushBuilder(UCurvedStairBuilder::StaticClass()));
+			CurvedStairAdditiveBrushBuilder->InnerRadius = 240.0f;
+			CurvedStairAdditiveBrushBuilder->StepHeight = 16.0f;
+			CurvedStairAdditiveBrushBuilder->StepWidth = 256.0f;
+			CurvedStairAdditiveBrushBuilder->AngleOfCurve = 90.0f;
+			CurvedStairAdditiveBrushBuilder->NumSteps = 4;
+			CurvedStairAdditiveBrushBuilder->AddToFirstStep = 0;
+			CurvedStairAdditiveBrushBuilder->CounterClockwise = false;
+			CurvedStairAdditiveBrushBuilder->Build(CurrentWorld);
+			GEditor->Exec(CurrentWorld, TEXT("BRUSH MOVETO X=-1290 Y=263 Z=193"));
+			GEditor->Exec(CurrentWorld, TEXT("BRUSH ADD"));
+			VertsAfter = CurrentLevel->Model->Verts.Num();
+
+			//Log results and test Undo/Redo
+			EditorBuildPromotionTestUtils::TestGeometryUndoRedo(CurrentLevel, VertsBefore, VertsAfter, TEXT("curved stair"), true);
+
+			//Spiral Stair Additive Brush
+			VertsBefore = CurrentLevel->Model->Verts.Num();
+			USpiralStairBuilder* SpiralStairAdditiveBrushBuilder = Cast<USpiralStairBuilder>(GEditor->FindBrushBuilder(USpiralStairBuilder::StaticClass()));
+			SpiralStairAdditiveBrushBuilder->InnerRadius = 64;
+			SpiralStairAdditiveBrushBuilder->StepWidth = 256.0f;
+			SpiralStairAdditiveBrushBuilder->StepHeight = 16.0f;
+			SpiralStairAdditiveBrushBuilder->StepThickness = 32.0f;
+			SpiralStairAdditiveBrushBuilder->NumStepsPer360 = 8;
+			SpiralStairAdditiveBrushBuilder->NumSteps = 8;
+			SpiralStairAdditiveBrushBuilder->SlopedCeiling = true;
+			SpiralStairAdditiveBrushBuilder->Build(CurrentWorld);
+			GEditor->Exec(CurrentWorld, TEXT("BRUSH MOVETO X=850 Y=263 Z=193"));
+			GEditor->Exec(CurrentWorld, TEXT("BRUSH ADD"));
+			VertsAfter = CurrentLevel->Model->Verts.Num();
+
+			//Log results and test Undo/Redo
+			EditorBuildPromotionTestUtils::TestGeometryUndoRedo(CurrentLevel, VertsBefore, VertsAfter, TEXT("spiral stair"), true);
+
+			//Cylinder Additive Brush
+			VertsBefore = CurrentLevel->Model->Verts.Num();
+			UCylinderBuilder* CylinderSubtractiveBrushBuilder = Cast<UCylinderBuilder>(GEditor->FindBrushBuilder(UCylinderBuilder::StaticClass()));
+			CylinderSubtractiveBrushBuilder->Z = 256;
+			CylinderSubtractiveBrushBuilder->OuterRadius = 512.0f;
+			CylinderSubtractiveBrushBuilder->InnerRadius = 384.0f;
+			CylinderSubtractiveBrushBuilder->Sides = 3;
+			CylinderSubtractiveBrushBuilder->Build(CurrentWorld);
+			GEditor->Exec(CurrentWorld, TEXT("BRUSH MOVETO X=0 Y=0 Z=128"));
+			GEditor->Exec(CurrentWorld, TEXT("BRUSH SUBTRACT"));
+			VertsAfter = CurrentLevel->Model->Verts.Num();
+
+			//Log results and test Undo/Redo
+			EditorBuildPromotionTestUtils::TestGeometryUndoRedo(CurrentLevel, VertsBefore, VertsAfter, TEXT("cylinder"), true);
+
+			//Cube Subtractive Brush
+			VertsBefore = CurrentLevel->Model->Verts.Num();
+			UCubeBuilder* CubeSubtractiveBrushBuilder = Cast<UCubeBuilder>(GEditor->FindBrushBuilder(UCubeBuilder::StaticClass()));
+			CubeSubtractiveBrushBuilder->X = 256.0f;
+			CubeSubtractiveBrushBuilder->Y = 1024.0f;
+			CubeSubtractiveBrushBuilder->Z = 256.0f;
+			CubeSubtractiveBrushBuilder->Build(CurrentWorld);
+			GEditor->Exec(CurrentWorld, TEXT("BRUSH MOVETO X=-88 Y=-1777 Z=535"));
+			GEditor->Exec(CurrentWorld, TEXT("BRUSH SUBTRACT"));
+			VertsAfter = CurrentLevel->Model->Verts.Num();
+
+			//Log results and test Undo/Redo
+			EditorBuildPromotionTestUtils::TestGeometryUndoRedo(CurrentLevel, VertsBefore, VertsAfter, TEXT("cube"), false);
+
+			//Add a playerstart
+			const FTransform Transform(FRotator(0, -90, 0), FVector(0.f, 1750.f, 166.f));
+			AActor* PlayerStart = GEditor->AddActor(CurrentWorld->GetCurrentLevel(), APlayerStart::StaticClass(), Transform);
+			if (PlayerStart)
+			{
+				UE_LOG(LogEditorBuildPromotionTests, Display, TEXT("Placed a player start"));
+			}
+			else
+			{
+				UE_LOG(LogEditorBuildPromotionTests, Error, TEXT("Failed to place a PlayerStart"));
+			}
+
+			return true;
+		}
+		
+		/**
+		* Lighting Test Stage: Place a point Light
+		*    Adds a point light to the map and sets its location, rotation, and scale
+		*/
+		bool Lighting_PlaceAPointLight()
+		{
+			const FTransform Transform(FVector(0.0f, 50.0f, 400.0f));
+			PointLight = Cast<APointLight>(GEditor->AddActor(CurrentWorld->GetCurrentLevel(), APointLight::StaticClass(), Transform));
+			if (PointLight)
+			{
+				UE_LOG(LogEditorBuildPromotionTests, Display, TEXT("Placed a PointLight"));
+				PointLight->SetActorLocation(FVector(100.f, 50.f, 200.f));
+				UE_LOG(LogEditorBuildPromotionTests, Display, TEXT("Modified the PointLight's Location"));
+				PointLight->SetActorRotation(FRotator(300, 250, -91));
+				UE_LOG(LogEditorBuildPromotionTests, Display, TEXT("Modified the PointLight's Rotation"));
+				PointLight->SetActorScale3D(FVector(2.f, 2.f, 2.f));
+				UE_LOG(LogEditorBuildPromotionTests, Display, TEXT("Modified the PointLight's Scale"));
+			}
+			else
+			{
+				UE_LOG(LogEditorBuildPromotionTests, Error, TEXT("Error creating point light"));
+			}
+
+			return true;
+		}
+
+		/**
+		* Lighting Test Stage: Point Light Properties
+		*    Modifies the point light properties (Intensity, LightColor, and AttenuationRadius)
+		*/
+		bool Lighting_PointLightProperties()
+		{
+			if (PointLight)
+			{
+				UPointLightComponent* PointLightComp = PointLight->PointLightComponent;
+				EditorBuildPromotionTestUtils::SetPropertyByName(PointLightComp, TEXT("Intensity"), TEXT("1000.f"));
+				UE_LOG(LogEditorBuildPromotionTests, Display, TEXT("Modified the PointLight's Intensity"));
+				EditorBuildPromotionTestUtils::SetPropertyByName(PointLightComp, TEXT("LightColor"), TEXT("(R=0,G=0,B=255)"));
+				UE_LOG(LogEditorBuildPromotionTests, Display, TEXT("Modified the PointLight's LightColor"));
+				EditorBuildPromotionTestUtils::SetPropertyByName(PointLightComp, TEXT("AttenuationRadius"), TEXT("1024.f"));
+				UE_LOG(LogEditorBuildPromotionTests, Display, TEXT("Modified the PointLight's AttenuationRadius"));
+			}
+			else
+			{
+				SkippedTests.Add(TEXT("Lighting: Point Light Properties. (Failed to create light)"));
+			}
+
+			return true;
+		}
+
+		/**
+		* Lighting Test Stage: Duplicate point light
+		*    Duplicates the point light created earlier using Copy/Paste and Duplicate
+		*/
+		bool Lighting_DuplicatePointLight()
+		{
+			if (PointLight)
+			{
+				int32 StartingLightCount = EditorBuildPromotionTestUtils::GetNumActors(CurrentWorld, APointLight::StaticClass());
+
+				//Select the light
+				GEditor->SelectNone(false, true, false);
+				GEditor->SelectActor(PointLight, true, false, true);
+
+				GEngine->Exec(CurrentWorld, TEXT("EDIT COPY"));
+				GEngine->Exec(CurrentWorld, TEXT("EDIT PASTE"));
+
+				//Check if we gained a new light
+				int32 CurrentLightCount = EditorBuildPromotionTestUtils::GetNumActors(CurrentWorld, APointLight::StaticClass());
+				if (CurrentLightCount - StartingLightCount == 1)
+				{
+					UE_LOG(LogEditorBuildPromotionTests, Display, TEXT("Duplicated the PointLight using copy / paste"));
+				}
+				else
+				{
+					UE_LOG(LogEditorBuildPromotionTests, Error, TEXT("Failed to duplicate light using copy / paste"));
+				}
+
+				//Reset the starting count
+				StartingLightCount = CurrentLightCount;
+
+				//Select the light
+				GEditor->SelectNone(false, true, false);
+				GEditor->SelectActor(PointLight, true, false, true);
+
+				GEngine->Exec(CurrentWorld, TEXT("DUPLICATE"));
+
+				//Check if we gained a new light
+				CurrentLightCount = EditorBuildPromotionTestUtils::GetNumActors(CurrentWorld, APointLight::StaticClass());
+				if (CurrentLightCount - StartingLightCount == 1)
+				{
+					UE_LOG(LogEditorBuildPromotionTests, Display, TEXT("Duplicated the PointLight using duplicate"));
+				}
+				else
+				{
+					UE_LOG(LogEditorBuildPromotionTests, Error, TEXT("Failed to duplicate light using duplicate"));
+				}
+			}
+			else
+			{
+				SkippedTests.Add(TEXT("Lighting: Duplicate point light. (Failed to create light)"));
+			}
+
+			return true;
+		}
+
+		/**
+		* Lighting Test Stage: Build Lighting (Part 1)
+		*    Sets the lighting quality level and starts a lighting build
+		*/
+		bool Lighting_BuildLighting_Part1()
+		{
+			//Set production quality
+			GConfig->SetInt(TEXT("LightingBuildOptions"), TEXT("QualityLevel"), (int32)ELightingBuildQuality::Quality_Production, GEditorUserSettingsIni);
+			UE_LOG(LogEditorBuildPromotionTests, Display, TEXT("Set the lighting quality to Production"));
+
+			//Force AutoApplyLighting on
+			ULevelEditorMiscSettings* LevelEdSettings = GetMutableDefault<ULevelEditorMiscSettings>();
+			bDisableAutoApplyLighting = !LevelEdSettings->bAutoApplyLightingEnable;
+			LevelEdSettings->bAutoApplyLightingEnable = true;
+
+			//Build Lighting
+			EditorBuildPromotionTestUtils::BuildLighting();
+
+			return true;
+		}
+
+		/**
+		* Lighting Test Stage: Build Lighting (Part 2)
+		*    Waits for lighting to finish
+		*/
+		bool Lighting_BuildLighting_Part2()
+		{
+			if (!GUnrealEd->IsLightingBuildCurrentlyRunning())
+			{
+				if (bDisableAutoApplyLighting)
+				{
+					ULevelEditorMiscSettings* LevelEdSettings = GetMutableDefault<ULevelEditorMiscSettings>();
+					LevelEdSettings->bAutoApplyLightingEnable = false;
+				}
+				UE_LOG(LogEditorBuildPromotionTests, Display, TEXT("Built Lighting"));
+				return true;
+			}
+			return false;
+		}
+
+#define IMPORT_ASSET_WITH_FACTORY(TFactoryClass,TObjectClass,ImportSetting,ClassVariable,ExtraSettings) \
+{ \
+	const FString FilePath = AutomationTestSettings->BuildPromotionTest.ImportWorkflow.ImportSetting.ImportFilePath.FilePath; \
+	if (FilePath.Len() > 0) \
+	{ \
+		TFactoryClass* FactoryInst = ConstructObject<TFactoryClass>(TFactoryClass::StaticClass()); \
+		ExtraSettings \
+		AutomationEditorCommonUtils::ApplyCustomFactorySettings(FactoryInst, AutomationTestSettings->BuildPromotionTest.ImportWorkflow.ImportSetting.FactorySettings); \
+		ClassVariable = Cast<TObjectClass>(EditorBuildPromotionTestUtils::ImportAsset(FactoryInst, FilePath)); \
+		EditorBuildPromotionTestUtils::PlaceImportedAsset(ClassVariable, PlaceLocation); \
+	} \
+	else \
+	{ \
+		SkippedTests.Add(FString::Printf(TEXT("Importing Workflow: Importing %s. (No file path)"),TEXT(#ImportSetting))); \
+		UE_LOG(LogEditorBuildPromotionTests, Warning, TEXT("No asset import path set for %s"), TEXT(#ImportSetting)); \
+	} \
+	TagPreviousLogs(TEXT(#ImportSetting)); \
+}
+
+		/**
+		* Workflow Test Stage: Importing Workflow
+		*    Imports a set of assets from the AutomationTestSettings and adds them to the map
+		*/
+		bool Workflow_ImportWorkflow()
+		{
+			UAutomationTestSettings const* AutomationTestSettings = GetDefault<UAutomationTestSettings>();
+			check(AutomationTestSettings);
+
+			FVector PlaceLocation(940, 810, 160);
+			const float ActorSpacing = 200.f;
+
+			//Diffuse
+			IMPORT_ASSET_WITH_FACTORY(UTextureFactory, UTexture, Diffuse, DiffuseTexture, );
+
+			//Normalmap
+			IMPORT_ASSET_WITH_FACTORY(UTextureFactory, UTexture, Normal, NormalTexture, FactoryInst->LODGroup = TEXTUREGROUP_WorldNormalMap;);
+
+			//Static Mesh
+			IMPORT_ASSET_WITH_FACTORY(UFbxFactory, UStaticMesh, StaticMesh, WindowMesh, );
+
+			//Reimport Static Mesh
+			IMPORT_ASSET_WITH_FACTORY(UFbxFactory, UStaticMesh, ReimportStaticMesh, ReimportMesh, );
+
+			//Blend Shape Mesh
+			IMPORT_ASSET_WITH_FACTORY(UFbxFactory, USkeletalMesh, BlendShapeMesh, BlendShape, FactoryInst->ImportUI->MeshTypeToImport == FBXIT_SkeletalMesh;);
+
+			//Morph Mesh
+			IMPORT_ASSET_WITH_FACTORY(UFbxFactory, USkeletalMesh, MorphMesh, MorphAndMorphAnim, FactoryInst->ImportUI->MeshTypeToImport == FBXIT_SkeletalMesh;);
+
+			//Skeletal Mesh
+			IMPORT_ASSET_WITH_FACTORY(UFbxFactory, USkeletalMesh, SkeletalMesh, SkeletalMesh_Test, FactoryInst->ImportUI->MeshTypeToImport == FBXIT_SkeletalMesh;);
+
+			if (SkeletalMesh_Test)
+			{
+				//Animation
+				IMPORT_ASSET_WITH_FACTORY(UFbxFactory, UAnimSequence, Animation, AnimationTest, FactoryInst->ImportUI->MeshTypeToImport == FBXIT_Animation; FactoryInst->ImportUI->Skeleton = SkeletalMesh_Test->Skeleton;);
+			}
+			else
+			{
+				SkippedTests.Add(TEXT("Importing Workflow: Importing Animation.  (No skeletal mesh.)")); 
+			}
+
+			//Sound
+			IMPORT_ASSET_WITH_FACTORY(USoundFactory, USoundWave, Sound, FemaleVoice, );
+
+			//SurroundSound is a special case.  We need to import 6 files separately
+			const FString SurroundFilePath = AutomationTestSettings->BuildPromotionTest.ImportWorkflow.SurroundSound.ImportFilePath.FilePath;
+			if (SurroundFilePath.Len() > 0)
+			{
+				const FString BaseFileName = FPaths::GetPath(SurroundFilePath) / FPaths::GetBaseFilename(SurroundFilePath).LeftChop(3);
+
+				USoundSurroundFactory* FactoryInst = ConstructObject<USoundSurroundFactory>(USoundSurroundFactory::StaticClass());
+				AutomationEditorCommonUtils::ApplyCustomFactorySettings(FactoryInst, AutomationTestSettings->BuildPromotionTest.ImportWorkflow.SurroundSound.FactorySettings); 
+
+				const FString SurroundChannels[] = { TEXT("_fl"), TEXT("_fr"), TEXT("_fc"), TEXT("_lf"), TEXT("_sl"), TEXT("_sr"), TEXT("_bl"), TEXT("_br") };
+
+				USoundWave* ImportedSound = NULL;
+				for (int32 ChannelID = 0; ChannelID < ARRAY_COUNT(SurroundChannels); ++ChannelID)
+				{
+					const FString ChannelFileName = FString::Printf(TEXT("%s%s.WAV"),*BaseFileName,*SurroundChannels[ChannelID]);
+					if (FPaths::FileExists(ChannelFileName))
+					{
+						USoundWave* CreatedWave = Cast<USoundWave>(EditorBuildPromotionTestUtils::ImportAsset(FactoryInst, ChannelFileName));
+						if (ImportedSound == NULL)
+						{
+							ImportedSound = CreatedWave;
+						}
+					}
+				}
+
+				if (ImportedSound)
+				{
+					EditorBuildPromotionTestUtils::PlaceImportedAsset(ImportedSound, PlaceLocation);
+				}
+				else
+				{
+					UE_LOG(LogEditorBuildPromotionTests, Error, TEXT("Failed to create surround sound asset at (%s)"), *SurroundFilePath);
+				}
+				TagPreviousLogs(TEXT("SurroundSound"));
+			}
+			else
+			{
+				SkippedTests.Add(TEXT("Importing Workflow: Importing SurroundSound. (No file path)"));
+			}
+
+			//Import the others
+			const TArray<FEditorImportWorkflowDefinition>& AssetsToImport = AutomationTestSettings->BuildPromotionTest.ImportWorkflow.OtherAssetsToImport;
+			for (int32 i = 0; i < AssetsToImport.Num(); ++i)
+			{
+				//Check the file path
+				const FString FilePath = AssetsToImport[i].ImportFilePath.FilePath;
+				if (FilePath.Len() > 0)
+				{
+					//Get the import factory class to use
+					const FString FileExtension = FPaths::GetExtension(FilePath);
+					UClass* FactoryClass = AutomationEditorCommonUtils::GetFactoryClassForType(FileExtension);
+					if (FactoryClass)
+					{
+						//Create the factory and import the asset
+						UFactory* FactoryInst = ConstructObject<UFactory>(FactoryClass);
+						AutomationEditorCommonUtils::ApplyCustomFactorySettings(FactoryInst, AssetsToImport[i].FactorySettings);
+						UObject* NewObject = EditorBuildPromotionTestUtils::ImportAsset(FactoryInst, FilePath);
+						if (NewObject)
+						{
+							EditorBuildPromotionTestUtils::PlaceImportedAsset(NewObject, PlaceLocation);
+						}
+						else
+						{
+							UE_LOG(LogEditorBuildPromotionTests, Warning, TEXT("Failed to create asset (%s) with factory (%s)"), *FilePath, *FactoryClass->GetName());
+						}
+					}
+					else
+					{
+						UE_LOG(LogEditorBuildPromotionTests, Warning, TEXT("Couldn't find import factory to use on assset (%s)"), *FilePath);
+					}
+				} 
+				else 
+				{ 
+					UE_LOG(LogEditorBuildPromotionTests, Warning, TEXT("No asset import path set for OtherAssetsToImport.  Index: %i"), i);
+				}
+				TagPreviousLogs(FString::Printf(TEXT("OtherAssetsToImport #%i"),i+1));
+			}
+
+			//Remove one from the test counts to keep the section from counting
+			SectionTestCount--;
+			SectionSuccessCount--;
+
+			//Save all the new assets
+			EditorBuildPromotionTestUtils::SaveAllAssetsInFolder(EditorBuildPromotionTestUtils::GetGamePath());
+
+			return true;
+		}
+#undef IMPORT_ASSET_WITH_FACTORY
+
+		/**
+		* ContentBrowser Test Stage: Source Control (part 1)
+		*    Opens the asset editor for the source control material
+		*/
+		bool ContentBrowser_SourceControl_Part1()
+		{
+			//Find the material to check out
+			UAutomationTestSettings const* AutomationTestSettings = GetDefault<UAutomationTestSettings>();
+			check(AutomationTestSettings);
+			
+			const FString SourceControlMaterialPath = AutomationTestSettings->BuildPromotionTest.SourceControlMaterial.FilePath;
+			if (SourceControlMaterialPath.Len() > 0)
+			{
+				FAssetData MaterialData = EditorBuildPromotionTestUtils::GetAssetDataFromPackagePath(SourceControlMaterialPath);
+				SCTestMat = Cast<UMaterial>(MaterialData.GetAsset());
+
+				if (SCTestMat)
+				{
+					//Open the asset editor
+					FAssetEditorManager::Get().OpenEditorForAsset(SCTestMat);
+					UE_LOG(LogEditorBuildPromotionTests, Display, TEXT("Opened the material editor for: %s"), *SCTestMat->GetName());
+
+					FString PackageFileName = FPackageName::LongPackageNameToFilename(MaterialData.PackageName.ToString(), FPackageName::GetAssetPackageExtension());
+					FString MaterialFilePath = FPaths::ConvertRelativePathToFull(PackageFileName);
+					AsyncHelper = SourceControlAutomationCommon::FAsyncCommandHelper(MaterialFilePath);
+				}
+				else
+				{
+					UE_LOG(LogEditorBuildPromotionTests, Error, TEXT("Failed to find material to modify for Content Browser tests."));
+				}
+			}
+			else
+			{
+				SkippedTests.Add(TEXT("ContentBrowser: Source Control. (No file path)"));
+			}
+
+			return true;
+		}
+
+		/**
+		* ContentBrowser Test Stage: Source Control (part 2)
+		*    Checks the source control material out of source control
+		*/
+		bool ContentBrowser_SourceControl_Part2()
+		{			
+			if (SCTestMat)
+			{
+				if (!AsyncHelper.IsDispatched())
+				{
+					if (ISourceControlModule::Get().GetProvider().Execute(
+						ISourceControlOperation::Create<FCheckOut>(),
+						SourceControlHelpers::PackageFilename(AsyncHelper.GetParameter()),
+						EConcurrency::Asynchronous,
+						FSourceControlOperationComplete::CreateRaw(&AsyncHelper, &SourceControlAutomationCommon::FAsyncCommandHelper::SourceControlOperationComplete)
+						) != ECommandResult::Succeeded)
+					{
+						UE_LOG(LogEditorBuildPromotionTests, Warning, TEXT("Failed to check out '%s'"), *AsyncHelper.GetParameter());
+						return true;
+					}
+
+					AsyncHelper.SetDispatched();
+				}
+
+				if (AsyncHelper.IsDone())
+				{
+					// check state now we are done
+					TSharedPtr<ISourceControlState, ESPMode::ThreadSafe> SourceControlState = ISourceControlModule::Get().GetProvider().GetState(SourceControlHelpers::PackageFilename(AsyncHelper.GetParameter()), EStateCacheUsage::Use);
+					if (!SourceControlState.IsValid())
+					{
+						UE_LOG(LogEditorBuildPromotionTests, Error, TEXT("Could not retrieve state for file '%s'"), *AsyncHelper.GetParameter());
+					}
+					else
+					{
+						if (!SourceControlState->IsCheckedOut())
+						{
+							UE_LOG(LogEditorBuildPromotionTests, Error, TEXT("Unexpected state following Check Out operation for file '%s'"), *AsyncHelper.GetParameter());
+						}
+						else
+						{
+							UE_LOG(LogEditorBuildPromotionTests, Display, TEXT("Checked out the source control material"));
+						}
+					}
+				}
+
+				return AsyncHelper.IsDone();
+			}
+			return true;
+		}
+
+		/**
+		* ContentBrowser Test Stage: Source Control (part 3)
+		*    Changes the source control material's color
+		*/
+		bool ContentBrowser_SourceControl_Part3()
+		{
+			if (SCTestMat)
+			{
+				UAutomationTestSettings const* AutomationTestSettings = GetDefault<UAutomationTestSettings>();
+				check(AutomationTestSettings);
+
+				IAssetEditorInstance* AssetEditor = FAssetEditorManager::Get().FindEditorForAsset(SCTestMat, true);
+				FMaterialEditor* MaterialEditor = (FMaterialEditor*)AssetEditor;
+
+				//Pick a random color
+				const FString AvailableColors[] = { 
+					TEXT("Red"),	TEXT("(R=1.0f,G=0.0f,B=0.0f)"), 
+					TEXT("Green"),	TEXT("(R=0.0f,G=1.0f,B=0.0f)"), 
+					TEXT("Blue"),	TEXT("(R=0.0f,G=0.0f,B=1.0f)"), 
+					TEXT("Pink"),	TEXT("(R=1.0f,G=0.0f,B=1.0f)"),
+					TEXT("Yellow"),	TEXT("(R=1.0f,G=1.0f,B=0.0f)"),
+					TEXT("Black"),	TEXT("(R=0.0f,G=0.0f,B=0.0f)"),
+					TEXT("White"),	TEXT("(R=1.0f,G=1.0f,B=1.0f)") };
+
+				const int32 ChosenIndex = FMath::RandHelper(ARRAY_COUNT(AvailableColors) / 2);
+				ChosenMaterialColor = AvailableColors[ChosenIndex * 2];
+				const FString ColorValue = AvailableColors[(ChosenIndex * 2)+1];
+
+				UE_LOG(LogEditorBuildPromotionTests, Display, TEXT("Changing source control test to %s"), *ChosenMaterialColor);
+
+				//Get the editor material
+				UMaterial* EditorMaterial = Cast<UMaterial>(MaterialEditor->GetMaterialInterface());
+				for (int32 i = 0; i < EditorMaterial->Expressions.Num(); ++i)
+				{
+					UMaterialExpressionConstant3Vector* ColorParam = Cast<UMaterialExpressionConstant3Vector>(EditorMaterial->Expressions[i]);
+					if (ColorParam)
+					{
+						EditorMaterial->Modify();
+						EditorBuildPromotionTestUtils::SetPropertyByName(ColorParam, TEXT("Constant"), ColorValue);
+						MaterialEditor->UpdateMaterialAfterGraphChange();
+						MaterialEditor->ForceRefreshExpressionPreviews();
+						EditorBuildPromotionTestUtils::SendUpdateMaterialCommand();
+					}
+				}
+
+				FAssetData MaterialData = EditorBuildPromotionTestUtils::GetAssetDataFromPackagePath(AutomationTestSettings->BuildPromotionTest.SourceControlMaterial.FilePath);
+				FString PackageFileName = FPackageName::LongPackageNameToFilename(MaterialData.PackageName.ToString(), FPackageName::GetAssetPackageExtension());
+				FString MaterialFilePath = FPaths::ConvertRelativePathToFull(PackageFileName);
+				AsyncHelper = SourceControlAutomationCommon::FAsyncCommandHelper(MaterialFilePath);
+			}
+			return true;
+		}
+
+		/**
+		* ContentBrowser Test Stage: Source Control (part 4)
+		*    Checks the source control material back in and sets the description to the new color
+		*/
+		bool ContentBrowser_SourceControl_Part4()
+		{
+			if (SCTestMat)
+			{
+				if (!AsyncHelper.IsDispatched())
+				{
+					TSharedRef<FCheckIn, ESPMode::ThreadSafe> CheckInOperation = ISourceControlOperation::Create<FCheckIn>();
+					FString CheckinDescription = FString::Printf(TEXT("[AUTOMATED TEST] Changed Material Color to %s"), *ChosenMaterialColor);
+					CheckInOperation->SetDescription(FText::FromString(CheckinDescription));
+
+					if (ISourceControlModule::Get().GetProvider().Execute(
+						CheckInOperation,
+						SourceControlHelpers::PackageFilename(AsyncHelper.GetParameter()),
+						EConcurrency::Asynchronous,
+						FSourceControlOperationComplete::CreateRaw(&AsyncHelper, &SourceControlAutomationCommon::FAsyncCommandHelper::SourceControlOperationComplete)
+						) != ECommandResult::Succeeded)
+					{
+						UE_LOG(LogEditorBuildPromotionTests, Warning, TEXT("Failed to check in '%s'"), *AsyncHelper.GetParameter());
+						return true;
+					}
+
+					AsyncHelper.SetDispatched();
+				}
+
+				if (AsyncHelper.IsDone())
+				{
+					// check state now we are done
+					TSharedPtr<ISourceControlState, ESPMode::ThreadSafe> SourceControlState = ISourceControlModule::Get().GetProvider().GetState(SourceControlHelpers::PackageFilename(AsyncHelper.GetParameter()), EStateCacheUsage::Use);
+					if (!SourceControlState.IsValid())
+					{
+						UE_LOG(LogEditorBuildPromotionTests, Error, TEXT("Could not retrieve state for file '%s'"), *AsyncHelper.GetParameter());
+					}
+					else
+					{
+						if (!SourceControlState->IsSourceControlled() || !SourceControlState->CanCheckout())
+						{
+							UE_LOG(LogEditorBuildPromotionTests, Error, TEXT("Unexpected state following Check In operation for file '%s'"), *AsyncHelper.GetParameter());
+						}
+						else
+						{
+							UE_LOG(LogEditorBuildPromotionTests, Display, TEXT("Checked in the source control material"));
+						}
+					}
+					return true;
+				}
+				return false;
+			}
+			return true;
+		}
+
+		/**
+		* ContentBrowser Test Stage: Source Control (part 5)
+		*    Closes the material editor
+		*/
+		bool ContentBrowser_SourceControl_Part5()
+		{
+			UE_LOG(LogEditorBuildPromotionTests, Display, TEXT("Closed the material editor"));
+			FAssetEditorManager::Get().CloseAllAssetEditors();
+			return true;
+		}
+
+		/**
+		* ContentBrowser Test Stage:  Open Assets (Part 1)
+		*   Queues up assets to be open, modified, and placed into the level
+		*/
+		bool ContentBrowser_OpenAssets_Part1()
+		{
+			UAutomationTestSettings const* AutomationTestSettings = GetDefault<UAutomationTestSettings>();
+			check(AutomationTestSettings);
+
+			TArray<FOpenAssetInfo> OpenInfo;
+			UObject* Asset = NULL;
+			FAssetData AssetData;
+			FString AssetPackagePath;
+
+			// Blueprint
+			AssetPackagePath = AutomationTestSettings->BuildPromotionTest.OpenAssets.BlueprintAsset.FilePath;
+			if (AssetPackagePath.Len() > 0)
+			{
+				AssetData = EditorBuildPromotionTestUtils::GetAssetDataFromPackagePath(AssetPackagePath);
+				Asset = AssetData.GetAsset();
+				OpenInfo.Add(FOpenAssetInfo(Asset, AssetData, TEXT("BlueprintDescription"), TEXT("Modified by BuildPromotionTest TM")));
+			}
+			else
+			{
+				SkippedTests.Add(TEXT("ContentBrowser: Open Blueprint. (No file path)"));
+				UE_LOG(LogEditorBuildPromotionTests, Warning, TEXT("Skipping Asset: BlueprintAsset setting is empty"));
+			}
+			
+			// Material
+			AssetPackagePath = AutomationTestSettings->BuildPromotionTest.OpenAssets.MaterialAsset.FilePath;
+			if (AssetPackagePath.Len() > 0)
+			{
+				AssetData = EditorBuildPromotionTestUtils::GetAssetDataFromPackagePath(AssetPackagePath);
+				Asset = AssetData.GetAsset();
+				OpenInfo.Add(FOpenAssetInfo(Asset, AssetData, TEXT("TwoSided"), TEXT("true")));
+			}
+			else
+			{
+				SkippedTests.Add(TEXT("ContentBrowser: Open Material. (No file path)"));
+				UE_LOG(LogEditorBuildPromotionTests, Warning, TEXT("Skipping Asset: MaterialAsset setting is empty"));
+			}
+			
+			// Particle System
+			AssetPackagePath = AutomationTestSettings->BuildPromotionTest.OpenAssets.ParticleSystemAsset.FilePath;
+			if (AssetPackagePath.Len() > 0)
+			{
+				AssetData = EditorBuildPromotionTestUtils::GetAssetDataFromPackagePath(AssetPackagePath);
+				Asset = AssetData.GetAsset();
+				OpenInfo.Add(FOpenAssetInfo(Asset, AssetData, TEXT("UpdateTime_FPS"), TEXT("100")));
+			}
+			else
+			{
+				SkippedTests.Add(TEXT("ContentBrowser: Open ParticleSystem. (No file path)"));
+				UE_LOG(LogEditorBuildPromotionTests, Warning, TEXT("Skipping Asset: ParticleSystemAsset setting is empty"));
+			}
+
+			// Skeletal Mesh
+			AssetPackagePath = AutomationTestSettings->BuildPromotionTest.OpenAssets.SkeletalMeshAsset.FilePath;
+			if (AssetPackagePath.Len() > 0)
+			{
+				AssetData = EditorBuildPromotionTestUtils::GetAssetDataFromPackagePath(AssetPackagePath);
+				Asset = AssetData.GetAsset();
+				OpenInfo.Add(FOpenAssetInfo(Asset, AssetData, TEXT("bUseFullPrecisionUVs"), TEXT("1")));
+			}
+			else
+			{
+				SkippedTests.Add(TEXT("ContentBrowser: Open SkeletalMesh. (No file path)"));
+				UE_LOG(LogEditorBuildPromotionTests, Warning, TEXT("Skipping Asset: SkeletalMeshAsset setting is empty"));
+			}
+			
+			// Static Mesh
+			AssetPackagePath = AutomationTestSettings->BuildPromotionTest.OpenAssets.StaticMeshAsset.FilePath;
+			if (AssetPackagePath.Len() > 0)
+			{
+				AssetData = EditorBuildPromotionTestUtils::GetAssetDataFromPackagePath(AssetPackagePath);
+				Asset = AssetData.GetAsset();
+				OpenInfo.Add(FOpenAssetInfo(Asset, AssetData, TEXT("AutoLODPixelError"), TEXT("42.f")));
+			}
+			else
+			{
+				SkippedTests.Add(TEXT("ContentBrowser: Open StaticMesh. (No file path)"));
+				UE_LOG(LogEditorBuildPromotionTests, Warning, TEXT("Skipping Asset: StaticMeshAsset setting is empty"));
+			}
+			
+			// Texture
+			AssetPackagePath = AutomationTestSettings->BuildPromotionTest.OpenAssets.TextureAsset.FilePath;
+			if (AssetPackagePath.Len() > 0)
+			{
+				AssetData = EditorBuildPromotionTestUtils::GetAssetDataFromPackagePath(AssetPackagePath);
+				Asset = AssetData.GetAsset();
+				OpenInfo.Add(FOpenAssetInfo(Asset, AssetData, TEXT("LODBias"), TEXT("2")));
+			}
+			else
+			{
+				SkippedTests.Add(TEXT("ContentBrowser: Open Texture. (No file path)"));
+				UE_LOG(LogEditorBuildPromotionTests, Warning, TEXT("Skipping Asset: TextureAsset setting is empty"));
+			}
+			
+			OpenAssetHelper = new FOpenAssetHelper(OpenInfo, TestExecutionInfo);
+
+			return true;
+		}
+
+		/**
+		* ContentBrowser Test Stage: Open Assets (Part 2)
+		*    Waits for the OpenAssetHeler to finish 
+		*/
+		bool ContentBrowser_OpenAssets_Part2()
+		{
+			if (OpenAssetHelper)
+			{
+				if (OpenAssetHelper->Update())
+				{
+					delete OpenAssetHelper;
+					OpenAssetHelper = NULL;
+					return true;
+				}
+				return false;
+			}
+			else
+			{
+				return true;
+			}
+		}
+
+		/**
+		* ContentBrowser Test Stage: Reimport Asset
+		*    Reimports the static mesh
+		*/
+		bool ContentBrowser_ReimportAsset()
+		{
+			if (ReimportMesh)
+			{
+				if (FReimportManager::Instance()->Reimport(ReimportMesh, /*bAskForNewFileIfMissing=*/false))
+				{
+					UE_LOG(LogEditorBuildPromotionTests, Display, TEXT("Reimported asset %s"), *ReimportMesh->GetName());
+				}
+				else
+				{
+					UE_LOG(LogEditorBuildPromotionTests, Error, TEXT("Failed to reimport asset %s"), *ReimportMesh->GetName());
+				}
+			}
+			else
+			{
+				SkippedTests.Add(TEXT("ContentBrowser: Reimport Asset.  (No Reimport mesh)"));
+			}
+
+			return true;
+		}
+
+		/**
+		* ContentBrowser Test Stage: Creating a material (Part 1)
+		*    Creates a material from the diffuse provided in the AutomationTestSettings and opens the material editor
+		*/
+		bool ContentBrowser_CreateAMaterial_Part1()
+		{
+			if (DiffuseTexture)
+			{
+				CreatedMaterial = EditorBuildPromotionTestUtils::CreateMaterialFromTexture(DiffuseTexture);
+				if (CreatedMaterial)
+				{
+					UE_LOG(LogEditorBuildPromotionTests, Display, TEXT("Created new material (%s) from diffuse (%s)"),*CreatedMaterial->GetName(), *DiffuseTexture->GetName());
+					FAssetEditorManager::Get().OpenEditorForAsset(CreatedMaterial);
+					UE_LOG(LogEditorBuildPromotionTests, Display, TEXT("Opened the material editor"));
+				}
+				else
+				{
+					UE_LOG(LogEditorBuildPromotionTests, Error, TEXT("Failed to create material fom diffuse"));
+				}
+			}
+			else
+			{
+				SkippedTests.Add(TEXT("ContentBrowser: Creating a material. (No diffuse)"));
+				UE_LOG(LogEditorBuildPromotionTests, Warning, TEXT("Skipping Create Material: DiffuseTexture is missing"));
+			}
+
+			return true;
+		}
+
+		/**
+		* ContentBrowser Test Stage: Creating a material (Part 2)
+		*    Adds the normalmap texture to the material and updates the shader
+		*/
+		bool ContentBrowser_CreateAMaterial_Part2()
+		{
+			if (CreatedMaterial)
+			{
+				if (NormalTexture)
+				{
+					IAssetEditorInstance* AssetEditor = FAssetEditorManager::Get().FindEditorForAsset(CreatedMaterial, true);
+					FMaterialEditor* MaterialEditor = (FMaterialEditor*)AssetEditor;
+
+					if (EditorBuildPromotionTestUtils::AssignNormalToMaterial(NormalTexture, CreatedMaterial))
+					{
+						FAssetEditorManager::Get().FindEditorForAsset(CreatedMaterial, true);
+						EditorBuildPromotionTestUtils::SendUpdateMaterialCommand();
+						UE_LOG(LogEditorBuildPromotionTests, Display, TEXT("Compiled the new material"));
+					}
+				}
+				else
+				{
+					SkippedTests.Add(TEXT("ContentBrowser: Creating a material. (No normalmap)"));
+					UE_LOG(LogEditorBuildPromotionTests, Warning, TEXT("No normal texture to apply to new material"));
+				}
+			}
+
+			return true;
+		}
+
+		/**
+		* ContentBrowser Test Stage: Creating a material (Part 3)
+		*    Closes the material editor and adds the material to a default object in the map
+		*/
+		bool ContentBrowser_CreateAMaterial_Part3()
+		{
+			FAssetEditorManager::Get().CloseAllAssetEditors();
+
+			const FVector PlaceLocation(0, 2240, 166);
+			AActor* PlacedActor = EditorBuildPromotionTestUtils::AddDefaultMeshToLevelWithMaterial(PlaceLocation, CreatedMaterial);
+			if ( PlacedActor )
+			{
+				UE_LOG(LogEditorBuildPromotionTests, Display, TEXT("Placed new material in the level"));
+			}
+
+			return true;
+		}
+
+		/**
+		* ContentBrowser Test Stage: Creating a particle system (Part 1)
+		*    Creates a new particle system and opens the cascade editor
+		*/
+		bool ContentBrowser_CreateAParticleSystem_Part1()
+		{
+			//Create a Particle system
+			UParticleSystemFactoryNew* PSFactory = ConstructObject<UParticleSystemFactoryNew>(UParticleSystemFactoryNew::StaticClass());
+			const FString& PSName(TEXT("BP_ParticleSystem"));
+			CreatedPS = Cast<UParticleSystem>(EditorBuildPromotionTestUtils::CreateAsset(PSFactory, UParticleSystem::StaticClass(), PSName));
+			if (CreatedPS)
+			{
+				FAssetEditorManager::Get().OpenEditorForAsset(CreatedPS);
+				UE_LOG(LogEditorBuildPromotionTests, Display, TEXT("Opened the cascade editor"));
+			}
+			else
+			{
+				UE_LOG(LogEditorBuildPromotionTests, Error, TEXT("Failed to create a new ParticleSystem"));
+			}
+
+			return true;
+		}
+
+		/**
+		* ContentBrowser Test Stage: Creating a particle system (Part 2)
+		*    Finds and modifies the StartSize of the particle system, saves the asset, and then closes the editor
+		*/
+		bool ContentBrowser_CreateAParticleSystem_Part2()
+		{
+			if (CreatedPS)
+			{
+				IAssetEditorInstance* AssetEditor = FAssetEditorManager::Get().FindEditorForAsset(CreatedPS, true);
+
+				bool bModifiedSize = false;
+				UParticleLODLevel* DefaultLOD = CreatedPS->Emitters[0]->LODLevels[0];
+				for (int32 i = 0; i < DefaultLOD->Modules.Num(); ++i)
+				{
+					UParticleModuleSize* SizeModule = Cast<UParticleModuleSize>(DefaultLOD->Modules[i]);
+					if (SizeModule)
+					{
+						UDistributionVectorUniform* Distribution = Cast<UDistributionVectorUniform>(SizeModule->StartSize.Distribution);
+						EditorBuildPromotionTestUtils::SetPropertyByName(Distribution, TEXT("Max"), TEXT("(X=100,Y=100,Z=100)"));
+						EditorBuildPromotionTestUtils::SetPropertyByName(Distribution, TEXT("Min"), TEXT("(X=100,Y=100,Z=100)"));
+						bModifiedSize = true;
+					}
+				}
+
+				if (bModifiedSize)
+				{
+					UE_LOG(LogEditorBuildPromotionTests, Display, TEXT("Modified ParticleSystem StartSize (Min and Max)"));
+					EditorBuildPromotionTestUtils::SendSaveCascadeCommand();
+					UE_LOG(LogEditorBuildPromotionTests, Display, TEXT("Saved the particle system"));
+				}
+				else
+				{
+					UE_LOG(LogEditorBuildPromotionTests, Error, TEXT("Failed to modify ParticleSystem StartSize"));
+				}
+
+				FAssetEditorManager::Get().CloseAllAssetEditors();
+				UE_LOG(LogEditorBuildPromotionTests, Display, TEXT("Closed the cascade editor"));
+			}
+			return true;
+		}
+
+		/**
+		* Saves the blueprint stored in BlueprintObject
+		*/
+		void SaveBlueprint()
+		{
+			if (BlueprintObject && BlueprintPackage)
+			{
+				BlueprintPackage->SetDirtyFlag(true);
+				BlueprintPackage->FullyLoad();
+				const FString PackagePath = EditorBuildPromotionTestUtils::GetGamePath() + TEXT("/") + EditorBuildPromotionTestUtils::BlueprintNameString;
+				if (UPackage::SavePackage(BlueprintPackage, NULL, RF_Standalone, *FPackageName::LongPackageNameToFilename(PackagePath, FPackageName::GetAssetPackageExtension()), GLog, nullptr, false, true, SAVE_None))
+				{
+					UE_LOG(LogEditorBuildPromotionTests, Display, TEXT("Saved blueprint (%s)"), *BlueprintObject->GetName());
+				}
+				else
+				{
+					UE_LOG(LogEditorBuildPromotionTests, Error, TEXT("Unable to save blueprint (%s)"), *BlueprintObject->GetName());
+				}
+			}
+		}
+
+		/**
+		* Blueprint Test Stage: Create a new Blueprint (Part 1)
+		*    Creates a new actor based blueprint and opens the editor
+		*/
+		bool Blueprint_CreateNewBlueprint_Part1()
+		{
+			//Make sure we have the required assets
+			UAutomationTestSettings const* AutomationTestSettings = GetDefault<UAutomationTestSettings>();
+			check(AutomationTestSettings);
+			FAssetData AssetData;
+
+			const FString FirstMeshPath = AutomationTestSettings->BuildPromotionTest.BlueprintSettings.FirstMeshPath.FilePath;
+			if (FirstMeshPath.Len() > 0)
+			{
+				AssetData = EditorBuildPromotionTestUtils::GetAssetDataFromPackagePath(FirstMeshPath);
+				FirstBlueprintMesh = Cast<UStaticMesh>(AssetData.GetAsset());
+			}
+
+			const FString SecondMeshPath = AutomationTestSettings->BuildPromotionTest.BlueprintSettings.SecondMeshPath.FilePath;
+			if (SecondMeshPath.Len() > 0)
+			{
+				AssetData = EditorBuildPromotionTestUtils::GetAssetDataFromPackagePath(SecondMeshPath);
+				SecondBlueprintMesh = Cast<UStaticMesh>(AssetData.GetAsset());
+			}
+
+			if (FirstBlueprintMesh && SecondBlueprintMesh && CreatedPS)
+			{
+				UBlueprintFactory* Factory = ConstructObject<UBlueprintFactory>(UBlueprintFactory::StaticClass());
+				Factory->ParentClass = AActor::StaticClass();
+
+				const FString PackageName = EditorBuildPromotionTestUtils::GetGamePath() + TEXT("/") + EditorBuildPromotionTestUtils::BlueprintNameString;
+				BlueprintPackage = CreatePackage(NULL, *PackageName);
+				EObjectFlags Flags = RF_Public | RF_Standalone;
+
+				UObject* ExistingBlueprint = FindObject<UBlueprint>(BlueprintPackage, *EditorBuildPromotionTestUtils::BlueprintNameString);
+				if (ExistingBlueprint)
+				{
+					UE_LOG(LogEditorBuildPromotionTests, Error, TEXT("Found existing blueprint.  Please use the cleanup test between runs."));
+					return true;
+				}
+
+				BlueprintObject = Cast<UBlueprint>(Factory->FactoryCreateNew(UBlueprint::StaticClass(), BlueprintPackage, FName(*EditorBuildPromotionTestUtils::BlueprintNameString), Flags, NULL, GWarn));
+
+				if (BlueprintObject)
+				{
+					UE_LOG(LogEditorBuildPromotionTests, Display, TEXT("Created a new Actor based blueprint (%s)"), *BlueprintObject->GetName());
+
+					// Notify the asset registry
+					FAssetRegistryModule::AssetCreated(BlueprintObject);
+
+					// Mark the package dirty...
+					BlueprintPackage->MarkPackageDirty();
+
+					UE_LOG(LogEditorBuildPromotionTests, Display, TEXT("Opening the blueprint editor for the first time"));
+					FAssetEditorManager::Get().OpenEditorForAsset(BlueprintObject);
+				}
+				else
+				{
+					UE_LOG(LogEditorBuildPromotionTests, Error, TEXT("Failed to create blueprint!"));
+				}
+			}
+			else
+			{
+				SkippedTests.Add(TEXT("All Blueprint tests. (Missing a required mesh)"));
+				UE_LOG(LogEditorBuildPromotionTests, Warning, TEXT("SKIPPING BLUEPRINT TESTS.  Invalid or missing FirstMeshPath or SecondMeshPath in AutomationTestSettings."));
+			}
+			
+			return true;
+		}
+
+		/**
+		* Blueprint Test Stage: Create a new Blueprint (Part 2)
+		*    Checks that the blueprint editor opened in the correct mode
+		*/
+		bool Blueprint_CreateNewBlueprint_Part2()
+		{
+			if (BlueprintObject)
+			{
+				IAssetEditorInstance* AssetEditor = FAssetEditorManager::Get().FindEditorForAsset(BlueprintObject, true);
+				IBlueprintEditor* BlueprintEditor = (IBlueprintEditor*)AssetEditor;
+				if (BlueprintEditor->GetCurrentMode() != FBlueprintEditorApplicationModes::BlueprintDefaultsMode)
+				{
+					UE_LOG(LogEditorBuildPromotionTests, Display, TEXT("Opened the correct initial editor"));
+				}
+				else
+				{
+					UE_LOG(LogEditorBuildPromotionTests, Error, TEXT("Didn't open the correct editor"));
+				}
+			}
+			
+			return true;
+		}
+
+		/**
+		* Blueprint Test Stage: Data-only Blueprint (Part 1)
+		*    Closes the blueprint editor
+		*/
+		bool Blueprint_DataOnlyBlueprint_Part1()
+		{
+			if (BlueprintObject)
+			{
+				UE_LOG(LogEditorBuildPromotionTests, Display, TEXT("Closing the blueprint editor"));
+				FAssetEditorManager::Get().CloseAllAssetEditors();
+			}
+			return true;
+		}
+
+		/**
+		* Blueprint Test Stage: Data-only Blueprint (Part 2)
+		*    Re opens the blueprint editor
+		*/
+		bool Blueprint_DataOnlyBlueprint_Part2()
+		{
+			if (BlueprintObject)
+			{
+				UE_LOG(LogEditorBuildPromotionTests, Display, TEXT("Opening the blueprint editor for the second time"));
+				FAssetEditorManager::Get().OpenEditorForAsset(BlueprintObject);
+			}
+			return true;
+		}
+
+		/**
+		* Blueprint Test Stage: Data-only Blueprint (Part 3)
+		*    Checks that the editor opened in the Defaults mode
+		*/
+		bool Blueprint_DataOnlyBlueprint_Part3()
+		{
+			if (BlueprintObject)
+			{
+				IAssetEditorInstance* AssetEditor = FAssetEditorManager::Get().FindEditorForAsset(BlueprintObject, true);
+				IBlueprintEditor* BlueprintEditor = (IBlueprintEditor*)AssetEditor;
+				if (BlueprintEditor->GetCurrentMode() == FBlueprintEditorApplicationModes::BlueprintDefaultsMode)
+				{
+					//Good.  Now switch to components mode
+					UE_LOG(LogEditorBuildPromotionTests, Display, TEXT("Opened the correct defaults editor"));
+					UE_LOG(LogEditorBuildPromotionTests, Display, TEXT("Switching to components mode"));
+					BlueprintEditor->SetCurrentMode(FBlueprintEditorApplicationModes::BlueprintComponentsMode);
+				}
+				else
+				{
+					UE_LOG(LogEditorBuildPromotionTests, Error, TEXT("Didn't open the correct editor the second time"));
+				}
+			}
+			return true;
+		}
+
+		/**
+		* Blueprint Test Stage:  Components Mode (Part 1)
+		*   Adds 3 new components to the blueprint, changes the root component, renames the components, and compiles the blueprint
+		*/
+		bool Blueprint_ComponentsMode_Part1()
+		{
+			if (BlueprintObject)
+			{
+				IAssetEditorInstance* AssetEditor = FAssetEditorManager::Get().FindEditorForAsset(BlueprintObject, true);
+				IBlueprintEditor* BlueprintEditor = (IBlueprintEditor*)AssetEditor;
+
+				MeshNode = EditorBuildPromotionTestUtils::CreateBlueprintComponent(BlueprintObject, FirstBlueprintMesh);
+				if (MeshNode)
+				{
+					UE_LOG(LogEditorBuildPromotionTests, Display, TEXT("Added a mesh component"));
+				}
+				else
+				{
+					UE_LOG(LogEditorBuildPromotionTests, Error, TEXT("Failed to create the first mesh component"));
+				}
+
+				OtherMeshNode = EditorBuildPromotionTestUtils::CreateBlueprintComponent(BlueprintObject, SecondBlueprintMesh);
+				if (OtherMeshNode)
+				{
+					UE_LOG(LogEditorBuildPromotionTests, Display, TEXT("Added a second mesh component"));
+				}
+				else
+				{
+					UE_LOG(LogEditorBuildPromotionTests, Error, TEXT("Failed to create the second mesh component"));
+				}
+
+				PSNode = EditorBuildPromotionTestUtils::CreateBlueprintComponent(BlueprintObject, CreatedPS);
+				if (PSNode)
+				{
+					UE_LOG(LogEditorBuildPromotionTests, Display, TEXT("Added a particle system component"));
+				}
+				else
+				{
+					UE_LOG(LogEditorBuildPromotionTests, Error, TEXT("Failed to create the particle system component"));
+				}
+
+				//Set the Particle System as the root
+				EditorBuildPromotionTestUtils::SetComponentAsRoot(BlueprintObject, PSNode);
+				UE_LOG(LogEditorBuildPromotionTests, Display, TEXT("Set the particle system component as the new root"));
+
+				//Set Names
+				const FName MeshName(TEXT("FirstMesh"));
+				FBlueprintEditorUtils::RenameComponentMemberVariable(BlueprintObject, MeshNode, MeshName);
+				UE_LOG(LogEditorBuildPromotionTests, Display, TEXT("Renamed the first mesh component to FirstMesh"));
+
+				const FName OtherMeshName(TEXT("SecondMesh"));
+				FBlueprintEditorUtils::RenameComponentMemberVariable(BlueprintObject, OtherMeshNode, OtherMeshName);
+				UE_LOG(LogEditorBuildPromotionTests, Display, TEXT("Renamed the second mesh component to SecondMesh"));
+
+				const FName PSName(TEXT("ParticleSys"));
+				FBlueprintEditorUtils::RenameComponentMemberVariable(BlueprintObject, PSNode, PSName);
+				UE_LOG(LogEditorBuildPromotionTests, Display, TEXT("Renamed the particle system component to ParticleSys"));
+
+				EditorBuildPromotionTestUtils::CompileBlueprint(BlueprintObject);
+
+				UE_LOG(LogEditorBuildPromotionTests, Display, TEXT("Switched to graph editing mode"));
+				BlueprintEditor->SetCurrentMode(FBlueprintEditorApplicationModes::StandardBlueprintEditorMode);
+			}
+			return true;
+		}
+
+		/**
+		* Blueprint Test Stage: Components Mode (Part 2)
+		*    Removes the 3 components added before and compiles the blueprint
+		*/
+		bool Blueprint_ComponentsMode_Part2()
+		{
+			if (BlueprintObject)
+			{
+				IAssetEditorInstance* AssetEditor = FAssetEditorManager::Get().FindEditorForAsset(BlueprintObject, true);
+				IBlueprintEditor* BlueprintEditor = (IBlueprintEditor*)AssetEditor;
+
+				EditorBuildPromotionTestUtils::TakeScreenshot(TEXT("BlueprintComponentVariables"), true);
+
+				UE_LOG(LogEditorBuildPromotionTests, Display, TEXT("Switched to components mode"));
+				BlueprintEditor->SetCurrentMode(FBlueprintEditorApplicationModes::BlueprintComponentsMode);
+
+				EditorBuildPromotionTestUtils::RemoveBlueprintComponent(BlueprintObject, MeshNode);
+				EditorBuildPromotionTestUtils::RemoveBlueprintComponent(BlueprintObject, OtherMeshNode);
+				EditorBuildPromotionTestUtils::RemoveBlueprintComponent(BlueprintObject, PSNode);
+
+				//There should only be the scene component left
+				if (BlueprintObject->SimpleConstructionScript->GetAllNodes().Num() == 1)
+				{
+					UE_LOG(LogEditorBuildPromotionTests, Display, TEXT("Successfully removed the blueprint components"));
+				}
+				else
+				{
+					UE_LOG(LogEditorBuildPromotionTests, Error, TEXT("Failed to remove the blueprint components"));
+				}
+
+				MeshNode = NULL;
+				OtherMeshNode = NULL;
+				PSNode = NULL;
+
+				UE_LOG(LogEditorBuildPromotionTests, Display, TEXT("Switched to graph mode"));
+				BlueprintEditor->SetCurrentMode(FBlueprintEditorApplicationModes::StandardBlueprintEditorMode);
+
+				EditorBuildPromotionTestUtils::CompileBlueprint(BlueprintObject);
+			}
+			return true;
+		}
+
+		/**
+		* Blueprint Test Stage: Construction Script
+		*    Adds an AddStaticMeshComponent to the construction graph and links it to the entry node
+		*/
+		bool Blueprint_ConstructionScript()
+		{
+			if (BlueprintObject)
+			{
+				const UEdGraphSchema_K2* K2Schema = GetDefault<UEdGraphSchema_K2>();
+
+				IAssetEditorInstance* AssetEditor = FAssetEditorManager::Get().FindEditorForAsset(BlueprintObject, true);
+				FBlueprintEditor* BlueprintEditor = (FBlueprintEditor*)AssetEditor;
+
+				UEdGraph* ConstructionGraph = FBlueprintEditorUtils::FindUserConstructionScript(BlueprintObject);
+				BlueprintEditor->OpenGraphAndBringToFront(ConstructionGraph);
+
+				AddMeshNode = EditorBuildPromotionTestUtils::CreateAddComponentActionNode(BlueprintObject, ConstructionGraph, FirstBlueprintMesh);
+				if (AddMeshNode)
+				{
+					UE_LOG(LogEditorBuildPromotionTests, Display, TEXT("Created an Add Static Mesh Component node"));
+				}
+				else
+				{
+					UE_LOG(LogEditorBuildPromotionTests, Error, TEXT("Failed to create an Add Static Mesh Component node"));
+				}
+
+				GEditor->UndoTransaction();
+				if (ConstructionGraph->Nodes.Num() == 0 || ConstructionGraph->Nodes[ConstructionGraph->Nodes.Num() - 1] != AddMeshNode)
+				{
+					UE_LOG(LogEditorBuildPromotionTests, Display, TEXT("Undo add component node successful"));
+				}
+				else
+				{
+					UE_LOG(LogEditorBuildPromotionTests, Error, TEXT("Failed to undo add component node"));
+				}
+
+				GEditor->RedoTransaction();
+				if (ConstructionGraph->Nodes.Num() > 0 && ConstructionGraph->Nodes[ConstructionGraph->Nodes.Num() - 1] == AddMeshNode)
+				{
+					UE_LOG(LogEditorBuildPromotionTests, Display, TEXT("Redo add component node successful"));
+				}
+				else
+				{
+					UE_LOG(LogEditorBuildPromotionTests, Error, TEXT("Failed to redo add component node"));
+				}
+
+
+				TArray<UK2Node_FunctionEntry*> EntryNodes;
+				ConstructionGraph->GetNodesOfClass(EntryNodes);
+				UEdGraphNode* EntryNode = EntryNodes.Num() > 0 ? EntryNodes[0] : NULL;
+				if (EntryNode)
+				{
+					EditorBuildPromotionTestUtils::ConnectGraphNodes(AddMeshNode, K2Schema->PN_Execute, EntryNode, K2Schema->PN_Then);
+					UE_LOG(LogEditorBuildPromotionTests, Display, TEXT("Connected the entry node to the Add Static Mesh node"));
+
+					UEdGraphPin* EntryOutPin = EntryNode->FindPin(K2Schema->PN_Then);
+					UEdGraphPin* AddStaticMeshInPin = AddMeshNode->FindPin(K2Schema->PN_Execute);
+
+					GEditor->UndoTransaction();
+					if (EntryOutPin->LinkedTo.Num() == 0)
+					{
+						UE_LOG(LogEditorBuildPromotionTests, Display, TEXT("Undo connection successful"));
+					}
+					else
+					{
+						UE_LOG(LogEditorBuildPromotionTests, Error, TEXT("Failed to undo connection to Add Static Mesh node"));
+					}
+
+
+					GEditor->RedoTransaction();
+					if (EntryOutPin->LinkedTo.Contains(AddStaticMeshInPin))
+					{
+						UE_LOG(LogEditorBuildPromotionTests, Display, TEXT("Redo connection successful"));
+					}
+					else
+					{
+						UE_LOG(LogEditorBuildPromotionTests, Error, TEXT("Failed to redo connection to Add Static Mesh node"));
+					}
+				}
+				else
+				{
+					UE_LOG(LogEditorBuildPromotionTests, Error, TEXT("Failed to connect the entry node to the Add Static Mesh node"));
+				}
+
+				EditorBuildPromotionTestUtils::CompileBlueprint(BlueprintObject);
+			}
+			return true;
+		}
+
+		/**
+		* Blueprint Test Stage: Variable from component mode (Part 1)
+		*    Promotes the return pin of the AddStaticMeshNode to a variable and then renames it
+		*/
+		bool Blueprint_PromoteVariable_Part1()
+		{
+			if (BlueprintObject)
+			{
+				IAssetEditorInstance* AssetEditor = FAssetEditorManager::Get().FindEditorForAsset(BlueprintObject, true);
+				FBlueprintEditor* BlueprintEditor = (FBlueprintEditor*)AssetEditor;
+
+				const UEdGraphSchema_K2* K2Schema = GetDefault<UEdGraphSchema_K2>();
+				EditorBuildPromotionTestUtils::PromotePinToVariable(BlueprintObject, AddMeshNode, K2Schema->PN_ReturnValue);
+
+				UE_LOG(LogEditorBuildPromotionTests, Display, TEXT("Promoted the return pin on the add mesh node to a variable"));
+
+				const FName OldVarName(TEXT("NewVar")); //<- Default variable name
+				const FName NewVarName(TEXT("MyMesh"));
+				FBlueprintEditorUtils::RenameMemberVariable(BlueprintObject, OldVarName, NewVarName);
+				UE_LOG(LogEditorBuildPromotionTests, Display, TEXT("Renamed the new variable to MyMesh"));
+
+				EditorBuildPromotionTestUtils::CompileBlueprint(BlueprintObject);
+
+				UE_LOG(LogEditorBuildPromotionTests, Display, TEXT("Switched to graph mode"));
+				BlueprintEditor->SetCurrentMode(FBlueprintEditorApplicationModes::StandardBlueprintEditorMode);
+			}
+			return true;
+		}
+
+		bool Blueprint_PromoteVariable_Part2()
+		{
+			if (BlueprintObject)
+			{
+				IAssetEditorInstance* AssetEditor = FAssetEditorManager::Get().FindEditorForAsset(BlueprintObject, true);
+				FBlueprintEditor* BlueprintEditor = (FBlueprintEditor*)AssetEditor;
+
+				EditorBuildPromotionTestUtils::TakeScreenshot(TEXT("BlueprintMeshVariable"), true);
+
+				UE_LOG(LogEditorBuildPromotionTests, Display, TEXT("Switched to components mode"));
+				BlueprintEditor->SetCurrentMode(FBlueprintEditorApplicationModes::BlueprintComponentsMode);
+
+				EditorBuildPromotionTestUtils::SendBlueprintResetViewCommand();
+			}
+			return true;
+		}
+
+		/**
+		* Blueprint Test Stage: Variable from component mode (Part 3)
+		*    Takes a screenshot of the mesh variable
+		*/
+		bool Blueprint_PromoteVariable_Part3()
+		{
+			if (BlueprintObject)
+			{
+				EditorBuildPromotionTestUtils::TakeScreenshot(TEXT("BlueprintComponent"), true);
+			}
+			return true;
+		}
+
+		/**
+		* Blueprint Test Stage: Event Graph
+		*    Adds a ReceiveBeginPlay and Delay node to the event graph 
+		*/
+		bool Blueprint_EventGraph()
+		{
+			if (BlueprintObject)
+			{
+				IAssetEditorInstance* AssetEditor = FAssetEditorManager::Get().FindEditorForAsset(BlueprintObject, true);
+				FBlueprintEditor* BlueprintEditor = (FBlueprintEditor*)AssetEditor;
+
+				UEdGraph* EventGraph = FBlueprintEditorUtils::FindEventGraph(BlueprintObject);
+				BlueprintEditor->OpenGraphAndBringToFront(EventGraph);
+				UE_LOG(LogEditorBuildPromotionTests, Display, TEXT("Opened the event graph"));
+
+				PostBeginPlayEventNode = EditorBuildPromotionTestUtils::CreatePostBeginPlayEvent(BlueprintObject, EventGraph);
+				if (PostBeginPlayEventNode)
+				{
+					UE_LOG(LogEditorBuildPromotionTests, Display, TEXT("Created a ReceiveBeginPlay node"));
+				}
+				else
+				{
+					UE_LOG(LogEditorBuildPromotionTests, Error, TEXT("Failed to create a ReceiveBeginPlay node"));
+				}
+
+				const UEdGraphSchema_K2* K2Schema = GetDefault<UEdGraphSchema_K2>();
+				UEdGraphPin* PlayThenPin = PostBeginPlayEventNode->FindPin(K2Schema->PN_Then);
+				
+				DelayNode = EditorBuildPromotionTestUtils::AddDelayNode(BlueprintObject, EventGraph, PlayThenPin);
+				if (DelayNode)
+				{
+					UE_LOG(LogEditorBuildPromotionTests, Display, TEXT("Created a Delay node"));
+				}
+				else
+				{
+					UE_LOG(LogEditorBuildPromotionTests, Error, TEXT("Failed to create a Delay node"));
+				}
+
+				GEditor->UndoTransaction();
+				if (EventGraph->Nodes.Num() == 0 || EventGraph->Nodes[EventGraph->Nodes.Num() - 1] != DelayNode)
+				{
+					UE_LOG(LogEditorBuildPromotionTests, Display, TEXT("Undo add Delay node successful"));
+				}
+				else
+				{
+					UE_LOG(LogEditorBuildPromotionTests, Error, TEXT("Failed to undo add delay node"));
+				}
+
+				GEditor->RedoTransaction();
+				if (EventGraph->Nodes.Num() > 0 && EventGraph->Nodes[EventGraph->Nodes.Num() - 1] == DelayNode)
+				{
+					UE_LOG(LogEditorBuildPromotionTests, Display, TEXT("Redo add Delay node successful"));
+				}
+				else
+				{
+					UE_LOG(LogEditorBuildPromotionTests, Error, TEXT("Failed to redo add delay node"));
+				}
+
+				EditorBuildPromotionTestUtils::SetPinDefaultValue(DelayNode, TEXT("Duration"), TEXT("2.0"));
+				UE_LOG(LogEditorBuildPromotionTests, Display, TEXT("Set the duration of the delay node to 2.0 seconds"));
+
+				EditorBuildPromotionTestUtils::CompileBlueprint(BlueprintObject);
+			}
+
+			return true;
+		}
+
+		/**
+		* Blueprint Test Stage: Custom Variable
+		*    Creates a custom string variable and adds Get/Set nodes for it
+		*/
+		bool Blueprint_CustomVariable()
+		{
+			if (BlueprintObject)
+			{
+				UEdGraph* EventGraph = FBlueprintEditorUtils::FindEventGraph(BlueprintObject);
+
+				UE_LOG(LogEditorBuildPromotionTests, Display, TEXT("Added a string member variable"));
+				EditorBuildPromotionTestUtils::AddStringMemberValue(BlueprintObject, EditorBuildPromotionTestUtils::BlueprintStringVariableName);
+
+				UE_LOG(LogEditorBuildPromotionTests, Display, TEXT("Added a new Set node for the string variable"));
+				SetNode = EditorBuildPromotionTestUtils::AddGetSetNode(BlueprintObject, EventGraph, EditorBuildPromotionTestUtils::BlueprintStringVariableName.ToString(), false);
+
+				GEditor->UndoTransaction();
+				if (EventGraph->Nodes.Num() == 0 || EventGraph->Nodes[EventGraph->Nodes.Num() - 1] != SetNode)
+				{
+					UE_LOG(LogEditorBuildPromotionTests, Display, TEXT("Undo add Set node successful"));
+				}
+				else
+				{
+					UE_LOG(LogEditorBuildPromotionTests, Error, TEXT("Failed to undo add Set node"));
+				}
+
+				GEditor->RedoTransaction();
+				if (EventGraph->Nodes.Num() > 0 && EventGraph->Nodes[EventGraph->Nodes.Num() - 1] == SetNode)
+				{
+					UE_LOG(LogEditorBuildPromotionTests, Display, TEXT("Redo add Set node successful"));
+				}
+				else
+				{
+					UE_LOG(LogEditorBuildPromotionTests, Error, TEXT("Failed to redo add set node"));
+				}
+
+				UE_LOG(LogEditorBuildPromotionTests, Display, TEXT("Adding a Get node for the string variable"));
+				GetNode = EditorBuildPromotionTestUtils::AddGetSetNode(BlueprintObject, EventGraph, EditorBuildPromotionTestUtils::BlueprintStringVariableName.ToString(), true, 400);
+				
+				GEditor->UndoTransaction();
+				if (EventGraph->Nodes.Num() == 0 || EventGraph->Nodes[EventGraph->Nodes.Num() - 1] != GetNode)
+				{
+					UE_LOG(LogEditorBuildPromotionTests, Display, TEXT("Undo add get node"));
+				}
+				else
+				{
+					UE_LOG(LogEditorBuildPromotionTests, Error, TEXT("Failed to undo add get node"));
+				}
+
+				GEditor->RedoTransaction();
+				if (EventGraph->Nodes.Num() > 0 && EventGraph->Nodes[EventGraph->Nodes.Num() - 1] == GetNode)
+				{
+					UE_LOG(LogEditorBuildPromotionTests, Display, TEXT("Redo add get node"));
+				}
+				else
+				{
+					UE_LOG(LogEditorBuildPromotionTests, Error, TEXT("Failed to redo add get node"));
+				}
+
+				UE_LOG(LogEditorBuildPromotionTests, Display, TEXT("Removed the Set node"));
+				EventGraph->RemoveNode(SetNode);
+				SetNode = NULL;
+			}
+			return true;
+		}
+
+		/**
+		* Blueprint Test Stage:  Using Variables
+		*   Adds a PrintString and SetStaticMesh then connects all the existing nodes
+		*/
+		bool Blueprint_UsingVariables()
+		{
+			if (BlueprintObject)
+			{
+				const bool bVariableIsHidden = false;
+				FBlueprintEditorUtils::SetBlueprintOnlyEditableFlag(BlueprintObject, EditorBuildPromotionTestUtils::BlueprintStringVariableName, bVariableIsHidden);
+				UE_LOG(LogEditorBuildPromotionTests, Display, TEXT("Exposed the blueprint string variable"));
+
+				UEdGraph* EventGraph = FBlueprintEditorUtils::FindEventGraph(BlueprintObject);
+				PrintNode = EditorBuildPromotionTestUtils::AddPrintStringNode(BlueprintObject, EventGraph);
+				if (PrintNode)
+				{
+					UE_LOG(LogEditorBuildPromotionTests, Display, TEXT("Added a PrintString node"));
+				}
+				else
+				{
+					UE_LOG(LogEditorBuildPromotionTests, Error, TEXT("Failed to create a PrintString node"));
+				}
+
+				const UEdGraphSchema_K2* K2Schema = GetDefault<UEdGraphSchema_K2>();
+
+				//Connect Get to printstring
+				UEdGraphPin* GetVarPin = GetNode->FindPin(EditorBuildPromotionTestUtils::BlueprintStringVariableName.ToString());
+				UEdGraphPin* InStringPin = PrintNode->FindPin(TEXT("InString"));
+				GetVarPin->MakeLinkTo(InStringPin);
+				UE_LOG(LogEditorBuildPromotionTests, Display, TEXT("Connected the string variable Get node to the PrintString node"));
+
+				//Connect Delay to PrintString
+				UEdGraphPin* DelayExecPin = DelayNode->FindPin(K2Schema->PN_Then);
+				UEdGraphPin* PrintStringPin = PrintNode->FindPin(K2Schema->PN_Execute);
+				DelayExecPin->MakeLinkTo(PrintStringPin);
+				UE_LOG(LogEditorBuildPromotionTests, Display, TEXT("Connected the Delay node to the PrintString node"));
+
+				const FName MyMeshVarName(TEXT("MyMesh"));
+				GetNode = EditorBuildPromotionTestUtils::AddGetSetNode(BlueprintObject, EventGraph, MyMeshVarName.ToString(), true, 680);
+				if (GetNode)
+				{
+					UE_LOG(LogEditorBuildPromotionTests, Display, TEXT("Added a Get node for the MyMesh variable"));
+				}
+				else
+				{
+					UE_LOG(LogEditorBuildPromotionTests, Error, TEXT("Failed to create a Get node for the MyMesh variable"));
+				}
+
+				SetStaticMeshNode = EditorBuildPromotionTestUtils::AddSetStaticMeshNode(BlueprintObject, EventGraph);
+				if (SetStaticMeshNode)
+				{
+					UE_LOG(LogEditorBuildPromotionTests, Display, TEXT("Added a SetStaticMesh node"));
+				}
+				else
+				{
+					UE_LOG(LogEditorBuildPromotionTests, Error, TEXT("Failed to create a SetStaticMesh node"));
+				}
+
+				UEdGraphPin* GetExecPin = GetNode->FindPin(TEXT("MyMesh"));
+				UEdGraphPin* SetStaticMeshSelfPin = SetStaticMeshNode->FindPin(K2Schema->PN_Self);
+				GetExecPin->MakeLinkTo(SetStaticMeshSelfPin);
+				UE_LOG(LogEditorBuildPromotionTests, Display, TEXT("Connected the Get MyMesh node to the SetStaticMesh node"));
+
+				UEdGraphPin* SetStaticMeshMeshPin = SetStaticMeshNode->FindPin(TEXT("NewMesh"));
+				SetStaticMeshMeshPin->DefaultObject = SecondBlueprintMesh;
+				UE_LOG(LogEditorBuildPromotionTests, Display, TEXT("Set the satic mesh on the SetStaticMesh node to: %s"), *SecondBlueprintMesh->GetName());
+
+				//Connect SetStaticMeshMesh to PrintString
+				UEdGraphPin* PrintStringThenPin = PrintNode->FindPin(K2Schema->PN_Then);
+				UEdGraphPin* SetStaticMeshExecPin = SetStaticMeshNode->FindPin(K2Schema->PN_Execute);
+				PrintStringThenPin->MakeLinkTo(SetStaticMeshExecPin);
+				UE_LOG(LogEditorBuildPromotionTests, Display, TEXT("Connected the PrintString node to the SetStaticMesh node"));
+			}
+			return true;
+		}
+
+		/**
+		* Blueprint Test Stage: Renaming custom event
+		*    Creates, renames, and then removes a custom event node
+		*/
+		bool Blueprint_RenameCustomEvent()
+		{
+			if (BlueprintObject)
+			{
+				UEdGraph* EventGraph = FBlueprintEditorUtils::FindEventGraph(BlueprintObject);
+				CustomEventNode = EditorBuildPromotionTestUtils::CreateCustomEvent(BlueprintObject, EventGraph, TEXT("NewEvent"));
+				if (CustomEventNode)
+				{
+					UE_LOG(LogEditorBuildPromotionTests, Display, TEXT("Created a custom event node"));
+					//Rename the event
+					CustomEventNode->OnRenameNode(TEXT("RenamedEvent"));
+					UE_LOG(LogEditorBuildPromotionTests, Display, TEXT("Renamed the custom event node to RenamedEvent"));
+					EventGraph->RemoveNode(CustomEventNode);
+					UE_LOG(LogEditorBuildPromotionTests, Display, TEXT("Removed the custom event node"));
+					CustomEventNode = NULL;
+				}
+				else
+				{
+					UE_LOG(LogEditorBuildPromotionTests, Error, TEXT("Failed to create a custom event node"));
+				}
+
+			}
+
+			return true;
+		}
+
+		/**
+		* Blueprint Test Stage: New function
+		*    Creates a new function graph and then hooks up a new AddParticleSystem inside it
+		*/
+		bool Blueprint_NewFunctions()
+		{
+			if (BlueprintObject)
+			{
+				const UEdGraphSchema_K2* K2Schema = GetDefault<UEdGraphSchema_K2>();
+
+				CustomGraph = EditorBuildPromotionTestUtils::CreateNewFunctionGraph(BlueprintObject, TEXT("NewFunction"));
+				if (CustomGraph)
+				{
+					UE_LOG(LogEditorBuildPromotionTests, Display, TEXT("Created a new function graph (NewFunction)"));
+				}
+				else
+				{
+					UE_LOG(LogEditorBuildPromotionTests, Error, TEXT("Failed to create a new function graph"));
+				}
+
+				AddParticleSystemNode = EditorBuildPromotionTestUtils::CreateAddComponentActionNode(BlueprintObject, CustomGraph, CreatedPS);
+				if (AddParticleSystemNode)
+				{
+					UE_LOG(LogEditorBuildPromotionTests, Display, TEXT("Created an AddParticleSystem node"));
+				}
+				else
+				{
+					UE_LOG(LogEditorBuildPromotionTests, Error, TEXT("Failed to create an AddParticleSystem node"));
+				}
+
+				UEdGraphPin* ExecutePin = AddParticleSystemNode ? AddParticleSystemNode->FindPin(K2Schema->PN_Execute) : NULL;
+
+				//Find the input for the function graph
+				TArray<UK2Node_FunctionEntry*> EntryNodes;
+				CustomGraph->GetNodesOfClass(EntryNodes);
+				UEdGraphNode* EntryNode = EntryNodes.Num() > 0 ? EntryNodes[0] : NULL;
+				if (EntryNode && ExecutePin)
+				{
+					UEdGraphPin* EntryPin = EntryNode->FindPin(K2Schema->PN_Then);
+					EntryPin->MakeLinkTo(ExecutePin);
+					UE_LOG(LogEditorBuildPromotionTests, Display, TEXT("Connected the AddParticleSystem node to the entry node"));
+				}
+
+				EditorBuildPromotionTestUtils::CompileBlueprint(BlueprintObject);
+			}
+			return true;
+		}
+
+		/**
+		* Blueprint Test Stage: Completing the blueprint
+		*    Adds a CallFunction node to call the custom function created in the previous step.
+		*/
+		bool Blueprint_CompleteBlueprint()
+		{
+			if (BlueprintObject)
+			{
+				const UEdGraphSchema_K2* K2Schema = GetDefault<UEdGraphSchema_K2>();
+				UEdGraph* EventGraph = FBlueprintEditorUtils::FindEventGraph(BlueprintObject);
+				UEdGraphPin* SetStaticMeshThenPin = SetStaticMeshNode->FindPin(K2Schema->PN_Then);
+				CallFunctionNode = EditorBuildPromotionTestUtils::AddCallFunctionGraphNode(BlueprintObject, EventGraph, TEXT("NewFunction"), SetStaticMeshThenPin);
+				if (CallFunctionNode)
+				{
+					UE_LOG(LogEditorBuildPromotionTests, Display, TEXT("Created a CallFunction node to call NewFunction"));
+					UE_LOG(LogEditorBuildPromotionTests, Display, TEXT("Hooked the CallFunction node to the output of the SetStaticMesh node"));
+				}
+				else
+				{
+					UE_LOG(LogEditorBuildPromotionTests, Error, TEXT("Failed to create a CallFunction node"));
+				}
+
+				EditorBuildPromotionTestUtils::CompileBlueprint(BlueprintObject);
+
+				SaveBlueprint();
+			}
+			return true;
+		}
+
+		/**
+		* Blueprint Test Stage: Blueprint placement (Part 1)
+		*    Places the blueprint in the level then starts PIE
+		*/
+		bool Blueprint_Placement_Part1()
+		{
+			if (BlueprintObject)
+			{
+				PlacedBlueprint = FActorFactoryAssetProxy::AddActorForAsset(BlueprintObject);
+
+				if (PlacedBlueprint)
+				{
+					//Set the text
+					EditorBuildPromotionTestUtils::SetPropertyByName(PlacedBlueprint, EditorBuildPromotionTestUtils::BlueprintStringVariableName.ToString(), TEXT("Print String works!"));
+					UE_LOG(LogEditorBuildPromotionTests, Display, TEXT("Placed the blueprint and set the display string to \"Print String works!\""));
+				}
+				else
+				{
+					UE_LOG(LogEditorBuildPromotionTests, Error, TEXT("Failed to place blueprint"));
+				}
+
+				GEditor->SelectNone(false, true);
+
+				//Start PIE
+				EditorBuildPromotionTestUtils::StartPIE(true);
+
+				//Make sure the timer is reset
+				DelayHelper.Reset();
+			}
+			return true;
+		}
+
+		/**
+		* Blueprint Test Stage: Blueprint placement (Part 2)
+		*    Takes a screenshot of the initial state of the blueprint
+		*/
+		bool Blueprint_Placement_Part2()
+		{
+			if (BlueprintObject)
+			{
+				EditorBuildPromotionTestUtils::TakeScreenshot(TEXT("BlueprintPIE_Start"));
+			}
+			return true;
+		}
+
+		/**
+		* Blueprint Test Stage: Blueprint placement (Part 3)
+		*    Waits for 2 seconds for the timer to finish
+		*/
+		bool Blueprint_Placement_Part3()
+		{
+			if (BlueprintObject)
+			{
+				//Set a timeout in case PIE doesn't work
+				if (!DelayHelper.IsRunning())
+				{
+					DelayHelper.Start(5.0f);
+				}
+				else if (DelayHelper.IsComplete())
+				{
+					//FAILED to hit breakpoint in time
+					UE_LOG(LogEditorBuildPromotionTests, Error, TEXT("Timed out waiting for PIE to start"));
+					DelayHelper.Reset();
+					return true;
+				}
+
+				//Wait for PIE to startup
+				if (GEditor->PlayWorld != NULL)
+				{
+					//Stop after 2 seconds of gameplay
+					if (GEditor->PlayWorld->TimeSeconds > 2.0f)
+					{
+						DelayHelper.Reset();
+						return true;
+					}
+				}
+				return false;
+			}
+			return true;
+		}
+
+		/**
+		* Blueprint Test Stage: Blueprint placement (Part 4)
+		*    Takes a screenshot after the blueprint has changed.  Ends the PIE session.
+		*/
+		bool Blueprint_Placement_Part4()
+		{
+			if (BlueprintObject)
+			{
+				if (GEditor->PlayWorld != NULL)
+				{
+					EditorBuildPromotionTestUtils::TakeScreenshot(TEXT("BlueprintPIE_End"));
+				}
+				EditorBuildPromotionTestUtils::EndPIE();
+			}
+			return true;
+		}
+
+		/**
+		* Blueprint Test Stage: Set Breakpoint (Part 1)
+		*    Sets a breakpoint on the PrintString node and starts PIE
+		*/
+		bool Blueprint_SetBreakpoint_Part1()
+		{
+			if (BlueprintObject)
+			{
+
+				//Add a breakpoint
+				UBreakpoint* NewBreakpoint = NewObject<UBreakpoint>(BlueprintObject);
+				FKismetDebugUtilities::SetBreakpointEnabled(NewBreakpoint, true);
+				FKismetDebugUtilities::SetBreakpointLocation(NewBreakpoint, PrintNode);
+				BlueprintObject->Breakpoints.Add(NewBreakpoint);
+				BlueprintObject->MarkPackageDirty();
+
+				UE_LOG(LogEditorBuildPromotionTests, Display, TEXT("Set a breakpoint on the PrintString node"));
+				EditorBuildPromotionTestUtils::StartPIE(true);
+
+				//Make sure the timer is reset
+				DelayHelper.Reset();
+			}
+			return true;
+		}
+
+		/**
+		* Blueprint Test Stage: Set Breakpoint (Part 2)
+		*    Waits for the breakpoint to get hit or a 10 second timeout to expire
+		*/
+		bool Blueprint_SetBreakpoint_Part2()
+		{
+			if (BlueprintObject)
+			{
+				//Set a timeout for hitting the breakpoint
+				if (!DelayHelper.IsRunning())
+				{
+					DelayHelper.Start(10.0f);
+				}
+				else if (DelayHelper.IsComplete())
+				{
+					//FAILED to hit breakpoint in time
+					UE_LOG(LogEditorBuildPromotionTests, Error, TEXT("Failed to hit the breakpoint after 10 seconds"));
+					DelayHelper.Reset();
+					return true;
+				}
+
+				UEdGraphNode* CurrentBreakpointNode = NULL;
+				{
+					//Hack.  GetMostRecentBreakpointHit only returns data if GIntraFrameDebuggingGameThread is true.
+					TGuardValue<bool> SignalGameThreadBeingDebugged(GIntraFrameDebuggingGameThread, true);
+					CurrentBreakpointNode = FKismetDebugUtilities::GetMostRecentBreakpointHit();
+				}
+
+				//Wait for breakpoint to get hit
+				if (CurrentBreakpointNode == PrintNode)
+				{
+					//Success!
+					UE_LOG(LogEditorBuildPromotionTests, Display, TEXT("Hit the PrintString breakpoint"));
+					DelayHelper.Reset();
+					return true;
+				}
+				return false;
+			}
+			return true;
+		}
+
+		/**
+		* Blueprint Test Stage: Set Breakpoint (Part 3)
+		*    Ends the PIE session and clears the breakpoint
+		*/
+		bool Blueprint_SetBreakpoint_Part3()
+		{
+			if (BlueprintObject)
+			{
+				UE_LOG(LogEditorBuildPromotionTests, Display, TEXT("Clearing the breakpoint"));
+				FKismetDebugUtilities::ClearBreakpoints(BlueprintObject);
+				EditorBuildPromotionTestUtils::EndPIE();
+				CurrentWorld = PlacedBlueprint->GetWorld();
+			}
+			return true;
+		}
+
+		/**
+		* Blueprint Test Stage: Level Script (Part 1)
+		*    Finds and opens the level script blueprint
+		*/
+		bool Blueprint_LevelScript_Part1()
+		{
+			if (BlueprintObject)
+			{
+				//Open the level script blueprint
+				ULevelScriptBlueprint* LSB = PlacedBlueprint->GetLevel()->GetLevelScriptBlueprint(false);
+				FAssetEditorManager::Get().OpenEditorForAsset(LSB);
+				UE_LOG(LogEditorBuildPromotionTests, Display, TEXT("Opened the level script blueprint"));
+			}
+			return true;
+		}
+
+		/**
+		* Blueprint Test Stage: Level Script (Part 2)
+		*    Copies the event nodes from the blueprint to the level script and compiles.  
+		*    Removes the variables and function from the level script and compiles again.
+		*    Modifies the Delay and PrintString values
+		*    Starts a PIE session
+		*/
+		bool Blueprint_LevelScript_Part2()
+		{
+			if (BlueprintObject)
+			{
+				IAssetEditorInstance* AssetEditor = FAssetEditorManager::Get().FindEditorForAsset(BlueprintObject, true);
+				FBlueprintEditor* BlueprintEditor = (FBlueprintEditor*)AssetEditor;
+
+				UEdGraph* EventGraph = FBlueprintEditorUtils::FindEventGraph(BlueprintObject);
+				TSet<UObject*> NodesToExport;
+				for (int32 i = 0; i < EventGraph->Nodes.Num(); ++i)
+				{
+					EventGraph->Nodes[i]->PrepareForCopying();
+					NodesToExport.Add(EventGraph->Nodes[i]);
+				}
+
+				//Saving the starting log counts so we can clear any entries created from the copy/paste action.  (Lots of expected spam)
+				const int32 PreCopyErrorCount = TestExecutionInfo->Errors.Num();
+				const int32 PreCopyWarningCount = TestExecutionInfo->Warnings.Num();
+				const int32 PreCopyLogCount = TestExecutionInfo->LogItems.Num();
+
+				FString OutNodeText;
+				FEdGraphUtilities::ExportNodesToText(NodesToExport, OutNodeText);
+				FPlatformMisc::ClipboardCopy(*OutNodeText);
+
+				ULevelScriptBlueprint* LSB = PlacedBlueprint->GetLevel()->GetLevelScriptBlueprint(true);
+				UEdGraph* LevelEventGraph = FBlueprintEditorUtils::FindEventGraph(LSB);
+				FKismetEditorUtilities::PasteNodesHere(LevelEventGraph, FVector2D(0, 0));
+
+				//remove any new logs
+				TestExecutionInfo->Errors.RemoveAt(PreCopyErrorCount, TestExecutionInfo->Errors.Num() - PreCopyErrorCount);
+				TestExecutionInfo->Warnings.RemoveAt(PreCopyWarningCount, TestExecutionInfo->Warnings.Num() - PreCopyWarningCount);
+				TestExecutionInfo->LogItems.RemoveAt(PreCopyLogCount, TestExecutionInfo->LogItems.Num() - PreCopyLogCount);
+
+
+				//Note: These are a little out of order because logs are disabled above
+				UE_LOG(LogEditorBuildPromotionTests, Display, TEXT("Copied the blueprint event nodes"));
+				UE_LOG(LogEditorBuildPromotionTests, Display, TEXT("Pasted the nodes into to level script"));
+				
+
+				//We are expecting errors here.  So copy any new errors to the logs section
+				{
+					//Store our error count before the compile
+					uint32 NumErrors = TestExecutionInfo->Errors.Num();
+
+					EditorBuildPromotionTestUtils::CompileBlueprint(LSB);
+
+					//Copy the errors over to the LogItems
+					for (int32 i = NumErrors; i < TestExecutionInfo->Errors.Num(); ++i)
+					{
+						TestExecutionInfo->LogItems.Add(TestExecutionInfo->Errors[i]);
+					}
+					TestExecutionInfo->Errors.RemoveAt(NumErrors, TestExecutionInfo->Errors.Num() - NumErrors);
+				}
+
+				//Remove the "Get" nodes
+				TArray<UK2Node_VariableGet*> AllGetNodes;
+				LevelEventGraph->GetNodesOfClass<UK2Node_VariableGet>(AllGetNodes);
+				for (int32 i = 0; i < AllGetNodes.Num(); ++i)
+				{
+					LevelEventGraph->RemoveNode(AllGetNodes[i]);
+				}
+				UE_LOG(LogEditorBuildPromotionTests, Display, TEXT("Removed get nodes from the level script"));
+
+				//Remove all the call function nodes except the PrintString and Delay nodes
+				TArray<UK2Node_CallFunction*> AllCallFunctionNodes;
+				LevelEventGraph->GetNodesOfClass<UK2Node_CallFunction>(AllCallFunctionNodes);
+				for (int32 i = 0; i < AllCallFunctionNodes.Num(); ++i)
+				{
+					UK2Node_CallFunction* NodeIt = AllCallFunctionNodes[i];
+					const FString MemberName = NodeIt->FunctionReference.GetMemberName().ToString();
+					if (MemberName == TEXT("PrintString"))
+					{
+						UE_LOG(LogEditorBuildPromotionTests, Display, TEXT("Setting the level PrintString to \"Level kismet string text\""));
+						EditorBuildPromotionTestUtils::SetPinDefaultValue(NodeIt, TEXT("InString"), TEXT("Level kismet string text"));
+					}
+					else if (MemberName == TEXT("Delay"))
+					{
+						UE_LOG(LogEditorBuildPromotionTests, Display, TEXT("Setting the level script Delay to 4 seconds"));
+						EditorBuildPromotionTestUtils::SetPinDefaultValue(NodeIt, TEXT("Duration"), TEXT("4.0"));
+					}
+					else
+					{
+						if (NodeIt == CallFunctionNode)
+						{
+							UE_LOG(LogEditorBuildPromotionTests, Display, TEXT("Removed CallFunction node from level script"));
+							CallFunctionNode = NULL;
+						}
+						else if (NodeIt == SetStaticMeshNode)
+						{
+							UE_LOG(LogEditorBuildPromotionTests, Display, TEXT("Removed SetStaticMesh node from level script"));
+							SetStaticMeshNode = NULL;
+						}
+						//Remove 
+						LevelEventGraph->RemoveNode(NodeIt);
+					}
+				}
+
+				//Compile the blueprint
+				EditorBuildPromotionTestUtils::CompileBlueprint(LSB);
+
+				//Test PIE
+				EditorBuildPromotionTestUtils::StartPIE(true);
+
+				//Make sure the timer is reset
+				DelayHelper.Reset();
+			}
+			return true;
+		}
+
+		/**
+		* Blueprint Test Stage: Level Script (Part 3)
+		*    Waits for the delay timer in the level script
+		*/
+		bool Blueprint_LevelScript_Part3()
+		{
+			if (BlueprintObject)
+			{
+				//Set a timeout in case PIE doesn't start
+				if (!DelayHelper.IsRunning())
+				{
+					DelayHelper.Start(8.0f);
+				}
+				else if (DelayHelper.IsComplete())
+				{
+					//FAILED to hit breakpoint in time
+					UE_LOG(LogEditorBuildPromotionTests, Error, TEXT("Timed out waiting for PIE to start"));
+					DelayHelper.Reset();
+					return true;
+				}
+
+				//Wait for PIE to startup
+				if (GEditor->PlayWorld != NULL)
+				{
+					//Stop after 4 seconds of gameplay
+					if (GEditor->PlayWorld->TimeSeconds > 4.0f)
+					{
+						DelayHelper.Reset();
+						return true;
+					}
+				}
+				return false;
+			}
+			return true;
+		}
+
+		/**
+		* Blueprint Test Stage: Level Script (Part 4)
+		*    Takes a screenshot and ends PIE
+		*/
+		bool Blueprint_LevelScript_Part4()
+		{
+			if (BlueprintObject)
+			{
+				if (GEditor->PlayWorld != NULL)
+				{
+					//Take a screenshot and end the PIE session
+					EditorBuildPromotionTestUtils::TakeScreenshot(TEXT("LevelBlueprint"), false);
+				}
+				EditorBuildPromotionTestUtils::EndPIE();
+			}
+			return true;
+		}
+
+		/**
+		* Blueprint Test Stage: Level Script (Part 5)
+		*    Closes the blueprint editor and saves the blueprint
+		*/
+		bool Blueprint_LevelScript_Part5()
+		{
+			if (BlueprintObject)
+			{
+				UE_LOG(LogEditorBuildPromotionTests, Display, TEXT("Closing the blueprint editor"));
+				FAssetEditorManager::Get().CloseAllAssetEditors();
+				UE_LOG(LogEditorBuildPromotionTests, Display, TEXT("Saving the blueprint"));
+				SaveBlueprint();
+			}
+			return true;
+		}
+
+		/**
+		* Building Test Stage: Building and Saving (Part 1)
+		*    Toggles the level visibility off
+		*/
+		bool Building_BuildLevel_Part1()
+		{
+			UWorld* World = GEditor->GetEditorWorldContext().World();
+			ULevel* Level = World->GetCurrentLevel();
+
+			//Store our error count before build
+			BuildStartErrorCount = TestExecutionInfo->Errors.Num();
+
+			//Save all the new assets
+			EditorBuildPromotionTestUtils::SaveAllAssetsInFolder(EditorBuildPromotionTestUtils::GetGamePath());
+
+			UE_LOG(LogEditorBuildPromotionTests, Display, TEXT("Turning level visibility off"));
+			bool bShouldBeVisible = false;
+			EditorLevelUtils::SetLevelVisibility(Level, bShouldBeVisible, false);
+			GEditor->RedrawAllViewports(true);
+
+			return true;
+		}
+
+		/**
+		* Building Test Stage:  Building and Saving (Part 2)
+		*   Takes a screenshot and toggles the level visibility back on
+		*/
+		bool  Building_BuildLevel_Part2()
+		{
+			UWorld* World = GEditor->GetEditorWorldContext().World();
+			ULevel* Level = World->GetCurrentLevel();
+
+			UE_LOG(LogEditorBuildPromotionTests, Display, TEXT("Taking a screenshot"));
+			EditorBuildPromotionTestUtils::TakeScreenshot(TEXT("VisibilityOff"));
+
+			UE_LOG(LogEditorBuildPromotionTests, Display, TEXT("Turning level visibility on"));
+			bool bShouldBeVisible = true;
+			EditorLevelUtils::SetLevelVisibility(Level, bShouldBeVisible, false);
+
+			return true;
+		}
+
+		/**
+		* Building Test Stage:  Building and Saving (Part 3)
+		*   Takes a screenshot and does a full level rebuild
+		*/
+		bool  Building_BuildLevel_Part3()
+		{
+			UWorld* World = GEditor->GetEditorWorldContext().World();
+			ULevel* Level = World->GetCurrentLevel();
+
+			UE_LOG(LogEditorBuildPromotionTests, Display, TEXT("Taking a screenshot"));
+			EditorBuildPromotionTestUtils::TakeScreenshot(TEXT("VisibilityOn"));
+
+			FEditorFileUtils::SaveLevel(Level, TEXT("/Game/Maps/EditorBuildPromotionTest"));
+			GUnrealEd->Exec(World, TEXT("MAP REBUILD ALLVISIBLE"));
+			UE_LOG(LogEditorBuildPromotionTests, Display, TEXT("Rebuilt the map"));
+
+#if WITH_NAVIGATION_GENERATOR
+			if (World->GetWorldSettings()->bEnableNavigationSystem &&
+				World->GetNavigationSystem())
+			{
+				// Invoke navmesh generator
+				World->GetNavigationSystem()->Build();
+				UE_LOG(LogEditorBuildPromotionTests, Display, TEXT("Built navigation"));
+			}
+#endif
+
+			//Force AutoApplyLighting on
+			ULevelEditorMiscSettings* LevelEdSettings = GetMutableDefault<ULevelEditorMiscSettings>();
+			bDisableAutoApplyLighting = !LevelEdSettings->bAutoApplyLightingEnable;
+			LevelEdSettings->bAutoApplyLightingEnable = true;
+
+			//Build Lighting
+			EditorBuildPromotionTestUtils::BuildLighting();
+
+			return true;
+		}
+
+		/**
+		* Building Test Stage: Building and Saving (Part 4)
+		*    Waits for the lighting to finish building and saves the level
+		*/
+		bool Building_BuildLevel_Part4()
+		{
+			if (!GUnrealEd->IsLightingBuildCurrentlyRunning())
+			{
+				UE_LOG(LogEditorBuildPromotionTests, Display, TEXT("Built Lighting"));
+
+				if (bDisableAutoApplyLighting)
+				{
+					ULevelEditorMiscSettings* LevelEdSettings = GetMutableDefault<ULevelEditorMiscSettings>();
+					LevelEdSettings->bAutoApplyLightingEnable = false;
+				}
+
+				UE_LOG(LogEditorBuildPromotionTests, Display, TEXT("Saved the Level (EditorBuildPromotionTest)"));
+
+				UWorld* World = GEditor->GetEditorWorldContext().World();
+				ULevel* Level = World->GetCurrentLevel();
+				FEditorFileUtils::SaveLevel(Level);
+
+				//Save all the new assets again because material usage flags may have changed.
+				EditorBuildPromotionTestUtils::SaveAllAssetsInFolder(EditorBuildPromotionTestUtils::GetGamePath());
+
+				//Copy the errors over to the LogItems
+				for (int32 i = BuildStartErrorCount; i < TestExecutionInfo->Errors.Num(); ++i)
+				{
+					TestExecutionInfo->LogItems.Add(TestExecutionInfo->Errors[i]);
+				}
+				TestExecutionInfo->Errors.RemoveAt(BuildStartErrorCount, TestExecutionInfo->Errors.Num() - BuildStartErrorCount);
+				
+				return true;
+			}
+			return false;
+		}
+	};
+}
+
+/**
+* Automation test that handles cleanup of the build promotion test
+*/
+IMPLEMENT_SIMPLE_AUTOMATION_TEST(FBuildPromotionInitialCleanupTest, "Promotion.Editor Promotion Pass.Step 1 Main Editor Test.Cleanup old files", EAutomationTestFlags::ATF_Editor);
+bool FBuildPromotionInitialCleanupTest::RunTest(const FString& Parameters)
+{
+	EditorBuildPromotionTestUtils::PerformCleanup();
+	return true;
+}
+
+/**
+* Latent command to check if PIE is running
+*/
+DEFINE_LATENT_AUTOMATION_COMMAND(FSettingsCheckForPIECommand);
+bool FSettingsCheckForPIECommand::Update()
+{
+	if (GEditor->PlayWorld != NULL)
+	{
+		//Success
+		UE_LOG(LogEditorBuildPromotionTests, Display, TEXT("PlayInEditor keyboard shotcut success"));
+		EditorBuildPromotionTestUtils::EndPIE();
+	}
+	else
+	{
+		UE_LOG(LogEditorBuildPromotionTests, Error, TEXT("PlayInEditor keyboard shotcut failed"));
+	}
+	return true;
+}
+
+/**
+* Latent command to run the main build promotion test
+*/
+DEFINE_LATENT_AUTOMATION_COMMAND_ONE_PARAMETER(FLogTestResultCommand, FAutomationTestExecutionInfo*, InExecutionInfo);
+bool FLogTestResultCommand::Update()
+{
+	if (InExecutionInfo && InExecutionInfo->Errors.Num() > 0)
+	{
+		UE_LOG(LogEditorBuildPromotionTests, Display, TEXT("Test failed!"));
+	}
+	else
+	{
+		UE_LOG(LogEditorBuildPromotionTests, Display, TEXT("Test successful!"));
+	}
+	return true;
+}
+
+/**
+* Automation test that handles setting keybindings and editor preferences
+*/
+IMPLEMENT_SIMPLE_AUTOMATION_TEST(FBuildPromotionSettingsTest, "Promotion.Editor Promotion Pass.Step 1 Main Editor Test.Settings", EAutomationTestFlags::ATF_Editor);
+
+bool FBuildPromotionSettingsTest::RunTest(const FString& Parameters)
+{
+
+//Previous Editor Settings
+	UE_LOG(LogEditorBuildPromotionTests, Display, TEXT("Exporting Current keybindings and editor settings"));
+
+	//Export the original keybindings
+	const FString TargetOriginalKeybindFile = FString::Printf(TEXT("%s/BuildPromotion/OriginalKeybindings-%d.ini"), *FPaths::AutomationDir(), GEngineVersion.GetChangelist());
+	EditorBuildPromotionTestUtils::ExportKeybindings(TargetOriginalKeybindFile);
+
+	//Export the original preferences
+	const FString TargetOriginalPreferenceFile = FString::Printf(TEXT("%s/BuildPromotion/OriginalPreferences-%d.ini"), *FPaths::AutomationDir(), GEngineVersion.GetChangelist());
+	EditorBuildPromotionTestUtils::ExportEditorSettings(TargetOriginalPreferenceFile);
+
+//New Editor Settings
+	UE_LOG(LogEditorBuildPromotionTests, Display, TEXT("Binding create empty layer shortcut"));
+
+	//Bind H to CreateEmptyLayer keybinding
+	FInputGesture NewCreateGesture(EKeys::H, EModifierKey::None);
+	EditorBuildPromotionTestUtils::SetEditorKeybinding(TEXT("LayersView"), TEXT("CreateEmptyLayer"), NewCreateGesture);
+
+	UE_LOG(LogEditorBuildPromotionTests, Display, TEXT("Binding request rename layer shortcut"));
+
+	//Bind J to RequestRenameLayer
+	FInputGesture NewRenameGesture(EKeys::J, EModifierKey::None);
+	EditorBuildPromotionTestUtils::SetEditorKeybinding(TEXT("LayersView"), TEXT("RequestRenameLayer"), NewRenameGesture);
+
+	UE_LOG(LogEditorBuildPromotionTests, Display, TEXT("Binding play shortcut (PIE)"));
+	FInputGesture NewPIEGesture(EKeys::L, EModifierKey::Control);
+	EditorBuildPromotionTestUtils::SetEditorKeybinding(TEXT("PlayWorld"), TEXT("RepeatLastPlay"), NewPIEGesture);
+
+	//Export the keybindings
+	const FString TargetKeybindFile = FString::Printf(TEXT("%s/BuildPromotion/Keybindings-%d.ini"), *FPaths::AutomationDir(), GEngineVersion.GetChangelist());
+	UE_LOG(LogEditorBuildPromotionTests, Display, TEXT("Exporting keybind"));
+	EditorBuildPromotionTestUtils::ExportKeybindings(TargetKeybindFile);
+
+
+	UEditorStyleSettings* EditorStyleSettings = GetMutableDefault<UEditorStyleSettings>();
+	FString OldStyleSetting = EditorBuildPromotionTestUtils::GetPropertyByName(EditorStyleSettings, TEXT("bUseSmallToolBarIcons"));
+	EditorBuildPromotionTestUtils::SetPropertyByName(EditorStyleSettings, TEXT("bUseSmallToolBarIcons"), TEXT("true"));
+
+	UE_LOG(LogEditorBuildPromotionTests, Display, TEXT("Set UseSmallToolBarIcons"));
+
+	//Export the preferences
+	const FString TargetPreferenceFile = FString::Printf(TEXT("%s/BuildPromotion/Preferences-%d.ini"), *FPaths::AutomationDir(), GEngineVersion.GetChangelist());
+	EditorBuildPromotionTestUtils::ExportEditorSettings(TargetPreferenceFile);
+
+	//Take a screenshot of the small icons
+	EditorBuildPromotionTestUtils::TakeScreenshot(TEXT("Small Toolbar Icons"));
+
+	UE_LOG(LogEditorBuildPromotionTests, Display, TEXT("Exported keybindings and preferences to %s/BuildPromotion"), *FPaths::AutomationDir());
+
+	//Change the setting back
+	EditorBuildPromotionTestUtils::SetPropertyByName(EditorStyleSettings, TEXT("bUseSmallToolBarIcons"), OldStyleSetting);
+
+	FInputGesture CurrentCreateGesture;
+	EditorBuildPromotionTestUtils::GetEditorKeybinding(TEXT("LayersView"), TEXT("CreateEmptyLayer"), CurrentCreateGesture);
+	if (CurrentCreateGesture.Key == NewCreateGesture.Key)
+	{
+		UE_LOG(LogEditorBuildPromotionTests, Display, TEXT("CreateEmptyLayer keybinding correct."));
+	}
+	else
+	{
+		UE_LOG(LogEditorBuildPromotionTests, Error, TEXT("CreateEmptyLayer keybinding incorrect."));
+	}
+
+	FInputGesture CurrentRenameGesture;
+	EditorBuildPromotionTestUtils::GetEditorKeybinding(TEXT("LayersView"), TEXT("RequestRenameLayer"), CurrentRenameGesture);
+	if (CurrentRenameGesture.Key == NewRenameGesture.Key)
+	{
+		UE_LOG(LogEditorBuildPromotionTests, Display, TEXT("RequestRenameLayer keybinding correct."));
+	}
+	else
+	{
+		UE_LOG(LogEditorBuildPromotionTests, Error, TEXT("RequestRenameLayer keybinding incorrect."));
+	}
+
+	//Focus the main editor
+	TArray< TSharedRef<SWindow> > AllWindows;
+	FSlateApplication::Get().GetAllVisibleWindowsOrdered(AllWindows);
+	FSlateApplication::Get().ProcessWindowActivatedEvent(FWindowActivateEvent(FWindowActivateEvent::EA_Activate, AllWindows[0]));
+
+	//Send the PIE event
+	FKeyboardEvent PIEKeyEvent(EKeys::L, FModifierKeysState(false, false, true, false, false, false, false, false, false), false, 0x4C);
+	FSlateApplication::Get().ProcessKeyDownEvent(PIEKeyEvent);
+	FSlateApplication::Get().ProcessKeyUpEvent(PIEKeyEvent);
+
+	UE_LOG(LogEditorBuildPromotionTests, Display, TEXT("Sent PIE keyboard shortcut"));
+
+	ADD_LATENT_AUTOMATION_COMMAND(FSettingsCheckForPIECommand());
+	ADD_LATENT_AUTOMATION_COMMAND(FLogTestResultCommand(&ExecutionInfo));
+
+	return true;
+}
+
+/**
+* Latent command to run the main build promotion test
+*/
+DEFINE_LATENT_AUTOMATION_COMMAND_ONE_PARAMETER(FRunBuildPromotionTestCommand, TSharedPtr<BuildPromotionTestHelper::FBuildPromotionTest>, BuildPromotionTest);
+bool FRunBuildPromotionTestCommand::Update()
+{
+	return BuildPromotionTest->Update();
+}
+
+/**
+* Automation test that handles the build promotion process
+*/
+IMPLEMENT_SIMPLE_AUTOMATION_TEST(FBuildPromotionTest, "Promotion.Editor Promotion Pass.Step 1 Main Editor Test.General Editor Test", EAutomationTestFlags::ATF_Editor);
+
+bool FBuildPromotionTest::RunTest(const FString& Parameters)
+{
+	TSharedPtr<BuildPromotionTestHelper::FBuildPromotionTest> BuildPromotionTest = MakeShareable(new BuildPromotionTestHelper::FBuildPromotionTest(&ExecutionInfo));
+	ADD_LATENT_AUTOMATION_COMMAND(FRunBuildPromotionTestCommand(BuildPromotionTest));
+	return true;
+}
+
+/**
+* Latent command to run the main build promotion test
+*/
+DEFINE_LATENT_AUTOMATION_COMMAND(FEndPIECommand);
+bool FEndPIECommand::Update()
+{
+	UE_LOG(LogEditorBuildPromotionTests, Display, TEXT("Ending PIE"));
+	EditorBuildPromotionTestUtils::EndPIE();
+	return true;
+}
+
+/**
+* Latent command to run the main build promotion test
+*/
+DEFINE_LATENT_AUTOMATION_COMMAND_ONE_PARAMETER(FPIEExecCommand, FString, ExecCommand);
+bool FPIEExecCommand::Update()
+{
+	if (GEditor->PlayWorld != NULL)
+	{
+		GEngine->Exec(GEditor->PlayWorld, *ExecCommand);
+	}
+	else
+	{
+		UE_LOG(LogEditorBuildPromotionTests, Error, TEXT("Tried to execute a PIE command but PIE is not running. (%s)"), *ExecCommand);
+	}
+	return true;
+}
+
+/**
+* Execute the loading of one map to verify PIE works
+*/
+IMPLEMENT_SIMPLE_AUTOMATION_TEST(FBuildPromotionPIETest, "Promotion.Editor Promotion Pass.Step 2 Run Map After Re-launch.Run Map", EAutomationTestFlags::ATF_Editor)
+bool FBuildPromotionPIETest::RunTest(const FString& Parameters)
+{
+	UAutomationTestSettings const* AutomationTestSettings = GetDefault<UAutomationTestSettings>();
+	check(AutomationTestSettings);
+
+	bool bLoadAsTemplate = false;
+	bool bShowProgress = false;
+	const FString MapName = FPaths::GameContentDir() + TEXT("Maps/EditorBuildPromotionTest.umap");
+	FEditorFileUtils::LoadMap(MapName, bLoadAsTemplate, bShowProgress);
+	UE_LOG(LogEditorBuildPromotionTests, Display, TEXT("Loading Map: %s"), *MapName);
+
+	EditorBuildPromotionTestUtils::StartPIE(false);
+	UE_LOG(LogEditorBuildPromotionTests, Display, TEXT("Starting PIE"));
+
+	//Find the main editor window
+	TArray<TSharedRef<SWindow> > AllWindows;
+	FSlateApplication::Get().GetAllVisibleWindowsOrdered(AllWindows);
+	if (AllWindows.Num() == 0)
+	{
+		UE_LOG(LogEditorAutomationTests, Error, TEXT("ERROR: Could not find the main editor window."));
+		UE_LOG(LogEditorAutomationTests, Display, TEXT("Test FAILED"));
+		return true;
+	}
+
+	WindowScreenshotParameters ScreenshotParams;
+	AutomationCommon::GetScreenshotPath(TEXT("EditorBuildPromotion/RunMap"), ScreenshotParams.ScreenshotName, false);
+	ScreenshotParams.CurrentWindow = AllWindows[0];
+	//Wait for the play world to come up
+	ADD_LATENT_AUTOMATION_COMMAND(FWaitLatentCommand(1.f));
+
+	//Toggle a stat and take a screenshot
+	UE_LOG(LogEditorBuildPromotionTests, Display, TEXT("Toggling \"Stat Memory\" and taking a screenshot"));
+	ADD_LATENT_AUTOMATION_COMMAND(FPIEExecCommand(TEXT("STAT Memory")));
+	//Stat memory doesn't update right away so wait another second.
+	ADD_LATENT_AUTOMATION_COMMAND(FWaitLatentCommand(1.f));
+	ADD_LATENT_AUTOMATION_COMMAND(FTakeEditorScreenshotCommand(ScreenshotParams));
+	ADD_LATENT_AUTOMATION_COMMAND(FPIEExecCommand(TEXT("STAT None")));
+	ADD_LATENT_AUTOMATION_COMMAND(FEndPIECommand());
+	ADD_LATENT_AUTOMATION_COMMAND(FLogTestResultCommand(&ExecutionInfo));
+
+	return true;
+}
+
+/**
+* Automation test that handles cleanup of the build promotion test
+*/
+IMPLEMENT_SIMPLE_AUTOMATION_TEST(FBuildPromotionCleanupTest, "Promotion.Editor Promotion Pass.Step 3 Test Cleanup.Cleanup", EAutomationTestFlags::ATF_Editor);
+bool FBuildPromotionCleanupTest::RunTest(const FString& Parameters)
+{
+	EditorBuildPromotionTestUtils::PerformCleanup();
+	return true;
+}
+
+#undef LOCTEXT_NAMESPACE
diff --git a/Engine/Source/Editor/UnrealEd/Public/Tests/AutomationEditorCommon.h b/Engine/Source/Editor/UnrealEd/Public/Tests/AutomationEditorCommon.h
--- a/Engine/Source/Editor/UnrealEd/Public/Tests/AutomationEditorCommon.h
+++ b/Engine/Source/Editor/UnrealEd/Public/Tests/AutomationEditorCommon.h
@@ -2,6 +2,93 @@
 
 #pragma once
 
+#include "Tests/AutomationTestSettings.h"
+
+namespace AutomationEditorCommonUtils
+{
+	/**
+	* Creates a new map for editing.  Also clears editor tools that could cause issues when changing maps
+	*
+	* @return - The UWorld for the new map
+	*/
+	static UWorld* CreateNewMap()
+	{
+		// Change out of Matinee when opening new map, so we avoid editing data in the old one.
+		if (GLevelEditorModeTools().IsModeActive(FBuiltinEditorModes::EM_InterpEdit))
+		{
+			GLevelEditorModeTools().DeactivateMode(FBuiltinEditorModes::EM_InterpEdit);
+		}
+
+		// Also change out of Landscape mode to ensure all references are cleared.
+		if (GLevelEditorModeTools().IsModeActive(FBuiltinEditorModes::EM_Landscape))
+		{
+			GLevelEditorModeTools().DeactivateMode(FBuiltinEditorModes::EM_Landscape);
+		}
+
+		// Also change out of Foliage mode to ensure all references are cleared.
+		if (GLevelEditorModeTools().IsModeActive(FBuiltinEditorModes::EM_Foliage))
+		{
+			GLevelEditorModeTools().DeactivateMode(FBuiltinEditorModes::EM_Foliage);
+		}
+
+		// Change out of mesh paint mode when opening a new map.
+		if (GLevelEditorModeTools().IsModeActive(FBuiltinEditorModes::EM_MeshPaint))
+		{
+			GLevelEditorModeTools().DeactivateMode(FBuiltinEditorModes::EM_MeshPaint);
+		}
+
+		return GEditor->NewMap();
+	}
+
+	/**
+	* Converts a package path to an asset path
+	*
+	* @param PackagePath - The package path to convert
+	*/
+	static FString ConvertPackagePathToAssetPath(const FString& PackagePath);
+
+	/**
+	* Imports an object using a given factory
+	*
+	* @param ImportFactory - The factory to use to import the object
+	* @param ObjectName - The name of the object to create
+	* @param PackagePath - The full path of the package file to create
+	* @param ImportPath - The path to the object to import
+	*/
+	static UObject* ImportAssetUsingFactory(UFactory* ImportFactory, const FString& ObjectName, const FString& PackageName, const FString& ImportPath);
+
+	/**
+	* Nulls out references to a given object
+	*
+	* @param InObject - Object to null references to
+	*/
+	static void NullReferencesToObject(UObject* InObject);
+
+	/**
+	* gets a factory class based off an asset file extension
+	*
+	* @param AssetExtension - The file extension to use to find a supporting UFactory
+	*/
+	static UClass* GetFactoryClassForType(const FString& AssetExtension);
+
+	/**
+	* Applies settings to an object by finding UProperties by name and calling ImportText
+	*
+	* @param InObject - The object to search for matching properties
+	* @param PropertyChain - The list UProperty names recursively to search through
+	* @param Value - The value to import on the found property
+	*/
+	static void ApplyCustomFactorySetting(UObject* InObject, TArray<FString>& PropertyChain, const FString& Value);
+
+	/**
+	* Applies the custom factory settings
+	*
+	* @param InFactory - The factory to apply custom settings to
+	* @param FactorySettings - An array of custom settings to apply to the factory
+	*/
+	static void ApplyCustomFactorySettings(UFactory* InFactory, const TArray<FImportFactorySettingValues>& FactorySettings);
+}
+
 
 //////////////////////////////////////////////////////////////////////////
 //Common latent commands used for automated editor testing.
diff --git a/Engine/Source/Runtime/AutomationWorker/Private/AutomationWorkerModule.cpp b/Engine/Source/Runtime/AutomationWorker/Private/AutomationWorkerModule.cpp
--- a/Engine/Source/Runtime/AutomationWorker/Private/AutomationWorkerModule.cpp
+++ b/Engine/Source/Runtime/AutomationWorker/Private/AutomationWorkerModule.cpp
@@ -504,7 +504,14 @@ static struct FQueueTests
 			if (CommandLine.Contains(TEXT("-AutomationTests=")))
 			{
 				//Let the current callstack complete then quit
-				GEngine->DeferredCommands.Add(TEXT("QUIT"));
+				if (GEngine->IsEditor())
+				{
+					GEngine->DeferredCommands.Add(TEXT("QUIT_EDITOR"));
+				}
+				else
+				{
+					GEngine->DeferredCommands.Add(TEXT("QUIT"));
+				}
 			}
 		}
 	}
@@ -564,6 +571,11 @@ bool DirectAutomationCommand(const TCHAR* Cmd, FOutputDevice* Ar = GLog)
 		else if (FParse::Command(&TempCmd, TEXT("CommandLineTests")))
 		{
 			TArray<FString> TestNames;
+
+			const bool bSkipScreenshots = FParse::Param(FCommandLine::Get(), TEXT("NoScreenshots"));
+			const bool bFullSizeScreenshots = FParse::Param(FCommandLine::Get(), TEXT("FullSizeScreenshots"));
+			FAutomationTestFramework::GetInstance().SetScreenshotOptions(!bSkipScreenshots, bFullSizeScreenshots);
+
 			if (GenerateTestNamesFromCommandLine(TempCmd, TestNames))
 			{
 				//Get the number of times to loop
diff --git a/Engine/Source/Runtime/Engine/Classes/Tests/AutomationTestSettings.h b/Engine/Source/Runtime/Engine/Classes/Tests/AutomationTestSettings.h
--- a/Engine/Source/Runtime/Engine/Classes/Tests/AutomationTestSettings.h
+++ b/Engine/Source/Runtime/Engine/Classes/Tests/AutomationTestSettings.h
@@ -104,6 +104,176 @@ struct FEditorImportExportTestDefinition
 	TArray<FImportFactorySettingValues> FactorySettings;
 };
 
+/**
+* Holds settings for the asset import workflow test
+*/
+USTRUCT()
+struct FEditorImportWorkflowDefinition
+{
+	GENERATED_USTRUCT_BODY()
+
+	/* The file to import */
+	UPROPERTY(config, EditAnywhere, Category = Automation, meta = (FilePathFilter = "*"))
+	FFilePath ImportFilePath;
+
+	/* Settings for the import factory */
+	UPROPERTY(config, EditAnywhere, Category = Automation)
+	TArray<FImportFactorySettingValues> FactorySettings;
+};
+
+/**
+* Holds settings for the import workflow stage of the build promotion test
+*/
+USTRUCT()
+struct FBuildPromotionImportWorkflowSettings
+{
+	GENERATED_USTRUCT_BODY()
+
+	/* Import settings for the Diffuse texture */
+	UPROPERTY(config, EditAnywhere, Category = Automation)
+	FEditorImportWorkflowDefinition Diffuse;
+
+	/* Import settings for the Normalmap texture */
+	UPROPERTY(config, EditAnywhere, Category = Automation)
+	FEditorImportWorkflowDefinition Normal;
+
+	/* Import settings for the static mesh */
+	UPROPERTY(config, EditAnywhere, Category = Automation)
+	FEditorImportWorkflowDefinition StaticMesh;
+
+	/* Import settings for the static mesh to re-import */
+	UPROPERTY(config, EditAnywhere, Category = Automation)
+	FEditorImportWorkflowDefinition ReimportStaticMesh;
+
+	/* Import settings for the blend shape */
+	UPROPERTY(config, EditAnywhere, Category = Automation)
+	FEditorImportWorkflowDefinition BlendShapeMesh;
+
+	/* Import settings for the morph mesh */
+	UPROPERTY(config, EditAnywhere, Category = Automation)
+	FEditorImportWorkflowDefinition MorphMesh;
+
+	/* Import settings for the skeletal mesh */
+	UPROPERTY(config, EditAnywhere, Category = Automation)
+	FEditorImportWorkflowDefinition SkeletalMesh;
+
+	/* Import settings for the animation asset.  (Will automatically use the skeleton of the skeletal mesh above) */
+	UPROPERTY(config, EditAnywhere, Category = Automation)
+	FEditorImportWorkflowDefinition Animation;
+
+	/* Import settings for the sound */
+	UPROPERTY(config, EditAnywhere, Category = Automation)
+	FEditorImportWorkflowDefinition Sound;
+
+	/* Import settings for the surround sound (Select any of the channels.  It will auto import the rest)*/
+	UPROPERTY(config, EditAnywhere, Category = Automation)
+	FEditorImportWorkflowDefinition SurroundSound;
+
+	/* Import settings for any other assets you may want to import */
+	UPROPERTY(config, EditAnywhere, Category = Automation)
+	TArray<FEditorImportWorkflowDefinition> OtherAssetsToImport;
+};
+
+/**
+* Holds settings for the open assets stage of the build promotion test
+*/
+USTRUCT()
+struct FBuildPromotionOpenAssetSettings
+{
+	GENERATED_USTRUCT_BODY()
+
+	/* The blueprint asset to open */
+	UPROPERTY(config, EditAnywhere, Category = Automation, meta = (FilePathFilter = "uasset"))
+	FFilePath BlueprintAsset;
+
+	/* The material asset to open */
+	UPROPERTY(config, EditAnywhere, Category = Automation, meta = (FilePathFilter = "uasset"))
+	FFilePath MaterialAsset;
+	
+	/* The particle system asset to open */
+	UPROPERTY(config, EditAnywhere, Category = Automation, meta = (FilePathFilter = "uasset"))
+	FFilePath ParticleSystemAsset;
+	
+	/* The skeletal mesh asset to open */
+	UPROPERTY(config, EditAnywhere, Category = Automation, meta = (FilePathFilter = "uasset"))
+	FFilePath SkeletalMeshAsset;
+	
+	/* The static mesh asset to open */
+	UPROPERTY(config, EditAnywhere, Category = Automation, meta = (FilePathFilter = "uasset"))
+	FFilePath StaticMeshAsset;
+	
+	/* The texture asset to open */
+	UPROPERTY(config, EditAnywhere, Category = Automation, meta = (FilePathFilter = "uasset"))
+	FFilePath TextureAsset;
+};
+
+/**
+* Holds settings for the blueprint stage of the build promotion test
+*/
+USTRUCT()
+struct FBuildPromotionBlueprintSettings
+{
+	GENERATED_USTRUCT_BODY()
+
+	/** The starting mesh for the blueprint **/
+	UPROPERTY(EditAnywhere, Category = Automation, meta = (FilePathFilter = "uasset"))
+	FFilePath FirstMeshPath;
+
+	/** The mesh to set on the blueprint after the delay **/
+	UPROPERTY(EditAnywhere, Category = Automation, meta = (FilePathFilter = "uasset"))
+	FFilePath SecondMeshPath;
+};
+
+/**
+* Holds settings for the new project stage of the build promotion test
+*/
+USTRUCT()
+struct FBuildPromotionNewProjectSettings
+{
+	GENERATED_USTRUCT_BODY()
+
+	/** The path for the new project */
+	UPROPERTY(EditAnywhere, Category = Automation)
+	FDirectoryPath NewProjectFolderOverride;
+
+	/** The name of the project **/
+	UPROPERTY(EditAnywhere, Category = Automation)
+	FString NewProjectNameOverride;
+};
+
+/**
+* Holds settings for the editor build promotion test
+*/
+USTRUCT()
+struct FBuildPromotionTestSettings
+{
+	GENERATED_USTRUCT_BODY()
+
+	/** Default static mesh asset to apply materials to **/
+	UPROPERTY(EditAnywhere, Category = Automation)
+	FFilePath DefaultStaticMeshAsset;
+
+	/** Import workflow settings **/
+	UPROPERTY(EditAnywhere, Category = Automation)
+	FBuildPromotionImportWorkflowSettings	ImportWorkflow;
+
+	/** Open assets settings **/
+	UPROPERTY(EditAnywhere, Category = Automation)
+	FBuildPromotionOpenAssetSettings	OpenAssets;
+
+	/** Blueprint settings **/
+	UPROPERTY(EditAnywhere, Category = Automation)
+	FBuildPromotionBlueprintSettings	BlueprintSettings;
+
+	/** New project settings **/
+	UPROPERTY(EditAnywhere, Category = Automation)
+	FBuildPromotionNewProjectSettings NewProjectSettings;
+
+	/** Material to modify for the content browser step **/
+	UPROPERTY(EditAnywhere, Category = Automation)
+	FFilePath	SourceControlMaterial;
+};
+
 /**
  * Implements the Editor's user settings.
  */
@@ -155,4 +325,10 @@ public:
 	 */
 	UPROPERTY(EditAnywhere, config, Category = Automation)
 	TArray<FEditorImportExportTestDefinition> ImportExportTestDefinitions;
+
+	/**
+	* Editor build promotion test settings
+	*/
+	UPROPERTY(EditAnywhere, config, Category = Automation)
+	FBuildPromotionTestSettings BuildPromotionTest;
 };
