Update to changelist 2246761 -  adds range clamping to the maximum loop iteration setting and updates the script core on change in the editor.
#TTP 340083 - +UE4: K2: Add editor exposed configuration of script max loop iterations count
#Branch UE4
#Proj UnrealEd

diff --git a/Engine/Source/Editor/UnrealEd/Private/Editor.cpp b/Engine/Source/Editor/UnrealEd/Private/Editor.cpp
--- a/Engine/Source/Editor/UnrealEd/Private/Editor.cpp
+++ b/Engine/Source/Editor/UnrealEd/Private/Editor.cpp
@@ -1689,6 +1689,15 @@ void UEditorEngine::PostEditChangeProperty(FPropertyChangedEvent& PropertyChange
 {
 	// Propagate the callback up to the superclass.
 	Super::PostEditChangeProperty(PropertyChangedEvent);
+
+	const FName PropertyName = PropertyChangedEvent.Property ? PropertyChangedEvent.Property->GetFName() : NAME_None;
+
+	if( PropertyName == FName( TEXT( "MaximumLoopIterationCount" )))
+	{
+		// Clamp to a reasonable range and feed the new value to the script core
+		MaximumLoopIterationCount = FMath::Clamp( MaximumLoopIterationCount, 100, 10000000 );
+		FBlueprintCoreDelegates::SetScriptMaximumLoopIterations( MaximumLoopIterationCount );
+	}
 }
 
 void UEditorEngine::Cleanse( bool ClearSelection, bool Redraw, const FText& TransReset )
diff --git a/Engine/Source/Editor/UnrealEd/Private/PlayLevel.cpp b/Engine/Source/Editor/UnrealEd/Private/PlayLevel.cpp
--- a/Engine/Source/Editor/UnrealEd/Private/PlayLevel.cpp
+++ b/Engine/Source/Editor/UnrealEd/Private/PlayLevel.cpp
@@ -1613,9 +1613,6 @@ void UEditorEngine::PlayInEditor( UWorld* InWorld, bool bInSimulateInEditor )
 	ULevelEditorPlaySettings const* EditorPlayInSettings = GetDefault<ULevelEditorPlaySettings>();
 	check(EditorPlayInSettings);
 
-	// Update Script Maximum loop iteration count
-	FBlueprintCoreDelegates::SetScriptMaximumLoopIterations( GEngine->MaximumLoopIterationCount );
-
 	// Prompt the user to compile any dirty Blueprints before PIE can occur.
 	TArray< UBlueprint* > ErrorBlueprintList;
 	bool bAnyBlueprintsDirty = false;
