Fix in FBlueprintEditorUtils::RefreshExternalBlueprintDependencyNodes
http://udn.unrealengine.com/questions/217137/issues-with-using-blueprint-interfaces-with-the-ou.html

diff --git a/Engine/Source/Editor/UnrealEd/Private/Kismet2/BlueprintEditorUtils.cpp b/Engine/Source/Editor/UnrealEd/Private/Kismet2/BlueprintEditorUtils.cpp
--- a/Engine/Source/Editor/UnrealEd/Private/Kismet2/BlueprintEditorUtils.cpp
+++ b/Engine/Source/Editor/UnrealEd/Private/Kismet2/BlueprintEditorUtils.cpp
@@ -345,7 +345,13 @@ void FBlueprintEditorUtils::RefreshExternalBlueprintDependencyNodes(UBlueprint*
 			{
 				for (UStruct* Struct : Dependencies)
 				{
-					if (Struct->IsChildOf(RefreshOnlyChild))
+					bool bShouldRefresh = Struct->IsChildOf(RefreshOnlyChild);
+					if (!bShouldRefresh)
+					{
+						UClass* OwnerClass = Struct->GetOwnerClass();
+						bShouldRefresh |= OwnerClass && OwnerClass->IsChildOf(RefreshOnlyChild);
+					}
+					if (bShouldRefresh)
 					{
 						//@todo:  Do we really need per-schema refreshing?
 						const UEdGraphSchema* Schema = Node->GetGraph()->GetSchema();
