Fix crash importing PCX files

diff --git a/Engine/Source/Editor/UnrealEd/Private/Factories/EditorFactories.cpp b/Engine/Source/Editor/UnrealEd/Private/Factories/EditorFactories.cpp
--- a/Engine/Source/Editor/UnrealEd/Private/Factories/EditorFactories.cpp
+++ b/Engine/Source/Editor/UnrealEd/Private/Factories/EditorFactories.cpp
@@ -3754,7 +3754,9 @@ UTexture* UTextureFactory::ImportTexture(UClass* Class, UObject* InParent, FName
 				uint8* PCXPalette = (uint8 *)(BufferEnd - 256 * 3);
 				TArray<FColor>	Palette;
 				for(uint32 i=0; i<256; i++ )
+				{
 					Palette.Add(FColor(PCXPalette[i*3+0],PCXPalette[i*3+1],PCXPalette[i*3+2],i == 0 ? 0 : 255));
+				}
 
 				// Import it.
 				FColor* DestEnd	= DestPtr + NewU * NewV;
@@ -3766,9 +3768,11 @@ UTexture* UTextureFactory::ImportTexture(UClass* Class, UObject* InParent, FName
 					{
 						uint32 RunLength = Color & 0x3f;
 						Color = *Buffer++;
-						checkf( (DestPtr+RunLength)<DestEnd, TEXT("RLE going off the end of buffer") );
+					
 						for(uint32 Index = 0;Index < RunLength;Index++)
+						{
 							*DestPtr++ = Palette[Color];
+						}
 					}
 					else *DestPtr++ = Palette[Color];
 				}
