Rearranging some of the editor automation tests.
Creating tests that open all of a certain asset. Lots of tests here.  This is current set to ATF_RequireUser as I want to get QA game cleaned up before having going fully automated.
The Open asset tests are still a work in progress.

diff --git a/Engine/Source/Editor/UnrealEd/Private/Tests/AutomationEditorCommon.cpp b/Engine/Source/Editor/UnrealEd/Private/Tests/AutomationEditorCommon.cpp
--- a/Engine/Source/Editor/UnrealEd/Private/Tests/AutomationEditorCommon.cpp
+++ b/Engine/Source/Editor/UnrealEd/Private/Tests/AutomationEditorCommon.cpp
@@ -2,6 +2,10 @@
 
 #include "UnrealEd.h"
 #include "AutomationEditorCommon.h"
+#include "AssetEditorManager.h"
+#include "ModuleManager.h"
+#include "LevelEditor.h"
+#include "ModuleManager.h"
 
 DEFINE_LOG_CATEGORY_STATIC(LogAutomationEditorCommon, Log, All);
 
@@ -291,5 +295,54 @@ bool FUndoRedoCommand::Update()
 		GEditor->RedoTransaction();
 	}
 
+	return true;
+}
+
+/**
+* Open editor for a particular asset
+*/
+bool FOpenEditorForAssetCommand::Update()
+{
+	UObject* Object = StaticLoadObject(UObject::StaticClass(), NULL, *AssetName);
+	if (Object)
+	{
+		FAssetEditorManager::Get().OpenEditorForAsset(Object);
+		UE_LOG(LogEditorAutomationTests, Log, TEXT("Verified asset editor for: %s."), *AssetName);
+	}
+	else
+	{
+		UE_LOG(LogEditorAutomationTests, Error, TEXT("Failed to find object: %s."), *AssetName);
+	}
+	return true;
+}
+
+/**
+* Close all sub-editors
+*/
+bool FCloseAllAssetEditorsCommand::Update()
+{
+	FAssetEditorManager::Get().CloseAllAssetEditors();
+
+	return true;
+}
+
+/**
+* Start PIE session
+*/
+bool FStartPIECommand::Update()
+{
+	FLevelEditorModule& LevelEditorModule = FModuleManager::Get().GetModuleChecked<FLevelEditorModule>(TEXT("LevelEditor"));
+	TSharedPtr<class ILevelViewport> ActiveLevelViewport = LevelEditorModule.GetFirstActiveViewport();
+
+	GUnrealEd->RequestPlaySession(false, ActiveLevelViewport, bSimulateInEditor, NULL, NULL, -1, false);
+	return true;
+}
+
+/**
+* End PlayMap session
+*/
+bool FEndPlayMapCommand::Update()
+{
+	GUnrealEd->RequestEndPlayMap();
 	return true;
 }
\ No newline at end of file
diff --git a/Engine/Source/Editor/UnrealEd/Private/Tests/EditorAutomationTests.cpp b/Engine/Source/Editor/UnrealEd/Private/Tests/EditorAutomationTests.cpp
--- a/Engine/Source/Editor/UnrealEd/Private/Tests/EditorAutomationTests.cpp
+++ b/Engine/Source/Editor/UnrealEd/Private/Tests/EditorAutomationTests.cpp
@@ -26,64 +26,6 @@
 
 //DEFINE_LOG_CATEGORY_STATIC(LogEditorAutomationTests, Log, All);
 
-/**
- * Start PIE session
- */
-DEFINE_LATENT_AUTOMATION_COMMAND_ONE_PARAMETER(FStartPIECommand, bool, bSimulateInEditor);
-
-bool FStartPIECommand::Update()
-{
-	FLevelEditorModule& LevelEditorModule = FModuleManager::Get().GetModuleChecked<FLevelEditorModule>( TEXT("LevelEditor") );
-	TSharedPtr<class ILevelViewport> ActiveLevelViewport = LevelEditorModule.GetFirstActiveViewport();
-
-	GUnrealEd->RequestPlaySession( false, ActiveLevelViewport, bSimulateInEditor, NULL, NULL, -1, false );
-	return true;
-}
-
-/**
- * End PlayMap session
- */
-DEFINE_LATENT_AUTOMATION_COMMAND(FEndPlayMapCommand);
-
-bool FEndPlayMapCommand::Update()
-{
-	GUnrealEd->RequestEndPlayMap();
-	return true;
-}
-
-/**
- * Open editor for a particular asset
- */
-DEFINE_LATENT_AUTOMATION_COMMAND_ONE_PARAMETER(FOpenEditorForAssetCommand, FString, AssetName);
-
-bool FOpenEditorForAssetCommand::Update()
-{
-	UObject* Object = StaticLoadObject(UObject::StaticClass(), NULL, *AssetName);
-	if (Object)
-	{
-		FAssetEditorManager::Get().OpenEditorForAsset(Object);
-		UE_LOG(LogEditorAutomationTests, Log, TEXT("Verified asset editor for: %s."), *AssetName);
-	}
-	else
-	{
-		UE_LOG(LogEditorAutomationTests, Error, TEXT("Failed to find object: %s."), *AssetName);
-	}
-
-	return true;
-}
-
-/**
- * Close all asset editors
- */
-DEFINE_LATENT_AUTOMATION_COMMAND(FCloseAllAssetEditorsCommand);
-
-bool FCloseAllAssetEditorsCommand::Update()
-{
-	FAssetEditorManager::Get().CloseAllAssetEditors();
-
-	return true;
-}
-
 /**
 * Change the attributes for a point light in the level.
 */
@@ -157,68 +99,6 @@ void TakeLatentAutomationScreenshot(struct WindowScreenshotParameters Screenshot
 	}
 }
 
-//////////////////////////////////////////////////////////////////////////
-// FEditorAutomationTestUtilities
-
-class FEditorAutomationTestUtilities
-{
-public:
-	/**
-	 * Loads the map specified by an automation test
-	 * 
-	 * @param MapName - Map to load
-	 */
-	static void LoadMap(const FString& MapName)
-	{
-		bool bLoadAsTemplate = false;
-		bool bShowProgress = false;
-		FEditorFileUtils::LoadMap(MapName, bLoadAsTemplate, bShowProgress);
-	}
-
-	/**
-	 * Run PIE
-	 */
-	static void RunPIE()
-	{
-		bool bInSimulateInEditor = true;
-		//once in the editor
-		ADD_LATENT_AUTOMATION_COMMAND(FStartPIECommand(true));
-		ADD_LATENT_AUTOMATION_COMMAND(FWaitLatentCommand(3.0f));
-		ADD_LATENT_AUTOMATION_COMMAND(FEndPlayMapCommand());
-
-		//wait between tests
-		ADD_LATENT_AUTOMATION_COMMAND(FWaitLatentCommand(1.0f));
-
-		//once not in the editor
-		ADD_LATENT_AUTOMATION_COMMAND(FStartPIECommand(false));
-		ADD_LATENT_AUTOMATION_COMMAND(FWaitLatentCommand(3.0f));
-		ADD_LATENT_AUTOMATION_COMMAND(FEndPlayMapCommand());
-	}
-
-	static void CollectTestsByClass(UClass * Class, TArray<FString>& OutBeautifiedNames, TArray <FString>& OutTestCommands) 
-	{
-		FAssetRegistryModule& AssetRegistryModule = FModuleManager::LoadModuleChecked<FAssetRegistryModule>(TEXT("AssetRegistry"));
-		TArray<FAssetData> ObjectList;
-		AssetRegistryModule.Get().GetAssetsByClass(Class->GetFName(), ObjectList);
-
-		for (auto ObjIter=ObjectList.CreateConstIterator(); ObjIter; ++ObjIter)
-		{
-			const FAssetData & Asset = *ObjIter;
-			FString Filename = Asset.ObjectPath.ToString();
-			//convert to full paths
-			Filename = FPackageName::LongPackageNameToFilename(Filename);
-			if (FAutomationTestFramework::GetInstance().ShouldTestContent(Filename))
-			{
-				FString BeautifiedFilename = Asset.AssetName.ToString();
-				OutBeautifiedNames.Add(BeautifiedFilename);
-				OutTestCommands.Add(Asset.ObjectPath.ToString());
-			}
-		}
-	}
-
-	
-};
-
 //////////////////////////////////////////////////////////////////////////
 
 /**
@@ -1389,122 +1269,6 @@ bool FConvertToValidation::RunTest(const FString& Parameters)
 }
 
 
-
-//////////////////////////////////////////////////////////////////////////
-/**
- * QA Test to open editor windows
- */
-IMPLEMENT_COMPLEX_AUTOMATION_TEST(FOpenAssetEditors, "QA.Open Asset Editors", EAutomationTestFlags::ATF_Editor);
-
-void FOpenAssetEditors::GetTests(TArray<FString>& OutBeautifiedNames, TArray <FString>& OutTestCommands) const
-{
-	UAutomationTestSettings const* AutomationTestSettings = GetDefault<UAutomationTestSettings>();
-	check(AutomationTestSettings);
-
-	bool bUnAttended = FApp::IsUnattended();
-
-	TArray<FString> AssetNames;
-	for (auto Iter = AutomationTestSettings->TestAssetsToOpen.CreateConstIterator(); Iter; ++Iter)
-	{
-		if (Iter->bSkipTestWhenUnAttended && bUnAttended)
-		{
-			continue;
-		}
-
-		if ( Iter->AssetToOpen.FilePath.Len() > 0 )
-		{
-			AssetNames.AddUnique(Iter->AssetToOpen.FilePath);
-		}
-	}
-
-	//Location of the engine folder
-	FString EngineContentFolderLocation = FPaths::ConvertRelativePathToFull(*FPaths::EngineContentDir());
-	//Put the Engine Content Folder Location into an array.  
-	TArray<FString> EngineContentFolderLocationArray;
-	EngineContentFolderLocation.ParseIntoArray(&EngineContentFolderLocationArray, TEXT("/"), true);
-
-	for (int32 i = 0; i < AssetNames.Num(); ++i)
-	{
-		//true means that the life is located in the Engine/Content folder.
-		bool bFileIsLocatedInTheEngineContentFolder = true;
-		
-		//Get the location of the asset that is being opened.
-		FString Filename = FPaths::ConvertRelativePathToFull(AssetNames[i]);
-
-		//Put File location into an array.  
-		TArray<FString> FilenameArray;
-		Filename.ParseIntoArray(&FilenameArray, TEXT("/"), true);
-
-		//Loop through the location array's and compare each element.  
-		//The loop runs while the index is less than the number of elements in the EngineContentFolderLocation array.
-		//If the elements are the same when the counter is up then it is assumed that the asset file is in the engine content folder. 
-		//Otherwise we'll it's in the game content folder.
-		for (int32 index = 0; index < EngineContentFolderLocationArray.Num(); index++)
-		{
-			if ((FilenameArray[index] != EngineContentFolderLocationArray[index]))
-			{
-				//If true it will proceed to add the asset to the Open Asset Editor test list.
-				//This will be false if the asset is on a different drive.
-				if (FPaths::MakePathRelativeTo(Filename, *FPaths::GameContentDir()))
-				{
-					FString ShortName = FPaths::GetBaseFilename(Filename);
-					FString PathName = FPaths::GetPath(Filename);
-					OutBeautifiedNames.Add(ShortName);
-					FString AssetName = FString::Printf(TEXT("/Game/%s/%s.%s"), *PathName, *ShortName, *ShortName);
-					OutTestCommands.Add(AssetName);
-					bFileIsLocatedInTheEngineContentFolder = false;
-					break;
-				}
-				else
-				{
-					UE_LOG(LogEditorAutomationTests, Error, TEXT("Invalid asset path: %s."), *Filename);
-					bFileIsLocatedInTheEngineContentFolder = false;
-					break;
-				}
-			}
-		}
-		//If true then the asset is in the Engine/Content folder and not in the Game/Content folder.
-		if (bFileIsLocatedInTheEngineContentFolder)
-		{
-			//If true it will proceed to add the asset to the Open Asset Editor test list.
-			//This will be false if the asset is on a different drive.
-			if (FPaths::MakePathRelativeTo(Filename, *FPaths::EngineContentDir()))
-			{
-				FString ShortName = FPaths::GetBaseFilename(Filename);
-				FString PathName = FPaths::GetPath(Filename);
-				OutBeautifiedNames.Add(ShortName);
-				FString AssetName = FString::Printf(TEXT("/Engine/%s/%s.%s"), *PathName, *ShortName, *ShortName);
-				OutTestCommands.Add(AssetName);
-			}
-			else
-			{
-				UE_LOG(LogEditorAutomationTests, Error, TEXT("Invalid asset path: %s."), *Filename);
-			}
-		}
-	}
-}
-
-bool FOpenAssetEditors::RunTest(const FString& Parameters)
-{
-	//start with all editors closed
-	FAssetEditorManager::Get().CloseAllAssetEditors();
-
-	// below is all latent action, so before sending there, verify the asset exists
-	UObject* Object = StaticLoadObject(UObject::StaticClass(), NULL, *Parameters);
-	if (!Object)
-	{
-		UE_LOG(LogEditorAutomationTests, Error, TEXT("Failed to find object: %s."), *Parameters);
-		return false;
-	}
-
-	ADD_LATENT_AUTOMATION_COMMAND(FOpenEditorForAssetCommand(*Parameters));
-	ADD_LATENT_AUTOMATION_COMMAND(FWaitLatentCommand(0.5f));
-	ADD_LATENT_AUTOMATION_COMMAND(FCloseAllAssetEditorsCommand());
-
-	return true;
-}
-
-
 //////////////////////////////////////////////////////////////////////////
 /**
  * QA Static Mesh Regression Testing
diff --git a/Engine/Source/Editor/UnrealEd/Private/Tests/EditorOpenAssetAutomationTests.cpp b/Engine/Source/Editor/UnrealEd/Private/Tests/EditorOpenAssetAutomationTests.cpp
new file mode 100644
--- /dev/null
+++ b/Engine/Source/Editor/UnrealEd/Private/Tests/EditorOpenAssetAutomationTests.cpp
@@ -0,0 +1,640 @@
+// Copyright 1998-2014 Epic Games, Inc. All Rights Reserved.
+
+#include "UnrealEd.h"
+#include "Tests/AutomationTestSettings.h"
+#include "AutomationEditorCommon.h"
+#include "AutomationTest.h"
+#include "AssertionMacros.h"
+#include "AutomationCommon.h"
+#include "AssetEditorManager.h"
+
+
+
+
+bool FOpenActualAssetEditors(const FString& Parameters)
+{
+	//start with all editors closed
+	FAssetEditorManager::Get().CloseAllAssetEditors();
+
+	// below is all latent action, so before sending there, verify the asset exists
+	UObject* Object = StaticLoadObject(UObject::StaticClass(), NULL, *Parameters);
+	if (!Object)
+	{
+		UE_LOG(LogEditorAutomationTests, Error, TEXT("Failed to find object: %s."), *Parameters);
+		return false;
+	}
+
+	ADD_LATENT_AUTOMATION_COMMAND(FOpenEditorForAssetCommand(*Parameters));
+	ADD_LATENT_AUTOMATION_COMMAND(FWaitLatentCommand(0.5f));
+	ADD_LATENT_AUTOMATION_COMMAND(FCloseAllAssetEditorsCommand());
+
+	return true;
+}
+//////////////////////////////////////////////////////////////////////////
+/**
+* Test to open the sub editor windows for a specified list of assets.
+* This list can be setup in the Editor Preferences window within the editor or the DefaultEngine.ini file for that particular project.
+*/
+IMPLEMENT_COMPLEX_AUTOMATION_TEST(FOpenAssetEditors, "QA.Open Asset Editors", EAutomationTestFlags::ATF_Editor);
+
+void FOpenAssetEditors::GetTests(TArray<FString>& OutBeautifiedNames, TArray <FString>& OutTestCommands) const
+{
+	UAutomationTestSettings const* AutomationTestSettings = GetDefault<UAutomationTestSettings>();
+	check(AutomationTestSettings);
+
+	bool bUnAttended = FApp::IsUnattended();
+
+	TArray<FString> AssetNames;
+	for (auto Iter = AutomationTestSettings->TestAssetsToOpen.CreateConstIterator(); Iter; ++Iter)
+	{
+		if (Iter->bSkipTestWhenUnAttended && bUnAttended)
+		{
+			continue;
+		}
+
+		if (Iter->AssetToOpen.FilePath.Len() > 0)
+		{
+			AssetNames.AddUnique(Iter->AssetToOpen.FilePath);
+		}
+	}
+
+	//Location of the engine folder
+	FString EngineContentFolderLocation = FPaths::ConvertRelativePathToFull(*FPaths::EngineContentDir());
+	//Put the Engine Content Folder Location into an array.  
+	TArray<FString> EngineContentFolderLocationArray;
+	EngineContentFolderLocation.ParseIntoArray(&EngineContentFolderLocationArray, TEXT("/"), true);
+
+	for (int32 i = 0; i < AssetNames.Num(); ++i)
+	{
+		//true means that the life is located in the Engine/Content folder.
+		bool bFileIsLocatedInTheEngineContentFolder = true;
+
+		//Get the location of the asset that is being opened.
+		FString Filename = FPaths::ConvertRelativePathToFull(AssetNames[i]);
+
+		//Put File location into an array.  
+		TArray<FString> FilenameArray;
+		Filename.ParseIntoArray(&FilenameArray, TEXT("/"), true);
+
+		//Loop through the location array's and compare each element.  
+		//The loop runs while the index is less than the number of elements in the EngineContentFolderLocation array.
+		//If the elements are the same when the counter is up then it is assumed that the asset file is in the engine content folder. 
+		//Otherwise we'll assume it's in the game content folder.
+		for (int32 index = 0; index < EngineContentFolderLocationArray.Num(); index++)
+		{
+			if ((FilenameArray[index] != EngineContentFolderLocationArray[index]))
+			{
+				//If true it will proceed to add the asset to the Open Asset Editor test list.
+				//This will be false if the asset is on a different drive.
+				if (FPaths::MakePathRelativeTo(Filename, *FPaths::GameContentDir()))
+				{
+					FString ShortName = FPaths::GetBaseFilename(Filename);
+					FString PathName = FPaths::GetPath(Filename);
+					OutBeautifiedNames.Add(ShortName);
+					FString AssetName = FString::Printf(TEXT("/Game/%s/%s.%s"), *PathName, *ShortName, *ShortName);
+					OutTestCommands.Add(AssetName);
+					bFileIsLocatedInTheEngineContentFolder = false;
+					break;
+				}
+				else
+				{
+					UE_LOG(LogEditorAutomationTests, Error, TEXT("Invalid asset path: %s."), *Filename);
+					bFileIsLocatedInTheEngineContentFolder = false;
+					break;
+				}
+			}
+		}
+		//If true then the asset is in the Engine/Content folder and not in the Game/Content folder.
+		if (bFileIsLocatedInTheEngineContentFolder)
+		{
+			//If true it will proceed to add the asset to the Open Asset Editor test list.
+			//This will be false if the asset is on a different drive.
+			if (FPaths::MakePathRelativeTo(Filename, *FPaths::EngineContentDir()))
+			{
+				FString ShortName = FPaths::GetBaseFilename(Filename);
+				FString PathName = FPaths::GetPath(Filename);
+				OutBeautifiedNames.Add(ShortName);
+				FString AssetName = FString::Printf(TEXT("/Engine/%s/%s.%s"), *PathName, *ShortName, *ShortName);
+				OutTestCommands.Add(AssetName);
+			}
+			else
+			{
+				UE_LOG(LogEditorAutomationTests, Error, TEXT("Invalid asset path: %s."), *Filename);
+			}
+		}
+	}
+}
+
+bool FOpenAssetEditors::RunTest(const FString& Parameters)
+{
+	bool bDidTheTestPass = FOpenActualAssetEditors(Parameters);
+	CollectGarbage(GARBAGE_COLLECTION_KEEPFLAGS);
+	return bDidTheTestPass;
+}
+
+
+//////////////////////////////////////////////////////////////////////////
+/**
+* This test opens each AimOffsetBlendSpace into its sub-editor.
+*/
+IMPLEMENT_COMPLEX_AUTOMATION_TEST(FOpenAimOffsetBlendSpaceAssetEditors, "Editor.Content.Animation.BlendSpaces.Open AimOffsetBlendSpace Assets", (EAutomationTestFlags::ATF_Editor | EAutomationTestFlags::ATF_RequiresUser));
+
+void FOpenAimOffsetBlendSpaceAssetEditors::GetTests(TArray<FString>& OutBeautifiedNames, TArray<FString>& OutTestCommands) const
+{
+	//This grabs each AimOffsetBlendSpace asset in the Game/Content directory
+	FEditorAutomationTestUtilities::CollectGameContentTestsByClass(UAimOffsetBlendSpace::StaticClass(), OutBeautifiedNames, OutTestCommands);
+}
+
+bool FOpenAimOffsetBlendSpaceAssetEditors::RunTest(const FString& Parameters)
+{
+	bool bDidTheTestPass = FOpenActualAssetEditors(Parameters);
+	CollectGarbage(GARBAGE_COLLECTION_KEEPFLAGS);
+	return bDidTheTestPass;
+}
+
+//////////////////////////////////////////////////////////////////////////
+/**
+* This test opens each AimOffsetBlendSpace1D into its sub-editor.
+*/
+IMPLEMENT_COMPLEX_AUTOMATION_TEST(FOpenAimOffsetBlendSpace1DAssetEditors, "Editor.Content.Animation.BlendSpaces.Open AimOffsetBlendSpace1D Assets", (EAutomationTestFlags::ATF_Editor | EAutomationTestFlags::ATF_RequiresUser));
+
+void FOpenAimOffsetBlendSpace1DAssetEditors::GetTests(TArray<FString>& OutBeautifiedNames, TArray<FString>& OutTestCommands) const
+{
+	//This grabs each AimOffsetBlendSpace1D asset in the Game/Content directory
+	FEditorAutomationTestUtilities::CollectGameContentTestsByClass(UAimOffsetBlendSpace1D::StaticClass(), OutBeautifiedNames, OutTestCommands);
+}
+
+bool FOpenAimOffsetBlendSpace1DAssetEditors::RunTest(const FString& Parameters)
+{
+	bool bDidTheTestPass = FOpenActualAssetEditors(Parameters);
+	CollectGarbage(GARBAGE_COLLECTION_KEEPFLAGS);
+	return bDidTheTestPass;
+}
+
+//////////////////////////////////////////////////////////////////////////
+/**
+* This test opens each AnimMontage into its sub-editor.
+*/
+IMPLEMENT_COMPLEX_AUTOMATION_TEST(FOpenAnimMontageAssetEditors, "Editor.Content.Animation.Open AnimMontage Assets", (EAutomationTestFlags::ATF_Editor | EAutomationTestFlags::ATF_RequiresUser));
+
+void FOpenAnimMontageAssetEditors::GetTests(TArray<FString>& OutBeautifiedNames, TArray<FString>& OutTestCommands) const
+{
+	//This grabs each AnimMontage asset in the Game/Content directory
+	FEditorAutomationTestUtilities::CollectGameContentTestsByClass(UAnimMontage::StaticClass(), OutBeautifiedNames, OutTestCommands);
+}
+
+bool FOpenAnimMontageAssetEditors::RunTest(const FString& Parameters)
+{
+	bool bDidTheTestPass = FOpenActualAssetEditors(Parameters);
+	CollectGarbage(GARBAGE_COLLECTION_KEEPFLAGS);
+	return bDidTheTestPass;
+}
+
+//////////////////////////////////////////////////////////////////////////
+/**
+* This test opens each BlendSpace into its sub-editor.
+*/
+IMPLEMENT_COMPLEX_AUTOMATION_TEST(FOpenBlendSpaceAssetEditors, "Editor.Content.Animation.BlendSpaces.Open BlendSpace Assets", (EAutomationTestFlags::ATF_Editor | EAutomationTestFlags::ATF_RequiresUser));
+
+void FOpenBlendSpaceAssetEditors::GetTests(TArray<FString>& OutBeautifiedNames, TArray<FString>& OutTestCommands) const
+{
+	//This grabs each BlendSpace asset in the Game/Content directory
+	FEditorAutomationTestUtilities::CollectGameContentTestsByClass(UBlendSpace::StaticClass(), OutBeautifiedNames, OutTestCommands);
+}
+
+bool FOpenBlendSpaceAssetEditors::RunTest(const FString& Parameters)
+{
+	bool bDidTheTestPass = FOpenActualAssetEditors(Parameters);
+	CollectGarbage(GARBAGE_COLLECTION_KEEPFLAGS);
+	return bDidTheTestPass;
+}
+
+//////////////////////////////////////////////////////////////////////////
+/**
+* This test opens each BlendSpace1D into its sub-editor.
+*/
+IMPLEMENT_COMPLEX_AUTOMATION_TEST(FOpenBlendSpace1DAssetEditors, "Editor.Content.Animation.BlendSpaces.Open BlendSpace1D Assets", (EAutomationTestFlags::ATF_Editor | EAutomationTestFlags::ATF_RequiresUser));
+
+void FOpenBlendSpace1DAssetEditors::GetTests(TArray<FString>& OutBeautifiedNames, TArray<FString>& OutTestCommands) const
+{
+	//This grabs each BlendSpace1D asset in the Game/Content directory
+	FEditorAutomationTestUtilities::CollectGameContentTestsByClass(UBlendSpace1D::StaticClass(), OutBeautifiedNames, OutTestCommands);
+}
+
+bool FOpenBlendSpace1DAssetEditors::RunTest(const FString& Parameters)
+{
+	bool bDidTheTestPass = FOpenActualAssetEditors(Parameters);
+	CollectGarbage(GARBAGE_COLLECTION_KEEPFLAGS);
+	return bDidTheTestPass;
+}
+
+//////////////////////////////////////////////////////////////////////////
+/**
+* This test opens each DialogueVoice into its sub-editor.
+*/
+IMPLEMENT_COMPLEX_AUTOMATION_TEST(FOpenDialogueVoiceAssetEditors, "Editor.Content.Sounds.Open Dialogue Voice Assets", (EAutomationTestFlags::ATF_Editor | EAutomationTestFlags::ATF_RequiresUser));
+
+void FOpenDialogueVoiceAssetEditors::GetTests(TArray<FString>& OutBeautifiedNames, TArray<FString>& OutTestCommands) const
+{
+	//This grabs each DialogueVoice asset in the Game/Content directory
+	FEditorAutomationTestUtilities::CollectGameContentTestsByClass(UDialogueVoice::StaticClass(), OutBeautifiedNames, OutTestCommands);
+}
+
+bool FOpenDialogueVoiceAssetEditors::RunTest(const FString& Parameters)
+{
+	bool bDidTheTestPass = FOpenActualAssetEditors(Parameters);
+	CollectGarbage(GARBAGE_COLLECTION_KEEPFLAGS);
+	return bDidTheTestPass;
+}
+
+//////////////////////////////////////////////////////////////////////////
+/**
+* This test opens each DialogueWave into its sub-editor.
+*/
+IMPLEMENT_COMPLEX_AUTOMATION_TEST(FOpenDialogueWaveAssetEditors, "Editor.Content.Sounds.Open Dialogue Wave Assets", (EAutomationTestFlags::ATF_Editor | EAutomationTestFlags::ATF_RequiresUser));
+
+void FOpenDialogueWaveAssetEditors::GetTests(TArray<FString>& OutBeautifiedNames, TArray<FString>& OutTestCommands) const
+{
+	//This grabs each DialogueWave asset in the Game/Content directory
+	FEditorAutomationTestUtilities::CollectGameContentTestsByClass(UDialogueWave::StaticClass(), OutBeautifiedNames, OutTestCommands);
+}
+
+bool FOpenDialogueWaveAssetEditors::RunTest(const FString& Parameters)
+{
+	bool bDidTheTestPass = FOpenActualAssetEditors(Parameters);
+	CollectGarbage(GARBAGE_COLLECTION_KEEPFLAGS);
+	return bDidTheTestPass;
+}
+
+//////////////////////////////////////////////////////////////////////////
+/**
+* This test opens each Material into its sub-editor.
+*/
+IMPLEMENT_COMPLEX_AUTOMATION_TEST(FOpenMaterialAssetEditors, "Editor.Content.Materials.Open Material Assets", (EAutomationTestFlags::ATF_Editor | EAutomationTestFlags::ATF_RequiresUser));
+
+void FOpenMaterialAssetEditors::GetTests(TArray<FString>& OutBeautifiedNames, TArray<FString>& OutTestCommands) const
+{
+	//This grabs each Material asset in the Game/Content directory
+	FEditorAutomationTestUtilities::CollectGameContentTestsByClass(UMaterial::StaticClass(), OutBeautifiedNames, OutTestCommands);
+}
+
+bool FOpenMaterialAssetEditors::RunTest(const FString& Parameters)
+{
+	bool bDidTheTestPass = FOpenActualAssetEditors(Parameters);
+	CollectGarbage(GARBAGE_COLLECTION_KEEPFLAGS);
+	return bDidTheTestPass;
+}
+
+////////////////////////////////////////////////////////////////////////////
+///**
+//* This test opens each MaterialFunction into its sub-editor.
+//*/
+//IMPLEMENT_COMPLEX_AUTOMATION_TEST(FOpenMaterialFunctionAssetEditors, "Editor.Content.Materials.Open Material Function Assets", (EAutomationTestFlags::ATF_Editor | EAutomationTestFlags::ATF_RequiresUser));
+//
+//void FOpenMaterialFunctionAssetEditors::GetTests(TArray<FString>& OutBeautifiedNames, TArray<FString>& OutTestCommands) const
+//{
+//	//This grabs each MaterialFunction asset in the Game/Content directory
+//	FEditorAutomationTestUtilities::CollectGameContentTestsByClass(UMaterialFunction::StaticClass(), OutBeautifiedNames, OutTestCommands);
+//}
+//
+//bool FOpenMaterialFunctionAssetEditors::RunTest(const FString& Parameters)
+//{
+//	bool bDidTheTestPass = FOpenActualAssetEditors(Parameters);
+//	CollectGarbage(GARBAGE_COLLECTION_KEEPFLAGS);
+//	return bDidTheTestPass;
+//}
+
+//////////////////////////////////////////////////////////////////////////
+/**
+* This test opens each MaterialInstanceConstant into its sub-editor.
+*/
+IMPLEMENT_COMPLEX_AUTOMATION_TEST(FOpenMaterialInstanceConstantAssetEditors, "Editor.Content.Materials.Open Material Instance Constant Assets", (EAutomationTestFlags::ATF_Editor | EAutomationTestFlags::ATF_RequiresUser));
+
+void FOpenMaterialInstanceConstantAssetEditors::GetTests(TArray<FString>& OutBeautifiedNames, TArray<FString>& OutTestCommands) const
+{
+	//This grabs each MaterialInstanceConstant asset in the Game/Content directory
+	FEditorAutomationTestUtilities::CollectGameContentTestsByClass(UMaterialInstanceConstant::StaticClass(), OutBeautifiedNames, OutTestCommands);
+}
+
+bool FOpenMaterialInstanceConstantAssetEditors::RunTest(const FString& Parameters)
+{
+	bool bDidTheTestPass = FOpenActualAssetEditors(Parameters);
+	CollectGarbage(GARBAGE_COLLECTION_KEEPFLAGS);
+	return bDidTheTestPass;
+}
+
+//////////////////////////////////////////////////////////////////////////
+/**
+* This test opens each MaterialParameterCollection into its sub-editor.
+*/
+IMPLEMENT_COMPLEX_AUTOMATION_TEST(FOpenMaterialParameterCollectionAssetEditors, "Editor.Content.Materials.Open Material Parameter Collection Assets", (EAutomationTestFlags::ATF_Editor | EAutomationTestFlags::ATF_RequiresUser));
+
+void FOpenMaterialParameterCollectionAssetEditors::GetTests(TArray<FString>& OutBeautifiedNames, TArray<FString>& OutTestCommands) const
+{
+	//This grabs each MaterialParameterCollection asset in the Game/Content directory
+	FEditorAutomationTestUtilities::CollectGameContentTestsByClass(UMaterialParameterCollection::StaticClass(), OutBeautifiedNames, OutTestCommands);
+}
+
+bool FOpenMaterialParameterCollectionAssetEditors::RunTest(const FString& Parameters)
+{
+	bool bDidTheTestPass = FOpenActualAssetEditors(Parameters);
+	CollectGarbage(GARBAGE_COLLECTION_KEEPFLAGS);
+	return bDidTheTestPass;
+}
+
+
+//////////////////////////////////////////////////////////////////////////
+/**
+* This test opens each PhysicalMaterial into its sub-editor.
+*/
+IMPLEMENT_COMPLEX_AUTOMATION_TEST(FOpenPhysicalMaterialAssetEditors, "Editor.Content.Physics.Open PhysicalMaterial Assets", (EAutomationTestFlags::ATF_Editor | EAutomationTestFlags::ATF_RequiresUser));
+
+void FOpenPhysicalMaterialAssetEditors::GetTests(TArray<FString>& OutBeautifiedNames, TArray<FString>& OutTestCommands) const
+{
+	//This grabs each PhysicalMaterial asset in the Game/Content directory
+	FEditorAutomationTestUtilities::CollectGameContentTestsByClass(UPhysicalMaterial::StaticClass(), OutBeautifiedNames, OutTestCommands);
+}
+
+bool FOpenPhysicalMaterialAssetEditors::RunTest(const FString& Parameters)
+{
+	bool bDidTheTestPass = FOpenActualAssetEditors(Parameters);
+	CollectGarbage(GARBAGE_COLLECTION_KEEPFLAGS);
+	return bDidTheTestPass;
+}
+
+
+
+//////////////////////////////////////////////////////////////////////////
+/**
+* This test opens each ReverbEffect into its sub-editor.
+*/
+IMPLEMENT_COMPLEX_AUTOMATION_TEST(FOpenReverbEffectAssetEditors, "Editor.Content.Sounds.Open Reverb Effect Assets", (EAutomationTestFlags::ATF_Editor | EAutomationTestFlags::ATF_RequiresUser));
+
+void FOpenReverbEffectAssetEditors::GetTests(TArray<FString>& OutBeautifiedNames, TArray<FString>& OutTestCommands) const
+{
+	//This grabs each ReverbEffect asset in the Game/Content directory
+	FEditorAutomationTestUtilities::CollectGameContentTestsByClass(UReverbEffect::StaticClass(), OutBeautifiedNames, OutTestCommands);
+}
+
+bool FOpenReverbEffectAssetEditors::RunTest(const FString& Parameters)
+{
+	bool bDidTheTestPass = FOpenActualAssetEditors(Parameters);
+	CollectGarbage(GARBAGE_COLLECTION_KEEPFLAGS);
+	return bDidTheTestPass;
+}
+
+
+
+//////////////////////////////////////////////////////////////////////////
+/**
+* This test opens each SlateWidgetStyleAsset into its sub-editor.
+*/
+IMPLEMENT_COMPLEX_AUTOMATION_TEST(FOpenSlateWidgetStyleAssetAssetEditors, "Editor.Content.Physics.Open SlateWidgetStyleAsset Assets", (EAutomationTestFlags::ATF_Editor | EAutomationTestFlags::ATF_RequiresUser));
+
+void FOpenSlateWidgetStyleAssetAssetEditors::GetTests(TArray<FString>& OutBeautifiedNames, TArray<FString>& OutTestCommands) const
+{
+	//This grabs each SlateWidgetStyleAsset asset in the Game/Content directory
+	FEditorAutomationTestUtilities::CollectGameContentTestsByClass(USlateWidgetStyleAsset::StaticClass(), OutBeautifiedNames, OutTestCommands);
+}
+
+bool FOpenSlateWidgetStyleAssetAssetEditors::RunTest(const FString& Parameters)
+{
+	bool bDidTheTestPass = FOpenActualAssetEditors(Parameters);
+	CollectGarbage(GARBAGE_COLLECTION_KEEPFLAGS);
+	return bDidTheTestPass;
+}
+
+
+//////////////////////////////////////////////////////////////////////////
+/**
+* This test opens each SoundAttenuation into its sub-editor.
+*/
+IMPLEMENT_COMPLEX_AUTOMATION_TEST(FOpenSoundAttenuationAssetEditors, "Editor.Content.Sounds.Open Sound Attenuation Assets", (EAutomationTestFlags::ATF_Editor | EAutomationTestFlags::ATF_RequiresUser));
+
+void FOpenSoundAttenuationAssetEditors::GetTests(TArray<FString>& OutBeautifiedNames, TArray<FString>& OutTestCommands) const
+{
+	//This grabs each SoundAttenuation asset in the Game/Content directory
+	FEditorAutomationTestUtilities::CollectGameContentTestsByClass(USoundAttenuation::StaticClass(), OutBeautifiedNames, OutTestCommands);
+}
+
+bool FOpenSoundAttenuationAssetEditors::RunTest(const FString& Parameters)
+{
+	bool bDidTheTestPass = FOpenActualAssetEditors(Parameters);
+	CollectGarbage(GARBAGE_COLLECTION_KEEPFLAGS);
+	return bDidTheTestPass;
+}
+
+//////////////////////////////////////////////////////////////////////////
+/**
+* This test opens each Static Mesh into their sub-editor.
+*/
+IMPLEMENT_COMPLEX_AUTOMATION_TEST(FOpenStaticMeshAssetEditors, "Editor.Content.Static Mesh.Open Static Mesh Assets", (EAutomationTestFlags::ATF_Editor | EAutomationTestFlags::ATF_RequiresUser));
+
+void FOpenStaticMeshAssetEditors::GetTests(TArray<FString>& OutBeautifiedNames, TArray<FString>& OutTestCommands) const
+{
+	//This grabs each Static Mesh in the Game/Content
+	FEditorAutomationTestUtilities::CollectGameContentTestsByClass(UStaticMesh::StaticClass(), OutBeautifiedNames, OutTestCommands);
+}
+
+bool FOpenStaticMeshAssetEditors::RunTest(const FString& Parameters)
+{
+	bool bDidTheTestPass = FOpenActualAssetEditors(Parameters);
+	CollectGarbage(GARBAGE_COLLECTION_KEEPFLAGS);
+	return bDidTheTestPass;
+}
+
+//////////////////////////////////////////////////////////////////////////
+/**
+* This test opens each Skeletal Mesh into their sub-editor.
+*/
+IMPLEMENT_COMPLEX_AUTOMATION_TEST(FOpenSkeletalMeshAssetEditors, "Editor.Content.Skeletal Mesh.Open Skeletal Mesh Assets", (EAutomationTestFlags::ATF_Editor | EAutomationTestFlags::ATF_RequiresUser));
+
+void FOpenSkeletalMeshAssetEditors::GetTests(TArray<FString>& OutBeautifiedNames, TArray<FString>& OutTestCommands) const
+{
+	//This grabs each SKeletal Mesh in the Game/Content
+	FEditorAutomationTestUtilities::CollectGameContentTestsByClass(USkeletalMesh::StaticClass(), OutBeautifiedNames, OutTestCommands);
+}
+
+bool FOpenSkeletalMeshAssetEditors::RunTest(const FString& Parameters)
+{
+	bool bDidTheTestPass = FOpenActualAssetEditors(Parameters);
+	CollectGarbage(GARBAGE_COLLECTION_KEEPFLAGS);
+	return bDidTheTestPass;
+}
+
+
+//////////////////////////////////////////////////////////////////////////
+/**
+* This test opens each SoundClass into its sub-editor.
+*/
+IMPLEMENT_COMPLEX_AUTOMATION_TEST(FOpenSoundClassAssetEditors, "Editor.Content.Sounds.Open SoundClass Assets", (EAutomationTestFlags::ATF_Editor | EAutomationTestFlags::ATF_RequiresUser));
+
+void FOpenSoundClassAssetEditors::GetTests(TArray<FString>& OutBeautifiedNames, TArray<FString>& OutTestCommands) const
+{
+	//This grabs each SoundClass asset in the Game/Content directory
+	FEditorAutomationTestUtilities::CollectGameContentTestsByClass(USoundClass::StaticClass(), OutBeautifiedNames, OutTestCommands);
+}
+
+bool FOpenSoundClassAssetEditors::RunTest(const FString& Parameters)
+{
+	bool bDidTheTestPass = FOpenActualAssetEditors(Parameters);
+	CollectGarbage(GARBAGE_COLLECTION_KEEPFLAGS);
+	return bDidTheTestPass;
+}
+
+
+//////////////////////////////////////////////////////////////////////////
+/**
+* This test opens each SoundCue into its sub-editor.
+*/
+IMPLEMENT_COMPLEX_AUTOMATION_TEST(FOpenSoundCueAssetEditors, "Editor.Content.Sounds.Open Sound Cue Assets", (EAutomationTestFlags::ATF_Editor | EAutomationTestFlags::ATF_RequiresUser));
+
+void FOpenSoundCueAssetEditors::GetTests(TArray<FString>& OutBeautifiedNames, TArray<FString>& OutTestCommands) const
+{
+	//This grabs each SoundCue asset in the Game/Content directory
+	FEditorAutomationTestUtilities::CollectGameContentTestsByClass(USoundCue::StaticClass(), OutBeautifiedNames, OutTestCommands);
+}
+
+bool FOpenSoundCueAssetEditors::RunTest(const FString& Parameters)
+{
+	bool bDidTheTestPass = FOpenActualAssetEditors(Parameters);
+	CollectGarbage(GARBAGE_COLLECTION_KEEPFLAGS);
+	return bDidTheTestPass;
+}
+
+
+//////////////////////////////////////////////////////////////////////////
+/**
+* This test opens each SoundMix into its sub-editor.
+*/
+IMPLEMENT_COMPLEX_AUTOMATION_TEST(FOpenSoundMixAssetEditors, "Editor.Content.Sounds.Open SoundMix Assets", (EAutomationTestFlags::ATF_Editor | EAutomationTestFlags::ATF_RequiresUser));
+
+void FOpenSoundMixAssetEditors::GetTests(TArray<FString>& OutBeautifiedNames, TArray<FString>& OutTestCommands) const
+{
+	//This grabs each SoundMix asset in the Game/Content directory
+	FEditorAutomationTestUtilities::CollectGameContentTestsByClass(USoundMix::StaticClass(), OutBeautifiedNames, OutTestCommands);
+}
+
+bool FOpenSoundMixAssetEditors::RunTest(const FString& Parameters)
+{
+	bool bDidTheTestPass = FOpenActualAssetEditors(Parameters);
+	CollectGarbage(GARBAGE_COLLECTION_KEEPFLAGS);
+	return bDidTheTestPass;
+}
+
+
+//////////////////////////////////////////////////////////////////////////
+/**
+* This test opens each SubsurfaceProfile into its sub-editor.
+*/
+IMPLEMENT_COMPLEX_AUTOMATION_TEST(FOpenSubsurfaceProfileAssetEditors, "Editor.Content.Materials.Open Subsurface Profile Assets", (EAutomationTestFlags::ATF_Editor | EAutomationTestFlags::ATF_RequiresUser));
+
+void FOpenSubsurfaceProfileAssetEditors::GetTests(TArray<FString>& OutBeautifiedNames, TArray<FString>& OutTestCommands) const
+{
+	//This grabs each SubsurfaceProfile asset in the Game/Content directory
+	FEditorAutomationTestUtilities::CollectGameContentTestsByClass(USubsurfaceProfile::StaticClass(), OutBeautifiedNames, OutTestCommands);
+}
+
+bool FOpenSubsurfaceProfileAssetEditors::RunTest(const FString& Parameters)
+{
+	bool bDidTheTestPass = FOpenActualAssetEditors(Parameters);
+	CollectGarbage(GARBAGE_COLLECTION_KEEPFLAGS);
+	return bDidTheTestPass;
+}
+
+
+
+//////////////////////////////////////////////////////////////////////////
+/**
+* This test opens each texture2d into its sub-editor.
+*/
+IMPLEMENT_COMPLEX_AUTOMATION_TEST(FOpenTextureAssetEditors, "Editor.Content.Textures.Open Texture Assets", (EAutomationTestFlags::ATF_Editor | EAutomationTestFlags::ATF_RequiresUser));
+
+void FOpenTextureAssetEditors::GetTests(TArray<FString>& OutBeautifiedNames, TArray<FString>& OutTestCommands) const
+{
+	//This grabs each texture file in the Game/Content
+	FEditorAutomationTestUtilities::CollectGameContentTestsByClass(UTexture2D::StaticClass(), OutBeautifiedNames, OutTestCommands);
+}
+
+bool FOpenTextureAssetEditors::RunTest(const FString& Parameters)
+{
+	bool bDidTheTestPass = FOpenActualAssetEditors(Parameters);
+	CollectGarbage(GARBAGE_COLLECTION_KEEPFLAGS);
+	return bDidTheTestPass;
+}
+
+//////////////////////////////////////////////////////////////////////////
+/**
+* This test opens each texture cube into its sub-editor.
+*/
+IMPLEMENT_COMPLEX_AUTOMATION_TEST(FOpenTextureCubeAssetEditors, "Editor.Content.Textures.Open Texture Cube Assets", (EAutomationTestFlags::ATF_Editor | EAutomationTestFlags::ATF_RequiresUser));
+
+void FOpenTextureCubeAssetEditors::GetTests(TArray<FString>& OutBeautifiedNames, TArray<FString>& OutTestCommands) const
+{
+	//This grabs each texture cube in the Game/Content
+	FEditorAutomationTestUtilities::CollectGameContentTestsByClass(UTextureCube::StaticClass(), OutBeautifiedNames, OutTestCommands);
+}
+
+bool FOpenTextureCubeAssetEditors::RunTest(const FString& Parameters)
+{
+	bool bDidTheTestPass = FOpenActualAssetEditors(Parameters);
+	CollectGarbage(GARBAGE_COLLECTION_KEEPFLAGS);
+	return bDidTheTestPass;
+}
+
+//////////////////////////////////////////////////////////////////////////
+/**
+* This test opens each TextureRenderTargetCube into its sub-editor.
+*/
+IMPLEMENT_COMPLEX_AUTOMATION_TEST(FOpenTextureRenderTargetCubeAssetEditors, "Editor.Content.Textures.Open Texture Render Target Cube Assets", (EAutomationTestFlags::ATF_Editor | EAutomationTestFlags::ATF_RequiresUser));
+
+void FOpenTextureRenderTargetCubeAssetEditors::GetTests(TArray<FString>& OutBeautifiedNames, TArray<FString>& OutTestCommands) const
+{
+	//This grabs each TextureRenderTargetCube asset in the Game/Content directory
+	FEditorAutomationTestUtilities::CollectGameContentTestsByClass(UTextureRenderTargetCube::StaticClass(), OutBeautifiedNames, OutTestCommands);
+}
+
+bool FOpenTextureRenderTargetCubeAssetEditors::RunTest(const FString& Parameters)
+{
+	bool bDidTheTestPass = FOpenActualAssetEditors(Parameters);
+	CollectGarbage(GARBAGE_COLLECTION_KEEPFLAGS);
+	return bDidTheTestPass;
+}
+
+//////////////////////////////////////////////////////////////////////////
+/**
+* This test opens each Font into its sub-editor.
+*/
+IMPLEMENT_COMPLEX_AUTOMATION_TEST(FOpenFontAssetEditors, "Editor.Content.Textures.Open Font Assets", (EAutomationTestFlags::ATF_Editor | EAutomationTestFlags::ATF_RequiresUser));
+
+void FOpenFontAssetEditors::GetTests(TArray<FString>& OutBeautifiedNames, TArray<FString>& OutTestCommands) const
+{
+	//This grabs each Font asset in the Game/Content directory
+	FEditorAutomationTestUtilities::CollectGameContentTestsByClass(UFont::StaticClass(), OutBeautifiedNames, OutTestCommands);
+}
+
+bool FOpenFontAssetEditors::RunTest(const FString& Parameters)
+{
+	bool bDidTheTestPass = FOpenActualAssetEditors(Parameters);
+	CollectGarbage(GARBAGE_COLLECTION_KEEPFLAGS);
+	return bDidTheTestPass;
+}
+
+//////////////////////////////////////////////////////////////////////////
+/**
+* This test opens each TextureRenderTarget2D into its sub-editor.
+*/
+IMPLEMENT_COMPLEX_AUTOMATION_TEST(FOpenTextureRenderTarget2DAssetEditors, "Editor.Content.Textures.Open Texture Render Target 2D Assets", (EAutomationTestFlags::ATF_Editor | EAutomationTestFlags::ATF_RequiresUser));
+
+void FOpenTextureRenderTarget2DAssetEditors::GetTests(TArray<FString>& OutBeautifiedNames, TArray<FString>& OutTestCommands) const
+{
+	//This grabs each TextureRenderTarget2D asset in the Game/Content directory
+	FEditorAutomationTestUtilities::CollectGameContentTestsByClass(UTextureRenderTarget2D::StaticClass(), OutBeautifiedNames, OutTestCommands);
+}
+
+bool FOpenTextureRenderTarget2DAssetEditors::RunTest(const FString& Parameters)
+{
+	bool bDidTheTestPass = FOpenActualAssetEditors(Parameters);
+	CollectGarbage(GARBAGE_COLLECTION_KEEPFLAGS);
+	return bDidTheTestPass;
+}
\ No newline at end of file
diff --git a/Engine/Source/Editor/UnrealEd/Public/Tests/AutomationEditorCommon.h b/Engine/Source/Editor/UnrealEd/Public/Tests/AutomationEditorCommon.h
--- a/Engine/Source/Editor/UnrealEd/Public/Tests/AutomationEditorCommon.h
+++ b/Engine/Source/Editor/UnrealEd/Public/Tests/AutomationEditorCommon.h
@@ -1,6 +1,9 @@
 // Copyright 1998-2014 Epic Games, Inc. All Rights Reserved.
 
 #pragma once
+#include "AssetRegistryModule.h"
+#include "ModuleManager.h"
+#include "AutomationCommon.h"
 
 #include "Tests/AutomationTestSettings.h"
 
@@ -90,12 +93,144 @@ namespace AutomationEditorCommonUtils
 }
 
 
+
 //////////////////////////////////////////////////////////////////////////
 //Common latent commands used for automated editor testing.
 
-
 /**
 * Creates a latent command which the user can either undo or redo an action.
 * True will trigger Undo, False will trigger a redo
 */
-DEFINE_LATENT_AUTOMATION_COMMAND_ONE_PARAMETER(FUndoRedoCommand, bool, bUndo);
\ No newline at end of file
+DEFINE_LATENT_AUTOMATION_COMMAND_ONE_PARAMETER(FUndoRedoCommand, bool, bUndo);
+
+/**
+* Open editor for a particular asset
+*/
+DEFINE_LATENT_AUTOMATION_COMMAND_ONE_PARAMETER(FOpenEditorForAssetCommand, FString, AssetName);
+
+/**
+* Close all asset editors
+*/
+DEFINE_LATENT_AUTOMATION_COMMAND(FCloseAllAssetEditorsCommand);
+
+/**
+* Start PIE session
+*/
+DEFINE_LATENT_AUTOMATION_COMMAND_ONE_PARAMETER(FStartPIECommand, bool, bSimulateInEditor);
+
+/**
+* End PlayMap session
+*/
+DEFINE_LATENT_AUTOMATION_COMMAND(FEndPlayMapCommand);
+
+
+//////////////////////////////////////////////////////////////////////////
+// FEditorAutomationTestUtilities
+
+class FEditorAutomationTestUtilities
+{
+public:
+	/**
+	* Loads the map specified by an automation test
+	*
+	* @param MapName - Map to load
+	*/
+	static void LoadMap(const FString& MapName)
+	{
+		bool bLoadAsTemplate = false;
+		bool bShowProgress = false;
+		FEditorFileUtils::LoadMap(MapName, bLoadAsTemplate, bShowProgress);
+	}
+
+	/**
+	* Run PIE
+	*/
+	static void RunPIE()
+	{
+		bool bInSimulateInEditor = true;
+		//once in the editor
+		ADD_LATENT_AUTOMATION_COMMAND(FStartPIECommand(true));
+		ADD_LATENT_AUTOMATION_COMMAND(FWaitLatentCommand(3.0f));
+		ADD_LATENT_AUTOMATION_COMMAND(FEndPlayMapCommand());
+
+		//wait between tests
+		ADD_LATENT_AUTOMATION_COMMAND(FWaitLatentCommand(1.0f));
+
+		//once not in the editor
+		ADD_LATENT_AUTOMATION_COMMAND(FStartPIECommand(false));
+		ADD_LATENT_AUTOMATION_COMMAND(FWaitLatentCommand(3.0f));
+		ADD_LATENT_AUTOMATION_COMMAND(FEndPlayMapCommand());
+	}
+
+	/**
+	* Generates a list of assets from the ENGINE and the GAME by a specific type.
+	* This is to be used by the GetTest() function.
+	*/
+	static void CollectTestsByClass(UClass * Class, TArray<FString>& OutBeautifiedNames, TArray <FString>& OutTestCommands)
+	{
+		FAssetRegistryModule& AssetRegistryModule = FModuleManager::LoadModuleChecked<FAssetRegistryModule>(TEXT("AssetRegistry"));
+		TArray<FAssetData> ObjectList;
+		AssetRegistryModule.Get().GetAssetsByClass(Class->GetFName(), ObjectList);
+
+		for (TObjectIterator<UClass> AllClassesIt; AllClassesIt; ++AllClassesIt)
+		{
+			UClass* ClassList = *AllClassesIt;
+			FName ClassName = ClassList->GetFName();
+			//UE_LOG(LogEditorAutomationTests, Log, TEXT("Here: %s"), *ClassName);
+		}
+
+
+		for (auto ObjIter = ObjectList.CreateConstIterator(); ObjIter; ++ObjIter)
+		{
+			const FAssetData & Asset = *ObjIter;
+			FString Filename = Asset.ObjectPath.ToString();
+			//convert to full paths
+			Filename = FPackageName::LongPackageNameToFilename(Filename);
+			if (FAutomationTestFramework::GetInstance().ShouldTestContent(Filename))
+			{
+				FString BeautifiedFilename = Asset.AssetName.ToString();
+				OutBeautifiedNames.Add(BeautifiedFilename);
+				OutTestCommands.Add(Asset.ObjectPath.ToString());
+			}
+		}
+	}
+
+	/**
+	* Generates a list of assets from the GAME by a specific type.
+	* This is to be used by the GetTest() function.
+	*/
+	static void CollectGameContentTestsByClass(UClass * Class, TArray<FString>& OutBeautifiedNames, TArray <FString>& OutTestCommands)
+	{
+		//Setting the Asset Registry
+
+		FAssetRegistryModule& AssetRegistryModule = FModuleManager::LoadModuleChecked<FAssetRegistryModule>(TEXT("AssetRegistry"));
+		
+		//Variable setups
+		TArray<FAssetData> ObjectList;
+		FARFilter AssetFilter;
+
+		//Generating the list of assets.
+		//This list is being filtered by the game folder and class type.  The results are placed into the ObjectList variable.
+		AssetFilter.ClassNames.Add(Class->GetFName());
+		AssetFilter.PackagePaths.Add("/Game");
+		AssetFilter.bRecursivePaths = true;
+		AssetRegistryModule.Get().GetAssets(AssetFilter, ObjectList);
+
+		//Loop through the list of assets, make their path full and a string, then add them to the test.
+		for (auto ObjIter = ObjectList.CreateConstIterator(); ObjIter; ++ObjIter)
+		{
+			const FAssetData & Asset = *ObjIter;
+			FString Filename = Asset.ObjectPath.ToString();
+			//convert to full paths
+			Filename = FPackageName::LongPackageNameToFilename(Filename);
+			if (FAutomationTestFramework::GetInstance().ShouldTestContent(Filename))
+			{
+				FString BeautifiedFilename = Asset.AssetName.ToString();
+				OutBeautifiedNames.Add(BeautifiedFilename);
+				OutTestCommands.Add(Asset.ObjectPath.ToString());
+			}
+		}
+	}
+};
+
+
