Adding support to the dump BP info commandlet for struct pin types, as well as reference pin types.

diff --git a/Engine/Source/Editor/UnrealEd/Private/Commandlets/DumpBlueprintsInfoCommandlet.cpp b/Engine/Source/Editor/UnrealEd/Private/Commandlets/DumpBlueprintsInfoCommandlet.cpp
--- a/Engine/Source/Editor/UnrealEd/Private/Commandlets/DumpBlueprintsInfoCommandlet.cpp
+++ b/Engine/Source/Editor/UnrealEd/Private/Commandlets/DumpBlueprintsInfoCommandlet.cpp
@@ -1378,27 +1378,36 @@ static bool DumpBlueprintInfoUtils::DumpPinContextActions(uint32 Indent, UEdGrap
 		UEdGraphSchema_K2 const* K2Schema = GetDefault<UEdGraphSchema_K2>();
 		bool const bUsePinTypeClass = (CommandOptions.DumpFlags & BPDUMP_PinTypeIsClass) != 0;
 
-		UClass* TypeClass = FindObject<UClass>(ANY_PACKAGE, *CommandOptions.PinType);
-		if ((TypeClass != nullptr) && UEdGraphSchema_K2::IsAllowableBlueprintVariableType(TypeClass))
+		bool bIsValidPinType = true;	
+		if (UClass* TypeClass = FindObject<UClass>(ANY_PACKAGE, *CommandOptions.PinType))
 		{
-			PinType.PinSubCategoryObject = TypeClass;
-			if (TypeClass->IsChildOf(UInterface::StaticClass()))
+			bIsValidPinType = UEdGraphSchema_K2::IsAllowableBlueprintVariableType(TypeClass);
+			if (bIsValidPinType)
 			{
-				PinType.PinCategory = K2Schema->PC_Interface;
+				PinType.PinSubCategoryObject = TypeClass;
+				if (TypeClass->IsChildOf(UInterface::StaticClass()))
+				{
+					PinType.PinCategory = K2Schema->PC_Interface;
+				}
+				else
+				{
+					PinType.PinCategory = K2Schema->PC_Object;
+				}
+
+				if (bUsePinTypeClass)
+				{
+					PinType.PinCategory = K2Schema->PC_Class;
+				}
 			}
-			else if (TypeClass->IsChildOf(UScriptStruct::StaticClass()))
+			
+		}
+		else if (UScriptStruct* StructType = FindObject<UScriptStruct>(ANY_PACKAGE, *CommandOptions.PinType))
+		{
+			bIsValidPinType = UEdGraphSchema_K2::IsAllowableBlueprintVariableType(StructType);
+			if (bIsValidPinType)
 			{
 				PinType.PinCategory = K2Schema->PC_Struct;
-				PinType.PinSubCategoryObject = TypeClass->GetDefaultObject();
-			}
-			else
-			{
-				PinType.PinCategory = K2Schema->PC_Object;
-			}
-
-			if (bUsePinTypeClass)
-			{
-				PinType.PinCategory = K2Schema->PC_Class;
+				PinType.PinSubCategoryObject = StructType;
 			}
 		}
 		else if (!CommandOptions.PinType.Compare("self", ESearchCase::IgnoreCase))
@@ -1409,15 +1418,30 @@ static bool DumpBlueprintInfoUtils::DumpPinContextActions(uint32 Indent, UEdGrap
 				PinType.PinCategory = K2Schema->PC_Class;
 			}
 			PinType.PinSubCategory = K2Schema->PSC_Self;
+			bIsValidPinType = true;
+		}
+		else if (!CommandOptions.PinType.Compare(K2Schema->PC_Delegate, ESearchCase::IgnoreCase) ||
+		         !CommandOptions.PinType.Compare(K2Schema->PC_MCDelegate, ESearchCase::IgnoreCase))
+		{
+			// @TODO: PC_Delegate, PC_MCDelegate
+			bIsValidPinType = false;
 		}
-		// @TODO: PC_Delegate, PC_MCDelegate
+		
+		if (bIsValidPinType)
+		{
+			DumpContextualPinTypeActions(Indent, Graph, PinType, FileOutWriter);
+			FileOutWriter->Serialize(TCHAR_TO_ANSI(TEXT(",\n")), 2);
 
-		DumpContextualPinTypeActions(Indent, Graph, PinType, FileOutWriter);
-		FileOutWriter->Serialize(TCHAR_TO_ANSI(TEXT(",\n")), 2);
-		PinType.bIsArray = true;
-		DumpContextualPinTypeActions(Indent, Graph, PinType, FileOutWriter);
+			PinType.bIsReference = true;
+			DumpContextualPinTypeActions(Indent, Graph, PinType, FileOutWriter);
+			FileOutWriter->Serialize(TCHAR_TO_ANSI(TEXT(",\n")), 2);
+			PinType.bIsReference = false;
 
-		bWroteToFile = true;
+			PinType.bIsArray = true;
+			DumpContextualPinTypeActions(Indent, Graph, PinType, FileOutWriter);
+
+			bWroteToFile = true;
+		}
 	}
 	else if (Graph->Schema->IsChildOf(UEdGraphSchema_K2::StaticClass()))
 	{
@@ -1450,6 +1474,11 @@ static bool DumpBlueprintInfoUtils::DumpTypeTreeActions(uint32 Indent, UEdGraph*
 		DumpContextualPinTypeActions(Indent, Graph, PinType, FileOutWriter);
 		FileOutWriter->Serialize(TCHAR_TO_ANSI(TEXT(",\n")), 2);
 
+		PinType.bIsReference = true;
+		DumpContextualPinTypeActions(Indent, Graph, PinType, FileOutWriter);
+		FileOutWriter->Serialize(TCHAR_TO_ANSI(TEXT(",\n")), 2);
+		PinType.bIsReference = false;
+
 		PinType.bIsArray = true;
 		DumpContextualPinTypeActions(Indent, Graph, PinType, FileOutWriter);
 
