FBX import :custom curve option wasn't added.

diff --git a/Engine/Source/Editor/UnrealEd/Private/Fbx/FbxOptionWindow.cpp b/Engine/Source/Editor/UnrealEd/Private/Fbx/FbxOptionWindow.cpp
--- a/Engine/Source/Editor/UnrealEd/Private/Fbx/FbxOptionWindow.cpp
+++ b/Engine/Source/Editor/UnrealEd/Private/Fbx/FbxOptionWindow.cpp
@@ -54,16 +54,6 @@ public:
 						.Text(LOCTEXT("FbxOptionWindow_SelectType", "Select Asset Type"))
 					]
 
-// 					+SHorizontalBox::Slot()
-// 					.HAlign(HAlign_Right)
-// 					[
-// 						SNew(SBox)
-// 						.VAlign(VAlign_Bottom)
-// 						[
-// 							SNew(STextBlock)
-// 							.Text(LOCTEXT("FbxOptionWindow_Help", "[Help]"))
-// 						]
-// 					]
 
 					+SHorizontalBox::Slot()
 					.HAlign(HAlign_Left)
@@ -890,7 +880,6 @@ TSharedRef<SWidget> SFbxOptionWindow::ConstructSkeletalMeshAdvanced()
 			.Text(LOCTEXT("FbxOptionWindow_ImportMeshesInBoneHierarchy", "Import Meshes in Bone Hierarchy"))
 		]
 	];
-
 	NewBox->AddSlot().AutoHeight() .Padding(4)
 	[
 		SNew(STextBlock)
@@ -1007,6 +996,16 @@ TSharedRef<SWidget> SFbxOptionWindow::ConstructAnimationOption()
 		]
 	];
 
+	NewBox->AddSlot().AutoHeight().Padding(2)
+	[
+		SNew(SCheckBox)
+		.IsChecked(ImportUI->AnimSequenceImportData->bImportCustomAttribute ? ESlateCheckBoxState::Checked : ESlateCheckBoxState::Unchecked)
+		.OnCheckStateChanged(this, &SFbxOptionWindow::SetAnimation_ImportCustomAttribute)
+		[
+			SNew(STextBlock)
+			.Text(LOCTEXT("FbxOptionWindow_ImportCustomAttribute", "Import Custom Attribute as Curves"))
+		]
+	];
 	NewBox->AddSlot().AutoHeight() .Padding(2)
 	[
 		SNew(SCheckBox)
@@ -1318,6 +1317,11 @@ void SFbxOptionWindow::SetAnimation_ReserveLocalTransform(ESlateCheckBoxState::T
 	ImportUI->bPreserveLocalTransform = (NewType == ESlateCheckBoxState::Checked)? 1:0;
 }
 
+void SFbxOptionWindow::SetAnimation_ImportCustomAttribute(ESlateCheckBoxState::Type NewType)
+{
+	ImportUI->AnimSequenceImportData->bImportCustomAttribute = (NewType == ESlateCheckBoxState::Checked)? 1:0;
+}
+
 TSharedRef<SWidget> SFbxOptionWindow::MakeSkeletonPickerMenu()
 {
 	FContentBrowserModule& ContentBrowserModule = FModuleManager::Get().LoadModuleChecked<FContentBrowserModule>(TEXT("ContentBrowser"));
