Fix spelling of Synchronize

diff --git a/Engine/Source/Editor/UnrealEd/Classes/CookOnTheSide/CookOnTheFlyServer.h b/Engine/Source/Editor/UnrealEd/Classes/CookOnTheSide/CookOnTheFlyServer.h
--- a/Engine/Source/Editor/UnrealEd/Classes/CookOnTheSide/CookOnTheFlyServer.h
+++ b/Engine/Source/Editor/UnrealEd/Classes/CookOnTheSide/CookOnTheFlyServer.h
@@ -39,7 +39,7 @@ private:
 
 	/** Array which has been made thread safe :) */
 	template<typename Type, typename SynchronizationObjectType, typename ScopeLockType>
-	struct FUnsyncronizedQueue
+	struct FUnsynchronizedQueue
 	{
 	private:
 		mutable SynchronizationObjectType	SynchronizationObject; // made this mutable so this class can have const functions and still be thread safe
@@ -117,20 +117,20 @@ public:
 
 
 	template<typename Type>
-	struct FThreadSafeQueue : public FUnsyncronizedQueue<Type, FCriticalSection, FScopeLock>
+	struct FThreadSafeQueue : public FUnsynchronizedQueue<Type, FCriticalSection, FScopeLock>
 	{
 		/**
 		 * Don't add any functions here, this is just a overqualified typedef
-		 * Add functions / functionality to the FUnsyncronizedQueue
+		 * Add functions / functionality to the FUnsynchronizedQueue
 		 */
 	};
 
 	template<typename Type>
-	struct FQueue : public FUnsyncronizedQueue<Type, FDummyCriticalSection, FDummyScopeLock>
+	struct FQueue : public FUnsynchronizedQueue<Type, FDummyCriticalSection, FDummyScopeLock>
 	{
 		/**
 		 * Don't add any functions here, this is just a overqualified typedef
-		 * Add functions / functionality to the FUnsyncronizedQueue
+		 * Add functions / functionality to the FUnsynchronizedQueue
 		 */
 	};
 
diff --git a/Engine/Source/Runtime/Core/Private/Serialization/AsyncIOSystemBase.h b/Engine/Source/Runtime/Core/Private/Serialization/AsyncIOSystemBase.h
--- a/Engine/Source/Runtime/Core/Private/Serialization/AsyncIOSystemBase.h
+++ b/Engine/Source/Runtime/Core/Private/Serialization/AsyncIOSystemBase.h
@@ -384,7 +384,7 @@ protected:
 	 */
 	void LogIORequest(const FString& Message, const FAsyncIORequest& IORequest);
 
-	/** Critical section used to syncronize access to outstanding requests map						*/
+	/** Critical section used to synchronize access to outstanding requests map						*/
 	FCriticalSection*				CriticalSection;
 	/** TMap of file name string hash to file handles												*/
 	TMap<uint32, IFileHandle*>		NameHashToHandleMap;
diff --git a/Engine/Source/Runtime/UMG/Private/Components/Border.cpp b/Engine/Source/Runtime/UMG/Private/Components/Border.cpp
--- a/Engine/Source/Runtime/UMG/Private/Components/Border.cpp
+++ b/Engine/Source/Runtime/UMG/Private/Components/Border.cpp
@@ -43,9 +43,9 @@ TSharedRef<SWidget> UBorder::RebuildWidget()
 	return BuildDesignTimeWidget( MyBorder.ToSharedRef() );
 }
 
-void UBorder::SyncronizeProperties()
+void UBorder::SynchronizeProperties()
 {
-	Super::SyncronizeProperties();
+	Super::SynchronizeProperties();
 	
 	MyBorder->SetBorderBackgroundColor(BrushColor);
 	MyBorder->SetColorAndOpacity(ContentColorAndOpacity);
diff --git a/Engine/Source/Runtime/UMG/Private/Components/BorderSlot.cpp b/Engine/Source/Runtime/UMG/Private/Components/BorderSlot.cpp
--- a/Engine/Source/Runtime/UMG/Private/Components/BorderSlot.cpp
+++ b/Engine/Source/Runtime/UMG/Private/Components/BorderSlot.cpp
@@ -63,7 +63,7 @@ void UBorderSlot::SetVerticalAlignment(EVerticalAlignment InVerticalAlignment)
 	}
 }
 
-void UBorderSlot::SyncronizeProperties()
+void UBorderSlot::SynchronizeProperties()
 {
 	SetPadding(Padding);
 	SetHorizontalAlignment(HorizontalAlignment);
diff --git a/Engine/Source/Runtime/UMG/Private/Components/Button.cpp b/Engine/Source/Runtime/UMG/Private/Components/Button.cpp
--- a/Engine/Source/Runtime/UMG/Private/Components/Button.cpp
+++ b/Engine/Source/Runtime/UMG/Private/Components/Button.cpp
@@ -46,9 +46,9 @@ TSharedRef<SWidget> UButton::RebuildWidget()
 	return MyButton.ToSharedRef();
 }
 
-void UButton::SyncronizeProperties()
+void UButton::SynchronizeProperties()
 {
-	Super::SyncronizeProperties();
+	Super::SynchronizeProperties();
 
 	SButton::FArguments ButtonDefaults;
 
diff --git a/Engine/Source/Runtime/UMG/Private/Components/ButtonSlot.cpp b/Engine/Source/Runtime/UMG/Private/Components/ButtonSlot.cpp
--- a/Engine/Source/Runtime/UMG/Private/Components/ButtonSlot.cpp
+++ b/Engine/Source/Runtime/UMG/Private/Components/ButtonSlot.cpp
@@ -59,7 +59,7 @@ void UButtonSlot::SetVerticalAlignment(EVerticalAlignment InVerticalAlignment)
 	}
 }
 
-void UButtonSlot::SyncronizeProperties()
+void UButtonSlot::SynchronizeProperties()
 {
 	SetPadding(Padding);
 	SetHorizontalAlignment(HorizontalAlignment);
diff --git a/Engine/Source/Runtime/UMG/Private/Components/CanvasPanel.cpp b/Engine/Source/Runtime/UMG/Private/Components/CanvasPanel.cpp
--- a/Engine/Source/Runtime/UMG/Private/Components/CanvasPanel.cpp
+++ b/Engine/Source/Runtime/UMG/Private/Components/CanvasPanel.cpp
@@ -64,11 +64,6 @@ TSharedRef<SWidget> UCanvasPanel::RebuildWidget()
 	return BuildDesignTimeWidget( MyCanvas.ToSharedRef() );
 }
 
-void UCanvasPanel::SyncronizeProperties()
-{
-	Super::SyncronizeProperties();
-}
-
 TSharedPtr<SConstraintCanvas> UCanvasPanel::GetCanvasWidget() const
 {
 	return MyCanvas;
diff --git a/Engine/Source/Runtime/UMG/Private/Components/CanvasPanelSlot.cpp b/Engine/Source/Runtime/UMG/Private/Components/CanvasPanelSlot.cpp
--- a/Engine/Source/Runtime/UMG/Private/Components/CanvasPanelSlot.cpp
+++ b/Engine/Source/Runtime/UMG/Private/Components/CanvasPanelSlot.cpp
@@ -29,7 +29,7 @@ void UCanvasPanelSlot::BuildSlot(TSharedRef<SConstraintCanvas> Canvas)
 			Content == NULL ? SNullWidget::NullWidget : Content->TakeWidget()
 		];
 
-	SyncronizeProperties();
+	SynchronizeProperties();
 }
 
 void UCanvasPanelSlot::SetDesiredPosition(FVector2D InPosition)
@@ -139,7 +139,7 @@ void UCanvasPanelSlot::SetZOrder(int32 InZOrder)
 	}
 }
 
-void UCanvasPanelSlot::SyncronizeProperties()
+void UCanvasPanelSlot::SynchronizeProperties()
 {
 	SetOffsets(LayoutData.Offsets);
 	SetAnchors(LayoutData.Anchors);
@@ -167,7 +167,7 @@ void UCanvasPanelSlot::PreEditChange(UProperty* PropertyAboutToChange)
 
 void UCanvasPanelSlot::PostEditChangeChainProperty(struct FPropertyChangedChainEvent& PropertyChangedEvent)
 {
-	SyncronizeProperties();
+	SynchronizeProperties();
 
 	static FName AnchorsProperty(TEXT("Anchors"));
 	
@@ -248,7 +248,7 @@ void UCanvasPanelSlot::PostEditChangeChainProperty(struct FPropertyChangedChainE
 			}
 
 			// Apply the changes to the properties.
-			SyncronizeProperties();
+			SynchronizeProperties();
 		}
 	}
 
diff --git a/Engine/Source/Runtime/UMG/Private/Components/CheckBox.cpp b/Engine/Source/Runtime/UMG/Private/Components/CheckBox.cpp
--- a/Engine/Source/Runtime/UMG/Private/Components/CheckBox.cpp
+++ b/Engine/Source/Runtime/UMG/Private/Components/CheckBox.cpp
@@ -67,9 +67,9 @@ TSharedRef<SWidget> UCheckBox::RebuildWidget()
 	return MyCheckbox.ToSharedRef();
 }
 
-void UCheckBox::SyncronizeProperties()
+void UCheckBox::SynchronizeProperties()
 {
-	Super::SyncronizeProperties();
+	Super::SynchronizeProperties();
 
 	const FCheckBoxStyle* StylePtr = ( Style != NULL ) ? Style->GetStyle<FCheckBoxStyle>() : NULL;
 	if ( StylePtr == NULL )
diff --git a/Engine/Source/Runtime/UMG/Private/Components/CircularThrobber.cpp b/Engine/Source/Runtime/UMG/Private/Components/CircularThrobber.cpp
--- a/Engine/Source/Runtime/UMG/Private/Components/CircularThrobber.cpp
+++ b/Engine/Source/Runtime/UMG/Private/Components/CircularThrobber.cpp
@@ -38,9 +38,9 @@ TSharedRef<SWidget> UCircularThrobber::RebuildWidget()
 	return MyCircularThrobber.ToSharedRef();
 }
 
-void UCircularThrobber::SyncronizeProperties()
+void UCircularThrobber::SynchronizeProperties()
 {
-	Super::SyncronizeProperties();
+	Super::SynchronizeProperties();
 
 	MyCircularThrobber->SetPieceImage(GetPieceBrush());
 	MyCircularThrobber->SetNumPieces(NumberOfPieces);
diff --git a/Engine/Source/Runtime/UMG/Private/Components/ComboBox.cpp b/Engine/Source/Runtime/UMG/Private/Components/ComboBox.cpp
--- a/Engine/Source/Runtime/UMG/Private/Components/ComboBox.cpp
+++ b/Engine/Source/Runtime/UMG/Private/Components/ComboBox.cpp
@@ -12,11 +12,6 @@ UComboBox::UComboBox(const FPostConstructInitializeProperties& PCIP)
 {
 }
 
-void UComboBox::ReleaseNativeWidget()
-{
-	Super::ReleaseNativeWidget();
-}
-
 TSharedRef<SWidget> UComboBox::RebuildWidget()
 {
 	TSharedRef< SComboBox<UObject*> > NewComboBox =
diff --git a/Engine/Source/Runtime/UMG/Private/Components/DragPanel.cpp b/Engine/Source/Runtime/UMG/Private/Components/DragPanel.cpp
--- a/Engine/Source/Runtime/UMG/Private/Components/DragPanel.cpp
+++ b/Engine/Source/Runtime/UMG/Private/Components/DragPanel.cpp
@@ -34,11 +34,6 @@ TSharedRef<SWidget> UDragPanel::RebuildWidget()
 	return MyDragPanel.ToSharedRef();
 }
 
-void UDragPanel::SyncronizeProperties()
-{
-	Super::SyncronizeProperties();
-}
-
 void UDragPanel::OnSlotAdded(UPanelSlot* Slot)
 {
 	// Add the child to the live slot if it already exists
diff --git a/Engine/Source/Runtime/UMG/Private/Components/DropPanel.cpp b/Engine/Source/Runtime/UMG/Private/Components/DropPanel.cpp
--- a/Engine/Source/Runtime/UMG/Private/Components/DropPanel.cpp
+++ b/Engine/Source/Runtime/UMG/Private/Components/DropPanel.cpp
@@ -37,11 +37,6 @@ TSharedRef<SWidget> UDropPanel::RebuildWidget()
 	return MyDropPanel.ToSharedRef();
 }
 
-void UDropPanel::SyncronizeProperties()
-{
-	Super::SyncronizeProperties();
-}
-
 void UDropPanel::OnSlotAdded(UPanelSlot* Slot)
 {
 	// Add the child to the live slot if it already exists
diff --git a/Engine/Source/Runtime/UMG/Private/Components/EditableText.cpp b/Engine/Source/Runtime/UMG/Private/Components/EditableText.cpp
--- a/Engine/Source/Runtime/UMG/Private/Components/EditableText.cpp
+++ b/Engine/Source/Runtime/UMG/Private/Components/EditableText.cpp
@@ -71,9 +71,9 @@ TSharedRef<SWidget> UEditableText::RebuildWidget()
 	return BuildDesignTimeWidget( MyEditableText.ToSharedRef() );
 }
 
-void UEditableText::SyncronizeProperties()
+void UEditableText::SynchronizeProperties()
 {
-	Super::SyncronizeProperties();
+	Super::SynchronizeProperties();
 
 	TAttribute<FText> TextBinding = OPTIONAL_BINDING(FText, Text);
 	TAttribute<FText> HintTextBinding = OPTIONAL_BINDING(FText, HintText);
diff --git a/Engine/Source/Runtime/UMG/Private/Components/EditableTextBox.cpp b/Engine/Source/Runtime/UMG/Private/Components/EditableTextBox.cpp
--- a/Engine/Source/Runtime/UMG/Private/Components/EditableTextBox.cpp
+++ b/Engine/Source/Runtime/UMG/Private/Components/EditableTextBox.cpp
@@ -65,9 +65,9 @@ TSharedRef<SWidget> UEditableTextBox::RebuildWidget()
 	return MyEditableTextBlock.ToSharedRef();
 }
 
-void UEditableTextBox::SyncronizeProperties()
+void UEditableTextBox::SynchronizeProperties()
 {
-	Super::SyncronizeProperties();
+	Super::SynchronizeProperties();
 
 	const FEditableTextBoxStyle* StylePtr = ( Style != NULL ) ? Style->GetStyle<FEditableTextBoxStyle>() : NULL;
 	if ( StylePtr )
diff --git a/Engine/Source/Runtime/UMG/Private/Components/GridPanel.cpp b/Engine/Source/Runtime/UMG/Private/Components/GridPanel.cpp
--- a/Engine/Source/Runtime/UMG/Private/Components/GridPanel.cpp
+++ b/Engine/Source/Runtime/UMG/Private/Components/GridPanel.cpp
@@ -66,9 +66,9 @@ TSharedRef<SWidget> UGridPanel::RebuildWidget()
 	return BuildDesignTimeWidget( MyGridPanel.ToSharedRef() );
 }
 
-void UGridPanel::SyncronizeProperties()
+void UGridPanel::SynchronizeProperties()
 {
-	Super::SyncronizeProperties();
+	Super::SynchronizeProperties();
 
 	for ( int ColumnIndex = 0; ColumnIndex < ColumnFill.Num(); ColumnIndex++ )
 	{
diff --git a/Engine/Source/Runtime/UMG/Private/Components/GridSlot.cpp b/Engine/Source/Runtime/UMG/Private/Components/GridSlot.cpp
--- a/Engine/Source/Runtime/UMG/Private/Components/GridSlot.cpp
+++ b/Engine/Source/Runtime/UMG/Private/Components/GridSlot.cpp
@@ -85,7 +85,7 @@ void UGridSlot::SetVerticalAlignment(EVerticalAlignment InVerticalAlignment)
 	}
 }
 
-void UGridSlot::SyncronizeProperties()
+void UGridSlot::SynchronizeProperties()
 {
 	SetRow(Row);
 	SetRowSpan(RowSpan);
diff --git a/Engine/Source/Runtime/UMG/Private/Components/HorizontalBoxSlot.cpp b/Engine/Source/Runtime/UMG/Private/Components/HorizontalBoxSlot.cpp
--- a/Engine/Source/Runtime/UMG/Private/Components/HorizontalBoxSlot.cpp
+++ b/Engine/Source/Runtime/UMG/Private/Components/HorizontalBoxSlot.cpp
@@ -69,7 +69,7 @@ void UHorizontalBoxSlot::SetVerticalAlignment(EVerticalAlignment InVerticalAlign
 	}
 }
 
-void UHorizontalBoxSlot::SyncronizeProperties()
+void UHorizontalBoxSlot::SynchronizeProperties()
 {
 	SetPadding(Padding);
 	SetSize(Size);
diff --git a/Engine/Source/Runtime/UMG/Private/Components/Image.cpp b/Engine/Source/Runtime/UMG/Private/Components/Image.cpp
--- a/Engine/Source/Runtime/UMG/Private/Components/Image.cpp
+++ b/Engine/Source/Runtime/UMG/Private/Components/Image.cpp
@@ -28,9 +28,9 @@ TSharedRef<SWidget> UImage::RebuildWidget()
 	return MyImage.ToSharedRef();
 }
 
-void UImage::SyncronizeProperties()
+void UImage::SynchronizeProperties()
 {
-	Super::SyncronizeProperties();
+	Super::SynchronizeProperties();
 
 	TAttribute<FSlateColor> ColorAndOpacityBinding = OPTIONAL_BINDING(FSlateColor, ColorAndOpacity);
 	TAttribute<const FSlateBrush*> ImageBinding = OPTIONAL_BINDING_CONVERT(USlateBrushAsset*, Image, const FSlateBrush*, ConvertImage);
diff --git a/Engine/Source/Runtime/UMG/Private/Components/MenuAnchor.cpp b/Engine/Source/Runtime/UMG/Private/Components/MenuAnchor.cpp
--- a/Engine/Source/Runtime/UMG/Private/Components/MenuAnchor.cpp
+++ b/Engine/Source/Runtime/UMG/Private/Components/MenuAnchor.cpp
@@ -35,11 +35,6 @@ TSharedRef<SWidget> UMenuAnchor::RebuildWidget()
 	return BuildDesignTimeWidget( MyMenuAnchor.ToSharedRef() );
 }
 
-void UMenuAnchor::SyncronizeProperties()
-{
-	Super::SyncronizeProperties();
-}
-
 void UMenuAnchor::OnSlotAdded(UPanelSlot* Slot)
 {
 	// Add the child to the live slot if it already exists
diff --git a/Engine/Source/Runtime/UMG/Private/Components/MultiLineEditableText.cpp b/Engine/Source/Runtime/UMG/Private/Components/MultiLineEditableText.cpp
--- a/Engine/Source/Runtime/UMG/Private/Components/MultiLineEditableText.cpp
+++ b/Engine/Source/Runtime/UMG/Private/Components/MultiLineEditableText.cpp
@@ -52,9 +52,9 @@ TSharedRef<SWidget> UMultiLineEditableText::RebuildWidget()
 	return BuildDesignTimeWidget( MyMultiLineEditableText.ToSharedRef() );
 }
 
-void UMultiLineEditableText::SyncronizeProperties()
+void UMultiLineEditableText::SynchronizeProperties()
 {
-	Super::SyncronizeProperties();
+	Super::SynchronizeProperties();
 
 	MyMultiLineEditableText->SetText(Text);
 //	MyMultiLineEditableText->SetHintText(HintText);
diff --git a/Engine/Source/Runtime/UMG/Private/Components/MultiLineEditableTextBox.cpp b/Engine/Source/Runtime/UMG/Private/Components/MultiLineEditableTextBox.cpp
--- a/Engine/Source/Runtime/UMG/Private/Components/MultiLineEditableTextBox.cpp
+++ b/Engine/Source/Runtime/UMG/Private/Components/MultiLineEditableTextBox.cpp
@@ -47,9 +47,9 @@ TSharedRef<SWidget> UMultiLineEditableTextBox::RebuildWidget()
 	return MyEditableTextBlock.ToSharedRef();
 }
 
-void UMultiLineEditableTextBox::SyncronizeProperties()
+void UMultiLineEditableTextBox::SynchronizeProperties()
 {
-	Super::SyncronizeProperties();
+	Super::SynchronizeProperties();
 
 	//const FMultiLineEditableTextBoxStyle* StylePtr = ( Style != NULL ) ? Style->GetStyle<FMultiLineEditableTextBoxStyle>() : NULL;
 	//if ( StylePtr )
diff --git a/Engine/Source/Runtime/UMG/Private/Components/OverlaySlot.cpp b/Engine/Source/Runtime/UMG/Private/Components/OverlaySlot.cpp
--- a/Engine/Source/Runtime/UMG/Private/Components/OverlaySlot.cpp
+++ b/Engine/Source/Runtime/UMG/Private/Components/OverlaySlot.cpp
@@ -57,7 +57,7 @@ void UOverlaySlot::SetVerticalAlignment(EVerticalAlignment InVerticalAlignment)
 	}
 }
 
-void UOverlaySlot::SyncronizeProperties()
+void UOverlaySlot::SynchronizeProperties()
 {
 	SetPadding(Padding);
 	SetHorizontalAlignment(HorizontalAlignment);
diff --git a/Engine/Source/Runtime/UMG/Private/Components/PanelWidget.cpp b/Engine/Source/Runtime/UMG/Private/Components/PanelWidget.cpp
--- a/Engine/Source/Runtime/UMG/Private/Components/PanelWidget.cpp
+++ b/Engine/Source/Runtime/UMG/Private/Components/PanelWidget.cpp
@@ -112,7 +112,7 @@ void UPanelWidget::ReplaceChildAt(int32 Index, UWidget* Content)
 		Content->Slot = Slot;
 	}
 
-	Slot->SyncronizeProperties();
+	Slot->SynchronizeProperties();
 }
 
 void UPanelWidget::InsertChildAt(int32 Index, UWidget* Content)
diff --git a/Engine/Source/Runtime/UMG/Private/Components/ProcessBar.cpp b/Engine/Source/Runtime/UMG/Private/Components/ProcessBar.cpp
--- a/Engine/Source/Runtime/UMG/Private/Components/ProcessBar.cpp
+++ b/Engine/Source/Runtime/UMG/Private/Components/ProcessBar.cpp
@@ -24,7 +24,7 @@ TSharedRef<SWidget> UProgressBar::RebuildWidget()
 	return MyProgressBar.ToSharedRef();
 }
 
-void UProgressBar::SyncronizeProperties()
+void UProgressBar::SynchronizeProperties()
 {
 	const FProgressBarStyle* StylePtr = ( Style != NULL ) ? Style->GetStyle<FProgressBarStyle>() : NULL;
 	if ( StylePtr )
diff --git a/Engine/Source/Runtime/UMG/Private/Components/ScaleBox.cpp b/Engine/Source/Runtime/UMG/Private/Components/ScaleBox.cpp
--- a/Engine/Source/Runtime/UMG/Private/Components/ScaleBox.cpp
+++ b/Engine/Source/Runtime/UMG/Private/Components/ScaleBox.cpp
@@ -35,9 +35,9 @@ TSharedRef<SWidget> UScaleBox::RebuildWidget()
 	return MyScaleBox.ToSharedRef();
 }
 
-void UScaleBox::SyncronizeProperties()
+void UScaleBox::SynchronizeProperties()
 {
-	Super::SyncronizeProperties();
+	Super::SynchronizeProperties();
 
 	MyScaleBox->SetStretchDirection(StretchDirection);
 	MyScaleBox->SetStretch(Stretch);
diff --git a/Engine/Source/Runtime/UMG/Private/Components/ScaleBoxSlot.cpp b/Engine/Source/Runtime/UMG/Private/Components/ScaleBoxSlot.cpp
--- a/Engine/Source/Runtime/UMG/Private/Components/ScaleBoxSlot.cpp
+++ b/Engine/Source/Runtime/UMG/Private/Components/ScaleBoxSlot.cpp
@@ -59,7 +59,7 @@ void UScaleBoxSlot::SetVerticalAlignment(EVerticalAlignment InVerticalAlignment)
 	}
 }
 
-void UScaleBoxSlot::SyncronizeProperties()
+void UScaleBoxSlot::SynchronizeProperties()
 {
 	SetPadding(Padding);
 	SetHorizontalAlignment(HorizontalAlignment);
diff --git a/Engine/Source/Runtime/UMG/Private/Components/ScrollBar.cpp b/Engine/Source/Runtime/UMG/Private/Components/ScrollBar.cpp
--- a/Engine/Source/Runtime/UMG/Private/Components/ScrollBar.cpp
+++ b/Engine/Source/Runtime/UMG/Private/Components/ScrollBar.cpp
@@ -37,7 +37,7 @@ TSharedRef<SWidget> UScrollBar::RebuildWidget()
 	return MyScrollBar.ToSharedRef();
 }
 
-void UScrollBar::SyncronizeProperties()
+void UScrollBar::SynchronizeProperties()
 {
 	//MyScrollBar->SetScrollOffset(DesiredScrollOffset);
 }
diff --git a/Engine/Source/Runtime/UMG/Private/Components/ScrollBox.cpp b/Engine/Source/Runtime/UMG/Private/Components/ScrollBox.cpp
--- a/Engine/Source/Runtime/UMG/Private/Components/ScrollBox.cpp
+++ b/Engine/Source/Runtime/UMG/Private/Components/ScrollBox.cpp
@@ -84,7 +84,7 @@ TSharedRef<SWidget> UScrollBox::RebuildWidget()
 	return BuildDesignTimeWidget( MyScrollBox.ToSharedRef() );
 }
 
-void UScrollBox::SyncronizeProperties()
+void UScrollBox::SynchronizeProperties()
 {
 	MyScrollBox->SetScrollOffset(DesiredScrollOffset);
 	MyScrollBox->SetOrientation(Orientation);
diff --git a/Engine/Source/Runtime/UMG/Private/Components/ScrollBoxSlot.cpp b/Engine/Source/Runtime/UMG/Private/Components/ScrollBoxSlot.cpp
--- a/Engine/Source/Runtime/UMG/Private/Components/ScrollBoxSlot.cpp
+++ b/Engine/Source/Runtime/UMG/Private/Components/ScrollBoxSlot.cpp
@@ -40,7 +40,7 @@ void UScrollBoxSlot::SetHorizontalAlignment(EHorizontalAlignment InHorizontalAli
 	}
 }
 
-void UScrollBoxSlot::SyncronizeProperties()
+void UScrollBoxSlot::SynchronizeProperties()
 {
 	SetPadding(Padding);
 	SetHorizontalAlignment(HorizontalAlignment);
diff --git a/Engine/Source/Runtime/UMG/Private/Components/SizeBox.cpp b/Engine/Source/Runtime/UMG/Private/Components/SizeBox.cpp
--- a/Engine/Source/Runtime/UMG/Private/Components/SizeBox.cpp
+++ b/Engine/Source/Runtime/UMG/Private/Components/SizeBox.cpp
@@ -147,9 +147,9 @@ void USizeBox::ClearMaxDesiredHeight()
 	}
 }
 
-void USizeBox::SyncronizeProperties()
+void USizeBox::SynchronizeProperties()
 {
-	Super::SyncronizeProperties();
+	Super::SynchronizeProperties();
 
 	if ( bOverride_WidthOverride )
 	{
diff --git a/Engine/Source/Runtime/UMG/Private/Components/SizeBoxSlot.cpp b/Engine/Source/Runtime/UMG/Private/Components/SizeBoxSlot.cpp
--- a/Engine/Source/Runtime/UMG/Private/Components/SizeBoxSlot.cpp
+++ b/Engine/Source/Runtime/UMG/Private/Components/SizeBoxSlot.cpp
@@ -27,7 +27,7 @@ void USizeBoxSlot::BuildSlot(TSharedRef<SBox> InSizeBox)
 {
 	SizeBox = InSizeBox;
 
-	SyncronizeProperties();
+	SynchronizeProperties();
 
 	SizeBox->SetContent(Content ? Content->TakeWidget() : SNullWidget::NullWidget);
 }
@@ -59,7 +59,7 @@ void USizeBoxSlot::SetVerticalAlignment(EVerticalAlignment InVerticalAlignment)
 	}
 }
 
-void USizeBoxSlot::SyncronizeProperties()
+void USizeBoxSlot::SynchronizeProperties()
 {
 	SetPadding(Padding);
 	SetHorizontalAlignment(HorizontalAlignment);
diff --git a/Engine/Source/Runtime/UMG/Private/Components/Slider.cpp b/Engine/Source/Runtime/UMG/Private/Components/Slider.cpp
--- a/Engine/Source/Runtime/UMG/Private/Components/Slider.cpp
+++ b/Engine/Source/Runtime/UMG/Private/Components/Slider.cpp
@@ -25,9 +25,9 @@ TSharedRef<SWidget> USlider::RebuildWidget()
 	return MySlider.ToSharedRef();
 }
 
-void USlider::SyncronizeProperties()
+void USlider::SynchronizeProperties()
 {
-	Super::SyncronizeProperties();
+	Super::SynchronizeProperties();
 
 	TAttribute<float> ValueBinding = OPTIONAL_BINDING(float, Value);
 	
diff --git a/Engine/Source/Runtime/UMG/Private/Components/Spacer.cpp b/Engine/Source/Runtime/UMG/Private/Components/Spacer.cpp
--- a/Engine/Source/Runtime/UMG/Private/Components/Spacer.cpp
+++ b/Engine/Source/Runtime/UMG/Private/Components/Spacer.cpp
@@ -34,9 +34,9 @@ TSharedRef<SWidget> USpacer::RebuildWidget()
 	return MySpacer.ToSharedRef();
 }
 
-void USpacer::SyncronizeProperties()
+void USpacer::SynchronizeProperties()
 {
-	Super::SyncronizeProperties();
+	Super::SynchronizeProperties();
 	
 	MySpacer->SetSize(Size);
 }
diff --git a/Engine/Source/Runtime/UMG/Private/Components/TextBlock.cpp b/Engine/Source/Runtime/UMG/Private/Components/TextBlock.cpp
--- a/Engine/Source/Runtime/UMG/Private/Components/TextBlock.cpp
+++ b/Engine/Source/Runtime/UMG/Private/Components/TextBlock.cpp
@@ -69,9 +69,9 @@ TSharedRef<SWidget> UTextBlock::RebuildWidget()
 	return MyTextBlock.ToSharedRef();
 }
 
-void UTextBlock::SyncronizeProperties()
+void UTextBlock::SynchronizeProperties()
 {
-	Super::SyncronizeProperties();
+	Super::SynchronizeProperties();
 
 	FString FontPath = FPaths::GameContentDir() / Font.FontName.ToString();
 
diff --git a/Engine/Source/Runtime/UMG/Private/Components/Throbber.cpp b/Engine/Source/Runtime/UMG/Private/Components/Throbber.cpp
--- a/Engine/Source/Runtime/UMG/Private/Components/Throbber.cpp
+++ b/Engine/Source/Runtime/UMG/Private/Components/Throbber.cpp
@@ -32,9 +32,9 @@ TSharedRef<SWidget> UThrobber::RebuildWidget()
 	return MyThrobber.ToSharedRef();
 }
 
-void UThrobber::SyncronizeProperties()
+void UThrobber::SynchronizeProperties()
 {
-	Super::SyncronizeProperties();
+	Super::SynchronizeProperties();
 
 	MyThrobber->SetPieceImage(GetPieceBrush());
 	MyThrobber->SetNumPieces(NumberOfPieces);
diff --git a/Engine/Source/Runtime/UMG/Private/Components/UniformGridPanel.cpp b/Engine/Source/Runtime/UMG/Private/Components/UniformGridPanel.cpp
--- a/Engine/Source/Runtime/UMG/Private/Components/UniformGridPanel.cpp
+++ b/Engine/Source/Runtime/UMG/Private/Components/UniformGridPanel.cpp
@@ -69,9 +69,9 @@ TSharedRef<SWidget> UUniformGridPanel::RebuildWidget()
 	return BuildDesignTimeWidget( MyUniformGridPanel.ToSharedRef() );
 }
 
-void UUniformGridPanel::SyncronizeProperties()
+void UUniformGridPanel::SynchronizeProperties()
 {
-	Super::SyncronizeProperties();
+	Super::SynchronizeProperties();
 
 	MyUniformGridPanel->SetSlotPadding(SlotPadding);
 }
diff --git a/Engine/Source/Runtime/UMG/Private/Components/UniformGridSlot.cpp b/Engine/Source/Runtime/UMG/Private/Components/UniformGridSlot.cpp
--- a/Engine/Source/Runtime/UMG/Private/Components/UniformGridSlot.cpp
+++ b/Engine/Source/Runtime/UMG/Private/Components/UniformGridSlot.cpp
@@ -64,7 +64,7 @@ void UUniformGridSlot::SetVerticalAlignment(EVerticalAlignment InVerticalAlignme
 	}
 }
 
-void UUniformGridSlot::SyncronizeProperties()
+void UUniformGridSlot::SynchronizeProperties()
 {
 	SetRow(Row);
 	SetColumn(Column);
diff --git a/Engine/Source/Runtime/UMG/Private/Components/VerticalBoxSlot.cpp b/Engine/Source/Runtime/UMG/Private/Components/VerticalBoxSlot.cpp
--- a/Engine/Source/Runtime/UMG/Private/Components/VerticalBoxSlot.cpp
+++ b/Engine/Source/Runtime/UMG/Private/Components/VerticalBoxSlot.cpp
@@ -70,7 +70,7 @@ void UVerticalBoxSlot::SetVerticalAlignment(EVerticalAlignment InVerticalAlignme
 	}
 }
 
-void UVerticalBoxSlot::SyncronizeProperties()
+void UVerticalBoxSlot::SynchronizeProperties()
 {
 	SetPadding(Padding);
 	SetSize(Size);
diff --git a/Engine/Source/Runtime/UMG/Private/Components/Viewport.cpp b/Engine/Source/Runtime/UMG/Private/Components/Viewport.cpp
--- a/Engine/Source/Runtime/UMG/Private/Components/Viewport.cpp
+++ b/Engine/Source/Runtime/UMG/Private/Components/Viewport.cpp
@@ -518,9 +518,9 @@ TSharedRef<SWidget> UViewport::RebuildWidget()
 	return BuildDesignTimeWidget( ViewportWidget.ToSharedRef() );
 }
 
-void UViewport::SyncronizeProperties()
+void UViewport::SynchronizeProperties()
 {
-	Super::SyncronizeProperties();
+	Super::SynchronizeProperties();
 
 	ViewportWidget->ViewportClient->SetBackgroundColor(BackgroundColor);
 	ViewportWidget->ViewportClient->SetEngineShowFlags(ShowFlags);
diff --git a/Engine/Source/Runtime/UMG/Private/Components/Widget.cpp b/Engine/Source/Runtime/UMG/Private/Components/Widget.cpp
--- a/Engine/Source/Runtime/UMG/Private/Components/Widget.cpp
+++ b/Engine/Source/Runtime/UMG/Private/Components/Widget.cpp
@@ -170,7 +170,7 @@ TSharedRef<SWidget> UWidget::TakeWidget()
 
 			MyGCWidget = SafeGCWidget;
 
-			SyncronizeProperties();
+			SynchronizeProperties();
 
 			return SafeGCWidget.ToSharedRef();
 		}
@@ -179,7 +179,7 @@ TSharedRef<SWidget> UWidget::TakeWidget()
 	{
 		if ( bNewlyCreated )
 		{
-			SyncronizeProperties();
+			SynchronizeProperties();
 		}
 
 		return SafeWidget.ToSharedRef();
@@ -273,7 +273,7 @@ void UWidget::PostEditChangeProperty(FPropertyChangedEvent& PropertyChangedEvent
 
 	if ( MyWidget.IsValid() )
 	{
-		SyncronizeProperties();
+		SynchronizeProperties();
 	}
 }
 
@@ -334,7 +334,7 @@ TSharedRef<SWidget> UWidget::RebuildWidget()
 	return SNew(SSpacer);
 }
 
-void UWidget::SyncronizeProperties()
+void UWidget::SynchronizeProperties()
 {
 	// We want to apply the bindings to the cached widget, which could be the SWidget, or the SObjectWidget, 
 	// in the case where it's a user widget.  We always want to prefer the SObjectWidget so that bindings to 
diff --git a/Engine/Source/Runtime/UMG/Private/Components/WidgetSwitcher.cpp b/Engine/Source/Runtime/UMG/Private/Components/WidgetSwitcher.cpp
--- a/Engine/Source/Runtime/UMG/Private/Components/WidgetSwitcher.cpp
+++ b/Engine/Source/Runtime/UMG/Private/Components/WidgetSwitcher.cpp
@@ -108,9 +108,9 @@ TSharedRef<SWidget> UWidgetSwitcher::RebuildWidget()
 	return BuildDesignTimeWidget( MyWidgetSwitcher.ToSharedRef() );
 }
 
-void UWidgetSwitcher::SyncronizeProperties()
+void UWidgetSwitcher::SynchronizeProperties()
 {
-	Super::SyncronizeProperties();
+	Super::SynchronizeProperties();
 
 	SetActiveWidgetIndex(ActiveWidgetIndex);
 }
diff --git a/Engine/Source/Runtime/UMG/Private/Components/WidgetSwitcherSlot.cpp b/Engine/Source/Runtime/UMG/Private/Components/WidgetSwitcherSlot.cpp
--- a/Engine/Source/Runtime/UMG/Private/Components/WidgetSwitcherSlot.cpp
+++ b/Engine/Source/Runtime/UMG/Private/Components/WidgetSwitcherSlot.cpp
@@ -58,7 +58,7 @@ void UWidgetSwitcherSlot::SetVerticalAlignment(EVerticalAlignment InVerticalAlig
 	}
 }
 
-void UWidgetSwitcherSlot::SyncronizeProperties()
+void UWidgetSwitcherSlot::SynchronizeProperties()
 {
 	SetPadding(Padding);
 	SetHorizontalAlignment(HorizontalAlignment);
diff --git a/Engine/Source/Runtime/UMG/Public/Components/Border.h b/Engine/Source/Runtime/UMG/Public/Components/Border.h
--- a/Engine/Source/Runtime/UMG/Public/Components/Border.h
+++ b/Engine/Source/Runtime/UMG/Public/Components/Border.h
@@ -76,7 +76,7 @@ public:
 	void SetForegroundColor(FLinearColor InForegroundColor);
 
 	// UWidget interface
-	virtual void SyncronizeProperties() override;
+	virtual void SynchronizeProperties() override;
 	// End of UWidget interface
 
 	// UVisual interface
diff --git a/Engine/Source/Runtime/UMG/Public/Components/BorderSlot.h b/Engine/Source/Runtime/UMG/Public/Components/BorderSlot.h
--- a/Engine/Source/Runtime/UMG/Public/Components/BorderSlot.h
+++ b/Engine/Source/Runtime/UMG/Public/Components/BorderSlot.h
@@ -36,7 +36,7 @@ public:
 	void SetVerticalAlignment(EVerticalAlignment InVerticalAlignment);
 
 	// UPanelSlot interface
-	virtual void SyncronizeProperties() override;
+	virtual void SynchronizeProperties() override;
 	// End of UPanelSlot interface
 
 	/** Builds the underlying slot for the slate border. */
diff --git a/Engine/Source/Runtime/UMG/Public/Components/Button.h b/Engine/Source/Runtime/UMG/Public/Components/Button.h
--- a/Engine/Source/Runtime/UMG/Public/Components/Button.h
+++ b/Engine/Source/Runtime/UMG/Public/Components/Button.h
@@ -70,7 +70,7 @@ public:
 	bool IsPressed() const;
 
 	// UWidget interface
-	virtual void SyncronizeProperties() override;
+	virtual void SynchronizeProperties() override;
 	// End of UWidget interface
 
 	// UVisual interface
diff --git a/Engine/Source/Runtime/UMG/Public/Components/ButtonSlot.h b/Engine/Source/Runtime/UMG/Public/Components/ButtonSlot.h
--- a/Engine/Source/Runtime/UMG/Public/Components/ButtonSlot.h
+++ b/Engine/Source/Runtime/UMG/Public/Components/ButtonSlot.h
@@ -36,7 +36,7 @@ public:
 	void SetVerticalAlignment(EVerticalAlignment InVerticalAlignment);
 
 	// UPanelSlot interface
-	virtual void SyncronizeProperties() override;
+	virtual void SynchronizeProperties() override;
 	// End of UPanelSlot interface
 
 	/** Builds the underlying slot for the slate button. */
diff --git a/Engine/Source/Runtime/UMG/Public/Components/CanvasPanel.h b/Engine/Source/Runtime/UMG/Public/Components/CanvasPanel.h
--- a/Engine/Source/Runtime/UMG/Public/Components/CanvasPanel.h
+++ b/Engine/Source/Runtime/UMG/Public/Components/CanvasPanel.h
@@ -20,10 +20,6 @@ public:
 	/** Computes the geometry for a particular slot based on the current geometry of the canvas. */
 	bool GetGeometryForSlot(UCanvasPanelSlot* Slot, FGeometry& ArrangedGeometry) const;
 
-	// UWidget interface
-	virtual void SyncronizeProperties() override;
-	// End of UWidget interface
-
 	void ReleaseNativeWidget() override;
 
 #if WITH_EDITOR
diff --git a/Engine/Source/Runtime/UMG/Public/Components/CanvasPanelSlot.h b/Engine/Source/Runtime/UMG/Public/Components/CanvasPanelSlot.h
--- a/Engine/Source/Runtime/UMG/Public/Components/CanvasPanelSlot.h
+++ b/Engine/Source/Runtime/UMG/Public/Components/CanvasPanelSlot.h
@@ -67,7 +67,7 @@ class UMG_API UCanvasPanelSlot : public UPanelSlot
 	void SetZOrder(int32 InZOrder);
 
 	// UPanelSlot interface
-	virtual void SyncronizeProperties() override;
+	virtual void SynchronizeProperties() override;
 	// End of UPanelSlot interface
 
 	virtual void ReleaseNativeWidget() override;
diff --git a/Engine/Source/Runtime/UMG/Public/Components/CheckBox.h b/Engine/Source/Runtime/UMG/Public/Components/CheckBox.h
--- a/Engine/Source/Runtime/UMG/Public/Components/CheckBox.h
+++ b/Engine/Source/Runtime/UMG/Public/Components/CheckBox.h
@@ -113,7 +113,7 @@ public:
 	void SetCheckedState(ESlateCheckBoxState::Type InCheckedState);
 	
 	// UWidget interface
-	virtual void SyncronizeProperties() override;
+	virtual void SynchronizeProperties() override;
 	// End of UWidget interface
 
 	virtual void ReleaseNativeWidget() override;
diff --git a/Engine/Source/Runtime/UMG/Public/Components/CircularThrobber.h b/Engine/Source/Runtime/UMG/Public/Components/CircularThrobber.h
--- a/Engine/Source/Runtime/UMG/Public/Components/CircularThrobber.h
+++ b/Engine/Source/Runtime/UMG/Public/Components/CircularThrobber.h
@@ -31,7 +31,7 @@ public:
 	void SetPieceImage(USlateBrushAsset* InPieceImage);
 
 	// UWidget interface
-	virtual void SyncronizeProperties() override;
+	virtual void SynchronizeProperties() override;
 	// End of UWidget interface
 
 	virtual void ReleaseNativeWidget() override;
diff --git a/Engine/Source/Runtime/UMG/Public/Components/ComboBox.h b/Engine/Source/Runtime/UMG/Public/Components/ComboBox.h
--- a/Engine/Source/Runtime/UMG/Public/Components/ComboBox.h
+++ b/Engine/Source/Runtime/UMG/Public/Components/ComboBox.h
@@ -20,8 +20,6 @@ public:
 	UPROPERTY(EditDefaultsOnly, Category=Events)
 	FGenerateWidgetUObject OnGenerateWidget;
 
-	virtual void ReleaseNativeWidget() override;
-	
 #if WITH_EDITOR
 	virtual const FSlateBrush* GetEditorIcon() override;
 #endif
diff --git a/Engine/Source/Runtime/UMG/Public/Components/DragPanel.h b/Engine/Source/Runtime/UMG/Public/Components/DragPanel.h
--- a/Engine/Source/Runtime/UMG/Public/Components/DragPanel.h
+++ b/Engine/Source/Runtime/UMG/Public/Components/DragPanel.h
@@ -22,10 +22,6 @@ public:
 
 public:
 	
-	// UWidget interface
-	virtual void SyncronizeProperties() override;
-	// End of UWidget interface
-
 	virtual void ReleaseNativeWidget() override;
 
 #if WITH_EDITOR
diff --git a/Engine/Source/Runtime/UMG/Public/Components/DropPanel.h b/Engine/Source/Runtime/UMG/Public/Components/DropPanel.h
--- a/Engine/Source/Runtime/UMG/Public/Components/DropPanel.h
+++ b/Engine/Source/Runtime/UMG/Public/Components/DropPanel.h
@@ -35,10 +35,6 @@ public:
 
 public:
 
-	// UWidget interface
-	virtual void SyncronizeProperties() override;
-	// End of UWidget interface
-
 	virtual void ReleaseNativeWidget() override;
 
 #if WITH_EDITOR
diff --git a/Engine/Source/Runtime/UMG/Public/Components/EditableText.h b/Engine/Source/Runtime/UMG/Public/Components/EditableText.h
--- a/Engine/Source/Runtime/UMG/Public/Components/EditableText.h
+++ b/Engine/Source/Runtime/UMG/Public/Components/EditableText.h
@@ -112,7 +112,7 @@ public:
 	void SetText(FText InText);
 	
 	// UWidget interface
-	virtual void SyncronizeProperties() override;
+	virtual void SynchronizeProperties() override;
 	// End of UWidget interface
 
 	virtual void ReleaseNativeWidget() override;
diff --git a/Engine/Source/Runtime/UMG/Public/Components/EditableTextBox.h b/Engine/Source/Runtime/UMG/Public/Components/EditableTextBox.h
--- a/Engine/Source/Runtime/UMG/Public/Components/EditableTextBox.h
+++ b/Engine/Source/Runtime/UMG/Public/Components/EditableTextBox.h
@@ -117,7 +117,7 @@ public:
 	void ClearError();
 
 	// UWidget interface
-	virtual void SyncronizeProperties() override;
+	virtual void SynchronizeProperties() override;
 	// End of UWidget interface
 
 	virtual void ReleaseNativeWidget() override;
diff --git a/Engine/Source/Runtime/UMG/Public/Components/GridPanel.h b/Engine/Source/Runtime/UMG/Public/Components/GridPanel.h
--- a/Engine/Source/Runtime/UMG/Public/Components/GridPanel.h
+++ b/Engine/Source/Runtime/UMG/Public/Components/GridPanel.h
@@ -23,7 +23,7 @@ public:
 public:
 
 	// UWidget interface
-	virtual void SyncronizeProperties() override;
+	virtual void SynchronizeProperties() override;
 	// End of UWidget interface
 
 	virtual void ReleaseNativeWidget() override;
diff --git a/Engine/Source/Runtime/UMG/Public/Components/GridSlot.h b/Engine/Source/Runtime/UMG/Public/Components/GridSlot.h
--- a/Engine/Source/Runtime/UMG/Public/Components/GridSlot.h
+++ b/Engine/Source/Runtime/UMG/Public/Components/GridSlot.h
@@ -74,7 +74,7 @@ public:
 	void SetVerticalAlignment(EVerticalAlignment InVerticalAlignment);
 
 	// UPanelSlot interface
-	virtual void SyncronizeProperties() override;
+	virtual void SynchronizeProperties() override;
 	// End of UPanelSlot interface
 
 	/** Builds the underlying FSlot for the Slate layout panel. */
diff --git a/Engine/Source/Runtime/UMG/Public/Components/HorizontalBoxSlot.h b/Engine/Source/Runtime/UMG/Public/Components/HorizontalBoxSlot.h
--- a/Engine/Source/Runtime/UMG/Public/Components/HorizontalBoxSlot.h
+++ b/Engine/Source/Runtime/UMG/Public/Components/HorizontalBoxSlot.h
@@ -42,7 +42,7 @@ public:
 	void BuildSlot(TSharedRef<SHorizontalBox> HorizontalBox);
 
 	// UPanelSlot interface
-	virtual void SyncronizeProperties() override;
+	virtual void SynchronizeProperties() override;
 	// End of UPanelSlot interface
 
 	virtual void ReleaseNativeWidget() override;
diff --git a/Engine/Source/Runtime/UMG/Public/Components/Image.h b/Engine/Source/Runtime/UMG/Public/Components/Image.h
--- a/Engine/Source/Runtime/UMG/Public/Components/Image.h
+++ b/Engine/Source/Runtime/UMG/Public/Components/Image.h
@@ -64,7 +64,7 @@ public:
 	UMaterialInstanceDynamic* GetDynamicMaterial();
 
 	// UWidget interface
-	virtual void SyncronizeProperties() override;
+	virtual void SynchronizeProperties() override;
 	// End of UWidget interface
 
 	virtual void ReleaseNativeWidget() override;
diff --git a/Engine/Source/Runtime/UMG/Public/Components/MenuAnchor.h b/Engine/Source/Runtime/UMG/Public/Components/MenuAnchor.h
--- a/Engine/Source/Runtime/UMG/Public/Components/MenuAnchor.h
+++ b/Engine/Source/Runtime/UMG/Public/Components/MenuAnchor.h
@@ -52,10 +52,6 @@ public:
 	///** @return Whether this menu has open submenus */
 	//bool HasOpenSubMenus() const;
 
-	// UWidget interface
-	virtual void SyncronizeProperties() override;
-	// End of UWidget interface
-
 	virtual void ReleaseNativeWidget() override;
 
 #if WITH_EDITOR
diff --git a/Engine/Source/Runtime/UMG/Public/Components/MultiLineEditableText.h b/Engine/Source/Runtime/UMG/Public/Components/MultiLineEditableText.h
--- a/Engine/Source/Runtime/UMG/Public/Components/MultiLineEditableText.h
+++ b/Engine/Source/Runtime/UMG/Public/Components/MultiLineEditableText.h
@@ -46,7 +46,7 @@ public:
 	void SetText(FText InText);
 	
 	// UWidget interface
-	virtual void SyncronizeProperties() override;
+	virtual void SynchronizeProperties() override;
 	// End of UWidget interface
 
 #if WITH_EDITOR
diff --git a/Engine/Source/Runtime/UMG/Public/Components/MultiLineEditableTextBox.h b/Engine/Source/Runtime/UMG/Public/Components/MultiLineEditableTextBox.h
--- a/Engine/Source/Runtime/UMG/Public/Components/MultiLineEditableTextBox.h
+++ b/Engine/Source/Runtime/UMG/Public/Components/MultiLineEditableTextBox.h
@@ -68,7 +68,7 @@ public:
 	void SetError(FText InError);
 
 	// UWidget interface
-	virtual void SyncronizeProperties() override;
+	virtual void SynchronizeProperties() override;
 	// End of UWidget interface
 
 #if WITH_EDITOR
diff --git a/Engine/Source/Runtime/UMG/Public/Components/OverlaySlot.h b/Engine/Source/Runtime/UMG/Public/Components/OverlaySlot.h
--- a/Engine/Source/Runtime/UMG/Public/Components/OverlaySlot.h
+++ b/Engine/Source/Runtime/UMG/Public/Components/OverlaySlot.h
@@ -36,7 +36,7 @@ class UMG_API UOverlaySlot : public UPanelSlot
 	void BuildSlot(TSharedRef<SHorizontalBox> HorizontalBox);
 
 	// UPanelSlot interface
-	virtual void SyncronizeProperties() override;
+	virtual void SynchronizeProperties() override;
 	// End of UPanelSlot interface
 
 	/** Builds the underlying FSlot for the Slate layout panel. */
diff --git a/Engine/Source/Runtime/UMG/Public/Components/PanelSlot.h b/Engine/Source/Runtime/UMG/Public/Components/PanelSlot.h
--- a/Engine/Source/Runtime/UMG/Public/Components/PanelSlot.h
+++ b/Engine/Source/Runtime/UMG/Public/Components/PanelSlot.h
@@ -35,7 +35,7 @@ public:
 	virtual void ReleaseNativeWidget() override;
 
 	/** Applies all properties to the live slot if possible. */
-	virtual void SyncronizeProperties()
+	virtual void SynchronizeProperties()
 	{
 	}
 
@@ -45,7 +45,7 @@ public:
 	{
 		Super::PostEditChangeProperty(PropertyChangedEvent);
 
-		SyncronizeProperties();
+		SynchronizeProperties();
 	}
 
 #endif
diff --git a/Engine/Source/Runtime/UMG/Public/Components/ProgressBar.h b/Engine/Source/Runtime/UMG/Public/Components/ProgressBar.h
--- a/Engine/Source/Runtime/UMG/Public/Components/ProgressBar.h
+++ b/Engine/Source/Runtime/UMG/Public/Components/ProgressBar.h
@@ -58,7 +58,7 @@ public:
 	void SetPercent(float InPercent);
 	
 	// UWidget interface
-	virtual void SyncronizeProperties() override;
+	virtual void SynchronizeProperties() override;
 	// End of UWidget interface
 
 #if WITH_EDITOR
diff --git a/Engine/Source/Runtime/UMG/Public/Components/ScaleBox.h b/Engine/Source/Runtime/UMG/Public/Components/ScaleBox.h
--- a/Engine/Source/Runtime/UMG/Public/Components/ScaleBox.h
+++ b/Engine/Source/Runtime/UMG/Public/Components/ScaleBox.h
@@ -27,7 +27,7 @@ public:
 public:
 
 	// UWidget interface
-	virtual void SyncronizeProperties() override;
+	virtual void SynchronizeProperties() override;
 	// End of UWidget interface
 
 	virtual void ReleaseNativeWidget() override;
diff --git a/Engine/Source/Runtime/UMG/Public/Components/ScaleBoxSlot.h b/Engine/Source/Runtime/UMG/Public/Components/ScaleBoxSlot.h
--- a/Engine/Source/Runtime/UMG/Public/Components/ScaleBoxSlot.h
+++ b/Engine/Source/Runtime/UMG/Public/Components/ScaleBoxSlot.h
@@ -38,7 +38,7 @@ public:
 	void SetVerticalAlignment(EVerticalAlignment InVerticalAlignment);
 
 	// UPanelSlot interface
-	virtual void SyncronizeProperties() override;
+	virtual void SynchronizeProperties() override;
 	// End of UPanelSlot interface
 
 	/** Builds the underlying slot for the slate ScaleBox. */
diff --git a/Engine/Source/Runtime/UMG/Public/Components/ScrollBar.h b/Engine/Source/Runtime/UMG/Public/Components/ScrollBar.h
--- a/Engine/Source/Runtime/UMG/Public/Components/ScrollBar.h
+++ b/Engine/Source/Runtime/UMG/Public/Components/ScrollBar.h
@@ -58,7 +58,7 @@ public:
 
 
 	// UWidget interface
-	virtual void SyncronizeProperties() override;
+	virtual void SynchronizeProperties() override;
 	// End of UWidget interface
 
 #if WITH_EDITOR
diff --git a/Engine/Source/Runtime/UMG/Public/Components/ScrollBox.h b/Engine/Source/Runtime/UMG/Public/Components/ScrollBox.h
--- a/Engine/Source/Runtime/UMG/Public/Components/ScrollBox.h
+++ b/Engine/Source/Runtime/UMG/Public/Components/ScrollBox.h
@@ -29,7 +29,7 @@ class UMG_API UScrollBox : public UPanelWidget
 	void SetScrollOffset(float NewScrollOffset);
 
 	// UWidget interface
-	virtual void SyncronizeProperties() override;
+	virtual void SynchronizeProperties() override;
 	// End of UWidget interface
 
 #if WITH_EDITOR
diff --git a/Engine/Source/Runtime/UMG/Public/Components/ScrollBoxSlot.h b/Engine/Source/Runtime/UMG/Public/Components/ScrollBoxSlot.h
--- a/Engine/Source/Runtime/UMG/Public/Components/ScrollBoxSlot.h
+++ b/Engine/Source/Runtime/UMG/Public/Components/ScrollBoxSlot.h
@@ -29,7 +29,7 @@ public:
 	void SetHorizontalAlignment(EHorizontalAlignment InHorizontalAlignment);
 
 	// UPanelSlot interface
-	virtual void SyncronizeProperties() override;
+	virtual void SynchronizeProperties() override;
 	// End of UPanelSlot interface
 
 	/** Builds the underlying FSlot for the Slate layout panel. */
diff --git a/Engine/Source/Runtime/UMG/Public/Components/SizeBox.h b/Engine/Source/Runtime/UMG/Public/Components/SizeBox.h
--- a/Engine/Source/Runtime/UMG/Public/Components/SizeBox.h
+++ b/Engine/Source/Runtime/UMG/Public/Components/SizeBox.h
@@ -62,7 +62,7 @@ public:
 public:
 
 	// UWidget interface
-	virtual void SyncronizeProperties() override;
+	virtual void SynchronizeProperties() override;
 	// End of UWidget interface
 
 	// UVisual interface
diff --git a/Engine/Source/Runtime/UMG/Public/Components/SizeBoxSlot.h b/Engine/Source/Runtime/UMG/Public/Components/SizeBoxSlot.h
--- a/Engine/Source/Runtime/UMG/Public/Components/SizeBoxSlot.h
+++ b/Engine/Source/Runtime/UMG/Public/Components/SizeBoxSlot.h
@@ -26,7 +26,7 @@ public:
 public:
 
 	// UPanelSlot interface
-	virtual void SyncronizeProperties() override;
+	virtual void SynchronizeProperties() override;
 	// End of UPanelSlot interface
 
 	/** Builds the underlying slot for the slate SizeBox. */
diff --git a/Engine/Source/Runtime/UMG/Public/Components/Slider.h b/Engine/Source/Runtime/UMG/Public/Components/Slider.h
--- a/Engine/Source/Runtime/UMG/Public/Components/Slider.h
+++ b/Engine/Source/Runtime/UMG/Public/Components/Slider.h
@@ -68,7 +68,7 @@ public:
 	void SetValue(float InValue);
 	
 	// UWidget interface
-	virtual void SyncronizeProperties() override;
+	virtual void SynchronizeProperties() override;
 	// End of UWidget interface
 
 #if WITH_EDITOR
diff --git a/Engine/Source/Runtime/UMG/Public/Components/Spacer.h b/Engine/Source/Runtime/UMG/Public/Components/Spacer.h
--- a/Engine/Source/Runtime/UMG/Public/Components/Spacer.h
+++ b/Engine/Source/Runtime/UMG/Public/Components/Spacer.h
@@ -23,7 +23,7 @@ public:
 	void SetSize(FVector2D InSize);
 	
 	// UWidget interface
-	virtual void SyncronizeProperties() override;
+	virtual void SynchronizeProperties() override;
 	// End of UWidget interface
 
 #if WITH_EDITOR
diff --git a/Engine/Source/Runtime/UMG/Public/Components/TextBlock.h b/Engine/Source/Runtime/UMG/Public/Components/TextBlock.h
--- a/Engine/Source/Runtime/UMG/Public/Components/TextBlock.h
+++ b/Engine/Source/Runtime/UMG/Public/Components/TextBlock.h
@@ -89,7 +89,7 @@ public:
 	void SetText(FText InText);
 
 	// UWidget interface
-	void SyncronizeProperties() override;
+	void SynchronizeProperties() override;
 	// End of UWidget interface
 
 #if WITH_EDITOR
diff --git a/Engine/Source/Runtime/UMG/Public/Components/Throbber.h b/Engine/Source/Runtime/UMG/Public/Components/Throbber.h
--- a/Engine/Source/Runtime/UMG/Public/Components/Throbber.h
+++ b/Engine/Source/Runtime/UMG/Public/Components/Throbber.h
@@ -35,7 +35,7 @@ public:
 	void SetPieceImage(USlateBrushAsset* InPieceImage);
 
 	// UWidget interface
-	virtual void SyncronizeProperties() override;
+	virtual void SynchronizeProperties() override;
 	// End of UWidget interface
 
 #if WITH_EDITOR
diff --git a/Engine/Source/Runtime/UMG/Public/Components/UniformGridPanel.h b/Engine/Source/Runtime/UMG/Public/Components/UniformGridPanel.h
--- a/Engine/Source/Runtime/UMG/Public/Components/UniformGridPanel.h
+++ b/Engine/Source/Runtime/UMG/Public/Components/UniformGridPanel.h
@@ -27,7 +27,7 @@ public:
 public:
 
 	// UWidget interface
-	virtual void SyncronizeProperties() override;
+	virtual void SynchronizeProperties() override;
 	// End of UWidget interface
 
 	virtual void ReleaseNativeWidget() override;
diff --git a/Engine/Source/Runtime/UMG/Public/Components/UniformGridSlot.h b/Engine/Source/Runtime/UMG/Public/Components/UniformGridSlot.h
--- a/Engine/Source/Runtime/UMG/Public/Components/UniformGridSlot.h
+++ b/Engine/Source/Runtime/UMG/Public/Components/UniformGridSlot.h
@@ -50,7 +50,7 @@ public:
 	void SetVerticalAlignment(EVerticalAlignment InVerticalAlignment);
 
 	// UPanelSlot interface
-	virtual void SyncronizeProperties() override;
+	virtual void SynchronizeProperties() override;
 	// End of UPanelSlot interface
 
 	/** Builds the underlying FSlot for the Slate layout panel. */
diff --git a/Engine/Source/Runtime/UMG/Public/Components/VerticalBoxSlot.h b/Engine/Source/Runtime/UMG/Public/Components/VerticalBoxSlot.h
--- a/Engine/Source/Runtime/UMG/Public/Components/VerticalBoxSlot.h
+++ b/Engine/Source/Runtime/UMG/Public/Components/VerticalBoxSlot.h
@@ -45,7 +45,7 @@ public:
 	void SetVerticalAlignment(EVerticalAlignment InVerticalAlignment);
 
 	// UPanelSlot interface
-	virtual void SyncronizeProperties() override;
+	virtual void SynchronizeProperties() override;
 	// End of UPanelSlot interface
 
 	virtual void ReleaseNativeWidget() override;
diff --git a/Engine/Source/Runtime/UMG/Public/Components/Viewport.h b/Engine/Source/Runtime/UMG/Public/Components/Viewport.h
--- a/Engine/Source/Runtime/UMG/Public/Components/Viewport.h
+++ b/Engine/Source/Runtime/UMG/Public/Components/Viewport.h
@@ -265,7 +265,7 @@ class UMG_API UViewport : public UContentWidget
 	AActor* Spawn(TSubclassOf<AActor> ActorClass);
 
 	// UWidget interface
-	virtual void SyncronizeProperties() override;
+	virtual void SynchronizeProperties() override;
 	// End of UWidget interface
 
 	virtual void ReleaseNativeWidget() override;
diff --git a/Engine/Source/Runtime/UMG/Public/Components/Widget.h b/Engine/Source/Runtime/UMG/Public/Components/Widget.h
--- a/Engine/Source/Runtime/UMG/Public/Components/Widget.h
+++ b/Engine/Source/Runtime/UMG/Public/Components/Widget.h
@@ -180,7 +180,7 @@ public:
 	 * It can also be called by the editor to update modified state, so ensure all initialization to a widgets
 	 * properties are performed here, or the property and visual state may become unsynced.
 	 */
-	virtual void SyncronizeProperties();
+	virtual void SynchronizeProperties();
 
 	/** Returns if the widget is currently being displayed in the designer, it may want to display different data. */
 	bool IsDesignTime() const;
diff --git a/Engine/Source/Runtime/UMG/Public/Components/WidgetSwitcher.h b/Engine/Source/Runtime/UMG/Public/Components/WidgetSwitcher.h
--- a/Engine/Source/Runtime/UMG/Public/Components/WidgetSwitcher.h
+++ b/Engine/Source/Runtime/UMG/Public/Components/WidgetSwitcher.h
@@ -36,7 +36,7 @@ public:
 	UFUNCTION(BlueprintCallable, Category="Switcher")
 	void SetActiveWidget(UWidget* Widget);
 	
-	virtual void SyncronizeProperties() override;
+	virtual void SynchronizeProperties() override;
 
 	virtual void ReleaseNativeWidget() override;
 
diff --git a/Engine/Source/Runtime/UMG/Public/Components/WidgetSwitcherSlot.h b/Engine/Source/Runtime/UMG/Public/Components/WidgetSwitcherSlot.h
--- a/Engine/Source/Runtime/UMG/Public/Components/WidgetSwitcherSlot.h
+++ b/Engine/Source/Runtime/UMG/Public/Components/WidgetSwitcherSlot.h
@@ -34,7 +34,7 @@ class UMG_API UWidgetSwitcherSlot : public UPanelSlot
 	void SetVerticalAlignment(EVerticalAlignment InVerticalAlignment);
 
 	// UPanelSlot interface
-	virtual void SyncronizeProperties() override;
+	virtual void SynchronizeProperties() override;
 	// End of UPanelSlot interface
 
 	virtual void ReleaseNativeWidget() override;
