Automation:  Fixed non unity build errors with the new editor build promotion test.

diff --git a/Engine/Source/Editor/UnrealEd/Private/Tests/AutomationEditorCommon.cpp b/Engine/Source/Editor/UnrealEd/Private/Tests/AutomationEditorCommon.cpp
--- a/Engine/Source/Editor/UnrealEd/Private/Tests/AutomationEditorCommon.cpp
+++ b/Engine/Source/Editor/UnrealEd/Private/Tests/AutomationEditorCommon.cpp
@@ -16,7 +16,7 @@ namespace AutomationEditorCommonUtils
 	*
 	* @param PackagePath - The package path to convert
 	*/
-	static FString ConvertPackagePathToAssetPath(const FString& PackagePath)
+	FString ConvertPackagePathToAssetPath(const FString& PackagePath)
 	{
 		const FString Filename = FPaths::ConvertRelativePathToFull(PackagePath);
 		FString EngineFileName = Filename;
@@ -50,7 +50,7 @@ namespace AutomationEditorCommonUtils
 	* @param PackagePath - The full path of the package file to create
 	* @param ImportPath - The path to the object to import
 	*/
-	static UObject* ImportAssetUsingFactory(UFactory* ImportFactory, const FString& ObjectName, const FString& PackagePath, const FString& ImportPath)
+	UObject* ImportAssetUsingFactory(UFactory* ImportFactory, const FString& ObjectName, const FString& PackagePath, const FString& ImportPath)
 	{
 		UObject* ImportedAsset = NULL;
 
@@ -94,7 +94,7 @@ namespace AutomationEditorCommonUtils
 	*
 	* @param InObject - Object to null references to
 	*/
-	static void NullReferencesToObject(UObject* InObject)
+	void NullReferencesToObject(UObject* InObject)
 	{
 		TArray<UObject*> ReplaceableObjects;
 		TMap<UObject*, UObject*> ReplacementMap;
@@ -161,7 +161,7 @@ namespace AutomationEditorCommonUtils
 	*
 	* @param AssetExtension - The file extension to use to find a supporting UFactory
 	*/
-	static UClass* GetFactoryClassForType(const FString& AssetExtension)
+	UClass* GetFactoryClassForType(const FString& AssetExtension)
 	{
 		// First instantiate one factory for each file extension encountered that supports the extension
 		for (TObjectIterator<UClass> ClassIt; ClassIt; ++ClassIt)
@@ -193,7 +193,7 @@ namespace AutomationEditorCommonUtils
 	* @param PropertyChain - The list UProperty names recursively to search through
 	* @param Value - The value to import on the found property
 	*/
-	static void ApplyCustomFactorySetting(UObject* InObject, TArray<FString>& PropertyChain, const FString& Value)
+	void ApplyCustomFactorySetting(UObject* InObject, TArray<FString>& PropertyChain, const FString& Value)
 	{
 		const FString PropertyName = PropertyChain[0];
 		PropertyChain.RemoveAt(0);
@@ -250,7 +250,7 @@ namespace AutomationEditorCommonUtils
 	* @param InFactory - The factory to apply custom settings to
 	* @param FactorySettings - An array of custom settings to apply to the factory
 	*/
-	static void ApplyCustomFactorySettings(UFactory* InFactory, const TArray<FImportFactorySettingValues>& FactorySettings)
+	void ApplyCustomFactorySettings(UFactory* InFactory, const TArray<FImportFactorySettingValues>& FactorySettings)
 	{
 		bool bCallConfigureProperties = true;
 
diff --git a/Engine/Source/Editor/UnrealEd/Private/Tests/EditorBuildPromotionTests.cpp b/Engine/Source/Editor/UnrealEd/Private/Tests/EditorBuildPromotionTests.cpp
--- a/Engine/Source/Editor/UnrealEd/Private/Tests/EditorBuildPromotionTests.cpp
+++ b/Engine/Source/Editor/UnrealEd/Private/Tests/EditorBuildPromotionTests.cpp
@@ -11,9 +11,11 @@
 #include "AssetRegistryModule.h"
 #include "IAssetRegistry.h"
 #include "ComponentAssetBroker.h"
+#include "AssetSelection.h"
 
 //Materials
 #include "AssetEditorManager.h"
+#include "AssetEditorToolkit.h"
 #include "IMaterialEditor.h"
 #include "Private/MaterialEditor.h"
 #include "Materials/MaterialExpressionTextureSample.h"
@@ -26,7 +28,7 @@
 #include "Particles/Size/ParticleModuleSize.h"
 
 //Blueprints
-#include "EdGraphSchema_K2_Actions.h"
+#include "BlueprintGraphDefinitions.h"
 #include "BlueprintEditor.h"
 #include "BlueprintEditorModes.h"
 #include "CompilerResultsLog.h"
@@ -34,6 +36,7 @@
 #include "KismetDebugUtilities.h"
 #include "Engine/Breakpoint.h"
 #include "Engine/LevelScriptBlueprint.h"
+#include "BlueprintEditorUtils.h"
 
 //Utils
 #include "EditorBuildUtils.h"
@@ -49,6 +52,10 @@
 #include "LightingBuildOptions.h"
 #include "../Private/StaticLightingSystem/StaticLightingPrivate.h"
 
+#include "ScopedTransaction.h"
+#include "LevelEditor.h"
+#include "ObjectTools.h"
+
 #define LOCTEXT_NAMESPACE "EditorBuildPromotionTests"
 
 DEFINE_LOG_CATEGORY_STATIC(LogEditorBuildPromotionTests, Log, All);
@@ -1086,8 +1093,8 @@ namespace EditorBuildPromotionTestUtils
 	*/
 	static UObject* ImportAsset(UFactory* ImportFactory, const FString& ImportPath)
 	{
-		FString Name = ObjectTools::SanitizeObjectName(FPaths::GetBaseFilename(ImportPath));
-		FString PackageName = FString::Printf(TEXT("%s/%s"), *GetGamePath(), *Name);
+		const FString Name = ObjectTools::SanitizeObjectName(FPaths::GetBaseFilename(ImportPath));
+		const FString PackageName = FString::Printf(TEXT("%s/%s"), *GetGamePath(), *Name);
 
 		UObject* ImportedAsset = AutomationEditorCommonUtils::ImportAssetUsingFactory(ImportFactory,Name,PackageName,ImportPath);
 
@@ -1474,7 +1481,7 @@ namespace BuildPromotionTestHelper
 					FAssetEditorManager::Get().FindEditorForAsset(CurrentAsset, true);
 				}
 
-				UE_LOG(LogEditorAssetAutomationTests, Display, TEXT("Opened asset (%s)"), *CurrentAsset->GetClass()->GetName());
+				UE_LOG(LogEditorBuildPromotionTests, Display, TEXT("Opened asset (%s)"), *CurrentAsset->GetClass()->GetName());
 
 				CurrentStage = EState::ChangeProperty;
 			}
@@ -1503,7 +1510,7 @@ namespace BuildPromotionTestHelper
 
 			FString OldPropertyValue = EditorBuildPromotionTestUtils::GetPropertyByName(CurrentAsset, PropertyName);
 			EditorBuildPromotionTestUtils::SetPropertyByName(CurrentAsset, PropertyName, NewPropertyValue);
-			UE_LOG(LogEditorAssetAutomationTests, Display, TEXT("Modified asset.  %s = %s"), *PropertyName, *NewPropertyValue);
+			UE_LOG(LogEditorBuildPromotionTests, Display, TEXT("Modified asset.  %s = %s"), *PropertyName, *NewPropertyValue);
 
 			//Get the property again and use that to compare the redo action.  Parsing the new value may change the formatting a bit. ie) 100 becomes 100.0000
 			FString ParsedNewValue = EditorBuildPromotionTestUtils::GetPropertyByName(CurrentAsset, PropertyName);
@@ -1512,7 +1519,7 @@ namespace BuildPromotionTestHelper
 			FString CurrentValue = EditorBuildPromotionTestUtils::GetPropertyByName(CurrentAsset, PropertyName);
 			if (CurrentValue == OldPropertyValue)
 			{
-				UE_LOG(LogEditorAssetAutomationTests, Display, TEXT("Undo %s change successful"), *PropertyName);
+				UE_LOG(LogEditorBuildPromotionTests, Display, TEXT("Undo %s change successful"), *PropertyName);
 			}
 			else
 			{
@@ -1523,7 +1530,7 @@ namespace BuildPromotionTestHelper
 			CurrentValue = EditorBuildPromotionTestUtils::GetPropertyByName(CurrentAsset, PropertyName);
 			if (CurrentValue == ParsedNewValue)
 			{
-				UE_LOG(LogEditorAssetAutomationTests, Display, TEXT("Redo %s change successful"), *PropertyName);
+				UE_LOG(LogEditorBuildPromotionTests, Display, TEXT("Redo %s change successful"), *PropertyName);
 			}
 			else
 			{
@@ -1551,11 +1558,11 @@ namespace BuildPromotionTestHelper
 					FPlatformFileManager::Get().GetPlatformFile().SetReadOnly(*PackageFileName, false);
 					if (UPackage::SavePackage(AssetPackage, NULL, RF_Standalone, *PackageFileName, GError, nullptr, false, true, SAVE_NoError))
 					{
-						UE_LOG(LogEditorAssetAutomationTests, Display, TEXT("Saved asset"));
+						UE_LOG(LogEditorBuildPromotionTests, Display, TEXT("Saved asset"));
 					}
 					else
 					{
-						UE_LOG(LogEditorAssetAutomationTests, Display, TEXT("Unable to save asset"));
+						UE_LOG(LogEditorBuildPromotionTests, Display, TEXT("Unable to save asset"));
 					}
 				}
 			}
@@ -1582,7 +1589,7 @@ namespace BuildPromotionTestHelper
 				AActor* PlacedActor = FActorFactoryAssetProxy::AddActorForAsset(CurrentAsset);
 				if (PlacedActor)
 				{
-					UE_LOG(LogEditorAssetAutomationTests, Display, TEXT("Placed %s in the level"), *CurrentAsset->GetName());
+					UE_LOG(LogEditorBuildPromotionTests, Display, TEXT("Placed %s in the level"), *CurrentAsset->GetName());
 					PlacedActor->SetActorLocation(PlaceLocation);
 				}
 				else
diff --git a/Engine/Source/Editor/UnrealEd/Public/Tests/AutomationEditorCommon.h b/Engine/Source/Editor/UnrealEd/Public/Tests/AutomationEditorCommon.h
--- a/Engine/Source/Editor/UnrealEd/Public/Tests/AutomationEditorCommon.h
+++ b/Engine/Source/Editor/UnrealEd/Public/Tests/AutomationEditorCommon.h
@@ -48,7 +48,7 @@ namespace AutomationEditorCommonUtils
 	*
 	* @param PackagePath - The package path to convert
 	*/
-	static FString ConvertPackagePathToAssetPath(const FString& PackagePath);
+	FString ConvertPackagePathToAssetPath(const FString& PackagePath);
 
 	/**
 	* Imports an object using a given factory
@@ -58,21 +58,21 @@ namespace AutomationEditorCommonUtils
 	* @param PackagePath - The full path of the package file to create
 	* @param ImportPath - The path to the object to import
 	*/
-	static UObject* ImportAssetUsingFactory(UFactory* ImportFactory, const FString& ObjectName, const FString& PackageName, const FString& ImportPath);
+	UObject* ImportAssetUsingFactory(UFactory* ImportFactory, const FString& ObjectName, const FString& PackageName, const FString& ImportPath);
 
 	/**
 	* Nulls out references to a given object
 	*
 	* @param InObject - Object to null references to
 	*/
-	static void NullReferencesToObject(UObject* InObject);
+	void NullReferencesToObject(UObject* InObject);
 
 	/**
 	* gets a factory class based off an asset file extension
 	*
 	* @param AssetExtension - The file extension to use to find a supporting UFactory
 	*/
-	static UClass* GetFactoryClassForType(const FString& AssetExtension);
+	UClass* GetFactoryClassForType(const FString& AssetExtension);
 
 	/**
 	* Applies settings to an object by finding UProperties by name and calling ImportText
@@ -81,7 +81,7 @@ namespace AutomationEditorCommonUtils
 	* @param PropertyChain - The list UProperty names recursively to search through
 	* @param Value - The value to import on the found property
 	*/
-	static void ApplyCustomFactorySetting(UObject* InObject, TArray<FString>& PropertyChain, const FString& Value);
+	void ApplyCustomFactorySetting(UObject* InObject, TArray<FString>& PropertyChain, const FString& Value);
 
 	/**
 	* Applies the custom factory settings
@@ -89,7 +89,7 @@ namespace AutomationEditorCommonUtils
 	* @param InFactory - The factory to apply custom settings to
 	* @param FactorySettings - An array of custom settings to apply to the factory
 	*/
-	static void ApplyCustomFactorySettings(UFactory* InFactory, const TArray<FImportFactorySettingValues>& FactorySettings);
+	void ApplyCustomFactorySettings(UFactory* InFactory, const TArray<FImportFactorySettingValues>& FactorySettings);
 }
 
 
