Fixed erroneous comment

diff --git a/Engine/Source/Editor/UnrealEd/Classes/ActorFactories/ActorFactory.h b/Engine/Source/Editor/UnrealEd/Classes/ActorFactories/ActorFactory.h
--- a/Engine/Source/Editor/UnrealEd/Classes/ActorFactories/ActorFactory.h
+++ b/Engine/Source/Editor/UnrealEd/Classes/ActorFactories/ActorFactory.h
@@ -60,7 +60,7 @@ class UNREALED_API UActorFactory : public UObject
 	/** Given an instance of an actor pertaining to this factory, find the asset that should be used to create a new actor */
 	virtual UObject* GetAssetFromActorInstance(AActor* ActorInstance);
 
-	/** Return a rotator which aligns this actor type to the specified surface normal */
+	/** Return a quaternion which aligns this actor type to the specified surface normal */
 	virtual FQuat AlignObjectToSurfaceNormal(const FVector& InSurfaceNormal, const FQuat& ActorRotation = FQuat::Identity) const;
 
 protected:
diff --git a/Engine/Source/Editor/UnrealEd/Private/Factories/ActorFactory.cpp b/Engine/Source/Editor/UnrealEd/Private/Factories/ActorFactory.cpp
--- a/Engine/Source/Editor/UnrealEd/Private/Factories/ActorFactory.cpp
+++ b/Engine/Source/Editor/UnrealEd/Private/Factories/ActorFactory.cpp
@@ -137,7 +137,6 @@ UObject* UActorFactory::GetAssetFromActorInstance(AActor* ActorInstance)
 	return NULL;
 }
 
-/** Return a rotator which aligns this actor type to the specified surface normal */
 FQuat UActorFactory::AlignObjectToSurfaceNormal(const FVector& InSurfaceNormal, const FQuat& ActorRotation) const
 {
 	if (bUseSurfaceOrientation)
