Checking only that UWorld object was removed instead of host package when unloading streaming level

diff --git a/Engine/Source/Editor/UnrealEd/Private/EditorLevelUtils.cpp b/Engine/Source/Editor/UnrealEd/Private/EditorLevelUtils.cpp
--- a/Engine/Source/Editor/UnrealEd/Private/EditorLevelUtils.cpp
+++ b/Engine/Source/Editor/UnrealEd/Private/EditorLevelUtils.cpp
@@ -423,15 +423,18 @@ namespace EditorLevelUtils
 			// Reset transaction buffer and run GC to clear out the destroyed level
 			GEditor->Cleanse(true, false, LOCTEXT("RemoveLevelTransReset", "Removing Levels from World"));
 
-			// Ensure that level package was removed
+			// Ensure that world was removed
 			UPackage* LevelPackage = FindObjectFast<UPackage>(NULL, LevelPackageName);
 			if (LevelPackage != nullptr)
 			{
-				StaticExec(NULL, *FString::Printf(TEXT("OBJ REFS CLASS=%s NAME=%s shortest"),*LevelPackage->GetClass()->GetName(), *LevelPackage->GetPathName()));
-				TMap<UObject*,UProperty*>	Route		= FArchiveTraceRoute::FindShortestRootPath( LevelPackage, true, GARBAGE_COLLECTION_KEEPFLAGS );
-				FString						ErrorString	= FArchiveTraceRoute::PrintRootPath( Route, LevelPackage );
-			
-				UE_LOG(LogStreaming, Fatal, TEXT("%s didn't get garbage collected!") LINE_TERMINATOR TEXT("%s"), *LevelPackage->GetFullName(), *ErrorString );
+				UWorld* TheWorld = UWorld::FindWorldInPackage(LevelPackage->GetOutermost());								
+				if (TheWorld != nullptr)
+				{	
+					StaticExec(NULL, *FString::Printf(TEXT("OBJ REFS CLASS=%s NAME=%s shortest"),*TheWorld->GetClass()->GetName(), *TheWorld->GetPathName()));
+					TMap<UObject*,UProperty*>	Route		= FArchiveTraceRoute::FindShortestRootPath( TheWorld, true, GARBAGE_COLLECTION_KEEPFLAGS );
+					FString						ErrorString	= FArchiveTraceRoute::PrintRootPath( Route, TheWorld );
+					UE_LOG(LogStreaming, Fatal, TEXT("%s didn't get garbage collected!") LINE_TERMINATOR TEXT("%s"), *TheWorld->GetFullName(), *ErrorString );
+				}
 			}
 		}
 		return bRemoveSuccessful;
@@ -529,10 +532,6 @@ namespace EditorLevelUtils
 		InLevel->GetOuter()->MarkPendingKill();
 		InLevel->MarkPendingKill();		
 		InLevel->GetOuter()->ClearFlags(RF_Public|RF_Standalone);
-		if (InLevel->GetOutermost()->MetaData)
-		{
-			InLevel->GetOutermost()->MetaData->MarkPendingKill();
-		}
 
 		World->MarkPackageDirty();
 		World->BroadcastLevelsChanged();
