Integrated GitHub PullRequest 346
https://github.com/EpicGames/UnrealEngine/pull/346

diff --git a/Engine/Source/Editor/UnrealEd/Private/Toolkits/SimpleAssetEditor.cpp b/Engine/Source/Editor/UnrealEd/Private/Toolkits/SimpleAssetEditor.cpp
--- a/Engine/Source/Editor/UnrealEd/Private/Toolkits/SimpleAssetEditor.cpp
+++ b/Engine/Source/Editor/UnrealEd/Private/Toolkits/SimpleAssetEditor.cpp
@@ -47,7 +47,7 @@ void FSimpleAssetEditor::InitEditor( const EToolkitMode::Type Mode, const TShare
 	const bool bIsLockable = false;
 
 	FPropertyEditorModule& PropertyEditorModule = FModuleManager::GetModuleChecked<FPropertyEditorModule>( "PropertyEditor" );
-	const FDetailsViewArgs DetailsViewArgs( bIsUpdatable, bIsLockable, true, false, false );
+	const FDetailsViewArgs DetailsViewArgs( bIsUpdatable, bIsLockable, true, true, false );
 	DetailsView = PropertyEditorModule.CreateDetailView( DetailsViewArgs );
 	const TSharedRef<FTabManager::FLayout> StandaloneDefaultLayout = FTabManager::NewLayout( "Standalone_SimpleAssetEditor_Layout_v3" )
 	->AddArea
