Fix Autosave appearing when focused and recently typed in a  text block

diff --git a/Engine/Source/Editor/UnrealEd/Private/PackageAutoSaver.cpp b/Engine/Source/Editor/UnrealEd/Private/PackageAutoSaver.cpp
--- a/Engine/Source/Editor/UnrealEd/Private/PackageAutoSaver.cpp
+++ b/Engine/Source/Editor/UnrealEd/Private/PackageAutoSaver.cpp
@@ -357,13 +357,22 @@ bool FPackageAutoSaver::CanAutoSave() const
 	// Don't allow auto-saving if the auto-save wouldn't save anything
 	const bool bPackagesNeedAutoSave = DoPackagesNeedAutoSave();
 
+	static FName EditableText("SEditableText");
+
+	double LastInteractionTime = FSlateApplication::Get().GetLastUserInteractionTime();
+	
+	const float EditableTextDelay = 15.0f;
+	TSharedPtr<SWidget> KeyboardFocusedWidget = FSlateApplication::Get().GetKeyboardFocusedWidget();
+
+	bool bDidTypeInATextBlockRecently = KeyboardFocusedWidget.IsValid() && KeyboardFocusedWidget->GetType() == EditableText && (FApp::GetCurrentTime() - LastInteractionTime) < EditableTextDelay;
+
 	const bool bAutosaveEnabled	= GetDefault<UEditorLoadingSavingSettings>()->bAutoSaveEnable && bPackagesNeedAutoSave;
 	const bool bSlowTask = GIsSlowTask;
 	const bool bInterpEditMode = GLevelEditorModeTools().IsModeActive(FBuiltinEditorModes::EM_InterpEdit);
 	const bool bPlayWorldValid = GUnrealEd->PlayWorld != nullptr;
 	const bool bAnyMenusVisible	= FSlateApplication::Get().AnyMenusVisible();
 	const bool bAutomationTesting = GIsAutomationTesting;
-	const bool bIsInteratcting = FSlateApplication::Get().GetMouseCaptor().IsValid() || GUnrealEd->IsUserInteracting();
+	const bool bIsInteratcting = FSlateApplication::Get().GetMouseCaptor().IsValid() || GUnrealEd->IsUserInteracting() || bDidTypeInATextBlockRecently;
 	const bool bHasGameOrProjectLoaded = FApp::HasGameName();
 
 	return (bAutosaveEnabled && !bSlowTask && !bInterpEditMode && !bPlayWorldValid && !bAnyMenusVisible && !bAutomationTesting && !bIsInteratcting && !GIsDemoMode && bHasGameOrProjectLoaded);
