Fix for cook on the fly generating asset registry files for all available platforms rather than just platforms being cooked.

diff --git a/Engine/Source/Editor/UnrealEd/Private/CookOnTheFlyServer.cpp b/Engine/Source/Editor/UnrealEd/Private/CookOnTheFlyServer.cpp
--- a/Engine/Source/Editor/UnrealEd/Private/CookOnTheFlyServer.cpp
+++ b/Engine/Source/Editor/UnrealEd/Private/CookOnTheFlyServer.cpp
@@ -699,7 +699,7 @@ void UCookOnTheFlyServer::Initialize( bool inCompressed, bool inIterativeCooking
 
 	
 	ITargetPlatformManagerModule& TPM = GetTargetPlatformManagerRef();
-	const TArray<ITargetPlatform*>& Platforms = TPM.GetTargetPlatforms();
+	const TArray<ITargetPlatform*>& Platforms = TPM.GetActiveTargetPlatforms();
 
 	// Local sandbox file wrapper. This will be used to handle path conversions,
 	// but will not be used to actually write/read files so we can safely
