CurveEditor - Add commands for zoom to fit horizontal and zoom to fit vertical and add buttons to the toolbar in the CurveAssetEditor.

diff --git a/Engine/Source/Editor/CurveAssetEditor/Private/CurveAssetEditor.cpp b/Engine/Source/Editor/CurveAssetEditor/Private/CurveAssetEditor.cpp
--- a/Engine/Source/Editor/CurveAssetEditor/Private/CurveAssetEditor.cpp
+++ b/Engine/Source/Editor/CurveAssetEditor/Private/CurveAssetEditor.cpp
@@ -4,6 +4,7 @@
 
 #include "CurveAssetEditor.h"
 #include "SCurveEditor.h"
+#include "RichCurveEditorCommands.h"
 //#include "Toolkits/IToolkitHost.h"
 #include "WorkspaceMenuStructureModule.h"
 
@@ -55,6 +56,7 @@ void FCurveAssetEditor::InitCurveAssetEditor( const EToolkitMode::Type Mode, con
 	
 	FCurveAssetEditorModule& CurveAssetEditorModule = FModuleManager::LoadModuleChecked<FCurveAssetEditorModule>( "CurveAssetEditor" );
 	AddMenuExtender(CurveAssetEditorModule.GetMenuExtensibilityManager()->GetAllExtenders(GetToolkitCommands(), GetEditingObjects()));
+	AddToolbarExtender(GetToolbarExtender());
 
 	// @todo toolkit world centric editing
 	/*// Setup our tool's layout
@@ -64,7 +66,10 @@ void FCurveAssetEditor::InitCurveAssetEditor( const EToolkitMode::Type Mode, con
 		SpawnToolkitTab( CurveTabId, TabInitializationPayload, EToolkitTabSpot::Details );
 	}*/
 
-	// NOTE: Could fill in asset editor commands here!
+	if (TrackWidget.IsValid())
+	{
+		RegenerateMenusAndToolbars();
+	}
 }
 
 FName FCurveAssetEditor::GetToolkitFName() const
@@ -110,6 +115,7 @@ TSharedRef<SDockTab> FCurveAssetEditor::SpawnTab_CurveAsset( const FSpawnTabArgs
 				.OnSetInputViewRange(this, &FCurveAssetEditor::SetInputViewRange)
 				.HideUI(false)
 				.AlwaysDisplayColorCurves(true)
+				.ShowZoomButtons(false)
 			]
 		];
 
@@ -138,7 +144,7 @@ TSharedRef<SDockTab> FCurveAssetEditor::SpawnTab_CurveAsset( const FSpawnTabArgs
 		// Set this curve as the SCurveEditor's selected curve
 		TrackWidget->SetCurveOwner(CurveOwner);
 	}
-	
+
 	return NewDockTab;
 }
 
@@ -153,4 +159,31 @@ void FCurveAssetEditor::SetInputViewRange(float InViewMinInput, float InViewMaxI
 	ViewMinInput = InViewMinInput;
 }
 
+TSharedPtr<FExtender> FCurveAssetEditor::GetToolbarExtender()
+{
+	struct Local
+	{
+		static void FillToolbar(FToolBarBuilder& ToolbarBuilder)
+		{
+			ToolbarBuilder.BeginSection("Curve");
+			{
+				ToolbarBuilder.AddToolBarButton(FRichCurveEditorCommands::Get().ZoomToFitHorizontal);
+				ToolbarBuilder.AddToolBarButton(FRichCurveEditorCommands::Get().ZoomToFitVertical);
+			}
+			ToolbarBuilder.EndSection();
+		}
+	};
+
+	TSharedPtr<FExtender> ToolbarExtender = MakeShareable(new FExtender);
+
+	ToolbarExtender->AddToolBarExtension(
+		"Asset",
+		EExtensionHook::After,
+		TrackWidget->GetCommands(),
+		FToolBarExtensionDelegate::CreateStatic(&Local::FillToolbar)
+		);
+
+	return ToolbarExtender;
+}
+
 #undef LOCTEXT_NAMESPACE
\ No newline at end of file
diff --git a/Engine/Source/Editor/CurveAssetEditor/Private/CurveAssetEditor.h b/Engine/Source/Editor/CurveAssetEditor/Private/CurveAssetEditor.h
--- a/Engine/Source/Editor/CurveAssetEditor/Private/CurveAssetEditor.h
+++ b/Engine/Source/Editor/CurveAssetEditor/Private/CurveAssetEditor.h
@@ -34,6 +34,8 @@ public:
 	virtual FString GetWorldCentricTabPrefix() const override;
 	virtual FLinearColor GetWorldCentricTabColorScale() const override;
 
+	TSharedPtr<FExtender> GetToolbarExtender();
+
 private:
 
 	float GetViewMinInput() const { return ViewMinInput; }
diff --git a/Engine/Source/Editor/CurveAssetEditor/Private/CurveAssetEditorModule.cpp b/Engine/Source/Editor/CurveAssetEditor/Private/CurveAssetEditorModule.cpp
--- a/Engine/Source/Editor/CurveAssetEditor/Private/CurveAssetEditorModule.cpp
+++ b/Engine/Source/Editor/CurveAssetEditor/Private/CurveAssetEditorModule.cpp
@@ -3,6 +3,7 @@
 #include "CurveAssetEditorPrivatePCH.h"
 #include "ModuleManager.h"
 #include "CurveAssetEditor.h"
+#include "RichCurveEditorCommands.h"
 //#include "Toolkits/ToolkitManager.h"
 
 IMPLEMENT_MODULE( FCurveAssetEditorModule, CurveAssetEditor );
@@ -13,6 +14,7 @@ const FName FCurveAssetEditorModule::CurveAssetEditorAppIdentifier( TEXT( "Curve
 void FCurveAssetEditorModule::StartupModule()
 {
 	MenuExtensibilityManager = MakeShareable(new FExtensibilityManager);
+	FRichCurveEditorCommands::Register();
 }
 
 void FCurveAssetEditorModule::ShutdownModule()
diff --git a/Engine/Source/Editor/EditorStyle/Private/SlateEditorStyle.cpp b/Engine/Source/Editor/EditorStyle/Private/SlateEditorStyle.cpp
--- a/Engine/Source/Editor/EditorStyle/Private/SlateEditorStyle.cpp
+++ b/Engine/Source/Editor/EditorStyle/Private/SlateEditorStyle.cpp
@@ -5155,6 +5155,12 @@ void FSlateEditorStyle::FStyle::SetupToolkitStyles()
 		Set( "CurveEditor.DeleteTab.Small", new IMAGE_BRUSH( "Icons/icon_CurveEditor_DeleteTab_40x", Icon20x20 ) );
 	}
 
+	// Rich Curve Editor
+	{
+		Set("RichCurveEditor.ZoomToFitHorizontal", new IMAGE_BRUSH("Icons/icon_CurveEditor_Horizontal_40x", Icon40x40));
+		Set("RichCurveEditor.ZoomToFitVertical", new IMAGE_BRUSH("Icons/icon_CurveEditor_Vertical_40x", Icon40x40));
+	}
+
 	// PhAT
 	{
 		Set( "PhAT.Tabs.Properties", new IMAGE_BRUSH( "/Icons/properties_16x", Icon16x16 ) );
diff --git a/Engine/Source/Editor/UnrealEd/Private/RichCurveEditorCommands.cpp b/Engine/Source/Editor/UnrealEd/Private/RichCurveEditorCommands.cpp
new file mode 100644
--- /dev/null
+++ b/Engine/Source/Editor/UnrealEd/Private/RichCurveEditorCommands.cpp
@@ -0,0 +1,8 @@
+#include "UnrealEd.h"
+#include "RichCurveEditorCommands.h"
+
+void FRichCurveEditorCommands::RegisterCommands()
+{
+	UI_COMMAND(ZoomToFitHorizontal, "Fit Horizontal", "Zoom to Fit - Horizontal", EUserInterfaceActionType::Button, FInputGesture());
+	UI_COMMAND(ZoomToFitVertical, "Fit Vertical", "Zoom to Fit - Vertical", EUserInterfaceActionType::Button, FInputGesture());
+}
diff --git a/Engine/Source/Editor/UnrealEd/Private/SCurveEditor.cpp b/Engine/Source/Editor/UnrealEd/Private/SCurveEditor.cpp
--- a/Engine/Source/Editor/UnrealEd/Private/SCurveEditor.cpp
+++ b/Engine/Source/Editor/UnrealEd/Private/SCurveEditor.cpp
@@ -56,6 +56,7 @@ void SCurveEditor::Construct(const FArguments& InArgs)
 	bHideUI = InArgs._HideUI;
 	bAllowZoomOutput = InArgs._AllowZoomOutput;
 	bAlwaysDisplayColorCurves = InArgs._AlwaysDisplayColorCurves;
+	bShowZoomButtons = InArgs._ShowZoomButtons;
 
 	OnCreateAsset = InArgs._OnCreateAsset;
 
@@ -70,16 +71,22 @@ void SCurveEditor::Construct(const FArguments& InArgs)
 
 	TransactionIndex = -1;
 
-	Commands->MapAction( FGenericCommands::Get().Undo,
-		FExecuteAction::CreateSP( this, &SCurveEditor::UndoAction ));
+	Commands->MapAction(FGenericCommands::Get().Undo,
+		FExecuteAction::CreateSP(this, &SCurveEditor::UndoAction));
 
-	Commands->MapAction( FGenericCommands::Get().Redo,
-		FExecuteAction::CreateSP( this, &SCurveEditor::RedoAction ));
+	Commands->MapAction(FGenericCommands::Get().Redo,
+		FExecuteAction::CreateSP(this, &SCurveEditor::RedoAction));
 
-	SAssignNew(WarningMessageText, SErrorText );
+	Commands->MapAction(FRichCurveEditorCommands::Get().ZoomToFitHorizontal,
+		FExecuteAction::CreateSP(this, &SCurveEditor::ZoomToFitHorizontal));
+
+	Commands->MapAction(FRichCurveEditorCommands::Get().ZoomToFitVertical,
+		FExecuteAction::CreateSP(this, &SCurveEditor::ZoomToFitVertical));
+
+	SAssignNew(WarningMessageText, SErrorText);
 
 	TSharedRef<SBorder> CurveSelector = SNew(SBorder)
-		.BorderImage( FEditorStyle::GetBrush("NoBorder") )
+		.BorderImage(FEditorStyle::GetBrush("NoBorder"))
 		.Visibility(this, &SCurveEditor::GetControlVisibility)
 		[
 			CreateCurveSelectionWidget()
@@ -122,13 +129,14 @@ void SCurveEditor::Construct(const FArguments& InArgs)
 				.Padding(FMargin(2, 12, 0, 0))
 				[
 					SNew(SHorizontalBox)
-					
+
 					+ SHorizontalBox::Slot()
 					.AutoWidth()
 					[
 						SNew(SButton)
 						.ToolTipText(LOCTEXT("ZoomToFitHorizontal", "Zoom To Fit Horizontal"))
-						.OnClicked(this, &SCurveEditor::ZoomToFitHorizontal)
+						.Visibility(this, &SCurveEditor::GetZoomButtonVisibility)
+						.OnClicked(this, &SCurveEditor::ZoomToFitHorizontalClicked)
 						.ContentPadding(1)
 						[
 							SNew(SImage) 
@@ -142,7 +150,8 @@ void SCurveEditor::Construct(const FArguments& InArgs)
 					[
 						SNew(SButton)
 						.ToolTipText(LOCTEXT("ZoomToFitVertical", "Zoom To Fit Vertical"))
-						.OnClicked(this, &SCurveEditor::ZoomToFitVertical)
+						.Visibility(this, &SCurveEditor::GetZoomButtonVisibility)
+						.OnClicked(this, &SCurveEditor::ZoomToFitVerticalClicked)
 						.ContentPadding(1)
 						[
 							SNew(SImage) 
@@ -367,6 +376,11 @@ EVisibility SCurveEditor::GetColorGradientVisibility() const
 	return bIsGradientEditorVisible && IsLinearColorCurve() ? EVisibility::Visible : EVisibility::Collapsed;
 }
 
+EVisibility SCurveEditor::GetZoomButtonVisibility() const
+{
+	return bShowZoomButtons ? EVisibility::Visible : EVisibility::Collapsed;
+}
+
 bool SCurveEditor::GetInputEditEnabled() const
 {
 	return (SelectedKeys.Num() == 1);
@@ -1475,7 +1489,7 @@ TArray<FRichCurve*> SCurveEditor::GetCurvesToFit() const
 }
 
 
-FReply SCurveEditor::ZoomToFitHorizontal()
+void SCurveEditor::ZoomToFitHorizontal()
 {
 	TArray<FRichCurve*> CurvesToFit = GetCurvesToFit();
 
@@ -1516,6 +1530,11 @@ FReply SCurveEditor::ZoomToFitHorizontal()
 
 		SetInputMinMax(InMin, InMax);
 	}
+}
+
+FReply SCurveEditor::ZoomToFitHorizontalClicked()
+{
+	ZoomToFitHorizontal();
 	return FReply::Handled();
 }
 
@@ -1526,7 +1545,7 @@ void SCurveEditor::SetDefaultOutput(const float MinZoomRange)
 	ViewMaxOutput = (ViewMaxOutput + (0.5f*MinZoomRange));
 }
 
-FReply SCurveEditor::ZoomToFitVertical()
+void SCurveEditor::ZoomToFitVertical()
 {
 	TArray<FRichCurve*> CurvesToFit = GetCurvesToFit();
 
@@ -1560,6 +1579,11 @@ FReply SCurveEditor::ZoomToFitVertical()
 		ViewMinOutput = (ViewMinOutput - CONST_FitMargin*Size);
 		ViewMaxOutput = (ViewMaxOutput + CONST_FitMargin*Size);
 	}
+}
+
+FReply SCurveEditor::ZoomToFitVerticalClicked()
+{
+	ZoomToFitVertical();
 	return FReply::Handled();
 }
 
@@ -1657,6 +1681,11 @@ void SCurveEditor::AddReferencedObjects( FReferenceCollector& Collector )
 	Collector.AddReferencedObject( CurveFactory );
 }
 
+TSharedPtr<FUICommandList> SCurveEditor::GetCommands()
+{
+	return Commands;
+}
+
 bool SCurveEditor::IsValidCurve( FRichCurve* Curve ) const
 {
 	bool bIsValid = false;
diff --git a/Engine/Source/Editor/UnrealEd/Public/RichCurveEditorCommands.h b/Engine/Source/Editor/UnrealEd/Public/RichCurveEditorCommands.h
new file mode 100644
--- /dev/null
+++ b/Engine/Source/Editor/UnrealEd/Public/RichCurveEditorCommands.h
@@ -0,0 +1,32 @@
+#pragma once
+
+/**
+ * Defines commands for the SCurveEditor.
+ */
+class UNREALED_API FRichCurveEditorCommands : public TCommands<FRichCurveEditorCommands>
+{
+public:
+	FRichCurveEditorCommands()
+		: TCommands<FRichCurveEditorCommands>
+		(
+			TEXT("RichCurveEditor"),
+			NSLOCTEXT("Contexts", "RichCurveEditor", "Curve Editor"),
+			NAME_None,
+			FEditorStyle::GetStyleSetName()
+		)
+	{
+	}
+
+	TSharedPtr<FUICommandInfo> ZoomToFitHorizontal;
+	TSharedPtr<FUICommandInfo> ZoomToFitVertical;
+
+public:
+	virtual void RegisterCommands() override;
+
+	~FRichCurveEditorCommands()
+	{
+		int i = 0;
+	}
+};
+
+#endif
\ No newline at end of file
diff --git a/Engine/Source/Editor/UnrealEd/Public/SCurveEditor.h b/Engine/Source/Editor/UnrealEd/Public/SCurveEditor.h
--- a/Engine/Source/Editor/UnrealEd/Public/SCurveEditor.h
+++ b/Engine/Source/Editor/UnrealEd/Public/SCurveEditor.h
@@ -93,6 +93,7 @@ public:
 		, _AlwaysDisplayColorCurves(false)
 		, _ZoomToFitVertical(true)
 		, _ZoomToFitHorizontal(true)
+		, _ShowZoomButtons(true)
 		{}
 
 		SLATE_ATTRIBUTE( float, ViewMinInput )
@@ -109,6 +110,7 @@ public:
 		SLATE_ARGUMENT( bool, AlwaysDisplayColorCurves )
 		SLATE_ARGUMENT( bool, ZoomToFitVertical )
 		SLATE_ARGUMENT( bool, ZoomToFitHorizontal )
+		SLATE_ARGUMENT( bool, ShowZoomButtons )
 		SLATE_EVENT( FOnSetInputViewRange, OnSetInputViewRange )
 		SLATE_EVENT( FSimpleDelegate, OnCreateAsset )
 	SLATE_END_ARGS()
@@ -148,6 +150,9 @@ public:
 	 */
 	UNREALED_API virtual void AddReferencedObjects( FReferenceCollector& Collector );
 
+	/** Gets a list of the commands handled by this control */
+	UNREALED_API TSharedPtr<FUICommandList> GetCommands();
+
 private:
 	/** Used to track a key and the curve that owns it */
 	struct FSelectedCurveKey
@@ -253,9 +258,11 @@ private:
 	/* Get the curves to will be used during a fit operation */
 	TArray<FRichCurve*> GetCurvesToFit()const;
 
-	FReply ZoomToFitHorizontal();
-	FReply ZoomToFitVertical();
+	void ZoomToFitHorizontal();
+	void ZoomToFitVertical();
 
+	FReply ZoomToFitHorizontalClicked();
+	FReply ZoomToFitVerticalClicked();
 
 	TOptional<float> OnGetTime() const;
 	void OnTimeComitted(float NewValue, ETextCommit::Type CommitType);
@@ -272,6 +279,7 @@ private:
 	EVisibility GetControlVisibility() const;
 	EVisibility GetEditVisibility() const;
 	EVisibility GetColorGradientVisibility() const;
+	EVisibility GetZoomButtonVisibility() const;
 
 	bool GetInputEditEnabled() const;
 
@@ -494,6 +502,9 @@ protected:
 
 	/** True if the sliders are being used to adjust point values **/
 	bool bIsUsingSlider;
+
+	/** True if the internal zoom buttons should be visible. */
+	bool bShowZoomButtons;
 };
 
 #endif // __SCurveEditor_h__
