Fix Focus on selected not working on brush vertices, faces, and edges

diff --git a/Engine/Source/Editor/UnrealEd/Private/EditorServer.cpp b/Engine/Source/Editor/UnrealEd/Private/EditorServer.cpp
--- a/Engine/Source/Editor/UnrealEd/Private/EditorServer.cpp
+++ b/Engine/Source/Editor/UnrealEd/Private/EditorServer.cpp
@@ -21,6 +21,7 @@
 #include "Editor/UnrealEd/Public/Kismet2/BlueprintEditorUtils.h"
 #include "../Private/GeomFitUtils.h"
 #include "Editor/GeometryMode/Public/GeometryEdMode.h"
+#include "Editor/GeometryMode/Public/EditorGeometry.h"
 #include "Landscape/LandscapeProxy.h"
 #include "Lightmass/PrecomputedVisibilityOverrideVolume.h"
 #include "Animation/AnimSet.h"
@@ -4190,6 +4191,34 @@ void UEditorEngine::MoveViewportCamerasToActor(const TArray<AActor*> &Actors, bo
 				const FBox DefaultSizeBox(Actor->GetActorLocation() - DefaultExtent, Actor->GetActorLocation() + DefaultExtent);
 				BoundingBox += DefaultSizeBox;
 			}
+			else if( Actor->IsA<ABrush>() && GLevelEditorModeTools().IsModeActive( FBuiltinEditorModes::EM_Geometry ) )
+			{
+				FEdModeGeometry* GeometryMode = GLevelEditorModeTools().GetActiveModeTyped<FEdModeGeometry>(FBuiltinEditorModes::EM_Geometry);
+
+				TArray<FGeomVertex*> SelectedVertices;
+				GeometryMode->GetSelectedVertices( SelectedVertices );
+				for( FGeomVertex* Vertex : SelectedVertices )
+				{
+					BoundingBox += Vertex->GetWidgetLocation();
+				}
+				
+				TArray<FGeomPoly*> SelectedPolys;
+				GeometryMode->GetSelectedPolygons( SelectedPolys );
+				for( FGeomPoly* Poly : SelectedPolys )
+				{
+					BoundingBox += Poly->GetWidgetLocation();
+				}
+
+				TArray<FGeomEdge*> SelectedEdges;
+				GeometryMode->GetSelectedEdges( SelectedEdges );
+				for( FGeomEdge* Edge : SelectedEdges )
+				{
+					BoundingBox += Edge->GetWidgetLocation();
+				}
+
+				// Zoom out a little bit so you can see the selection
+				BoundingBox = BoundingBox.ExpandBy( 25 );
+			}
 			else
 			{
 				TArray<UPrimitiveComponent*> Components;
