RichCurveEditorCommands - Submit unsaved file (removes unneeded endif)

diff --git a/Engine/Source/Editor/UnrealEd/Public/RichCurveEditorCommands.h b/Engine/Source/Editor/UnrealEd/Public/RichCurveEditorCommands.h
--- a/Engine/Source/Editor/UnrealEd/Public/RichCurveEditorCommands.h
+++ b/Engine/Source/Editor/UnrealEd/Public/RichCurveEditorCommands.h
@@ -27,6 +27,4 @@ public:
 	{
 		int i = 0;
 	}
-};
-
-#endif
\ No newline at end of file
+};
\ No newline at end of file
