Redraw viewports when reimporting to ensure all meshes are refreshed in non-realtime viewports

diff --git a/Engine/Source/Editor/UnrealEd/Private/Editor.cpp b/Engine/Source/Editor/UnrealEd/Private/Editor.cpp
--- a/Engine/Source/Editor/UnrealEd/Private/Editor.cpp
+++ b/Engine/Source/Editor/UnrealEd/Private/Editor.cpp
@@ -2686,6 +2686,8 @@ bool FReimportManager::Reimport( UObject* Obj, bool bAskForNewFileIfMissing )
 	// Let listeners know whether the reimport was successful or not
 	PostReimport.Broadcast( Obj, bSuccess );
 
+	GEditor->RedrawAllViewports();
+
 	return bSuccess;
 }
 
