Remove undo warning when opening matinee

diff --git a/Engine/Source/Editor/UnrealEd/Private/Editor.cpp b/Engine/Source/Editor/UnrealEd/Private/Editor.cpp
--- a/Engine/Source/Editor/UnrealEd/Private/Editor.cpp
+++ b/Engine/Source/Editor/UnrealEd/Private/Editor.cpp
@@ -3877,15 +3877,6 @@ bool UEditorEngine::ShouldOpenMatinee(AMatineeActor* MatineeActor) const
 		return false;
 	}
 
-	FSuppressableWarningDialog::FSetupInfo Info( LOCTEXT("MatineeUndoWarningBody","Opening the Matinee Editor will reset all current Undo data. (Undo/Redo will only be available on matinee actions while it remains open)"), LOCTEXT("MatineeUndoWarningTitle","Matinee Undo Warning"), "MatineeUndoWarning" );
-	Info.ConfirmText =  LOCTEXT( "Continue", "Continue");
-	Info.CancelText =  LOCTEXT( "Cancel", "Cancel");	
-	Info.bDefaultToSupressInTheFuture = true;
-	FSuppressableWarningDialog MatineeUndoWarning( Info );
-	if( MatineeUndoWarning.ShowModal() == FSuppressableWarningDialog::Cancel )
-	{
-		return false;
-	}
 	return true;
 }
 
