Step in blueprint debugging mapped to F10 now. Possess/Unpossess moved from F10 to F8

diff --git a/Engine/Source/Editor/UnrealEd/Private/Kismet2/DebuggerCommands.cpp b/Engine/Source/Editor/UnrealEd/Private/Kismet2/DebuggerCommands.cpp
--- a/Engine/Source/Editor/UnrealEd/Private/Kismet2/DebuggerCommands.cpp
+++ b/Engine/Source/Editor/UnrealEd/Private/Kismet2/DebuggerCommands.cpp
@@ -219,9 +219,9 @@ void FPlayWorldCommands::RegisterCommands()
 	UI_COMMAND( PausePlaySession, "Pause", "Pause simulation", EUserInterfaceActionType::Button, FInputGesture() );
 	UI_COMMAND( SingleFrameAdvance, "Frame Skip", "Advances a single frame", EUserInterfaceActionType::Button, FInputGesture() );
 	UI_COMMAND( TogglePlayPauseOfPlaySession, "Toggle Play/Pause", "Resume playing if paused, or pause if playing", EUserInterfaceActionType::Button, FInputGesture( EKeys::Pause ) );
-	UI_COMMAND( PossessEjectPlayer, "Possess or Eject Player", "Possesses or ejects the player from the camera", EUserInterfaceActionType::Button, FInputGesture( EKeys::F10 ) );
+	UI_COMMAND( PossessEjectPlayer, "Possess or Eject Player", "Possesses or ejects the player from the camera", EUserInterfaceActionType::Button, FInputGesture( EKeys::F8 ) );
 	UI_COMMAND( ShowCurrentStatement, "Find Node", "Show the current node", EUserInterfaceActionType::Button, FInputGesture() );
-	UI_COMMAND( StepInto, "Step", "Step to the next node to be executed", EUserInterfaceActionType::Button, FInputGesture( EKeys::F5 ) );
+	UI_COMMAND( StepInto, "Step", "Step to the next node to be executed", EUserInterfaceActionType::Button, FInputGesture( EKeys::F10) );
 
 	// Launch
 	UI_COMMAND( RepeatLastLaunch, "Launch", "Launches the game on the device as the last session launched from the dropdown next to the Play on Device button on the level editor toolbar", EUserInterfaceActionType::Button, FInputGesture( EKeys::P, EModifierKey::Alt | EModifierKey::Shift ) )
