SCurveEditor - Remove a spacer and fix up some margins to fix layout issues.

diff --git a/Engine/Source/Editor/UnrealEd/Private/SCurveEditor.cpp b/Engine/Source/Editor/UnrealEd/Private/SCurveEditor.cpp
--- a/Engine/Source/Editor/UnrealEd/Private/SCurveEditor.cpp
+++ b/Engine/Source/Editor/UnrealEd/Private/SCurveEditor.cpp
@@ -85,8 +85,7 @@ void SCurveEditor::Construct(const FArguments& InArgs)
 
 	SAssignNew(WarningMessageText, SErrorText);
 
-	TSharedRef<SBorder> CurveSelector = SNew(SBorder)
-		.BorderImage(FEditorStyle::GetBrush("NoBorder"))
+	TSharedRef<SBox> CurveSelector = SNew(SBox)
 		.Visibility(this, &SCurveEditor::GetControlVisibility)
 		[
 			CreateCurveSelectionWidget()
@@ -104,15 +103,9 @@ void SCurveEditor::Construct(const FArguments& InArgs)
 			SNew(SHorizontalBox)
 			.Visibility( this, &SCurveEditor::GetCurveAreaVisibility )
 
-			+ SHorizontalBox::Slot()
-			.Padding(FMargin(20, 12, 0, 0))
-			[
-				SNew(SSpacer)
-			]
-			
 			+ SHorizontalBox::Slot()
 			.AutoWidth()
-			.Padding(FMargin(2, 12, 0, 0))
+			.Padding(FMargin(30, 12, 0, 0))
 			[
 				CurveSelector
 			]
@@ -363,7 +356,7 @@ EVisibility SCurveEditor::GetCurveAreaVisibility() const
 
 EVisibility SCurveEditor::GetControlVisibility() const
 {
-	return (IsHovered() || (false == bHideUI)) ? EVisibility::Visible : EVisibility::Collapsed;
+	return (IsHovered() || (false == bHideUI)) ? EVisibility::Visible : EVisibility::Hidden;
 }
 
 EVisibility SCurveEditor::GetEditVisibility() const
diff --git a/Engine/Source/Editor/UnrealEd/Public/SCurveEditor.h b/Engine/Source/Editor/UnrealEd/Public/SCurveEditor.h
--- a/Engine/Source/Editor/UnrealEd/Public/SCurveEditor.h
+++ b/Engine/Source/Editor/UnrealEd/Public/SCurveEditor.h
@@ -409,7 +409,7 @@ protected:
 private:
 
 	/** Curve selection*/
-	TWeakPtr<SBorder>		CurveSelectionWidget;
+	TWeakPtr<SBox>		CurveSelectionWidget;
 
 	/** Text block used to display warnings related to curves */
 	TSharedPtr< SErrorText > WarningMessageText;
