#UE4 nonunity fix

diff --git a/Engine/Source/Editor/UnrealEd/Private/SCurveEditor.cpp b/Engine/Source/Editor/UnrealEd/Private/SCurveEditor.cpp
--- a/Engine/Source/Editor/UnrealEd/Private/SCurveEditor.cpp
+++ b/Engine/Source/Editor/UnrealEd/Private/SCurveEditor.cpp
@@ -2,6 +2,7 @@
 
 
 #include "UnrealEd.h"
+#include "RichCurveEditorCommands.h"
 #include "SCurveEditor.h"
 #include "ScopedTransaction.h"
 #include "SColorGradientEditor.h"
