Curve Editor Commands - Register them along with the other misc commands so that core engine usages of the SCurveEditor don't crash.

diff --git a/Engine/Source/Editor/UnrealEd/Private/UnrealEdMisc.cpp b/Engine/Source/Editor/UnrealEd/Private/UnrealEdMisc.cpp
--- a/Engine/Source/Editor/UnrealEd/Private/UnrealEdMisc.cpp
+++ b/Engine/Source/Editor/UnrealEd/Private/UnrealEdMisc.cpp
@@ -6,6 +6,7 @@
 #include "Kismet2/DebuggerCommands.h"
 #include "Toolkits/AssetEditorCommonCommands.h"
 #include "SoundCueGraphEditorCommands.h"
+#include "RichCurveEditorCommands.h"
 #include "FileHelpers.h"
 #include "EditorBuildUtils.h"
 #include "AssetToolsModule.h"
@@ -112,6 +113,9 @@ void FUnrealEdMisc::OnInit()
 	// Register navigation commands for all viewports
 	FViewportNavigationCommands::Register();
 
+	// Register curve editor commands.
+	FRichCurveEditorCommands::Register();
+
 	FEditorModeRegistry::Initialize();
 	GLevelEditorModeTools().ActivateDefaultMode();
 
