EditorOpenAssetAutomationTests.cpp
AutomationEditorCommon.h

More updates to the large Open Asset tests.

diff --git a/Engine/Source/Editor/UnrealEd/Private/Tests/EditorOpenAssetAutomationTests.cpp b/Engine/Source/Editor/UnrealEd/Private/Tests/EditorOpenAssetAutomationTests.cpp
--- a/Engine/Source/Editor/UnrealEd/Private/Tests/EditorOpenAssetAutomationTests.cpp
+++ b/Engine/Source/Editor/UnrealEd/Private/Tests/EditorOpenAssetAutomationTests.cpp
@@ -8,6 +8,9 @@
 #include "AutomationCommon.h"
 #include "AssetEditorManager.h"
 
+//Includes needed for opening certain assets
+#include "Materials/MaterialFunction.h"
+#include "Slate/SlateBrushAsset.h"
 
 
 
@@ -137,12 +140,12 @@ bool FOpenAssetEditors::RunTest(const FString& Parameters)
 /**
 * This test opens each AimOffsetBlendSpace into its sub-editor.
 */
-IMPLEMENT_COMPLEX_AUTOMATION_TEST(FOpenAimOffsetBlendSpaceAssetEditors, "Editor.Content.Animation.BlendSpaces.Open AimOffsetBlendSpace Assets", (EAutomationTestFlags::ATF_Editor | EAutomationTestFlags::ATF_RequiresUser));
+IMPLEMENT_COMPLEX_AUTOMATION_TEST(FOpenAimOffsetBlendSpaceAssetEditors, "Editor.Content.Skeletal Mesh.Animation.BlendSpaces.Open AimOffsetBlendSpace Assets", (EAutomationTestFlags::ATF_Editor | EAutomationTestFlags::ATF_RequiresUser));
 
 void FOpenAimOffsetBlendSpaceAssetEditors::GetTests(TArray<FString>& OutBeautifiedNames, TArray<FString>& OutTestCommands) const
 {
 	//This grabs each AimOffsetBlendSpace asset in the Game/Content directory
-	FEditorAutomationTestUtilities::CollectGameContentTestsByClass(UAimOffsetBlendSpace::StaticClass(), OutBeautifiedNames, OutTestCommands);
+	FEditorAutomationTestUtilities::CollectGameContentTestsByClass(UAimOffsetBlendSpace::StaticClass(), true, OutBeautifiedNames, OutTestCommands);
 }
 
 bool FOpenAimOffsetBlendSpaceAssetEditors::RunTest(const FString& Parameters)
@@ -156,12 +159,12 @@ bool FOpenAimOffsetBlendSpaceAssetEditors::RunTest(const FString& Parameters)
 /**
 * This test opens each AimOffsetBlendSpace1D into its sub-editor.
 */
-IMPLEMENT_COMPLEX_AUTOMATION_TEST(FOpenAimOffsetBlendSpace1DAssetEditors, "Editor.Content.Animation.BlendSpaces.Open AimOffsetBlendSpace1D Assets", (EAutomationTestFlags::ATF_Editor | EAutomationTestFlags::ATF_RequiresUser));
+IMPLEMENT_COMPLEX_AUTOMATION_TEST(FOpenAimOffsetBlendSpace1DAssetEditors, "Editor.Content.Skeletal Mesh.Animation.BlendSpaces.Open AimOffsetBlendSpace1D Assets", (EAutomationTestFlags::ATF_Editor | EAutomationTestFlags::ATF_RequiresUser));
 
 void FOpenAimOffsetBlendSpace1DAssetEditors::GetTests(TArray<FString>& OutBeautifiedNames, TArray<FString>& OutTestCommands) const
 {
 	//This grabs each AimOffsetBlendSpace1D asset in the Game/Content directory
-	FEditorAutomationTestUtilities::CollectGameContentTestsByClass(UAimOffsetBlendSpace1D::StaticClass(), OutBeautifiedNames, OutTestCommands);
+	FEditorAutomationTestUtilities::CollectGameContentTestsByClass(UAimOffsetBlendSpace1D::StaticClass(), true, OutBeautifiedNames, OutTestCommands);
 }
 
 bool FOpenAimOffsetBlendSpace1DAssetEditors::RunTest(const FString& Parameters)
@@ -175,12 +178,12 @@ bool FOpenAimOffsetBlendSpace1DAssetEditors::RunTest(const FString& Parameters)
 /**
 * This test opens each AnimMontage into its sub-editor.
 */
-IMPLEMENT_COMPLEX_AUTOMATION_TEST(FOpenAnimMontageAssetEditors, "Editor.Content.Animation.Open AnimMontage Assets", (EAutomationTestFlags::ATF_Editor | EAutomationTestFlags::ATF_RequiresUser));
+IMPLEMENT_COMPLEX_AUTOMATION_TEST(FOpenAnimMontageAssetEditors, "Editor.Content.Skeletal Mesh.Animation.Open AnimMontage Assets", (EAutomationTestFlags::ATF_Editor | EAutomationTestFlags::ATF_RequiresUser));
 
 void FOpenAnimMontageAssetEditors::GetTests(TArray<FString>& OutBeautifiedNames, TArray<FString>& OutTestCommands) const
 {
 	//This grabs each AnimMontage asset in the Game/Content directory
-	FEditorAutomationTestUtilities::CollectGameContentTestsByClass(UAnimMontage::StaticClass(), OutBeautifiedNames, OutTestCommands);
+	FEditorAutomationTestUtilities::CollectGameContentTestsByClass(UAnimMontage::StaticClass(), true, OutBeautifiedNames, OutTestCommands);
 }
 
 bool FOpenAnimMontageAssetEditors::RunTest(const FString& Parameters)
@@ -190,16 +193,35 @@ bool FOpenAnimMontageAssetEditors::RunTest(const FString& Parameters)
 	return bDidTheTestPass;
 }
 
+//////////////////////////////////////////////////////////////////////////
+/**
+* This test opens each AnimSequence into its sub-editor.
+*/
+IMPLEMENT_COMPLEX_AUTOMATION_TEST(FOpenAnimSequenceAssetEditors, "Editor.Content.Skeletal Mesh.Animation.Open AnimSequence Assets", (EAutomationTestFlags::ATF_Editor | EAutomationTestFlags::ATF_RequiresUser));
+
+void FOpenAnimSequenceAssetEditors::GetTests(TArray<FString>& OutBeautifiedNames, TArray<FString>& OutTestCommands) const
+{
+	//This grabs each AnimSequence asset in the Game/Content directory
+	FEditorAutomationTestUtilities::CollectGameContentTestsByClass(UAnimSequence::StaticClass(), true, OutBeautifiedNames, OutTestCommands);
+}
+
+bool FOpenAnimSequenceAssetEditors::RunTest(const FString& Parameters)
+{
+	bool bDidTheTestPass = FOpenActualAssetEditors(Parameters);
+	CollectGarbage(GARBAGE_COLLECTION_KEEPFLAGS);
+	return bDidTheTestPass;
+}
+
 //////////////////////////////////////////////////////////////////////////
 /**
 * This test opens each BlendSpace into its sub-editor.
 */
-IMPLEMENT_COMPLEX_AUTOMATION_TEST(FOpenBlendSpaceAssetEditors, "Editor.Content.Animation.BlendSpaces.Open BlendSpace Assets", (EAutomationTestFlags::ATF_Editor | EAutomationTestFlags::ATF_RequiresUser));
+IMPLEMENT_COMPLEX_AUTOMATION_TEST(FOpenBlendSpaceAssetEditors, "Editor.Content.Skeletal Mesh.Animation.BlendSpaces.Open BlendSpace Assets", (EAutomationTestFlags::ATF_Editor | EAutomationTestFlags::ATF_RequiresUser));
 
 void FOpenBlendSpaceAssetEditors::GetTests(TArray<FString>& OutBeautifiedNames, TArray<FString>& OutTestCommands) const
 {
 	//This grabs each BlendSpace asset in the Game/Content directory
-	FEditorAutomationTestUtilities::CollectGameContentTestsByClass(UBlendSpace::StaticClass(), OutBeautifiedNames, OutTestCommands);
+	FEditorAutomationTestUtilities::CollectGameContentTestsByClass(UBlendSpace::StaticClass(), false, OutBeautifiedNames, OutTestCommands);
 }
 
 bool FOpenBlendSpaceAssetEditors::RunTest(const FString& Parameters)
@@ -213,12 +235,12 @@ bool FOpenBlendSpaceAssetEditors::RunTest(const FString& Parameters)
 /**
 * This test opens each BlendSpace1D into its sub-editor.
 */
-IMPLEMENT_COMPLEX_AUTOMATION_TEST(FOpenBlendSpace1DAssetEditors, "Editor.Content.Animation.BlendSpaces.Open BlendSpace1D Assets", (EAutomationTestFlags::ATF_Editor | EAutomationTestFlags::ATF_RequiresUser));
+IMPLEMENT_COMPLEX_AUTOMATION_TEST(FOpenBlendSpace1DAssetEditors, "Editor.Content.Skeletal Mesh.Animation.BlendSpaces.Open BlendSpace1D Assets", (EAutomationTestFlags::ATF_Editor | EAutomationTestFlags::ATF_RequiresUser));
 
 void FOpenBlendSpace1DAssetEditors::GetTests(TArray<FString>& OutBeautifiedNames, TArray<FString>& OutTestCommands) const
 {
 	//This grabs each BlendSpace1D asset in the Game/Content directory
-	FEditorAutomationTestUtilities::CollectGameContentTestsByClass(UBlendSpace1D::StaticClass(), OutBeautifiedNames, OutTestCommands);
+	FEditorAutomationTestUtilities::CollectGameContentTestsByClass(UBlendSpace1D::StaticClass(), false, OutBeautifiedNames, OutTestCommands);
 }
 
 bool FOpenBlendSpace1DAssetEditors::RunTest(const FString& Parameters)
@@ -228,6 +250,26 @@ bool FOpenBlendSpace1DAssetEditors::RunTest(const FString& Parameters)
 	return bDidTheTestPass;
 }
 
+
+//////////////////////////////////////////////////////////////////////////
+/**
+* This test opens each DESTRUCTIBLE MESH into its sub-editor.
+*/
+IMPLEMENT_COMPLEX_AUTOMATION_TEST(FOpenDestructibleMeshAssetEditors, "Editor.Content.Skeletal Mesh.Open Destructible Mesh Assets", (EAutomationTestFlags::ATF_Editor | EAutomationTestFlags::ATF_RequiresUser));
+
+void FOpenDestructibleMeshAssetEditors::GetTests(TArray<FString>& OutBeautifiedNames, TArray<FString>& OutTestCommands) const
+{
+	//This grabs each DestructibleMesh asset in the Game/Content directory
+	FEditorAutomationTestUtilities::CollectGameContentTestsByClass(UDestructibleMesh::StaticClass(), false, OutBeautifiedNames, OutTestCommands);
+}
+
+bool FOpenDestructibleMeshAssetEditors::RunTest(const FString& Parameters)
+{
+	bool bDidTheTestPass = FOpenActualAssetEditors(Parameters);
+	CollectGarbage(GARBAGE_COLLECTION_KEEPFLAGS);
+	return bDidTheTestPass;
+}
+
 //////////////////////////////////////////////////////////////////////////
 /**
 * This test opens each DialogueVoice into its sub-editor.
@@ -237,7 +279,7 @@ IMPLEMENT_COMPLEX_AUTOMATION_TEST(FOpenDialogueVoiceAssetEditors, "Editor.Conten
 void FOpenDialogueVoiceAssetEditors::GetTests(TArray<FString>& OutBeautifiedNames, TArray<FString>& OutTestCommands) const
 {
 	//This grabs each DialogueVoice asset in the Game/Content directory
-	FEditorAutomationTestUtilities::CollectGameContentTestsByClass(UDialogueVoice::StaticClass(), OutBeautifiedNames, OutTestCommands);
+	FEditorAutomationTestUtilities::CollectGameContentTestsByClass(UDialogueVoice::StaticClass(), true, OutBeautifiedNames, OutTestCommands);
 }
 
 bool FOpenDialogueVoiceAssetEditors::RunTest(const FString& Parameters)
@@ -256,7 +298,7 @@ IMPLEMENT_COMPLEX_AUTOMATION_TEST(FOpenDialogueWaveAssetEditors, "Editor.Content
 void FOpenDialogueWaveAssetEditors::GetTests(TArray<FString>& OutBeautifiedNames, TArray<FString>& OutTestCommands) const
 {
 	//This grabs each DialogueWave asset in the Game/Content directory
-	FEditorAutomationTestUtilities::CollectGameContentTestsByClass(UDialogueWave::StaticClass(), OutBeautifiedNames, OutTestCommands);
+	FEditorAutomationTestUtilities::CollectGameContentTestsByClass(UDialogueWave::StaticClass(), true, OutBeautifiedNames, OutTestCommands);
 }
 
 bool FOpenDialogueWaveAssetEditors::RunTest(const FString& Parameters)
@@ -275,7 +317,7 @@ IMPLEMENT_COMPLEX_AUTOMATION_TEST(FOpenMaterialAssetEditors, "Editor.Content.Mat
 void FOpenMaterialAssetEditors::GetTests(TArray<FString>& OutBeautifiedNames, TArray<FString>& OutTestCommands) const
 {
 	//This grabs each Material asset in the Game/Content directory
-	FEditorAutomationTestUtilities::CollectGameContentTestsByClass(UMaterial::StaticClass(), OutBeautifiedNames, OutTestCommands);
+	FEditorAutomationTestUtilities::CollectGameContentTestsByClass(UMaterial::StaticClass(), true, OutBeautifiedNames, OutTestCommands);
 }
 
 bool FOpenMaterialAssetEditors::RunTest(const FString& Parameters)
@@ -289,20 +331,20 @@ bool FOpenMaterialAssetEditors::RunTest(const FString& Parameters)
 ///**
 //* This test opens each MaterialFunction into its sub-editor.
 //*/
-//IMPLEMENT_COMPLEX_AUTOMATION_TEST(FOpenMaterialFunctionAssetEditors, "Editor.Content.Materials.Open Material Function Assets", (EAutomationTestFlags::ATF_Editor | EAutomationTestFlags::ATF_RequiresUser));
-//
-//void FOpenMaterialFunctionAssetEditors::GetTests(TArray<FString>& OutBeautifiedNames, TArray<FString>& OutTestCommands) const
-//{
-//	//This grabs each MaterialFunction asset in the Game/Content directory
-//	FEditorAutomationTestUtilities::CollectGameContentTestsByClass(UMaterialFunction::StaticClass(), OutBeautifiedNames, OutTestCommands);
-//}
-//
-//bool FOpenMaterialFunctionAssetEditors::RunTest(const FString& Parameters)
-//{
-//	bool bDidTheTestPass = FOpenActualAssetEditors(Parameters);
-//	CollectGarbage(GARBAGE_COLLECTION_KEEPFLAGS);
-//	return bDidTheTestPass;
-//}
+IMPLEMENT_COMPLEX_AUTOMATION_TEST(FOpenMaterialFunctionAssetEditors, "Editor.Content.Materials.Open Material Function Assets", (EAutomationTestFlags::ATF_Editor | EAutomationTestFlags::ATF_RequiresUser));
+
+void FOpenMaterialFunctionAssetEditors::GetTests(TArray<FString>& OutBeautifiedNames, TArray<FString>& OutTestCommands) const
+{
+	//This grabs each MaterialFunction asset in the Game/Content directory
+	FEditorAutomationTestUtilities::CollectGameContentTestsByClass(UMaterialFunction::StaticClass(), true, OutBeautifiedNames, OutTestCommands);
+}
+
+bool FOpenMaterialFunctionAssetEditors::RunTest(const FString& Parameters)
+{
+	bool bDidTheTestPass = FOpenActualAssetEditors(Parameters);
+	CollectGarbage(GARBAGE_COLLECTION_KEEPFLAGS);
+	return bDidTheTestPass;
+}
 
 //////////////////////////////////////////////////////////////////////////
 /**
@@ -313,7 +355,7 @@ IMPLEMENT_COMPLEX_AUTOMATION_TEST(FOpenMaterialInstanceConstantAssetEditors, "Ed
 void FOpenMaterialInstanceConstantAssetEditors::GetTests(TArray<FString>& OutBeautifiedNames, TArray<FString>& OutTestCommands) const
 {
 	//This grabs each MaterialInstanceConstant asset in the Game/Content directory
-	FEditorAutomationTestUtilities::CollectGameContentTestsByClass(UMaterialInstanceConstant::StaticClass(), OutBeautifiedNames, OutTestCommands);
+	FEditorAutomationTestUtilities::CollectGameContentTestsByClass(UMaterialInstanceConstant::StaticClass(), true, OutBeautifiedNames, OutTestCommands);
 }
 
 bool FOpenMaterialInstanceConstantAssetEditors::RunTest(const FString& Parameters)
@@ -332,7 +374,7 @@ IMPLEMENT_COMPLEX_AUTOMATION_TEST(FOpenMaterialParameterCollectionAssetEditors,
 void FOpenMaterialParameterCollectionAssetEditors::GetTests(TArray<FString>& OutBeautifiedNames, TArray<FString>& OutTestCommands) const
 {
 	//This grabs each MaterialParameterCollection asset in the Game/Content directory
-	FEditorAutomationTestUtilities::CollectGameContentTestsByClass(UMaterialParameterCollection::StaticClass(), OutBeautifiedNames, OutTestCommands);
+	FEditorAutomationTestUtilities::CollectGameContentTestsByClass(UMaterialParameterCollection::StaticClass(), true, OutBeautifiedNames, OutTestCommands);
 }
 
 bool FOpenMaterialParameterCollectionAssetEditors::RunTest(const FString& Parameters)
@@ -342,6 +384,24 @@ bool FOpenMaterialParameterCollectionAssetEditors::RunTest(const FString& Parame
 	return bDidTheTestPass;
 }
 
+//////////////////////////////////////////////////////////////////////////
+/**
+* This test opens each PARTICLE SYSTEM into its sub-editor.
+*/
+IMPLEMENT_COMPLEX_AUTOMATION_TEST(FOpenParticleSystemAssetEditors, "Editor.Content.Particle System.Open Particle System Assets", (EAutomationTestFlags::ATF_Editor | EAutomationTestFlags::ATF_RequiresUser));
+
+void FOpenParticleSystemAssetEditors::GetTests(TArray<FString>& OutBeautifiedNames, TArray<FString>& OutTestCommands) const
+{
+	//This grabs each ParticleSystem asset in the Game/Content directory
+	FEditorAutomationTestUtilities::CollectGameContentTestsByClass(UParticleSystem::StaticClass(), true, OutBeautifiedNames, OutTestCommands);
+}
+
+bool FOpenParticleSystemAssetEditors::RunTest(const FString& Parameters)
+{
+	bool bDidTheTestPass = FOpenActualAssetEditors(Parameters);
+	CollectGarbage(GARBAGE_COLLECTION_KEEPFLAGS);
+	return bDidTheTestPass;
+}
 
 //////////////////////////////////////////////////////////////////////////
 /**
@@ -352,7 +412,7 @@ IMPLEMENT_COMPLEX_AUTOMATION_TEST(FOpenPhysicalMaterialAssetEditors, "Editor.Con
 void FOpenPhysicalMaterialAssetEditors::GetTests(TArray<FString>& OutBeautifiedNames, TArray<FString>& OutTestCommands) const
 {
 	//This grabs each PhysicalMaterial asset in the Game/Content directory
-	FEditorAutomationTestUtilities::CollectGameContentTestsByClass(UPhysicalMaterial::StaticClass(), OutBeautifiedNames, OutTestCommands);
+	FEditorAutomationTestUtilities::CollectGameContentTestsByClass(UPhysicalMaterial::StaticClass(), true, OutBeautifiedNames, OutTestCommands);
 }
 
 bool FOpenPhysicalMaterialAssetEditors::RunTest(const FString& Parameters)
@@ -363,6 +423,25 @@ bool FOpenPhysicalMaterialAssetEditors::RunTest(const FString& Parameters)
 }
 
 
+//////////////////////////////////////////////////////////////////////////
+/**
+* This test opens each PhysicsAsset into its sub-editor.
+*/
+IMPLEMENT_COMPLEX_AUTOMATION_TEST(FOpenPhysicsAssetAssetEditors, "Editor.Content.Sounds.Open Physics Asset Assets", (EAutomationTestFlags::ATF_Editor | EAutomationTestFlags::ATF_RequiresUser));
+
+void FOpenPhysicsAssetAssetEditors::GetTests(TArray<FString>& OutBeautifiedNames, TArray<FString>& OutTestCommands) const
+{
+	//This grabs each PhysicsAsset asset in the Game/Content directory
+	FEditorAutomationTestUtilities::CollectGameContentTestsByClass(UPhysicsAsset::StaticClass(), true, OutBeautifiedNames, OutTestCommands);
+}
+
+bool FOpenPhysicsAssetAssetEditors::RunTest(const FString& Parameters)
+{
+	bool bDidTheTestPass = FOpenActualAssetEditors(Parameters);
+	CollectGarbage(GARBAGE_COLLECTION_KEEPFLAGS);
+	return bDidTheTestPass;
+}
+
 
 //////////////////////////////////////////////////////////////////////////
 /**
@@ -373,7 +452,7 @@ IMPLEMENT_COMPLEX_AUTOMATION_TEST(FOpenReverbEffectAssetEditors, "Editor.Content
 void FOpenReverbEffectAssetEditors::GetTests(TArray<FString>& OutBeautifiedNames, TArray<FString>& OutTestCommands) const
 {
 	//This grabs each ReverbEffect asset in the Game/Content directory
-	FEditorAutomationTestUtilities::CollectGameContentTestsByClass(UReverbEffect::StaticClass(), OutBeautifiedNames, OutTestCommands);
+	FEditorAutomationTestUtilities::CollectGameContentTestsByClass(UReverbEffect::StaticClass(), true, OutBeautifiedNames, OutTestCommands);
 }
 
 bool FOpenReverbEffectAssetEditors::RunTest(const FString& Parameters)
@@ -389,12 +468,12 @@ bool FOpenReverbEffectAssetEditors::RunTest(const FString& Parameters)
 /**
 * This test opens each SlateWidgetStyleAsset into its sub-editor.
 */
-IMPLEMENT_COMPLEX_AUTOMATION_TEST(FOpenSlateWidgetStyleAssetAssetEditors, "Editor.Content.Physics.Open SlateWidgetStyleAsset Assets", (EAutomationTestFlags::ATF_Editor | EAutomationTestFlags::ATF_RequiresUser));
+IMPLEMENT_COMPLEX_AUTOMATION_TEST(FOpenSlateWidgetStyleAssetAssetEditors, "Editor.Content.Slate.Open Slate Widget Style Assets", (EAutomationTestFlags::ATF_Editor | EAutomationTestFlags::ATF_RequiresUser));
 
 void FOpenSlateWidgetStyleAssetAssetEditors::GetTests(TArray<FString>& OutBeautifiedNames, TArray<FString>& OutTestCommands) const
 {
 	//This grabs each SlateWidgetStyleAsset asset in the Game/Content directory
-	FEditorAutomationTestUtilities::CollectGameContentTestsByClass(USlateWidgetStyleAsset::StaticClass(), OutBeautifiedNames, OutTestCommands);
+	FEditorAutomationTestUtilities::CollectGameContentTestsByClass(USlateWidgetStyleAsset::StaticClass(), true, OutBeautifiedNames, OutTestCommands);
 }
 
 bool FOpenSlateWidgetStyleAssetAssetEditors::RunTest(const FString& Parameters)
@@ -404,7 +483,24 @@ bool FOpenSlateWidgetStyleAssetAssetEditors::RunTest(const FString& Parameters)
 	return bDidTheTestPass;
 }
 
+//////////////////////////////////////////////////////////////////////////
+/**
+* This test opens each Slate Brush Asset into its sub-editor.
+*/
+IMPLEMENT_COMPLEX_AUTOMATION_TEST(FOpenSlateBrushAssetAssetEditors, "Editor.Content.Slate.Open Slate Brush Assets", (EAutomationTestFlags::ATF_Editor | EAutomationTestFlags::ATF_RequiresUser));
+
+void FOpenSlateBrushAssetAssetEditors::GetTests(TArray<FString>& OutBeautifiedNames, TArray<FString>& OutTestCommands) const
+{
+	//This grabs each SlateBrushAsset asset in the Game/Content directory
+	FEditorAutomationTestUtilities::CollectGameContentTestsByClass(USlateBrushAsset::StaticClass(), true, OutBeautifiedNames, OutTestCommands);
+}
 
+bool FOpenSlateBrushAssetAssetEditors::RunTest(const FString& Parameters)
+{
+	bool bDidTheTestPass = FOpenActualAssetEditors(Parameters);
+	CollectGarbage(GARBAGE_COLLECTION_KEEPFLAGS);
+	return bDidTheTestPass;
+}
 //////////////////////////////////////////////////////////////////////////
 /**
 * This test opens each SoundAttenuation into its sub-editor.
@@ -414,7 +510,7 @@ IMPLEMENT_COMPLEX_AUTOMATION_TEST(FOpenSoundAttenuationAssetEditors, "Editor.Con
 void FOpenSoundAttenuationAssetEditors::GetTests(TArray<FString>& OutBeautifiedNames, TArray<FString>& OutTestCommands) const
 {
 	//This grabs each SoundAttenuation asset in the Game/Content directory
-	FEditorAutomationTestUtilities::CollectGameContentTestsByClass(USoundAttenuation::StaticClass(), OutBeautifiedNames, OutTestCommands);
+	FEditorAutomationTestUtilities::CollectGameContentTestsByClass(USoundAttenuation::StaticClass(), true, OutBeautifiedNames, OutTestCommands);
 }
 
 bool FOpenSoundAttenuationAssetEditors::RunTest(const FString& Parameters)
@@ -433,7 +529,7 @@ IMPLEMENT_COMPLEX_AUTOMATION_TEST(FOpenStaticMeshAssetEditors, "Editor.Content.S
 void FOpenStaticMeshAssetEditors::GetTests(TArray<FString>& OutBeautifiedNames, TArray<FString>& OutTestCommands) const
 {
 	//This grabs each Static Mesh in the Game/Content
-	FEditorAutomationTestUtilities::CollectGameContentTestsByClass(UStaticMesh::StaticClass(), OutBeautifiedNames, OutTestCommands);
+	FEditorAutomationTestUtilities::CollectGameContentTestsByClass(UStaticMesh::StaticClass(), true, OutBeautifiedNames, OutTestCommands);
 }
 
 bool FOpenStaticMeshAssetEditors::RunTest(const FString& Parameters)
@@ -452,7 +548,7 @@ IMPLEMENT_COMPLEX_AUTOMATION_TEST(FOpenSkeletalMeshAssetEditors, "Editor.Content
 void FOpenSkeletalMeshAssetEditors::GetTests(TArray<FString>& OutBeautifiedNames, TArray<FString>& OutTestCommands) const
 {
 	//This grabs each SKeletal Mesh in the Game/Content
-	FEditorAutomationTestUtilities::CollectGameContentTestsByClass(USkeletalMesh::StaticClass(), OutBeautifiedNames, OutTestCommands);
+	FEditorAutomationTestUtilities::CollectGameContentTestsByClass(USkeletalMesh::StaticClass(), false, OutBeautifiedNames, OutTestCommands);
 }
 
 bool FOpenSkeletalMeshAssetEditors::RunTest(const FString& Parameters)
@@ -462,6 +558,24 @@ bool FOpenSkeletalMeshAssetEditors::RunTest(const FString& Parameters)
 	return bDidTheTestPass;
 }
 
+//////////////////////////////////////////////////////////////////////////
+/**
+* This test opens each Skeleton into its sub-editor.
+*/
+IMPLEMENT_COMPLEX_AUTOMATION_TEST(FOpenSkeletonAssetEditors, "Editor.Content.Skeletal Mesh.Open Skeleton Assets", (EAutomationTestFlags::ATF_Editor | EAutomationTestFlags::ATF_RequiresUser));
+
+void FOpenSkeletonAssetEditors::GetTests(TArray<FString>& OutBeautifiedNames, TArray<FString>& OutTestCommands) const
+{
+	//This grabs each Skeleton asset in the Game/Content directory
+	FEditorAutomationTestUtilities::CollectGameContentTestsByClass(USkeleton::StaticClass(), true, OutBeautifiedNames, OutTestCommands);
+}
+
+bool FOpenSkeletonAssetEditors::RunTest(const FString& Parameters)
+{
+	bool bDidTheTestPass = FOpenActualAssetEditors(Parameters);
+	CollectGarbage(GARBAGE_COLLECTION_KEEPFLAGS);
+	return bDidTheTestPass;
+}
 
 //////////////////////////////////////////////////////////////////////////
 /**
@@ -472,7 +586,7 @@ IMPLEMENT_COMPLEX_AUTOMATION_TEST(FOpenSoundClassAssetEditors, "Editor.Content.S
 void FOpenSoundClassAssetEditors::GetTests(TArray<FString>& OutBeautifiedNames, TArray<FString>& OutTestCommands) const
 {
 	//This grabs each SoundClass asset in the Game/Content directory
-	FEditorAutomationTestUtilities::CollectGameContentTestsByClass(USoundClass::StaticClass(), OutBeautifiedNames, OutTestCommands);
+	FEditorAutomationTestUtilities::CollectGameContentTestsByClass(USoundClass::StaticClass(), true, OutBeautifiedNames, OutTestCommands);
 }
 
 bool FOpenSoundClassAssetEditors::RunTest(const FString& Parameters)
@@ -492,7 +606,7 @@ IMPLEMENT_COMPLEX_AUTOMATION_TEST(FOpenSoundCueAssetEditors, "Editor.Content.Sou
 void FOpenSoundCueAssetEditors::GetTests(TArray<FString>& OutBeautifiedNames, TArray<FString>& OutTestCommands) const
 {
 	//This grabs each SoundCue asset in the Game/Content directory
-	FEditorAutomationTestUtilities::CollectGameContentTestsByClass(USoundCue::StaticClass(), OutBeautifiedNames, OutTestCommands);
+	FEditorAutomationTestUtilities::CollectGameContentTestsByClass(USoundCue::StaticClass(), true, OutBeautifiedNames, OutTestCommands);
 }
 
 bool FOpenSoundCueAssetEditors::RunTest(const FString& Parameters)
@@ -507,12 +621,12 @@ bool FOpenSoundCueAssetEditors::RunTest(const FString& Parameters)
 /**
 * This test opens each SoundMix into its sub-editor.
 */
-IMPLEMENT_COMPLEX_AUTOMATION_TEST(FOpenSoundMixAssetEditors, "Editor.Content.Sounds.Open SoundMix Assets", (EAutomationTestFlags::ATF_Editor | EAutomationTestFlags::ATF_RequiresUser));
+IMPLEMENT_COMPLEX_AUTOMATION_TEST(FOpenSoundMixAssetEditors, "Editor.Content.Sounds.Open Sound Mix Assets", (EAutomationTestFlags::ATF_Editor | EAutomationTestFlags::ATF_RequiresUser));
 
 void FOpenSoundMixAssetEditors::GetTests(TArray<FString>& OutBeautifiedNames, TArray<FString>& OutTestCommands) const
 {
 	//This grabs each SoundMix asset in the Game/Content directory
-	FEditorAutomationTestUtilities::CollectGameContentTestsByClass(USoundMix::StaticClass(), OutBeautifiedNames, OutTestCommands);
+	FEditorAutomationTestUtilities::CollectGameContentTestsByClass(USoundMix::StaticClass(), true, OutBeautifiedNames, OutTestCommands);
 }
 
 bool FOpenSoundMixAssetEditors::RunTest(const FString& Parameters)
@@ -525,82 +639,63 @@ bool FOpenSoundMixAssetEditors::RunTest(const FString& Parameters)
 
 //////////////////////////////////////////////////////////////////////////
 /**
-* This test opens each SubsurfaceProfile into its sub-editor.
+* This test opens each SoundWave into its sub-editor.
 */
-IMPLEMENT_COMPLEX_AUTOMATION_TEST(FOpenSubsurfaceProfileAssetEditors, "Editor.Content.Materials.Open Subsurface Profile Assets", (EAutomationTestFlags::ATF_Editor | EAutomationTestFlags::ATF_RequiresUser));
+IMPLEMENT_COMPLEX_AUTOMATION_TEST(FOpenSoundWaveAssetEditors, "Editor.Content.Sounds.Open Sound Wave Assets", (EAutomationTestFlags::ATF_Editor | EAutomationTestFlags::ATF_RequiresUser));
 
-void FOpenSubsurfaceProfileAssetEditors::GetTests(TArray<FString>& OutBeautifiedNames, TArray<FString>& OutTestCommands) const
+void FOpenSoundWaveAssetEditors::GetTests(TArray<FString>& OutBeautifiedNames, TArray<FString>& OutTestCommands) const
 {
-	//This grabs each SubsurfaceProfile asset in the Game/Content directory
-	FEditorAutomationTestUtilities::CollectGameContentTestsByClass(USubsurfaceProfile::StaticClass(), OutBeautifiedNames, OutTestCommands);
+	//This grabs each SoundWave asset in the Game/Content directory
+	FEditorAutomationTestUtilities::CollectGameContentTestsByClass(USoundWave::StaticClass(), true, OutBeautifiedNames, OutTestCommands);
 }
 
-bool FOpenSubsurfaceProfileAssetEditors::RunTest(const FString& Parameters)
+bool FOpenSoundWaveAssetEditors::RunTest(const FString& Parameters)
 {
 	bool bDidTheTestPass = FOpenActualAssetEditors(Parameters);
 	CollectGarbage(GARBAGE_COLLECTION_KEEPFLAGS);
 	return bDidTheTestPass;
 }
-
-
-
 //////////////////////////////////////////////////////////////////////////
 /**
-* This test opens each texture2d into its sub-editor.
+* This test opens each SubsurfaceProfile into its sub-editor.
 */
-IMPLEMENT_COMPLEX_AUTOMATION_TEST(FOpenTextureAssetEditors, "Editor.Content.Textures.Open Texture Assets", (EAutomationTestFlags::ATF_Editor | EAutomationTestFlags::ATF_RequiresUser));
+IMPLEMENT_COMPLEX_AUTOMATION_TEST(FOpenSubsurfaceProfileAssetEditors, "Editor.Content.Materials.Open Subsurface Profile Assets", (EAutomationTestFlags::ATF_Editor | EAutomationTestFlags::ATF_RequiresUser));
 
-void FOpenTextureAssetEditors::GetTests(TArray<FString>& OutBeautifiedNames, TArray<FString>& OutTestCommands) const
+void FOpenSubsurfaceProfileAssetEditors::GetTests(TArray<FString>& OutBeautifiedNames, TArray<FString>& OutTestCommands) const
 {
-	//This grabs each texture file in the Game/Content
-	FEditorAutomationTestUtilities::CollectGameContentTestsByClass(UTexture2D::StaticClass(), OutBeautifiedNames, OutTestCommands);
+	//This grabs each SubsurfaceProfile asset in the Game/Content directory
+	FEditorAutomationTestUtilities::CollectGameContentTestsByClass(USubsurfaceProfile::StaticClass(), true, OutBeautifiedNames, OutTestCommands);
 }
 
-bool FOpenTextureAssetEditors::RunTest(const FString& Parameters)
+bool FOpenSubsurfaceProfileAssetEditors::RunTest(const FString& Parameters)
 {
 	bool bDidTheTestPass = FOpenActualAssetEditors(Parameters);
 	CollectGarbage(GARBAGE_COLLECTION_KEEPFLAGS);
 	return bDidTheTestPass;
 }
 
-//////////////////////////////////////////////////////////////////////////
-/**
-* This test opens each texture cube into its sub-editor.
-*/
-IMPLEMENT_COMPLEX_AUTOMATION_TEST(FOpenTextureCubeAssetEditors, "Editor.Content.Textures.Open Texture Cube Assets", (EAutomationTestFlags::ATF_Editor | EAutomationTestFlags::ATF_RequiresUser));
-
-void FOpenTextureCubeAssetEditors::GetTests(TArray<FString>& OutBeautifiedNames, TArray<FString>& OutTestCommands) const
-{
-	//This grabs each texture cube in the Game/Content
-	FEditorAutomationTestUtilities::CollectGameContentTestsByClass(UTextureCube::StaticClass(), OutBeautifiedNames, OutTestCommands);
-}
 
-bool FOpenTextureCubeAssetEditors::RunTest(const FString& Parameters)
-{
-	bool bDidTheTestPass = FOpenActualAssetEditors(Parameters);
-	CollectGarbage(GARBAGE_COLLECTION_KEEPFLAGS);
-	return bDidTheTestPass;
-}
 
 //////////////////////////////////////////////////////////////////////////
 /**
-* This test opens each TextureRenderTargetCube into its sub-editor.
+* This test opens each TEXTURE asset into its sub-editor.
 */
-IMPLEMENT_COMPLEX_AUTOMATION_TEST(FOpenTextureRenderTargetCubeAssetEditors, "Editor.Content.Textures.Open Texture Render Target Cube Assets", (EAutomationTestFlags::ATF_Editor | EAutomationTestFlags::ATF_RequiresUser));
+IMPLEMENT_COMPLEX_AUTOMATION_TEST(FOpenTextureAssetEditors, "Editor.Content.Textures.Open Texture Assets", (EAutomationTestFlags::ATF_Editor | EAutomationTestFlags::ATF_RequiresUser));
 
-void FOpenTextureRenderTargetCubeAssetEditors::GetTests(TArray<FString>& OutBeautifiedNames, TArray<FString>& OutTestCommands) const
+void FOpenTextureAssetEditors::GetTests(TArray<FString>& OutBeautifiedNames, TArray<FString>& OutTestCommands) const
 {
-	//This grabs each TextureRenderTargetCube asset in the Game/Content directory
-	FEditorAutomationTestUtilities::CollectGameContentTestsByClass(UTextureRenderTargetCube::StaticClass(), OutBeautifiedNames, OutTestCommands);
+	//This grabs each texture file in the Game/Content
+	FEditorAutomationTestUtilities::CollectGameContentTestsByClass(UTexture::StaticClass(), true, OutBeautifiedNames, OutTestCommands);
 }
 
-bool FOpenTextureRenderTargetCubeAssetEditors::RunTest(const FString& Parameters)
+bool FOpenTextureAssetEditors::RunTest(const FString& Parameters)
 {
 	bool bDidTheTestPass = FOpenActualAssetEditors(Parameters);
 	CollectGarbage(GARBAGE_COLLECTION_KEEPFLAGS);
 	return bDidTheTestPass;
 }
 
+
 //////////////////////////////////////////////////////////////////////////
 /**
 * This test opens each Font into its sub-editor.
@@ -610,7 +705,7 @@ IMPLEMENT_COMPLEX_AUTOMATION_TEST(FOpenFontAssetEditors, "Editor.Content.Texture
 void FOpenFontAssetEditors::GetTests(TArray<FString>& OutBeautifiedNames, TArray<FString>& OutTestCommands) const
 {
 	//This grabs each Font asset in the Game/Content directory
-	FEditorAutomationTestUtilities::CollectGameContentTestsByClass(UFont::StaticClass(), OutBeautifiedNames, OutTestCommands);
+	FEditorAutomationTestUtilities::CollectGameContentTestsByClass(UFont::StaticClass(), true, OutBeautifiedNames, OutTestCommands);
 }
 
 bool FOpenFontAssetEditors::RunTest(const FString& Parameters)
@@ -622,17 +717,17 @@ bool FOpenFontAssetEditors::RunTest(const FString& Parameters)
 
 //////////////////////////////////////////////////////////////////////////
 /**
-* This test opens each TextureRenderTarget2D into its sub-editor.
+* This test opens each misc asset into its sub-editor.
 */
-IMPLEMENT_COMPLEX_AUTOMATION_TEST(FOpenTextureRenderTarget2DAssetEditors, "Editor.Content.Textures.Open Texture Render Target 2D Assets", (EAutomationTestFlags::ATF_Editor | EAutomationTestFlags::ATF_RequiresUser));
+IMPLEMENT_COMPLEX_AUTOMATION_TEST(FOpenMiscAssetEditors, "Editor.Content.Misc.Open Misc Assets", (EAutomationTestFlags::ATF_Editor | EAutomationTestFlags::ATF_RequiresUser));
 
-void FOpenTextureRenderTarget2DAssetEditors::GetTests(TArray<FString>& OutBeautifiedNames, TArray<FString>& OutTestCommands) const
+void FOpenMiscAssetEditors::GetTests(TArray<FString>& OutBeautifiedNames, TArray<FString>& OutTestCommands) const
 {
 	//This grabs each TextureRenderTarget2D asset in the Game/Content directory
-	FEditorAutomationTestUtilities::CollectGameContentTestsByClass(UTextureRenderTarget2D::StaticClass(), OutBeautifiedNames, OutTestCommands);
+	FEditorAutomationTestUtilities::CollectMiscGameContentTestsByClass(OutBeautifiedNames, OutTestCommands);
 }
 
-bool FOpenTextureRenderTarget2DAssetEditors::RunTest(const FString& Parameters)
+bool FOpenMiscAssetEditors::RunTest(const FString& Parameters)
 {
 	bool bDidTheTestPass = FOpenActualAssetEditors(Parameters);
 	CollectGarbage(GARBAGE_COLLECTION_KEEPFLAGS);
diff --git a/Engine/Source/Editor/UnrealEd/Public/Tests/AutomationEditorCommon.h b/Engine/Source/Editor/UnrealEd/Public/Tests/AutomationEditorCommon.h
--- a/Engine/Source/Editor/UnrealEd/Public/Tests/AutomationEditorCommon.h
+++ b/Engine/Source/Editor/UnrealEd/Public/Tests/AutomationEditorCommon.h
@@ -7,6 +7,10 @@
 
 #include "Tests/AutomationTestSettings.h"
 
+//Includes needed for opening certain assets
+#include "Materials/MaterialFunction.h"
+#include "Slate/SlateBrushAsset.h"
+
 namespace AutomationEditorCommonUtils
 {
 	/**
@@ -176,10 +180,8 @@ public:
 		{
 			UClass* ClassList = *AllClassesIt;
 			FName ClassName = ClassList->GetFName();
-			//UE_LOG(LogEditorAutomationTests, Log, TEXT("Here: %s"), *ClassName);
 		}
-
-
+		
 		for (auto ObjIter = ObjectList.CreateConstIterator(); ObjIter; ++ObjIter)
 		{
 			const FAssetData & Asset = *ObjIter;
@@ -199,7 +201,7 @@ public:
 	* Generates a list of assets from the GAME by a specific type.
 	* This is to be used by the GetTest() function.
 	*/
-	static void CollectGameContentTestsByClass(UClass * Class, TArray<FString>& OutBeautifiedNames, TArray <FString>& OutTestCommands)
+	static void CollectGameContentTestsByClass(UClass * Class, bool bRecursiveClass, TArray<FString>& OutBeautifiedNames, TArray <FString>& OutTestCommands)
 	{
 		//Setting the Asset Registry
 
@@ -213,6 +215,7 @@ public:
 		//This list is being filtered by the game folder and class type.  The results are placed into the ObjectList variable.
 		AssetFilter.ClassNames.Add(Class->GetFName());
 		AssetFilter.PackagePaths.Add("/Game");
+		AssetFilter.bRecursiveClasses = bRecursiveClass;
 		AssetFilter.bRecursivePaths = true;
 		AssetRegistryModule.Get().GetAssets(AssetFilter, ObjectList);
 
@@ -231,6 +234,105 @@ public:
 			}
 		}
 	}
+
+	/**
+	* Generates a list of misc. assets from the GAME.
+	* This is to be used by the GetTest() function.
+	*/
+	static void CollectMiscGameContentTestsByClass(TArray<FString>& OutBeautifiedNames, TArray <FString>& OutTestCommands)
+	{
+		//Setting the Asset Registry
+		FAssetRegistryModule& AssetRegistryModule = FModuleManager::LoadModuleChecked<FAssetRegistryModule>(TEXT("AssetRegistry"));
+		TArray<FAssetData> ObjectList;
+		FARFilter AssetFilter;
+
+		//This is a list of classes that we don't want to be in the misc category.
+		TArray<FName> ExcludeClassesList;
+		ExcludeClassesList.Add(UTexture::StaticClass()->GetFName());
+		ExcludeClassesList.Add(UTexture2D::StaticClass()->GetFName());
+		ExcludeClassesList.Add(UTextureRenderTarget::StaticClass()->GetFName());
+		ExcludeClassesList.Add(UTextureRenderTarget2D::StaticClass()->GetFName());
+		ExcludeClassesList.Add(UTextureCube::StaticClass()->GetFName());
+		ExcludeClassesList.Add(UAimOffsetBlendSpace::StaticClass()->GetFName());
+		ExcludeClassesList.Add(UAimOffsetBlendSpace1D::StaticClass()->GetFName());
+		ExcludeClassesList.Add(UAnimMontage::StaticClass()->GetFName());
+		ExcludeClassesList.Add(UAnimSequence::StaticClass()->GetFName());
+		ExcludeClassesList.Add(UBlendSpace::StaticClass()->GetFName());
+		ExcludeClassesList.Add(UBlendSpace1D::StaticClass()->GetFName());
+		ExcludeClassesList.Add(UDestructibleMesh::StaticClass()->GetFName());
+		ExcludeClassesList.Add(UDialogueVoice::StaticClass()->GetFName());
+		ExcludeClassesList.Add(UDialogueWave::StaticClass()->GetFName());
+		ExcludeClassesList.Add(UMaterial::StaticClass()->GetFName());
+		ExcludeClassesList.Add(UMaterialFunction::StaticClass()->GetFName());
+		ExcludeClassesList.Add(UMaterialInstanceConstant::StaticClass()->GetFName());
+		ExcludeClassesList.Add(UMaterialParameterCollection::StaticClass()->GetFName());
+		ExcludeClassesList.Add(UParticleSystem::StaticClass()->GetFName());
+		ExcludeClassesList.Add(UPhysicalMaterial::StaticClass()->GetFName());
+		ExcludeClassesList.Add(UReverbEffect::StaticClass()->GetFName());
+		ExcludeClassesList.Add(USlateWidgetStyleAsset::StaticClass()->GetFName());
+		ExcludeClassesList.Add(USlateBrushAsset::StaticClass()->GetFName());
+		ExcludeClassesList.Add(USoundAttenuation::StaticClass()->GetFName());
+		ExcludeClassesList.Add(UStaticMesh::StaticClass()->GetFName());
+		ExcludeClassesList.Add(USkeletalMesh::StaticClass()->GetFName());
+		ExcludeClassesList.Add(USkeleton::StaticClass()->GetFName());
+		ExcludeClassesList.Add(USoundClass::StaticClass()->GetFName());
+		ExcludeClassesList.Add(USoundCue::StaticClass()->GetFName());
+		ExcludeClassesList.Add(USoundMix::StaticClass()->GetFName());
+		ExcludeClassesList.Add(USoundWave::StaticClass()->GetFName());
+		ExcludeClassesList.Add(USubsurfaceProfile::StaticClass()->GetFName());
+		ExcludeClassesList.Add(UFont::StaticClass()->GetFName());
+		ExcludeClassesList.Add(UBlueprint::StaticClass()->GetFName());
+		ExcludeClassesList.Add(UAnimBlueprint::StaticClass()->GetFName());
+		
+		//Generating the list of assets.
+		//This list isn't expected to have anything that may be obtained with another function.
+		//This list is being filtered by the game folder and class type.  The results are placed into the ObjectList variable.
+		AssetFilter.PackagePaths.Add("/Game");
+		AssetFilter.bRecursivePaths = true;
+		AssetRegistryModule.Get().GetAssets(AssetFilter, ObjectList);
+		
+		//Loop through the list of assets, make their path full and a string, then add them to the test.
+		for (auto ObjIter = ObjectList.CreateConstIterator(); ObjIter; ++ObjIter)
+		{
+			const FAssetData & Asset = *ObjIter;
+			//Variable that holds the class FName for the current asset iteration.
+			FName AssetClassFName = Asset.GetClass()->GetFName();
+
+			//Counter used to keep track for the following for loop.
+			float ExcludedClassesCounter = 1;
+
+			for (auto ExcludeIter = ExcludeClassesList.CreateConstIterator(); ExcludeIter; ++ExcludeIter)
+			{
+				FName ExludedName = *ExcludeIter;
+				
+				//If the classes are the same then we don't want this asset. So we move onto the next one instead.
+				if (AssetClassFName == ExludedName)
+				{
+					break;
+				}
+
+				//We run out of class names in our Excluded list then we want the current ObjectList asset.
+				if ((ExcludedClassesCounter + 1) > ExcludeClassesList.Num())
+				{
+					FString Filename = Asset.ObjectPath.ToString();
+					//convert to full paths
+					Filename = FPackageName::LongPackageNameToFilename(Filename);
+
+					if (FAutomationTestFramework::GetInstance().ShouldTestContent(Filename))
+					{
+						FString BeautifiedFilename = Asset.AssetName.ToString();
+						OutBeautifiedNames.Add(BeautifiedFilename);
+						OutTestCommands.Add(Asset.ObjectPath.ToString());
+					}
+
+					break;
+				}
+				
+				//increment the counter.
+				ExcludedClassesCounter++;
+			}
+		}
+	}
 };
 
 
