Updates to the actual commands that open and close the sub editors for the open all asset tests.

diff --git a/Engine/Source/Editor/UnrealEd/Private/Tests/AutomationEditorCommon.cpp b/Engine/Source/Editor/UnrealEd/Private/Tests/AutomationEditorCommon.cpp
--- a/Engine/Source/Editor/UnrealEd/Private/Tests/AutomationEditorCommon.cpp
+++ b/Engine/Source/Editor/UnrealEd/Private/Tests/AutomationEditorCommon.cpp
@@ -307,7 +307,13 @@ bool FOpenEditorForAssetCommand::Update()
 	if (Object)
 	{
 		FAssetEditorManager::Get().OpenEditorForAsset(Object);
-		UE_LOG(LogEditorAutomationTests, Log, TEXT("Verified asset editor for: %s."), *AssetName);
+		//This checks to see if the asset sub editor is loaded.
+		if (FAssetEditorManager::Get().FindEditorForAsset(Object, true) != NULL)
+		{
+			UE_LOG(LogEditorAutomationTests, Log, TEXT("Verified asset editor for: %s."), *AssetName);
+			UE_LOG(LogEditorAutomationTests, Display, TEXT("The editor successffully loaded for: %s."), *AssetName);
+			return true;
+		}
 	}
 	else
 	{
@@ -323,6 +329,15 @@ bool FCloseAllAssetEditorsCommand::Update()
 {
 	FAssetEditorManager::Get().CloseAllAssetEditors();
 
+	//Get all assets currently being tracked with open editors and make sure they are not still opened.
+	if(FAssetEditorManager::Get().GetAllEditedAssets().Num() >= 1)
+	{
+		UE_LOG(LogEditorAutomationTests, Warning, TEXT("Not all of the editors were closed."));
+		return true;
+	}
+
+	UE_LOG(LogEditorAutomationTests, Log, TEXT("Verified asset editors were closed"));
+	UE_LOG(LogEditorAutomationTests, Display, TEXT("The asset editors closed successffully"));
 	return true;
 }
 
diff --git a/Engine/Source/Editor/UnrealEd/Private/Tests/EditorOpenAssetAutomationTests.cpp b/Engine/Source/Editor/UnrealEd/Private/Tests/EditorOpenAssetAutomationTests.cpp
--- a/Engine/Source/Editor/UnrealEd/Private/Tests/EditorOpenAssetAutomationTests.cpp
+++ b/Engine/Source/Editor/UnrealEd/Private/Tests/EditorOpenAssetAutomationTests.cpp
@@ -140,7 +140,7 @@ bool FOpenAssetEditors::RunTest(const FString& Parameters)
 /**
 * This test opens each AimOffsetBlendSpace into its sub-editor.
 */
-IMPLEMENT_COMPLEX_AUTOMATION_TEST(FOpenAimOffsetBlendSpaceAssetEditors, "Editor.Content.Skeletal Mesh.Animation.BlendSpaces.Open AimOffsetBlendSpace Assets", (EAutomationTestFlags::ATF_Editor | EAutomationTestFlags::ATF_RequiresUser));
+IMPLEMENT_COMPLEX_AUTOMATION_TEST(FOpenAimOffsetBlendSpaceAssetEditors, "Editor.Content.Skeletal Mesh.Animation.BlendSpaces.Open Aim Offset BlendSpace Assets", (EAutomationTestFlags::ATF_Editor | EAutomationTestFlags::ATF_RequiresUser));
 
 void FOpenAimOffsetBlendSpaceAssetEditors::GetTests(TArray<FString>& OutBeautifiedNames, TArray<FString>& OutTestCommands) const
 {
@@ -159,7 +159,7 @@ bool FOpenAimOffsetBlendSpaceAssetEditors::RunTest(const FString& Parameters)
 /**
 * This test opens each AimOffsetBlendSpace1D into its sub-editor.
 */
-IMPLEMENT_COMPLEX_AUTOMATION_TEST(FOpenAimOffsetBlendSpace1DAssetEditors, "Editor.Content.Skeletal Mesh.Animation.BlendSpaces.Open AimOffsetBlendSpace1D Assets", (EAutomationTestFlags::ATF_Editor | EAutomationTestFlags::ATF_RequiresUser));
+IMPLEMENT_COMPLEX_AUTOMATION_TEST(FOpenAimOffsetBlendSpace1DAssetEditors, "Editor.Content.Skeletal Mesh.Animation.BlendSpaces.Open Aim Offset BlendSpace1D Assets", (EAutomationTestFlags::ATF_Editor | EAutomationTestFlags::ATF_RequiresUser));
 
 void FOpenAimOffsetBlendSpace1DAssetEditors::GetTests(TArray<FString>& OutBeautifiedNames, TArray<FString>& OutTestCommands) const
 {
@@ -407,7 +407,7 @@ bool FOpenParticleSystemAssetEditors::RunTest(const FString& Parameters)
 /**
 * This test opens each PhysicalMaterial into its sub-editor.
 */
-IMPLEMENT_COMPLEX_AUTOMATION_TEST(FOpenPhysicalMaterialAssetEditors, "Editor.Content.Physics.Open PhysicalMaterial Assets", (EAutomationTestFlags::ATF_Editor | EAutomationTestFlags::ATF_RequiresUser));
+IMPLEMENT_COMPLEX_AUTOMATION_TEST(FOpenPhysicalMaterialAssetEditors, "Editor.Content.Physics.Open Physical Material Assets", (EAutomationTestFlags::ATF_Editor | EAutomationTestFlags::ATF_RequiresUser));
 
 void FOpenPhysicalMaterialAssetEditors::GetTests(TArray<FString>& OutBeautifiedNames, TArray<FString>& OutTestCommands) const
 {
@@ -427,7 +427,7 @@ bool FOpenPhysicalMaterialAssetEditors::RunTest(const FString& Parameters)
 /**
 * This test opens each PhysicsAsset into its sub-editor.
 */
-IMPLEMENT_COMPLEX_AUTOMATION_TEST(FOpenPhysicsAssetAssetEditors, "Editor.Content.Sounds.Open Physics Asset Assets", (EAutomationTestFlags::ATF_Editor | EAutomationTestFlags::ATF_RequiresUser));
+IMPLEMENT_COMPLEX_AUTOMATION_TEST(FOpenPhysicsAssetAssetEditors, "Editor.Content.Physics.Open Physics Asset Assets", (EAutomationTestFlags::ATF_Editor | EAutomationTestFlags::ATF_RequiresUser));
 
 void FOpenPhysicsAssetAssetEditors::GetTests(TArray<FString>& OutBeautifiedNames, TArray<FString>& OutTestCommands) const
 {
