Removed commented out reference code that was accidentally left in a previous check in.

diff --git a/Engine/Source/Editor/UnrealEd/Private/SkeletalMeshEdit.cpp b/Engine/Source/Editor/UnrealEd/Private/SkeletalMeshEdit.cpp
--- a/Engine/Source/Editor/UnrealEd/Private/SkeletalMeshEdit.cpp
+++ b/Engine/Source/Editor/UnrealEd/Private/SkeletalMeshEdit.cpp
@@ -1030,32 +1030,6 @@ bool UnFbx::FFbxImporter::ImportAnimation(USkeleton * Skeleton, UAnimSequence *
 							{
 								// now see if we have one already exists. If so, just overwrite that. if not, add new one. 
 								ImportCurveToAnimSequence(DestSeq, *ChannelName, Curve,  ACF_DrivesMorphTarget | ACF_TriggerEvent, AnimTimeSpan, 0.01f /** for some reason blend shape values are coming as 100 scaled **/);
-
-
-//								FName Name = *ChannelName;
-//								USkeleton::AnimCurveUID CurveId;
-//								Skeleton->SmartNames.GetContainer(USkeleton::AnimCurveMappingName)->AddName(Name, CurveId);
-//								FFloatCurve * CurveToImport = DestSeq->RawCurveData.GetCurveData(CurveId);
-//								if (CurveToImport==NULL)
-//								{
-//									if ( DestSeq->RawCurveData.AddCurveData(CurveId, ACF_DrivesMorphTarget | ACF_TriggerEvent) )
-//									{
-//										CurveToImport = DestSeq->RawCurveData.GetCurveData(CurveId);
-//									}
-//									else
-//									{
-//										// this should not happen, we already checked before adding
-//										ensureMsg (0, TEXT("FBX Import: Critical error: no memory?"));
-//									}
-//								}
-//								else
-//								{
-//									CurveToImport->FloatCurve.Reset();
-//								}
-//	
-//								ImportCurve(Curve, CurveToImport, AnimTimeSpan);
-
-
 							}
 						}
 					}
