Fixed issue with importing lesser # of frame since when we calculate sample rate, we don't round up if closer to the bigger one.

: The issue came out as 29.999999 cutting to 29, and I think using [REDACTED_PHONE] would work fine, but if any other case, where it hits closer to bigger number, it seems fair to go bigger number than smaller number.

diff --git a/Engine/Source/Editor/UnrealEd/Private/SkeletalMeshEdit.cpp b/Engine/Source/Editor/UnrealEd/Private/SkeletalMeshEdit.cpp
--- a/Engine/Source/Editor/UnrealEd/Private/SkeletalMeshEdit.cpp
+++ b/Engine/Source/Editor/UnrealEd/Private/SkeletalMeshEdit.cpp
@@ -609,7 +609,7 @@ int32 UnFbx::FFbxImporter::GetMaxSampleRate(TArray<FbxNode*>& SortedLinks, TArra
 						if(KeyAnimLength != 0.0)
 						{
 							// 30 fps animation has 31 keys because it includes index 0 key for 0.0 second
-							int32 NewRate = static_cast<int32>((KeyCount-1) / KeyAnimLength);
+							int32 NewRate = FPlatformMath::RoundToInt((KeyCount-1) / KeyAnimLength);
 							MaxStackResampleRate = FMath::Max(NewRate, MaxStackResampleRate);
 						}
 					}
