fix for not properly adding the maps to the file list when cooking with -Map
#ue4

diff --git a/Engine/Source/Editor/UnrealEd/Private/Commandlets/CookCommandlet.cpp b/Engine/Source/Editor/UnrealEd/Private/Commandlets/CookCommandlet.cpp
--- a/Engine/Source/Editor/UnrealEd/Private/Commandlets/CookCommandlet.cpp
+++ b/Engine/Source/Editor/UnrealEd/Private/Commandlets/CookCommandlet.cpp
@@ -785,10 +785,10 @@ void UCookCommandlet::CollectFilesToCook(TArray<FString>& FilesInPath)
 		};
 
 		// Check for -MAP=<name of map> entries
-		CmdLineMapEntries = GetSwitchValueElements(TEXT("MAP"));
+		CmdLineMapEntries += GetSwitchValueElements(TEXT("MAP"));
 
 		// Check for -COOKDIR=<path to directory> entries
-		CmdLineDirEntries = GetSwitchValueElements(TEXT("COOKDIR"));
+		CmdLineDirEntries += GetSwitchValueElements(TEXT("COOKDIR"));
 		for(FString& Entry : CmdLineDirEntries)
 		{
 			Entry = Entry.TrimQuotes();
@@ -796,7 +796,7 @@ void UCookCommandlet::CollectFilesToCook(TArray<FString>& FilesInPath)
 		}
 
 		// Check for -COOKCULTURES=<culture name> entries
-		CmdLineCultEntries = GetSwitchValueElements(TEXT("COOKCULTURES"));
+		CmdLineCultEntries += GetSwitchValueElements(TEXT("COOKCULTURES"));
 	}
 
 	// Also append any cookdirs from the project ini files; these dirs are relative to the game content directory
