Added new Lightmass WorldSetting VolumeLightSamplePlacementScale which allows users to reduce IndirectLightingCache cost and memory in their level

diff --git a/Engine/Source/Editor/UnrealEd/Private/Lightmass/Lightmass.cpp b/Engine/Source/Editor/UnrealEd/Private/Lightmass/Lightmass.cpp
--- a/Engine/Source/Editor/UnrealEd/Private/Lightmass/Lightmass.cpp
+++ b/Engine/Source/Editor/UnrealEd/Private/Lightmass/Lightmass.cpp
@@ -1864,6 +1864,10 @@ void FLightmassExporter::WriteSceneSettings( Lightmass::FSceneFileHeader& Scene
 		verify(GConfig->GetBool(TEXT("DevOptions.PrecomputedDynamicObjectLighting"), TEXT("bUseMaxSurfaceSampleNum"), bConfigBool, GLightmassIni));
 		Scene.DynamicObjectSettings.bUseMaxSurfaceSampleNum = bConfigBool;
 		verify(GConfig->GetInt(TEXT("DevOptions.PrecomputedDynamicObjectLighting"), TEXT("MaxSurfaceLightSamples"), Scene.DynamicObjectSettings.MaxSurfaceLightSamples, GLightmassIni));
+
+		Scene.DynamicObjectSettings.SurfaceLightSampleSpacing *= LevelSettings.VolumeLightSamplePlacementScale;
+		Scene.DynamicObjectSettings.VolumeLightSampleSpacing *= LevelSettings.VolumeLightSamplePlacementScale;
+		Scene.DynamicObjectSettings.DetailVolumeSampleSpacing *= LevelSettings.VolumeLightSamplePlacementScale;
 	}
 	{
 		verify(GConfig->GetBool(TEXT("DevOptions.PrecomputedVisibility"), TEXT("bVisualizePrecomputedVisibility"), bConfigBool, GLightmassIni));
diff --git a/Engine/Source/Runtime/Engine/Classes/GameFramework/WorldSettings.h b/Engine/Source/Runtime/Engine/Classes/GameFramework/WorldSettings.h
--- a/Engine/Source/Runtime/Engine/Classes/GameFramework/WorldSettings.h
+++ b/Engine/Source/Runtime/Engine/Classes/GameFramework/WorldSettings.h
@@ -121,6 +121,13 @@ struct FLightmassWorldInfoSettings
 	UPROPERTY(EditAnywhere, Category=LightmassDebug, AdvancedDisplay)
 	uint32 bVisualizeAmbientOcclusion:1;
 
+	/** 
+	 * Scales the distances at which volume lighting samples are placed.  Volume lighting samples are computed by Lightmass and are used for GI on movable components.
+	 * Using larger scales results in less sample memory usage and reduces Indirect Lighting Cache update times.
+	 */
+	UPROPERTY(EditAnywhere, Category=LightmassGeneral, AdvancedDisplay, meta=(UIMin = "0.1", UIMax = "100.0"))
+	float VolumeLightSamplePlacementScale;
+
 	FLightmassWorldInfoSettings()
 		: StaticLightingLevelScale(1)
 		, NumIndirectLightingBounces(3)
@@ -138,6 +145,7 @@ struct FLightmassWorldInfoSettings
 		, MaxOcclusionDistance(200.0f)
 		, bVisualizeMaterialDiffuse(false)
 		, bVisualizeAmbientOcclusion(false)
+		, VolumeLightSamplePlacementScale(1)
 	{
 	}
 };
diff --git a/Engine/Source/Runtime/Engine/Private/WorldSettings.cpp b/Engine/Source/Runtime/Engine/Private/WorldSettings.cpp
--- a/Engine/Source/Runtime/Engine/Private/WorldSettings.cpp
+++ b/Engine/Source/Runtime/Engine/Private/WorldSettings.cpp
@@ -242,6 +242,7 @@ void AWorldSettings::PostEditChangeProperty(FPropertyChangedEvent& PropertyChang
 
 	LightmassSettings.NumIndirectLightingBounces = FMath::Clamp(LightmassSettings.NumIndirectLightingBounces, 0, 100);
 	LightmassSettings.IndirectLightingSmoothness = FMath::Clamp(LightmassSettings.IndirectLightingSmoothness, .25f, 10.0f);
+	LightmassSettings.VolumeLightSamplePlacementScale = FMath::Clamp(LightmassSettings.VolumeLightSamplePlacementScale, .1f, 100.0f);
 	LightmassSettings.IndirectLightingQuality = FMath::Clamp(LightmassSettings.IndirectLightingQuality, .1f, 10.0f);
 	LightmassSettings.StaticLightingLevelScale = FMath::Clamp(LightmassSettings.StaticLightingLevelScale, .001f, 1000.0f);
 	LightmassSettings.EmissiveBoost = FMath::Max(LightmassSettings.EmissiveBoost, 0.0f);
