Dumping out database stats to help keep an eye one the memory impact of the new blueprint menu system (also adding command option for adding interfaces).

diff --git a/Engine/Source/Editor/BlueprintGraph/Private/BlueprintActionDatabase.cpp b/Engine/Source/Editor/BlueprintGraph/Private/BlueprintActionDatabase.cpp
--- a/Engine/Source/Editor/BlueprintGraph/Private/BlueprintActionDatabase.cpp
+++ b/Engine/Source/Editor/BlueprintGraph/Private/BlueprintActionDatabase.cpp
@@ -1043,4 +1043,15 @@ FBlueprintActionDatabase::FClassActionMap const& FBlueprintActionDatabase::GetAl
 	return ClassActions;
 }
 
+//------------------------------------------------------------------------------
+int32 FBlueprintActionDatabase::Size() const
+{
+	int32 TotalSize = sizeof(*this);
+	for (TObjectIterator<UBlueprintNodeSpawner> NodeSpawnerIt; NodeSpawnerIt; ++NodeSpawnerIt)
+	{
+		TotalSize += sizeof(**NodeSpawnerIt);
+	}
+	return TotalSize;
+}
+
 #undef LOCTEXT_NAMESPACE
diff --git a/Engine/Source/Editor/BlueprintGraph/Private/BlueprintNodeTemplateCache.cpp b/Engine/Source/Editor/BlueprintGraph/Private/BlueprintNodeTemplateCache.cpp
--- a/Engine/Source/Editor/BlueprintGraph/Private/BlueprintNodeTemplateCache.cpp
+++ b/Engine/Source/Editor/BlueprintGraph/Private/BlueprintNodeTemplateCache.cpp
@@ -124,7 +124,7 @@ FBlueprintNodeTemplateCache::FBlueprintNodeTemplateCache()
 	using namespace BlueprintNodeTemplateCacheImpl; // for MakeCompatibleBlueprint()
 
 	TemplateOuters.Add(MakeCompatibleBlueprint(UBlueprint::StaticClass(), AActor::StaticClass(), UBlueprintGeneratedClass::StaticClass()));
-	TemplateOuters.Add(MakeCompatibleBlueprint(UAnimBlueprint::StaticClass(), UAnimInstance::StaticClass(), UAnimBlueprintGeneratedClass::StaticClass()));
+	//TemplateOuters.Add(MakeCompatibleBlueprint(UAnimBlueprint::StaticClass(), UAnimInstance::StaticClass(), UAnimBlueprintGeneratedClass::StaticClass()));
 }
 
 //------------------------------------------------------------------------------
diff --git a/Engine/Source/Editor/BlueprintGraph/Public/BlueprintActionDatabase.h b/Engine/Source/Editor/BlueprintGraph/Public/BlueprintActionDatabase.h
--- a/Engine/Source/Editor/BlueprintGraph/Public/BlueprintActionDatabase.h
+++ b/Engine/Source/Editor/BlueprintGraph/Public/BlueprintActionDatabase.h
@@ -82,6 +82,11 @@ public:
 	 */
 	FClassActionMap const& GetAllActions();
 
+	/**
+	 * @return An estimated memory footprint of this database (in bytes).
+	 */
+	int32 Size() const;
+
 private:
 	/** Private constructor for singleton purposes. */
 	FBlueprintActionDatabase();
diff --git a/Engine/Source/Editor/UnrealEd/Private/Commandlets/DumpBlueprintsInfoCommandlet.cpp b/Engine/Source/Editor/UnrealEd/Private/Commandlets/DumpBlueprintsInfoCommandlet.cpp
--- a/Engine/Source/Editor/UnrealEd/Private/Commandlets/DumpBlueprintsInfoCommandlet.cpp
+++ b/Engine/Source/Editor/UnrealEd/Private/Commandlets/DumpBlueprintsInfoCommandlet.cpp
@@ -23,6 +23,7 @@
 #include "K2Node_Composite.h"
 #include "ModuleManager.h"
 #include "AssetToolsModule.h"
+#include "BlueprintNodeSpawner.h"
 
 DEFINE_LOG_CATEGORY_STATIC(LogBlueprintInfoDump, Log, All);
 
@@ -100,6 +101,10 @@ DumpBlueprintsInfo commandlet params: \n\
     -name=<Filename>    Overrides the default filename. Leave off the extention\n\
                         (this will add .json to the end). When -multifile is   \n\
                         supplied, the class name will be postfixed to the name.\n\
+\n\
+    -interface=<Class>  Appends the desired interface to blueprints that are   \n\
+                        being dumped. The <Class> param has to match a known   \n\
+                        interface class.                                       \n\
 \n\
     -help, -h, -?       Display this message and then exit.                    \n\
 \n");
@@ -137,6 +142,7 @@ DumpBlueprintsInfo commandlet params: \n\
 			, PaletteFilter(nullptr)
 			, GraphFilter(GT_MAX)
 			, SelectedObjectType(nullptr)
+			, InterfaceClass(nullptr)
 		{
 		}
 
@@ -156,6 +162,7 @@ DumpBlueprintsInfo commandlet params: \n\
 		FString    DiffPath;
 		FString    DiffCommand;
 		FString    Filename;
+		UClass*	   InterfaceClass;
 	};
 
 	/** Static instance of the command switches (so we don't have to pass one along the call stack) */
@@ -233,12 +240,17 @@ DumpBlueprintsInfo commandlet params: \n\
 	static FString GetActionKey(FGraphActionListBuilderBase::ActionGroup const& Action);
 
 	/**
-	* Goes through all of the blueprint skeleton's object properties and pulls
-	* out the ones that are associated with an UActorComponent (and are visbile
-	* to the blueprint).
-	*/
+	 * Goes through all of the blueprint skeleton's object properties and pulls
+	 * out the ones that are associated with an UActorComponent (and are visbile
+	 * to the blueprint).
+	 */
 	static void GetComponentProperties(UBlueprint* Blueprint, TArray<UObjectProperty*>& PropertiesOut);
 	
+	/**
+	 * 
+	 */
+	static bool DumpActionDatabaseInfo(uint32 Indent, FArchive* FileOutWriter);
+
 	/**
 	 * Constructs a temporary blueprint (of the class type specified) and kicks 
 	 * off a dump of all its nested information (palette, graph, contextual 
@@ -365,6 +377,7 @@ DumpBlueprintInfoUtils::CommandletOptions::CommandletOptions(TArray<FString> con
 	, PaletteFilter(nullptr)
 	, GraphFilter(GT_MAX)
 	, SelectedObjectType(nullptr)
+	, InterfaceClass(nullptr)	
 {
 	uint32 NewDumpFlags = BPDUMP_UseLegacyMenuBuilder;
 	for (FString const& Switch : Switches)
@@ -535,7 +548,25 @@ DumpBlueprintInfoUtils::CommandletOptions::CommandletOptions(TArray<FString> con
 			FString NameSwitch;
 			Switch.Split(TEXT("="), &NameSwitch, &Filename);
 		}
-		else
+		else if (Switch.StartsWith("interface="))
+		{
+			FString InterfaceSwitch, InterfaceName;
+			Switch.Split(TEXT("="), &InterfaceSwitch, &InterfaceName);
+
+			if (UClass* Class = FindObject<UClass>(ANY_PACKAGE, *InterfaceName))
+			{
+				if (Class->IsChildOf<UInterface>())
+				{
+					InterfaceClass = Class;
+				}
+			}
+			
+			if (InterfaceClass == nullptr)
+			{
+				UE_LOG(LogBlueprintInfoDump, Warning, TEXT("Could not find a matching interface class matching this name: '%s'"), *InterfaceName);
+			}
+		}
+		else if (!Switch.StartsWith("run=")) // account for the first switch, which invoked this commandlet
 		{
 			UE_LOG(LogBlueprintInfoDump, Warning, TEXT("Unrecognized command switch '%s', use -help for a listing of all accepted params"), *Switch);
 		}
@@ -957,6 +988,89 @@ static void DumpBlueprintInfoUtils::GetComponentProperties(UBlueprint* Blueprint
 	}
 }
 
+static bool DumpBlueprintInfoUtils::DumpActionDatabaseInfo(uint32 Indent, FArchive* FileOutWriter)
+{
+	bool bWroteToFile = false;
+	if ((CommandOptions.DumpFlags & BPDUMP_UseLegacyMenuBuilder) == 0)
+	{
+		UE_LOG(LogBlueprintInfoDump, Display, TEXT("%sDumping Database info..."), *BuildIndentString(Indent, true));
+
+		double DatabaseBuildTime = 0.0;
+		{
+			FScopedDurationTimer DatabaseBuildTimer(DatabaseBuildTime);
+			// prime the database so we can record information from it
+			FBlueprintActionDatabase::Get();
+		}
+		FBlueprintActionDatabase& Database = FBlueprintActionDatabase::Get();
+		
+		TSet<UBlueprint*> TemplateOuters;
+		int32 DatabaseCount = 0;
+		double NodeCacheTimer = 0.0;
+		int32 TemplateCount = 0;
+
+		for (auto const& DbEntry : Database.GetAllActions())
+		{
+			for (UBlueprintNodeSpawner* BpAction : DbEntry.Value)
+			{
+				FDurationTimer NodeCacheTimer(NodeCacheTimer);
+				NodeCacheTimer.Start();
+				UEdGraphNode* TemplateNode = BpAction->GetTemplateNode();
+				NodeCacheTimer.Stop();
+
+				if (TemplateNode != nullptr)
+				{
+					NodeCacheTimer.Start();
+					TemplateNode->AllocateDefaultPins();
+					NodeCacheTimer.Stop();
+
+					UObject* TemplateOuter = TemplateNode->GetOuter();
+					while ((TemplateOuter != nullptr) && (Cast<UBlueprint>(TemplateOuter) == nullptr))
+					{
+						TemplateOuter = TemplateOuter->GetOuter();
+					}
+					UBlueprint* OuterBlueprint = CastChecked<UBlueprint>(TemplateOuter);
+					TemplateOuters.Add(OuterBlueprint);
+
+					++TemplateCount;
+				}
+				++DatabaseCount;
+			}
+		}
+
+		float const DatabaseKbSize = Database.Size() / 1024.f;
+		float EstimatedCacheKbSize = 0.0;
+		for (UBlueprint* CacheBlueprint : TemplateOuters)
+		{
+			TArray<UObject*> ChildObjs;
+			GetObjectsWithOuter(CacheBlueprint, ChildObjs);
+
+			EstimatedCacheKbSize += sizeof(*CacheBlueprint);
+			for (UObject* ChildObj : ChildObjs)
+			{
+				EstimatedCacheKbSize += sizeof(*ChildObj);
+			}
+		}
+		EstimatedCacheKbSize /= 1024.f;
+
+		FString const OriginalIndent  = BuildIndentString(Indent);
+		FString const IndentedNewline = "\n" + BuildIndentString(Indent + 1);
+
+		FString DatabaseInfoEntry = FString::Printf(TEXT("%s\"ActionDatabaseInfo\" : {"), *OriginalIndent);
+		DatabaseInfoEntry += FString::Printf(TEXT("%s\"DatabaseBuildTime\"      : \"%.3f seconds\","), *IndentedNewline, DatabaseBuildTime);
+		DatabaseInfoEntry += FString::Printf(TEXT("%s\"CachingNodesDuration\"   : \"%.3f seconds\","), *IndentedNewline, NodeCacheTimer);
+		DatabaseInfoEntry += FString::Printf(TEXT("%s\"DatabaseEntryCount\"     : %d,"), *IndentedNewline, DatabaseCount);
+		DatabaseInfoEntry += FString::Printf(TEXT("%s\"NodeCacheCount\"         : %d,"), *IndentedNewline, TemplateCount);
+		DatabaseInfoEntry += FString::Printf(TEXT("%s\"DatabaseSize\"           : \"%.2f KB\","), *IndentedNewline, DatabaseKbSize);
+		DatabaseInfoEntry += FString::Printf(TEXT("%s\"AvgSizePerDatabaseEntry\": \"%.2f bytes\","), *IndentedNewline, (DatabaseKbSize * 1024.f) / DatabaseCount);
+		DatabaseInfoEntry += FString::Printf(TEXT("%s\"EstimatedNodeCacheSize\" : \"%.2f KB\","), *IndentedNewline, EstimatedCacheKbSize);
+		DatabaseInfoEntry += FString::Printf(TEXT("%s\"AvgSizePerCachedNode\"   : \"%.2f bytes\","), *IndentedNewline, (EstimatedCacheKbSize * 1024.f)/ TemplateCount);
+		DatabaseInfoEntry += FString::Printf(TEXT("%s\"EstimatedSystemSize\"    : \"%.2f KB\"\n%s}"), *IndentedNewline, DatabaseKbSize + EstimatedCacheKbSize, *OriginalIndent);
+		FileOutWriter->Serialize(TCHAR_TO_ANSI(*DatabaseInfoEntry), DatabaseInfoEntry.Len());
+		bWroteToFile = true;
+	}
+	return bWroteToFile;
+}
+
 //------------------------------------------------------------------------------
 static void DumpBlueprintInfoUtils::DumpInfoForClass(uint32 Indent, UClass* BlueprintClass, FArchive* FileOutWriter)
 {
@@ -972,6 +1086,13 @@ static void DumpBlueprintInfoUtils::DumpInfoForClass(uint32 Indent, UClass* Blue
 	FileOutWriter->Serialize(TCHAR_TO_ANSI(*BeginClassEntry), BeginClassEntry.Len());
 
 	UBlueprint* TempBlueprint = MakeTempBlueprint(BlueprintClass);
+	if (CommandOptions.InterfaceClass != nullptr)
+	{
+		if (!FBlueprintEditorUtils::ImplementNewInterface(TempBlueprint, CommandOptions.InterfaceClass->GetFName()))
+		{
+			UE_LOG(LogBlueprintInfoDump, Warning, TEXT("Failed to add interface (%s), to class blueprint: '%s'"), *CommandOptions.InterfaceClass->GetName(), *ClassName);
+		}
+	}
 
 	bool bNeedsClosingComma = false;
 	if (CommandOptions.DumpFlags & BPDUMP_PaletteMask)
@@ -1809,13 +1930,6 @@ int32 UDumpBlueprintsInfoCommandlet::Main(FString const& Params)
 	bool const bSplitFilesByClass = (CommandOptions.DumpFlags & DumpBlueprintInfoUtils::BPDUMP_FilePerBlueprint) != 0;
 	bool const bDiffGeneratedFile = !CommandOptions.DiffPath.IsEmpty() || !CommandOptions.DiffCommand.IsEmpty();
 
-	bool const bUseLegacyMenu = (CommandOptions.DumpFlags & DumpBlueprintInfoUtils::BPDUMP_UseLegacyMenuBuilder) != 0;
-	if (!bUseLegacyMenu)
-	{
-		// prime the database so it's not recorded in our timing capture
-		FBlueprintActionDatabase::Get();
-	}
-
 	UClass* const SelectedObjType = DumpBlueprintInfoUtils::CommandOptions.SelectedObjectType;
 	// if the user specified that they want a level actor selected during the 
 	// dump, then spawn one and select it (extra blueprint context actions 
@@ -1843,12 +1957,12 @@ int32 UDumpBlueprintsInfoCommandlet::Main(FString const& Params)
 				check(!ActiveFilePath.IsEmpty());
 				DumpBlueprintInfoUtils::DiffDumpFiles(ActiveFilePath, CommandOptions.DiffPath, CommandOptions.DiffCommand);
 			}
-
 			delete FileOut;
 			FileOut = nullptr;
 		}
 	};
 
+	
 	// this lambda is responsible for opening a file for write, and adding 
 	// opening json characters to the file (contextually tracks the active filepath as well)
 	auto OpenFileStream = [&ActiveFilePath, &CloseFileStream, &FileOut](UClass* Class)->FArchive*
@@ -1860,6 +1974,10 @@ int32 UDumpBlueprintsInfoCommandlet::Main(FString const& Params)
 		check(FileOut != nullptr);
 		FileOut->Serialize(TCHAR_TO_ANSI(TEXT("{\n")), 2);
 
+		if (DumpBlueprintInfoUtils::DumpActionDatabaseInfo(1, FileOut))
+		{
+			FileOut->Serialize(TCHAR_TO_ANSI(TEXT(",\n")), 2);
+		}
 		return FileOut;
 	};
 	
