Resave packages commandlet now has the option to strip editor only content

diff --git a/Engine/Source/Editor/UnrealEd/Classes/Commandlets/ResavePackagesCommandlet.h b/Engine/Source/Editor/UnrealEd/Classes/Commandlets/ResavePackagesCommandlet.h
--- a/Engine/Source/Editor/UnrealEd/Classes/Commandlets/ResavePackagesCommandlet.h
+++ b/Engine/Source/Editor/UnrealEd/Classes/Commandlets/ResavePackagesCommandlet.h
@@ -49,6 +49,9 @@ class UResavePackagesCommandlet : public UCommandlet
 	// If non-empty, this substring has to be present in the package name for the commandlet to process it
 	FString PackageSubstring;
 
+	// strip editor only content
+	bool bStripEditorOnlyContent;
+
 	// skip the assert when a package can not be opened
 	bool bCanIgnoreFails;
 
diff --git a/Engine/Source/Editor/UnrealEd/Private/Commandlets/ContentCommandlets.cpp b/Engine/Source/Editor/UnrealEd/Private/Commandlets/ContentCommandlets.cpp
--- a/Engine/Source/Editor/UnrealEd/Private/Commandlets/ContentCommandlets.cpp
+++ b/Engine/Source/Editor/UnrealEd/Private/Commandlets/ContentCommandlets.cpp
@@ -325,6 +325,12 @@ void UResavePackagesCommandlet::LoadAndSaveOnePackage(const FString& Filename)
 				}
 			}
 
+			if (bStripEditorOnlyContent)
+			{
+				UE_LOG(LogContentCommandlet, Log, TEXT("Removing editor only data"));
+				Package->PackageFlags |= PKG_FilterEditorOnly;
+			}
+
 			// Now based on the computation above we will see if we should actually attempt
 			// to save this package
 			if (bSavePackage == true)
@@ -425,6 +431,8 @@ int32 UResavePackagesCommandlet::Main( const FString& Params )
 	// Ensure source control is initialized and shut down properly
 	FScopedSourceControl SourceControl;
 
+	// strip editor only content
+	bStripEditorOnlyContent = Switches.Contains(TEXT("STRIPEDITORONLY"));
 	// skip the assert when a package can not be opened
 	bCanIgnoreFails = Switches.Contains(TEXT("SKIPFAILS"));
 	/** load all packages, and display warnings for those packages which would have been resaved but were read-only */
