Fixing build error; accidently named two variables the same thing.

diff --git a/Engine/Source/Editor/UnrealEd/Private/Commandlets/DumpBlueprintsInfoCommandlet.cpp b/Engine/Source/Editor/UnrealEd/Private/Commandlets/DumpBlueprintsInfoCommandlet.cpp
--- a/Engine/Source/Editor/UnrealEd/Private/Commandlets/DumpBlueprintsInfoCommandlet.cpp
+++ b/Engine/Source/Editor/UnrealEd/Private/Commandlets/DumpBlueprintsInfoCommandlet.cpp
@@ -988,6 +988,7 @@ static void DumpBlueprintInfoUtils::GetComponentProperties(UBlueprint* Blueprint
 	}
 }
 
+//------------------------------------------------------------------------------
 static bool DumpBlueprintInfoUtils::DumpActionDatabaseInfo(uint32 Indent, FArchive* FileOutWriter)
 {
 	bool bWroteToFile = false;
@@ -1005,14 +1006,14 @@ static bool DumpBlueprintInfoUtils::DumpActionDatabaseInfo(uint32 Indent, FArchi
 		
 		TSet<UBlueprint*> TemplateOuters;
 		int32 DatabaseCount = 0;
-		double NodeCacheTimer = 0.0;
+		double NodeCachingDuration = 0.0;
 		int32 TemplateCount = 0;
 
 		for (auto const& DbEntry : Database.GetAllActions())
 		{
 			for (UBlueprintNodeSpawner* BpAction : DbEntry.Value)
 			{
-				FDurationTimer NodeCacheTimer(NodeCacheTimer);
+				FDurationTimer NodeCacheTimer(NodeCachingDuration);
 				NodeCacheTimer.Start();
 				UEdGraphNode* TemplateNode = BpAction->GetTemplateNode();
 				NodeCacheTimer.Stop();
@@ -1057,7 +1058,7 @@ static bool DumpBlueprintInfoUtils::DumpActionDatabaseInfo(uint32 Indent, FArchi
 
 		FString DatabaseInfoEntry = FString::Printf(TEXT("%s\"ActionDatabaseInfo\" : {"), *OriginalIndent);
 		DatabaseInfoEntry += FString::Printf(TEXT("%s\"DatabaseBuildTime\"      : \"%.3f seconds\","), *IndentedNewline, DatabaseBuildTime);
-		DatabaseInfoEntry += FString::Printf(TEXT("%s\"CachingNodesDuration\"   : \"%.3f seconds\","), *IndentedNewline, NodeCacheTimer);
+		DatabaseInfoEntry += FString::Printf(TEXT("%s\"CachingNodesDuration\"   : \"%.3f seconds\","), *IndentedNewline, NodeCachingDuration);
 		DatabaseInfoEntry += FString::Printf(TEXT("%s\"DatabaseEntryCount\"     : %d,"), *IndentedNewline, DatabaseCount);
 		DatabaseInfoEntry += FString::Printf(TEXT("%s\"NodeCacheCount\"         : %d,"), *IndentedNewline, TemplateCount);
 		DatabaseInfoEntry += FString::Printf(TEXT("%s\"DatabaseSize\"           : \"%.2f KB\","), *IndentedNewline, DatabaseKbSize);
@@ -1088,10 +1089,14 @@ static void DumpBlueprintInfoUtils::DumpInfoForClass(uint32 Indent, UClass* Blue
 	UBlueprint* TempBlueprint = MakeTempBlueprint(BlueprintClass);
 	if (CommandOptions.InterfaceClass != nullptr)
 	{
-		if (!FBlueprintEditorUtils::ImplementNewInterface(TempBlueprint, CommandOptions.InterfaceClass->GetFName()))
+		UClass* BpClass = (TempBlueprint->SkeletonGeneratedClass != nullptr) ? TempBlueprint->SkeletonGeneratedClass : TempBlueprint->ParentClass;
+		if (!BpClass->ImplementsInterface(CommandOptions.InterfaceClass))
 		{
-			UE_LOG(LogBlueprintInfoDump, Warning, TEXT("Failed to add interface (%s), to class blueprint: '%s'"), *CommandOptions.InterfaceClass->GetName(), *ClassName);
-		}
+			if (!FBlueprintEditorUtils::ImplementNewInterface(TempBlueprint, CommandOptions.InterfaceClass->GetFName()))
+			{
+				UE_LOG(LogBlueprintInfoDump, Warning, TEXT("Failed to add interface (%s), to class blueprint: '%s'"), *CommandOptions.InterfaceClass->GetName(), *ClassName);
+			}
+		}		
 	}
 
 	bool bNeedsClosingComma = false;
