Speed up to chunk manifest generator by removing calls to GetAssetsByPackageName()

diff --git a/Engine/Source/Editor/UnrealEd/Private/Commandlets/ChunkManifestGenerator.cpp b/Engine/Source/Editor/UnrealEd/Private/Commandlets/ChunkManifestGenerator.cpp
--- a/Engine/Source/Editor/UnrealEd/Private/Commandlets/ChunkManifestGenerator.cpp
+++ b/Engine/Source/Editor/UnrealEd/Private/Commandlets/ChunkManifestGenerator.cpp
@@ -211,7 +211,7 @@ void FChunkManifestGenerator::AddPackageToChunkManifest(UPackage* Package, const
 		}
 
 		// Collect all chunk IDs associated with this package from the asset registry
-		TArray<int32> RegistryChunkIDs = GetAssetRegistryChunkAssignments(PackageFName);
+		TArray<int32> RegistryChunkIDs = GetAssetRegistryChunkAssignments(Package);
 		ExistingChunkIDs = GetExistingPackageChunkAssignments(PackageFName);
 
 		// Try to call game-specific delegate to determine the target chunk ID
diff --git a/Engine/Source/Editor/UnrealEd/Private/Commandlets/ChunkManifestGenerator.h b/Engine/Source/Editor/UnrealEd/Private/Commandlets/ChunkManifestGenerator.h
--- a/Engine/Source/Editor/UnrealEd/Private/Commandlets/ChunkManifestGenerator.h
+++ b/Engine/Source/Editor/UnrealEd/Private/Commandlets/ChunkManifestGenerator.h
@@ -134,20 +134,25 @@ class FChunkManifestGenerator
 	}
 
 	/**
-	* Returns an array of chunks ID for a package name that have been assigned in the editor.
+	* Returns an array of chunks IDs for a package that have been assigned in the editor.
 	*/
-	FORCEINLINE TArray<int32> GetAssetRegistryChunkAssignments(FName PackageFName)
+	FORCEINLINE TArray<int32> GetAssetRegistryChunkAssignments(UPackage* Package)
 	{
 		TArray<int32> RegistryChunkIDs;
-		TArray<FAssetData> PkgAssets;
-		AssetRegistry.GetAssetsByPackageName(PackageFName, PkgAssets);
-		for (const auto& Asset : PkgAssets)
+		TArray<UObject*> ObjectsInPackage;
+		GetObjectsWithOuter(Package, ObjectsInPackage, false);
+		for (const auto* Object : ObjectsInPackage)
 		{
-			for (const auto& AssetChunk : Asset.ChunkIDs)
+			if (Object && Object->IsAsset())
 			{
-				RegistryChunkIDs.AddUnique(AssetChunk);
+				FAssetData Asset(Object);
+				for (auto AssetChunk : Asset.ChunkIDs)
+				{
+					RegistryChunkIDs.AddUnique(AssetChunk);
+				}
 			}
 		}
+
 		return RegistryChunkIDs;
 	}
 
