fix for TTP #3390990 "GNavigationBuildingNotification has ill-defined destruction order" #UE4

diff --git a/Engine/Source/Editor/UnrealEd/Private/NavigationBuildingNotification.cpp b/Engine/Source/Editor/UnrealEd/Private/NavigationBuildingNotification.cpp
--- a/Engine/Source/Editor/UnrealEd/Private/NavigationBuildingNotification.cpp
+++ b/Engine/Source/Editor/UnrealEd/Private/NavigationBuildingNotification.cpp
@@ -4,53 +4,7 @@
 #include "MainFrame.h"
 #include "Kismet2/DebuggerCommands.h"
 #include "EditorBuildUtils.h"
-
-/** Notification class for asynchronous shader compiling. */
-class FNavigationBuildingNotificationImpl
-	: public FTickableEditorObject
-{
-
-public:
-
-	FNavigationBuildingNotificationImpl()
-	{
-		bPreviouslyDetectedBuild = false;
-		LastEnableTime = 0;
-	}
-
-	/** Starts the notification. */
-	void BuildStarted();
-
-	/** Ends the notification. */
-	void BuildFinished();
-
-protected:
-	/** FTickableEditorObject interface */
-	virtual void Tick(float DeltaTime);
-	virtual bool IsTickable() const
-	{
-		return true;
-	}
-	virtual TStatId GetStatId() const override;
-
-private:
-	void ClearCompleteNotification();
-
-	bool bPreviouslyDetectedBuild;
-	double TimeOfStartedBuild;
-	double TimeOfStoppedBuild;
-
-	/** Tracks the last time the notification was started, used to avoid spamming. */
-	double LastEnableTime;
-
-	TWeakPtr<SNotificationItem> NavigationBuiltCompleteNotification;
-
-	/** The source code symbol query in progress message */
-	TWeakPtr<SNotificationItem> NavigationBuildNotificationPtr;
-};
-
-/** Instance. */
-FNavigationBuildingNotificationImpl GNavigationBuildingNotification;
+#include "NavigationBuildingNotification.h"
 
 void FNavigationBuildingNotificationImpl::BuildStarted()
 {
diff --git a/Engine/Source/Editor/UnrealEd/Private/NavigationBuildingNotification.h b/Engine/Source/Editor/UnrealEd/Private/NavigationBuildingNotification.h
new file mode 100644
--- /dev/null
+++ b/Engine/Source/Editor/UnrealEd/Private/NavigationBuildingNotification.h
@@ -0,0 +1,47 @@
+// Copyright 1998-2014 Epic Games, Inc. All Rights Reserved.
+#include "UnrealEd.h"
+
+/** Notification class for asynchronous shader compiling. */
+class FNavigationBuildingNotificationImpl
+	: public FTickableEditorObject
+{
+
+public:
+
+	FNavigationBuildingNotificationImpl()
+	{
+		bPreviouslyDetectedBuild = false;
+		LastEnableTime = 0;
+	}
+
+	/** Starts the notification. */
+	void BuildStarted();
+
+	/** Ends the notification. */
+	void BuildFinished();
+
+protected:
+	/** FTickableEditorObject interface */
+	virtual void Tick(float DeltaTime);
+	virtual bool IsTickable() const
+	{
+		return true;
+	}
+	virtual TStatId GetStatId() const override;
+
+private:
+	void ClearCompleteNotification();
+
+	bool bPreviouslyDetectedBuild;
+	double TimeOfStartedBuild;
+	double TimeOfStoppedBuild;
+
+	/** Tracks the last time the notification was started, used to avoid spamming. */
+	double LastEnableTime;
+
+	TWeakPtr<SNotificationItem> NavigationBuiltCompleteNotification;
+
+	/** The source code symbol query in progress message */
+	TWeakPtr<SNotificationItem> NavigationBuildNotificationPtr;
+};
+
diff --git a/Engine/Source/Editor/UnrealEd/Private/UnrealEdMisc.cpp b/Engine/Source/Editor/UnrealEd/Private/UnrealEdMisc.cpp
--- a/Engine/Source/Editor/UnrealEd/Private/UnrealEdMisc.cpp
+++ b/Engine/Source/Editor/UnrealEd/Private/UnrealEdMisc.cpp
@@ -33,6 +33,7 @@
 #include "Slate.h"
 #include "DesktopPlatformModule.h"
 #include "ShaderCompiler.h"
+#include "NavigationBuildingNotification.h"
 
 #define LOCTEXT_NAMESPACE "UnrealEd"
 
@@ -50,7 +51,8 @@ FUnrealEdMisc::FUnrealEdMisc() :
 	bInitialized( false ),
 	bSaveLayoutOnClose( true ),
 	bDeletePreferences( false ),
-	bIsAssetAnalyticsPending( false )
+	bIsAssetAnalyticsPending( false ),
+	NavigationBuildingNotificationHandler(NULL)
 {
 }
 
@@ -325,6 +327,8 @@ void FUnrealEdMisc::OnInit()
 	Delegate.BindRaw( this, &FUnrealEdMisc::EditorAnalyticsHeartbeat );
 	GEditor->GetTimerManager()->SetTimer( Delegate, Seconds, true );
 
+	// add handler to notify about navmesh building process
+	NavigationBuildingNotificationHandler = MakeShareable(new FNavigationBuildingNotificationImpl());
 }
 
 void FUnrealEdMisc::InitEngineAnalytics()
@@ -510,6 +514,11 @@ void FUnrealEdMisc::OnExit()
 		CancelPerformanceSurvey();
 	}
 
+	if (NavigationBuildingNotificationHandler.IsValid())
+	{
+		NavigationBuildingNotificationHandler = NULL;
+	}
+
 	// Report session maximum window and tab counts to engine analytics, if available
 	if (FEngineAnalytics::IsAvailable())
 	{
diff --git a/Engine/Source/Editor/UnrealEd/Public/UnrealEdMisc.h b/Engine/Source/Editor/UnrealEd/Public/UnrealEdMisc.h
--- a/Engine/Source/Editor/UnrealEd/Public/UnrealEdMisc.h
+++ b/Engine/Source/Editor/UnrealEd/Public/UnrealEdMisc.h
@@ -300,4 +300,7 @@ private:
 
 	/** An array of frame rate samples used by the performance survey */
 	TArray<float> FrameRateSamples;
+
+	/** handler to notify about navigation building process */
+	TSharedPtr<FTickableEditorObject> NavigationBuildingNotificationHandler;
 };
