Remove log spam when moving camera in sim mode with physics object selected

diff --git a/Engine/Source/Editor/UnrealEd/Private/PhysicsManipulationMode.cpp b/Engine/Source/Editor/UnrealEd/Private/PhysicsManipulationMode.cpp
--- a/Engine/Source/Editor/UnrealEd/Private/PhysicsManipulationMode.cpp
+++ b/Engine/Source/Editor/UnrealEd/Private/PhysicsManipulationMode.cpp
@@ -60,7 +60,7 @@ void FPhysicsManipulationEdMode::Exit()
 
 bool FPhysicsManipulationEdMode::InputDelta( FEditorViewportClient* InViewportClient,FViewport* InViewport,FVector& InDrag,FRotator& InRot,FVector& InScale )
 {
-	UE_LOG(LogEditorPhysMode, Warning, TEXT("Mouse: %s InDrag: %s  InRot: %s"), *GEditor->MouseMovement.ToString(), *InDrag.ToString(), *InRot.ToString());
+	//UE_LOG(LogEditorPhysMode, Warning, TEXT("Mouse: %s InDrag: %s  InRot: %s"), *GEditor->MouseMovement.ToString(), *InDrag.ToString(), *InRot.ToString());
 
 	const float GrabMoveSpeed = 1.0f;
 	const float GrabRotateSpeed = 1.0f;
@@ -83,7 +83,7 @@ bool FPhysicsManipulationEdMode::InputDelta( FEditorViewportClient* InViewportCl
 
 bool FPhysicsManipulationEdMode::StartTracking( FEditorViewportClient* InViewportClient, FViewport* InViewport )
 {
-	UE_LOG(LogEditorPhysMode, Warning, TEXT("Start Tracking"));
+	//UE_LOG(LogEditorPhysMode, Warning, TEXT("Start Tracking"));
 
 	FVector GrabLocation(0,0,0);
 	UPrimitiveComponent* ComponentToGrab = NULL;
@@ -121,7 +121,7 @@ bool FPhysicsManipulationEdMode::StartTracking( FEditorViewportClient* InViewpor
 
 bool FPhysicsManipulationEdMode::EndTracking( FEditorViewportClient* InViewportClient, FViewport* InViewport )
 {
-	UE_LOG(LogEditorPhysMode, Warning, TEXT("End Tracking"));
+	//UE_LOG(LogEditorPhysMode, Warning, TEXT("End Tracking"));
 
 	HandleComp->ReleaseComponent();
 
