EditorOpenAssetAutomationTests.cpp
--Added the missing blueprint test.

diff --git a/Engine/Source/Editor/UnrealEd/Private/Tests/EditorOpenAssetAutomationTests.cpp b/Engine/Source/Editor/UnrealEd/Private/Tests/EditorOpenAssetAutomationTests.cpp
--- a/Engine/Source/Editor/UnrealEd/Private/Tests/EditorOpenAssetAutomationTests.cpp
+++ b/Engine/Source/Editor/UnrealEd/Private/Tests/EditorOpenAssetAutomationTests.cpp
@@ -251,6 +251,25 @@ bool FOpenBlendSpace1DAssetEditors::RunTest(const FString& Parameters)
 }
 
 
+//////////////////////////////////////////////////////////////////////////
+/**
+* This test opens each BLUEPRINT into its sub-editor.
+*/
+IMPLEMENT_COMPLEX_AUTOMATION_TEST(FOpenBlueprintAssetEditors, "Editor.Content.Blueprint.Open Blueprint Assets", (EAutomationTestFlags::ATF_Editor | EAutomationTestFlags::ATF_RequiresUser));
+
+void FOpenBlueprintAssetEditors::GetTests(TArray<FString>& OutBeautifiedNames, TArray<FString>& OutTestCommands) const
+{
+	//This grabs each BlendSpace asset in the Game/Content directory
+	FEditorAutomationTestUtilities::CollectGameContentTestsByClass(UBlueprint::StaticClass(), true, OutBeautifiedNames, OutTestCommands);
+}
+
+bool FOpenBlueprintAssetEditors::RunTest(const FString& Parameters)
+{
+	bool bDidTheTestPass = FOpenActualAssetEditors(Parameters);
+	CollectGarbage(GARBAGE_COLLECTION_KEEPFLAGS);
+	return bDidTheTestPass;
+}
+
 //////////////////////////////////////////////////////////////////////////
 /**
 * This test opens each DESTRUCTIBLE MESH into its sub-editor.
