Improved performance for cook on the fly.

diff --git a/Engine/Source/Editor/UnrealEd/Classes/CookOnTheSide/CookOnTheFlyServer.h b/Engine/Source/Editor/UnrealEd/Classes/CookOnTheSide/CookOnTheFlyServer.h
--- a/Engine/Source/Editor/UnrealEd/Classes/CookOnTheSide/CookOnTheFlyServer.h
+++ b/Engine/Source/Editor/UnrealEd/Classes/CookOnTheSide/CookOnTheFlyServer.h
@@ -134,43 +134,147 @@ public:
 		 */
 	};
 
+	template<typename Type>
+	struct FLookupQueue
+	{
+	private:
+		TSet<Type> Set;
+		FQueue<Type> Queue;
+
+	public:
+		void Enqueue(const Type& Item)
+		{
+			Set.Add(Item);
+			Queue.Enqueue(Item);
+		}
+		void EnqueueUnique( const Type& Item )
+		{
+			if ( Set.Find(Item) == NULL )
+			{
+				Enqueue( Item );
+			}
+		}
+
+		bool Dequeue(Type* Result)
+		{
+			if (Queue.Num())
+			{
+				Queue.Dequeue(Result);
+				Set.Remove(*Result);
+				return true;
+			}
+			return false;
+		}
+		void DequeueAll(TArray<Type>& Results)
+		{
+			Queue.DequeueAll(Results);
+			Set.Empty();
+		}
+
+		bool HasItems() const
+		{
+			return Queue.Num() > 0;
+		}
+
+		void Remove( const Type& Item ) 
+		{
+			Queue.Remove( Item );
+			Set.Remove(Item);
+		}
+
+		int Num() const 
+		{
+			return Queue.Num();
+		}
+
+	};
+
 public:
 	/** cooked file requests which includes platform which file is requested for */
 	struct FFilePlatformRequest
 	{
+	private:
+		FName Filename;
+		FName Platformname;
+		uint32 CachedHash;
 	public:
-		FFilePlatformRequest() { }
-		FFilePlatformRequest( const FString &InFilename, const FName &InPlatformname ) : Filename( InFilename ), Platformname( InPlatformname ) 
+		FFilePlatformRequest() : CachedHash(0) { }
+
+		FFilePlatformRequest( const FName &InFilename, const FName &InPlatformname ) : Filename( InFilename ), Platformname( InPlatformname ) 
 		{
+			UpdateHash();
 		}
 
-		FString Filename;
-		FName Platformname;
+		FFilePlatformRequest( const FString &InFilename, const FName &InPlatformname ) : Filename( *InFilename ), Platformname( InPlatformname ) 
+		{
+			UpdateHash();
+		}
+
+		FFilePlatformRequest( const FFilePlatformRequest &InFilePlatformRequest ) : Filename( InFilePlatformRequest.Filename ), Platformname( InFilePlatformRequest.Platformname ), CachedHash( InFilePlatformRequest.CachedHash )
+		{
+#if DO_CHECK
+			uint32 OldHash = CachedHash;
+			UpdateHash();
+			check( OldHash == CachedHash );
+#endif
+		}
+
+		void SetFilename( const FString &InFilename ) 
+		{
+			Filename = FName(*InFilename);
+			UpdateHash();
+		}
+
+		const FName &GetFilename() const
+		{
+			return Filename;
+		}
+
+		const FName &GetPlatformname() const
+		{
+			return Platformname;
+		}
+
+		void UpdateHash()
+		{
+			CachedHash = GetTypeHash(Filename) ^ GetTypeHash(Platformname);
+		}
+
+		uint32 GetCachedHash() const
+		{
+			return CachedHash;
+		}
 
 		bool IsValid()  const
 		{
-			return Filename.Len() > 0;
+			return Filename != NAME_None;
 		}
 
 		void Clear()
 		{
 			Filename = TEXT("");
 			Platformname = TEXT("");
+			UpdateHash();
 		}
 
 		bool operator ==( const FFilePlatformRequest &InFileRequest ) const
 		{
-			if ( InFileRequest.Filename == Filename )
+			if ( CachedHash == InFileRequest.CachedHash )
 			{
-				if ( InFileRequest.Platformname == Platformname )
-					return true;
+				if ( InFileRequest.Filename == Filename )
+				{
+					if ( InFileRequest.Platformname == Platformname )
+					{
+						return true;
+					}
+				}
 			}
 			return false;
 		}
 
-		FORCEINLINE FString ToString() const
+		FORCEINLINE FString &&ToString() const
 		{
-			return FString::Printf( TEXT("%s %s"), *Filename, *Platformname.ToString() );
+			return MoveTemp( FString::Printf( TEXT("%s %s"), *Filename.ToString(), *Platformname.ToString() ) );
 		}
 
 	};
@@ -178,7 +282,7 @@ public:
 private:
 
 	FThreadSafeQueue<FFilePlatformRequest> FileRequests; // list of requested files
-	FQueue<FFilePlatformRequest> UnsolicitedFileRequests; // list of files which haven't been requested but we think should cook based on previous requests
+	FLookupQueue<FFilePlatformRequest> UnsolicitedFileRequests; // list of files which haven't been requested but we think should cook based on previous requests
 	FThreadSafeQueue<FFilePlatformRequest> UnsolicitedCookedPackages; // list of files which weren't requested but were cooked (based on UnsolicitedFileRequests)
 
 
@@ -189,6 +293,16 @@ private:
 		mutable FCriticalSection	SynchronizationObject;
 		TSet<FFilePlatformRequest> FilesProcessed;
 	public:
+
+		void Lock()
+		{
+			SynchronizationObject.Lock();
+		}
+		void Unlock()
+		{
+			SynchronizationObject.Unlock();
+		}
+
 		void Add(const FFilePlatformRequest &Filename)
 		{
 			FScopeLock ScopeLock(&SynchronizationObject);
@@ -200,14 +314,14 @@ private:
 			FScopeLock ScopeLock(&SynchronizationObject);
 			return FilesProcessed.Contains(Filename);
 		}
-		int RemoveAll( const FString &Filename )
+		int RemoveAll( const FName &Filename )
 		{
 			FScopeLock ScopeLock( &SynchronizationObject );
 			int NumRemoved = 0;
 			// for ( TSet<FFilePlatformRequest>::TIterator It( FilesProcessed ); It; ++It )
 			for ( auto It = FilesProcessed.CreateIterator(); It; ++It )
 			{
-				if ( It->Filename == Filename )
+				if ( It->GetFilename() == Filename )
 				{
 					It.RemoveCurrent();
 					++NumRemoved;
@@ -399,6 +513,7 @@ private:
 
 FORCEINLINE uint32 GetTypeHash(const UCookOnTheFlyServer::FFilePlatformRequest &Key)
 {
-	return GetTypeHash( Key.Filename ) ^ GetTypeHash( Key.Platformname );
+	//return GetTypeHash( Key.Filename ) ^ GetTypeHash( Key.Platformname );
+	return Key.GetCachedHash();
 }
 
diff --git a/Engine/Source/Editor/UnrealEd/Private/CookOnTheFlyServer.cpp b/Engine/Source/Editor/UnrealEd/Private/CookOnTheFlyServer.cpp
--- a/Engine/Source/Editor/UnrealEd/Private/CookOnTheFlyServer.cpp
+++ b/Engine/Source/Editor/UnrealEd/Private/CookOnTheFlyServer.cpp
@@ -27,8 +27,134 @@
 DEFINE_LOG_CATEGORY_STATIC(LogCookOnTheFly, Log, All);
 
 #define DEBUG_COOKONTHEFLY 0
+#define OUTPUT_TIMING 1
 
+#if OUTPUT_TIMING
 
+struct FTimerInfo
+{
+public:
+
+	FTimerInfo( FTimerInfo &&InTimerInfo )
+	{
+		Swap( Name, InTimerInfo.Name );
+		Length = InTimerInfo.Length;
+	}
+
+	FTimerInfo( FString &&InName, double InLength ) : Name(MoveTemp(InName)), Length(InLength) { }
+
+	FString Name;
+	double Length;
+};
+
+
+static TArray<FTimerInfo> GTimerInfo;
+
+struct FScopeTimer
+{
+private:
+	bool Started;
+	bool DecrementScope;
+	static int GScopeDepth;
+public:
+
+	FScopeTimer( const FScopeTimer &outer )
+	{
+		Index = outer.Index;
+		DecrementScope = false;
+		Started = false;	
+	}
+
+	FScopeTimer( const FString &InName, bool IncrementScope = false )
+	{
+		DecrementScope = IncrementScope;
+		if( DecrementScope)
+		{
+			++GScopeDepth;
+		}
+		FString Name = InName;
+		for ( int i =0; i < GScopeDepth; ++i )
+		{
+			Name = FString(TEXT("  ")) + Name;
+		}
+		Index = GTimerInfo.Emplace(MoveTemp(Name), 0.0);
+		Started = false;
+	}
+
+	void Start()
+	{
+		if ( !Started )
+		{
+			GTimerInfo[Index].Length -= FPlatformTime::Seconds();
+			Started = true;
+		}
+	}
+
+	void Stop()
+	{
+		if ( Started )
+		{
+			GTimerInfo[Index].Length += FPlatformTime::Seconds();
+			Started = false;
+		}
+	}
+
+	~FScopeTimer()
+	{
+		Stop();
+		if ( DecrementScope )
+		{
+			--GScopeDepth;
+		}
+	}
+
+	int Index;
+};
+
+int FScopeTimer::GScopeDepth = 0;
+
+
+
+void OutputTimers()
+{
+	UE_LOG( LogCookOnTheFly, Display, TEXT("Timing information for cook") );
+	UE_LOG( LogCookOnTheFly, Display, TEXT("Name\tlength(ms)") );
+	for ( auto TimerInfo : GTimerInfo )
+	{
+		UE_LOG( LogCookOnTheFly, Display, TEXT("%s\t%.2f"), *TimerInfo.Name, TimerInfo.Length * 1000.0f );
+	}
+
+	GTimerInfo.Empty();
+}
+
+
+#define CREATE_TIMER(name, incrementScope) FScopeTimer ScopeTimer##name(#name, incrementScope); 
+
+#define SCOPE_TIMER(name) CREATE_TIMER(name, true); ScopeTimer##name.Start();
+#define STOP_TIMER( name ) ScopeTimer##name.Stop();
+
+
+#define ACCUMULATE_TIMER(name) CREATE_TIMER(name, false);
+#define ACCUMULATE_TIMER_SCOPE(name) FScopeTimer ScopeTimerInner##name(ScopeTimer##name); ScopeTimerInner##name.Start();
+#define ACCUMULATE_TIMER_START(name) ScopeTimer##name.Start();
+#define ACCUMULATE_TIMER_STOP(name) ScopeTimer##name.Stop();
+
+#define OUTPUT_TIMERS() OutputTimers();
+
+#else
+#define CREATE_TIMER(name)
+
+#define SCOPE_TIMER(name)
+#define STOP_TIMER(name)
+
+#define ACCUMULATE_TIMER(name) 
+#define ACCUMULATE_TIMER_SCOPE(name) 
+#define ACCUMULATE_TIMER_START(name) 
+#define ACCUMULATE_TIMER_STOP(name) 
+
+#define OUTPUT_TIMERS()
+
+#endif
 ////////////////////////////////////////////////////////////////
 /// Cook on the fly server
 ///////////////////////////////////////////////////////////////
@@ -37,18 +163,88 @@ DEFINE_LOG_CATEGORY_STATIC(LogCookOnTheFly, Log, All);
 /* Static functions
  *****************************************************************************/
 
+struct FCachedPackageFilename
+{
+public:
+	FCachedPackageFilename(FString &&InPackageFilename, FString &&InStandardFilename, FName InStandardFileFName ) :
+		PackageFilename( MoveTemp( InPackageFilename )),
+		StandardFilename(MoveTemp(InStandardFilename)),
+		StandardFileFName( InStandardFileFName )
+	{
+	}
+
+	FCachedPackageFilename( const FCachedPackageFilename &In )
+	{
+		PackageFilename = In.PackageFilename;
+		StandardFilename = In.StandardFilename;
+		StandardFileFName = In.StandardFileFName;
+	}
+
+	FCachedPackageFilename( FCachedPackageFilename &&In )
+	{
+		PackageFilename = MoveTemp(In.PackageFilename);
+		StandardFilename = MoveTemp(In.StandardFilename);
+		StandardFileFName = In.StandardFileFName;
+	}
+
+	FString PackageFilename; // this is also full path
+	FString StandardFilename;
+	FName StandardFileFName;
+};
+
+static TMap<FString, FCachedPackageFilename> PackageFilenameCache;
+
+
+static const FCachedPackageFilename &Cache(UPackage *Package)
+{
+	FCachedPackageFilename *Cached = PackageFilenameCache.Find( Package->GetName() );
+	if ( Cached != NULL )
+	{
+		return *Cached;
+	}
+
+
+	FString Filename;
+	FString StandardFilename;
+	FName StandardFileFName = NAME_None;
+	if (FPackageName::DoesPackageExist(Package->GetName(), NULL, &Filename))
+	{
+		StandardFilename = FPaths::ConvertRelativePathToFull(Filename);
+
+		FPaths::MakeStandardFilename(StandardFilename);
+		StandardFileFName = FName(*StandardFilename);
+	}
+
+	return PackageFilenameCache.Emplace( Package->GetName(), MoveTemp(FCachedPackageFilename(MoveTemp(Filename),MoveTemp(StandardFilename), StandardFileFName)) );
+}
 
 
-static FString GetPackageFilename( UPackage* Package )
+/*static FString GetPackageFilename( UPackage* Package )
 {
 	FString Filename;
 	if (FPackageName::DoesPackageExist(Package->GetName(), NULL, &Filename))
 	{
 		Filename = FPaths::ConvertRelativePathToFull(Filename);
 	}
+
 	return Filename;
+}*/
+
+
+static FString GetCachedPackageFilename( UPackage *Package )
+{
+	return Cache( Package ).PackageFilename;
 }
 
+static FString GetCachedStandardPackageFilename( UPackage* Package )
+{
+	return Cache(Package).StandardFilename;
+}
+
+static FName GetCachedStandardPackageFileFName( UPackage *Package )
+{
+	return Cache(Package).StandardFileFName;
+}
 
 /* helper structs
  *****************************************************************************/
@@ -161,7 +357,6 @@ bool UCookOnTheFlyServer::BroadcastFileserverPresence( const FGuid &InstanceId )
 uint32 UCookOnTheFlyServer::TickCookOnTheSide( const float TimeSlice, uint32 &CookedPackageCount )
 {
 	double StartTime = FPlatformTime::Seconds();
-
 	uint32 Result = 0;
 
 	// This is all the target platforms which we needed to process requests for this iteration
@@ -170,6 +365,8 @@ uint32 UCookOnTheFlyServer::TickCookOnTheSide( const float TimeSlice, uint32 &Co
 
 	while (!GIsRequestingExit)
 	{
+		SCOPE_TIMER(TickCookOnTheSide);
+
 		bool bCookedPackage = false;
 		FFilePlatformRequest ToBuild;
 
@@ -191,165 +388,227 @@ uint32 UCookOnTheFlyServer::TickCookOnTheSide( const float TimeSlice, uint32 &Co
 
 		check( ToBuild.IsValid() );
 		TArray<FName> TargetPlatformNames;
-		
-		if( ToBuild.Platformname != NAME_None )
+
+		if( ToBuild.GetPlatformname() != NAME_None )
 		{
-			TargetPlatformNames.AddUnique( ToBuild.Platformname );
+			TargetPlatformNames.AddUnique( ToBuild.GetPlatformname() );
 		}
 
 		bool bWasUpToDate = false;
 
 		if (!ThreadSafeFilenameSet.Exists(ToBuild))
 		{
-
-			// if we have no target platforms then we want to cook because this will cook for all target platforms in that case
-			bool bShouldCook = TargetPlatformNames.Num() > 0 ? false : ShouldCook( ToBuild.Filename, NAME_None );
-			for ( int Index = 0; Index < TargetPlatformNames.Num(); ++Index )
 			{
-				bShouldCook |= ShouldCook( ToBuild.Filename, TargetPlatformNames[Index] );
-			}
+				const FString BuildFilename = ToBuild.GetFilename().ToString();
 
-			
-			if ( bShouldCook ) // if we should cook the package then cook it otherwise add it to the list of already cooked packages below
-			{
-				FString PackageName;	
-				UPackage *Package = NULL;
-				if ( FPackageName::TryConvertFilenameToLongPackageName(ToBuild.Filename, PackageName) )
+				// if we have no target platforms then we want to cook because this will cook for all target platforms in that case
+				bool bShouldCook = TargetPlatformNames.Num() > 0 ? false : ShouldCook( BuildFilename, NAME_None );
+				for ( int Index = 0; Index < TargetPlatformNames.Num(); ++Index )
 				{
-					Package = FindObject<UPackage>( ANY_PACKAGE, *PackageName );
+					SCOPE_TIMER(ShouldCook);
+					bShouldCook |= ShouldCook( ToBuild.GetFilename().ToString(), TargetPlatformNames[Index] );
 				}
 
-				UE_LOG( LogCookOnTheFly, Display, TEXT("Processing request %s"), *ToBuild.Filename)
 
-				//  if the package is already loaded then try to avoid reloading it :)
-				if ( ( Package == NULL ) || ( Package->IsFullyLoaded() == false ) )
+				if ( bShouldCook ) // if we should cook the package then cook it otherwise add it to the list of already cooked packages below
 				{
-					Package = LoadPackage( NULL, *ToBuild.Filename, LOAD_None );
-				}
-				else
-				{
-					UE_LOG(LogCookOnTheFly, Display, TEXT("Package already loaded %s avoiding reload"), *ToBuild.Filename);
-				}
+					FString PackageName;	
+					UPackage *Package = NULL;
+					if ( FPackageName::TryConvertFilenameToLongPackageName(BuildFilename, PackageName) )
+					{
+						Package = FindObject<UPackage>( ANY_PACKAGE, *PackageName );
+					}
 
-				if( Package == NULL )
-				{
-					UE_LOG(LogCookOnTheFly, Error, TEXT("Error loading %s!"), *ToBuild.Filename );
-					Result |= COSR_ErrorLoadingPackage;
-				}
-				else
-				{
-					if( SaveCookedPackage(Package, SAVE_KeepGUID | SAVE_Async, bWasUpToDate, TargetPlatformNames) )
+					UE_LOG( LogCookOnTheFly, Display, TEXT("Processing request %s"), *BuildFilename)
+
+					//  if the package is already loaded then try to avoid reloading it :)
+					if ( ( Package == NULL ) || ( Package->IsFullyLoaded() == false ) )
+					{
+						SCOPE_TIMER(LoadPackage);
+						Package = LoadPackage( NULL, *BuildFilename, LOAD_None );
+					}
+					else
 					{
-						FString Name = Package->GetPathName();
-						FString PackageFilename(GetPackageFilename(Package));
-						FPaths::MakeStandardFilename(PackageFilename);
+						UE_LOG(LogCookOnTheFly, Display, TEXT("Package already loaded %s avoiding reload"), *BuildFilename );
+					}
 
-						if ( PackageFilename != ToBuild.Filename )
+					if( Package == NULL )
+					{
+						UE_LOG(LogCookOnTheFly, Error, TEXT("Error loading %s!"), *BuildFilename );
+						Result |= COSR_ErrorLoadingPackage;
+					}
+					else
+					{
+						SCOPE_TIMER(SaveCookedPackage);
+						if( SaveCookedPackage(Package, SAVE_KeepGUID | SAVE_Async, bWasUpToDate, TargetPlatformNames) )
 						{
-							// we have saved something which we didn't mean to load 
-							//  sounds unpossible.... but it is due to searching for files and such
-							//  mark the original request as processed (if this isn't actually the file they were requesting then it will fail)
-							//	and then also save our new request as processed so we don't do it again
-							ThreadSafeFilenameSet.Add( ToBuild );
-							UE_LOG( LogCookOnTheFly, Display, TEXT("Request for %s received saved %s"), *ToBuild.Filename, *PackageFilename );
-							ToBuild.Filename = PackageFilename;
+							FString Name = Package->GetPathName();
+							/*FString PackageFilename(GetPackageFilename(Package));
+							FPaths::MakeStandardFilename(PackageFilename);*/
+							FString PackageFilename( GetCachedStandardPackageFilename( Package ) );
+							if ( PackageFilename != BuildFilename )
+							{
+								// we have saved something which we didn't mean to load 
+								//  sounds unpossible.... but it is due to searching for files and such
+								//  mark the original request as processed (if this isn't actually the file they were requesting then it will fail)
+								//	and then also save our new request as processed so we don't do it again
+								ThreadSafeFilenameSet.Add( ToBuild );
+								UE_LOG( LogCookOnTheFly, Display, TEXT("Request for %s received saved %s"), *BuildFilename, *PackageFilename );
+								ToBuild.SetFilename( PackageFilename );
+							}
+							// ToBuild.Filename = PackageFilename;
+
+							// Update flags used to determine garbage collection.
+							if (Package->ContainsMap())
+							{
+								Result |= COSR_CookedMap;
+								//bCookedAMapSinceLastGC = true;
+							}
+							else
+							{
+								++CookedPackageCount;
+								Result |= COSR_CookedPackage;
+								//NonMapPackageCountSinceLastGC++;
+								// LastCookActionTime = FPlatformTime::Seconds();
+							}
+							bCookedPackage = true;
 						}
-						// ToBuild.Filename = PackageFilename;
 
-						// Update flags used to determine garbage collection.
-						if (Package->ContainsMap())
-						{
-							Result |= COSR_CookedMap;
-							//bCookedAMapSinceLastGC = true;
-						}
-						else
+						//@todo ResetLoaders outside of this (ie when Package is NULL) causes problems w/ default materials
+						if (Package->HasAnyFlags(RF_RootSet) == false)
 						{
-							++CookedPackageCount;
-							Result |= COSR_CookedPackage;
-							//NonMapPackageCountSinceLastGC++;
-							// LastCookActionTime = FPlatformTime::Seconds();
+							SCOPE_TIMER(ResetLoaders);
+							ResetLoaders(Package);
 						}
-						bCookedPackage = true;
-					}
-
-					//@todo ResetLoaders outside of this (ie when Package is NULL) causes problems w/ default materials
-					if (Package->HasAnyFlags(RF_RootSet) == false)
-					{
-						ResetLoaders(Package);
 					}
 				}
 			}
 
+
 			for ( int Index = 0; Index < TargetPlatformNames.Num(); ++Index )
 			{
 				AllTargetPlatformNames.AddUnique(TargetPlatformNames[Index]);
 
-				FFilePlatformRequest FileRequest( ToBuild.Filename, TargetPlatformNames[Index]);
+				FFilePlatformRequest FileRequest( ToBuild.GetFilename(), TargetPlatformNames[Index]);
 
 				if ( bIsUnsolicitedRequest )
 				{
-					if ((FPaths::FileExists(FileRequest.Filename) == true) &&
+					if ((FPaths::FileExists(FileRequest.GetFilename().ToString()) == true) &&
 						(bWasUpToDate == false))
 					{
 						UnsolicitedCookedPackages.EnqueueUnique( FileRequest );
-						UE_LOG(LogCookOnTheFly, Display, TEXT("UnsolicitedCookedPackages: %s"), *FileRequest.Filename);
+						UE_LOG(LogCookOnTheFly, Display, TEXT("UnsolicitedCookedPackages: %s"), *FileRequest.GetFilename().ToString());
 					}
 				}
-
 				ThreadSafeFilenameSet.Add( FileRequest );
-
-
 			}
 		}
 		else
 		{
-			UE_LOG(LogCookOnTheFly, Display, TEXT("Package for platform already cooked %s, discarding request"), *ToBuild.Filename);
+			UE_LOG(LogCookOnTheFly, Display, TEXT("Package for platform already cooked %s, discarding request"), *ToBuild.GetFilename().ToString());
 		}
 
 
 		if ( bCookedPackage && !bIsUnsolicitedRequest)
 		{
+			SCOPE_TIMER(UnsolicitedMarkup);
+
 			TArray<UObject *> ObjectsInOuter;
-			GetObjectsWithOuter(NULL, ObjectsInOuter, false);
-			for( int32 Index = 0; Index < ObjectsInOuter.Num() && !GIsRequestingExit; Index++ )
 			{
-				UPackage* Package = Cast<UPackage>(ObjectsInOuter[Index]);
+				SCOPE_TIMER(GetObjectsWithOuter);
+				GetObjectsWithOuter(NULL, ObjectsInOuter, false);
+			}
+
+			TArray<FName> PackageNames;
+			PackageNames.Empty(ObjectsInOuter.Num());
+			{
+				SCOPE_TIMER(GeneratePackageNames);
+				ACCUMULATE_TIMER(GetPathName);
+				ACCUMULATE_TIMER(GetPackageFilename);
+				ACCUMULATE_TIMER(PackageCast);
+				ACCUMULATE_TIMER(ConvertRelativePathToFull);
+				ACCUMULATE_TIMER(PackageNamesAdd);
+				ACCUMULATE_TIMER(PackageFNameConvert);
+				for( int32 Index = 0; Index < ObjectsInOuter.Num() && !GIsRequestingExit; Index++ )
+				{
+					ACCUMULATE_TIMER_START(PackageCast);
+					UPackage* Package = Cast<UPackage>(ObjectsInOuter[Index]);
+					ACCUMULATE_TIMER_STOP(PackageCast);
 
-				if (Package)
+				/*for ( TObjectIterator<UPackage> It; It; ++It ) // ~700ms
 				{
-					FString Name = Package->GetPathName();
-					FString PackageFilename(GetPackageFilename(Package));
-					FPaths::MakeStandardFilename(PackageFilename);
+					UPackage *Package = *It;*/
 
-					const FString FullPath = FPaths::ConvertRelativePathToFull(PackageFilename);
 
+					if (Package)
+					{
+						ACCUMULATE_TIMER_START(GetPackageFilename);
+						FName PackageFilename(GetCachedStandardPackageFileFName(Package));
+						ACCUMULATE_TIMER_STOP(GetPackageFilename);
+		
+						if ( PackageFilename != NAME_None )
+						{
+							ACCUMULATE_TIMER_SCOPE(PackageNamesAdd);
+							PackageNames.Add(PackageFilename);
+						}
+						
+					}
+				}
+
+			}
+
+			{
+				SCOPE_TIMER(EnqueueUnsolicitedPackages);
+				{
+					SCOPE_TIMER(ThreadSafeFileNameSetLock);
+					ThreadSafeFilenameSet.Lock(); // manually lock and unlock here so that we don't thrash the 
+				}
+				
+
+				ACCUMULATE_TIMER(CreateRequest);
+				ACCUMULATE_TIMER(ThreadSafeFilenameSetExists);
+				ACCUMULATE_TIMER(ReenqueueUnsolicitedRequests);
+				for ( const FName &PackageFilename : PackageNames )
+				{
 					bool AlreadyCooked = true;
 					for ( const FName TargetPlatformName : AllTargetPlatformNames )
 					{
-						if ( !ThreadSafeFilenameSet.Exists(FFilePlatformRequest(PackageFilename, TargetPlatformName)) )
+						ACCUMULATE_TIMER_START(CreateRequest);
+						auto Request = FFilePlatformRequest(PackageFilename, TargetPlatformName);
+						ACCUMULATE_TIMER_STOP(CreateRequest);
+
+						ACCUMULATE_TIMER_SCOPE(ThreadSafeFilenameSetExists);
+						if ( !ThreadSafeFilenameSet.Exists(Request) )
 						{
-							AlreadyCooked = false;
-							break;
+							ACCUMULATE_TIMER_SCOPE(ReenqueueUnsolicitedRequests);
+							UnsolicitedFileRequests.EnqueueUnique(Request);
 						}
 					}
-					if ( !AlreadyCooked )
+					/*if ( !AlreadyCooked )
 					{
+						ACCUMULATE_TIMER_SCOPE(ReenqueueUnsolicitedRequests);
 						for (const FName TargetPlatformName : AllTargetPlatformNames)
 						{
 							const FFilePlatformRequest Request( PackageFilename, TargetPlatformName );
 							UnsolicitedFileRequests.EnqueueUnique(Request);
 						}
-					}
+					}*/
 				}
+				ThreadSafeFilenameSet.Unlock();
 			}
 		}
 
+
+		STOP_TIMER(TickCookOnTheSide);
+		// kill the timer 
+		OUTPUT_TIMERS();
+
 		if ( (FPlatformTime::Seconds() - StartTime) > TimeSlice)
 		{
 			break;
 		}
-	}
 
+	}
 	return Result;
 }
 
@@ -375,10 +634,11 @@ void UCookOnTheFlyServer::MarkPackageDirtyForCooker( UPackage *Package )
 {
 	// force that package to be recooked
 	const FString Name = Package->GetPathName();
-	FString PackageFilename(GetPackageFilename(Package));
-	FPaths::MakeStandardFilename(PackageFilename);
+	/*FString PackageFilename(GetPackageFilename(Package));
+	FPaths::MakeStandardFilename(PackageFilename);*/
+	FString PackageFilename = GetCachedStandardPackageFilename(Package);
 	UE_LOG(LogCookOnTheFly, Display, TEXT("Modification detected to package %s"), *PackageFilename);
-	ThreadSafeFilenameSet.RemoveAll( PackageFilename );
+	ThreadSafeFilenameSet.RemoveAll( FName(*PackageFilename) );
 
 	// TODO: go through package dependencies of this package, if a found dependency is not loaded then mark it as need recooking.
 }
@@ -518,7 +778,7 @@ bool UCookOnTheFlyServer::SaveCookedPackage( UPackage* Package, uint32 SaveFlags
 
 	bool bSavedCorrectly = true;
 
-	FString Filename(GetPackageFilename(Package));
+	FString Filename(GetCachedPackageFilename(Package));
 
 	if (Filename.Len())
 	{
@@ -916,9 +1176,9 @@ void UCookOnTheFlyServer::HandleNetworkFileServerFileRequest( const FString& Fil
 		UnsolicitedCookedPackages.DequeueAll(UnsolicitedRequests);
 		for (const FFilePlatformRequest &Request : UnsolicitedRequests)
 		{
-			if ( Request.Platformname == PlatformFname )
+			if ( Request.GetPlatformname() == PlatformFname )
 			{
-				FString StandardFilename = Request.Filename;
+				FString StandardFilename = Request.GetFilename().ToString();
 				FPaths::MakeStandardFilename( StandardFilename );
 				UnsolicitedFiles.Add( StandardFilename );
 			}
@@ -941,20 +1201,21 @@ void UCookOnTheFlyServer::HandleNetworkFileServerFileRequest( const FString& Fil
 	}
 	while (!ThreadSafeFilenameSet.Exists(FileRequest));
 
-	UE_LOG( LogCookOnTheFly, Display, TEXT("Cook complete %s"), *FileRequest.Filename)
+	UE_LOG( LogCookOnTheFly, Display, TEXT("Cook complete %s"), *FileRequest.GetFilename().ToString())
 
 	TArray<FFilePlatformRequest> UnsolicitedRequests;
 	UnsolicitedCookedPackages.DequeueAll(UnsolicitedRequests);
 	for (const FFilePlatformRequest &Request : UnsolicitedRequests)
 	{
-		if ( Request.Filename == FileRequest.Filename )
+		if ( Request.GetFilename() == FileRequest.GetFilename() )
 		{
 			// don't do anything we don't want this guy
 		}
-		else if ( Request.Platformname == PlatformFname )
+		else if ( Request.GetPlatformname() == PlatformFname )
 		{
-			UE_LOG( LogCookOnTheFly, Display, TEXT("Returning unsolicited file %s"), *Request.Filename );
-			UnsolicitedFiles.Add( Request.Filename );
+			const FString RequestFilename = Request.GetFilename().ToString();
+			UE_LOG( LogCookOnTheFly, Display, TEXT("Returning unsolicited file %s"), *RequestFilename );
+			UnsolicitedFiles.Add( RequestFilename );
 		}
 		else
 		{
