Fixed warning and made more descriptive warning.

diff --git a/Engine/Source/Editor/UnrealEd/Private/SkeletalMeshEdit.cpp b/Engine/Source/Editor/UnrealEd/Private/SkeletalMeshEdit.cpp
--- a/Engine/Source/Editor/UnrealEd/Private/SkeletalMeshEdit.cpp
+++ b/Engine/Source/Editor/UnrealEd/Private/SkeletalMeshEdit.cpp
@@ -44,7 +44,7 @@ UAnimSequence * UEditorEngine::ImportFbxAnimation( USkeleton * Skeleton, UObject
 
 		if (!SkeletonRoot)
 		{
-			FFbxImporter->AddTokenizedErrorMessage(FTokenizedMessage::Create(EMessageSeverity::Error, FText::Format(LOCTEXT("Error_CouldNotFindFbxTrack", "Could not find needed track ({0}).\nImport failed."), FText::FromName(Skeleton->GetReferenceSkeleton().GetBoneName(0)))), FFbxErrors::Animation_CouldNotFindRootTrack);
+			FFbxImporter->AddTokenizedErrorMessage(FTokenizedMessage::Create(EMessageSeverity::Error, FText::Format(LOCTEXT("Error_CouldNotFindFbxTrack", "Mesh contains {0} bone as root but animation doesn't contain the root track.\nImport failed."), FText::FromName(Skeleton->GetReferenceSkeleton().GetBoneName(0)))), FFbxErrors::Animation_CouldNotFindRootTrack);
 
 			FFbxImporter->ReleaseScene();
 			GWarn->EndSlowTask();
@@ -141,7 +141,7 @@ bool UEditorEngine::ReimportFbxAnimation( USkeleton * Skeleton, UAnimSequence *
 
 		if (!SkeletonRoot)
 		{
-			FbxImporter->AddTokenizedErrorMessage(FTokenizedMessage::Create(EMessageSeverity::Error, FText::Format(LOCTEXT("Error_CouldNotFindFbxTrack", "Could not find needed track ({0}).\nImport failed."), FText::FromName(Skeleton->GetReferenceSkeleton().GetBoneName(0)))), FFbxErrors::Animation_CouldNotFindTrack);
+			FbxImporter->AddTokenizedErrorMessage(FTokenizedMessage::Create(EMessageSeverity::Error, FText::Format(LOCTEXT("Error_CouldNotFindFbxTrack", "Mesh contains {0} bone as root but animation doesn't contain the root track.\nImport failed."), FText::FromName(Skeleton->GetReferenceSkeleton().GetBoneName(0)))), FFbxErrors::Animation_CouldNotFindTrack);
 
 			FbxImporter->ReleaseScene();
 			GWarn->EndSlowTask();
