Disabled timing information for cook.

diff --git a/Engine/Source/Editor/UnrealEd/Private/CookOnTheFlyServer.cpp b/Engine/Source/Editor/UnrealEd/Private/CookOnTheFlyServer.cpp
--- a/Engine/Source/Editor/UnrealEd/Private/CookOnTheFlyServer.cpp
+++ b/Engine/Source/Editor/UnrealEd/Private/CookOnTheFlyServer.cpp
@@ -27,7 +27,7 @@
 DEFINE_LOG_CATEGORY_STATIC(LogCookOnTheFly, Log, All);
 
 #define DEBUG_COOKONTHEFLY 0
-#define OUTPUT_TIMING 1
+#define OUTPUT_TIMING 0
 
 #if OUTPUT_TIMING
 
