Fixed mac compilation error, missing explicit copy constructor.

diff --git a/Engine/Source/Editor/UnrealEd/Private/CookOnTheFlyServer.cpp b/Engine/Source/Editor/UnrealEd/Private/CookOnTheFlyServer.cpp
--- a/Engine/Source/Editor/UnrealEd/Private/CookOnTheFlyServer.cpp
+++ b/Engine/Source/Editor/UnrealEd/Private/CookOnTheFlyServer.cpp
@@ -41,6 +41,12 @@ public:
 		Length = InTimerInfo.Length;
 	}
 
+	FTimerInfo( const FTimerInfo &InTimerInfo )
+	{
+		Name = InTimerInfo.Name;
+		Length = InTimerInfo.Length;
+	}
+
 	FTimerInfo( FString &&InName, double InLength ) : Name(MoveTemp(InName)), Length(InLength) { }
 
 	FString Name;
