CurveEditor - Various enhancements.
+ Change the key and tangent handle icons to be smaller and cleaner.
+ Change the time and value entry boxes to be a fixed width.
+ Refactor the mouse handling code to support marquee select.

diff --git a/Engine/Source/Editor/EditorStyle/Private/SlateEditorStyle.cpp b/Engine/Source/Editor/EditorStyle/Private/SlateEditorStyle.cpp
--- a/Engine/Source/Editor/EditorStyle/Private/SlateEditorStyle.cpp
+++ b/Engine/Source/Editor/EditorStyle/Private/SlateEditorStyle.cpp
@@ -781,10 +781,11 @@ void FSlateEditorStyle::FStyle::SetupGeneralStyles()
 		Set( "CurveEd.TimelineArea", new IMAGE_BRUSH( "Old/White", Icon16x16, FLinearColor(1,1,1,0.25f) ) );
 		Set( "CurveEd.FitHorizontal", new IMAGE_BRUSH( "Icons/FitHorz_16x", Icon16x16 ) );
 		Set( "CurveEd.FitVertical", new IMAGE_BRUSH( "Icons/FitVert_16x", Icon16x16 ) );
-		Set( "CurveEd.CurveKey", new IMAGE_BRUSH( "Icons/CurveKey_8x", Icon8x8 ) );
-		Set( "CurveEd.CurveKeySelected", new IMAGE_BRUSH( "Icons/CurveKey_8x", Icon8x8, SelectionColor ) );
+		Set( "CurveEd.CurveKey", new IMAGE_BRUSH( "Common/Key", FVector2D(11.0f, 11.0f) ) );
+		Set( "CurveEd.CurveKeySelected", new IMAGE_BRUSH( "Common/Key", FVector2D(11.0f, 11.0f), SelectionColor ) );
 		Set( "CurveEd.InfoFont", TTF_CORE_FONT( "Fonts/Roboto-Regular", 8 ) );
-		Set( "CurveEd.Tangent", new IMAGE_BRUSH( "Icons/Tangente_8x", Icon8x8 ) );
+		Set( "CurveEd.Tangent", new IMAGE_BRUSH( "Common/Key", FVector2D(11.0f, 11.0f), FLinearColor(0.0f, 0.66f, 0.7f) ) );
+		Set( "CurveEd.TangentColor", FLinearColor(0.0f, 0.66f, 0.7f) );
 	}
 	
 	// Scrub control buttons
@@ -1398,7 +1399,7 @@ void FSlateEditorStyle::FStyle::SetupGeneralStyles()
 		Set( "Sequencer.Timeline.ScrubHandleWhole", new BOX_BRUSH( "Sequencer/ScrubHandleWhole", FMargin( 6.f/13.f, 10/24.f, 6/13.f, 10/24.f  ) ) );
 		Set( "Sequencer.Timeline.NotifyAlignmentMarker", new IMAGE_BRUSH( "Sequencer/NotifyAlignmentMarker", FVector2D(10,19) ) );
 		Set( "Sequencer.NotificationImage_AddedPlayMovieSceneEvent", new IMAGE_BRUSH( "Old/Checkbox_checked", Icon16x16 ) );
-		Set( "Sequencer.Key", new IMAGE_BRUSH( "Sequencer/Key", FVector2D(7.0f,12.0f) ) );
+		Set( "Sequencer.Key", new IMAGE_BRUSH( "Common/Key", FVector2D(7.0f,12.0f) ) );
 		Set( "Sequencer.GenericSection.Background", new BOX_BRUSH( "Sequencer/GenericSectionBackground", FMargin(5.f/16.f) ) );
 		Set( "Sequencer.SectionGripLeft", new BOX_BRUSH( "Sequencer/SectionGripLeft", FMargin(5.f/16.f) ) );
 		Set( "Sequencer.SectionGripRight", new BOX_BRUSH( "Sequencer/SectionGripRight", FMargin(5.f/16.f) ) );
diff --git a/Engine/Source/Editor/UnrealEd/Private/SCurveEditor.cpp b/Engine/Source/Editor/UnrealEd/Private/SCurveEditor.cpp
--- a/Engine/Source/Editor/UnrealEd/Private/SCurveEditor.cpp
+++ b/Engine/Source/Editor/UnrealEd/Private/SCurveEditor.cpp
@@ -9,7 +9,7 @@
 
 #define LOCTEXT_NAMESPACE "SCurveEditor"
 
-const static FVector2D	CONST_KeySize		= FVector2D(10,10);
+const static FVector2D	CONST_KeySize		= FVector2D(11,11);
 const static FVector2D	CONST_CurveSize		= FVector2D(12,12);
 
 const static float		CONST_FitMargin		= 0.05f;
@@ -43,8 +43,6 @@ void SCurveEditor::Construct(const FArguments& InArgs)
 
 	bIsUsingSlider = false;
 
-	ResetDragValues();
-
 	// if editor size is set, use it, otherwise, use default value
 	if (DesiredSize.Get().IsZero())
 	{
@@ -61,9 +59,8 @@ void SCurveEditor::Construct(const FArguments& InArgs)
 
 	OnCreateAsset = InArgs._OnCreateAsset;
 
-	bMovingKeys = false;
-	bPanning = false;
-	bMovingTangent = false;
+	DragState = EDragState::None;
+	DragThreshold = 4;
 
 	//Simple r/g/b for now
 	CurveColors.Add(FLinearColor(1.0f, 0.0f, 0.0f));
@@ -153,7 +150,7 @@ void SCurveEditor::Construct(const FArguments& InArgs)
 							.ColorAndOpacity( FSlateColor::UseForeground() ) 
 						]
 					]
-					
+
 					+ SHorizontalBox::Slot()
 					.AutoWidth()
 					[
@@ -177,6 +174,8 @@ void SCurveEditor::Construct(const FArguments& InArgs)
 							.MaxValue(TOptional<float>())
 							.MaxSliderValue(TOptional<float>())
 							.MinSliderValue(TOptional<float>())
+							.Delta(0.001f)
+							.MinDesiredValueWidth(60.0f)
 							.Label()
 							[
 								SNew(STextBlock)
@@ -207,6 +206,8 @@ void SCurveEditor::Construct(const FArguments& InArgs)
 							.MaxValue(TOptional<float>())
 							.MaxSliderValue(TOptional<float>())
 							.MinSliderValue(TOptional<float>())
+							.Delta(.001f)
+							.MinDesiredValueWidth(60.0f)
 							.Label()
 							[
 								SNew(STextBlock)
@@ -471,6 +472,11 @@ int32 SCurveEditor::OnPaint(const FPaintArgs& Args, const FGeometry& AllottedGeo
 		}
 	}
 
+	if (DragState == EDragState::MarqueeSelect)
+	{
+		PaintMarquee(AllottedGeometry, MyClippingRect, OutDrawElements, LayerId);
+	}
+
 	// Paint children
 	SCompoundWidget::OnPaint(Args, CurveAreaGeometry, MyClippingRect, OutDrawElements, LayerId, InWidgetStyle, bParentEnabled);
 
@@ -648,6 +654,7 @@ void SCurveEditor::PaintTangent( FTrackScaleInfo &ScaleInfo, FRichCurve* Curve,
 	LeaveTangentPos  -= OffsetKeySize;
 
 	const FSlateBrush* TangentBrush = FEditorStyle::GetBrush("CurveEd.Tangent");
+	const FLinearColor TangentColor = FEditorStyle::GetColor("CurveEd.TangentColor");
 
 	FVector2D KeyOffset(CONST_KeySize.X*0.5f, CONST_KeySize.Y*0.5f);
 
@@ -662,7 +669,7 @@ void SCurveEditor::PaintTangent( FTrackScaleInfo &ScaleInfo, FRichCurve* Curve,
 		LinePoints,
 		MyClippingRect,
 		DrawEffects,
-		FLinearColor(0.0f,0.66f,0.7f)
+		TangentColor
 		);
 
 	LinePoints.Empty();
@@ -675,7 +682,7 @@ void SCurveEditor::PaintTangent( FTrackScaleInfo &ScaleInfo, FRichCurve* Curve,
 		LinePoints,
 		MyClippingRect,
 		DrawEffects,
-		FLinearColor(0.0f,0.66f,0.7f)
+		TangentColor
 		);
 
 	//Arrive tangent control
@@ -810,6 +817,27 @@ void SCurveEditor::PaintGridLines(const FGeometry &AllottedGeometry, FTrackScale
 	}
 }
 
+void SCurveEditor::PaintMarquee(const FGeometry& AllottedGeometry, const FSlateRect& MyClippingRect, FSlateWindowElementList& OutDrawElements, int32 LayerId) const
+{
+	FVector2D MarqueTopLeft(
+		FMath::Min(MouseDownLocation.X, MouseMoveLocation.X),
+		FMath::Min(MouseDownLocation.Y, MouseMoveLocation.Y)
+		);
+
+	FVector2D MarqueBottomRight(
+		FMath::Max(MouseDownLocation.X, MouseMoveLocation.X),
+		FMath::Max(MouseDownLocation.Y, MouseMoveLocation.Y)
+		);
+
+	FSlateDrawElement::MakeBox(
+		OutDrawElements,
+		LayerId,
+		AllottedGeometry.ToPaintGeometry(MarqueTopLeft, MarqueBottomRight - MarqueTopLeft),
+		FEditorStyle::GetBrush(TEXT("MarqueeSelection")),
+		MyClippingRect
+		);
+}
+
 void SCurveEditor::SetCurveOwner(FCurveOwnerInterface* InCurveOwner, bool bCanEdit)
 {
 	if(InCurveOwner != CurveOwner)
@@ -895,100 +923,15 @@ void SCurveEditor::DeleteSelectedKeys()
 	CurveOwner->OnCurveChanged();
 }
 
-
-FReply SCurveEditor::OnMouseButtonDown( const FGeometry& InMyGeometry, const FPointerEvent& InMouseEvent )
+FReply SCurveEditor::OnMouseButtonDown( const FGeometry& InMyGeometry, const FPointerEvent& InMouseEvent)
 {
 	const bool bLeftMouseButton = InMouseEvent.GetEffectingButton() == EKeys::LeftMouseButton;
 	const bool bRightMouseButton = InMouseEvent.GetEffectingButton() == EKeys::RightMouseButton;
-	const bool bControlDown = InMouseEvent.IsControlDown();
-	const bool bShiftDown = InMouseEvent.IsShiftDown();
 
-	// See if we hit a key
-	FSelectedCurveKey HitKey = HitTestKeys(InMyGeometry, InMouseEvent.GetScreenSpacePosition() );
-
-	//Turn off error msg
-	//WarningMessageText->SetError(FString());
-
-	if( bLeftMouseButton || bRightMouseButton )
+	DragState = EDragState::PreDrag;
+	if (bLeftMouseButton || bRightMouseButton)
 	{
-		bMovingTangent = false;
-		// Hit a key
-		if( IsEditingEnabled() && HitKey.IsValid())
-		{
-			if(bLeftMouseButton)
-			{
-				if(!IsKeySelected(HitKey))
-				{
-					// If control is down, add to current selection
-					if(!bControlDown)
-					{
-						EmptySelection();
-					}
-
-					AddToSelection(HitKey);
-				}
-				else if(bControlDown)
-				{
-					RemoveFromSelection(HitKey);
-				}
-
-				// If th key we clicked on is now selected, start movement mode
-				if(IsKeySelected(HitKey))
-				{
-					BeginDragTransaction();
-					bMovingKeys = true;
-				}
-			}
-			else
-			{
-				// Make sure key is selected in readyness for context menu
-				if(!IsKeySelected(HitKey))
-				{
-					EmptySelection();
-					AddToSelection(HitKey);
-				}
-				
-				PushInterpolationMenu(InMouseEvent.GetScreenSpacePosition());
-			}
-		}
-		// Hit background
-		else
-		{
-			if( bLeftMouseButton && IsEditingEnabled() )
-			{
-				// shift click on background = add new key
-				if(bShiftDown)
-				{
-					AddNewKey(InMyGeometry, InMouseEvent.GetScreenSpacePosition());
-				}
-				else
-				{
-					// clicking on background clears selection
-					FSelectedTangent Tangent = HitTestCubicTangents(InMyGeometry, InMouseEvent.GetScreenSpacePosition());
-					if(Tangent.IsValid())
-					{
-						BeginDragTransaction();
-						SelectedTangent = Tangent;
-						bMovingTangent = true;
-					}
-					else
-					{
-						// clicking on background clears key selection
-						EmptySelection();
-
-						//select any curve that is nearby
-						if(Curves.Num() > 1)
-						{
-							HitTestCurves(InMyGeometry, InMouseEvent);
-						}
-					}
-				}
-			}
-			else if(bRightMouseButton)
-			{
-				bPanning = true;
-			}
-		}
+		MouseDownLocation = InMyGeometry.AbsoluteToLocal(InMouseEvent.GetScreenSpacePosition());
 
 		// Set keyboard focus to this so that selected text box doesn't try to apply to newly selected keys
 		if(!HasKeyboardFocus())
@@ -1041,60 +984,31 @@ void SCurveEditor::AddNewKey(FGeometry InMyGeometry, FVector2D ScreenPosition)
 	CurveOwner->OnCurveChanged();
 }
 
-void SCurveEditor::ResetDragValues()
-{
-	TotalDragDist = 0.f;
-	bPanning = false;
-	bMovingKeys = false;
-	bMovingTangent = false;
-}
-
 void SCurveEditor::OnMouseCaptureLost()
 {
 	// if we began a drag transaction we need to finish it to make sure undo doesn't get out of sync
-	if ( IsEditingEnabled() && ( bMovingTangent || bMovingKeys ) )
+	if ( IsEditingEnabled() && ( DragState == EDragState::DragKey || DragState == EDragState::DragTangent ) )
 	{
 		EndDragTransaction();
-
-		ResetDragValues();
 	}
+	DragState = EDragState::None;
 }
 
 FReply SCurveEditor::OnMouseButtonUp( const FGeometry& InMyGeometry, const FPointerEvent& InMouseEvent )
 {
-	const bool bLeftMouseButton = InMouseEvent.GetEffectingButton() == EKeys::LeftMouseButton;
-	const bool bRightMouseButton = InMouseEvent.GetEffectingButton() == EKeys::RightMouseButton;
-
-	if( this->HasMouseCapture() && (bLeftMouseButton || bRightMouseButton) )
+	if (this->HasMouseCapture())
 	{
-		if( IsEditingEnabled() )
+		if (DragState == EDragState::PreDrag)
 		{
-			EndDragTransaction();
+			// If the user didn't start dragging, handle the mouse operation as a click.
+			ProcessClick(InMyGeometry, InMouseEvent);
 		}
-
-		//Allow context menu to appear if drag distance is less than threshold distance.
-		const float DragThresholdDist = 5.0f;
-		if( bRightMouseButton && TotalDragDist < DragThresholdDist )
+		else
 		{
-			FVector2D ScreenPos = InMouseEvent.GetScreenSpacePosition();
-
-			if( IsEditingEnabled() && SelectedKeys.Num() == 0 && HitTestCurves(InMyGeometry, InMouseEvent))
-			{
-				PushInterpolationMenu(InMouseEvent.GetScreenSpacePosition());
-			}
-			else if(
-				!HitTestCubicTangents(InMyGeometry, ScreenPos).IsValid() && 
-				!HitTestKeys(InMyGeometry, ScreenPos).IsValid())
-			{
-				CreateContextMenu(InMyGeometry, ScreenPos);
-			}
+			EndDrag(InMyGeometry, InMouseEvent);
 		}
-
-		ResetDragValues();
-
 		return FReply::Handled().ReleaseMouseCapture();
 	}
-
 	return FReply::Unhandled();
 }
 
@@ -1120,49 +1034,19 @@ void ClampViewRangeToDataIfBound( float & NewViewMin, float & NewViewMax, const
 FReply SCurveEditor::OnMouseMove( const FGeometry& InMyGeometry, const FPointerEvent& InMouseEvent )
 {
 	FRichCurve* Curve = GetCurve(0);
-	if( Curve != NULL)
+	if( Curve != NULL && this->HasMouseCapture())
 	{
-		// When mouse moves, if we are moving a key, update its 'input' position
-		if (bPanning || bMovingKeys || bMovingTangent)
+		if (DragState == EDragState::PreDrag)
 		{
-			FTrackScaleInfo ScaleInfo(ViewMinInput.Get(),  ViewMaxInput.Get(), ViewMinOutput, ViewMaxOutput, InMyGeometry.Size);
-			FVector2D ScreenDelta = InMouseEvent.GetCursorDelta();
-			TotalDragDist += ScreenDelta.Size(); // accumulate distance moved
-
-			FVector2D InputDelta;
-			InputDelta.X = ScreenDelta.X/ScaleInfo.PixelsPerInput;
-			InputDelta.Y = -ScreenDelta.Y/ScaleInfo.PixelsPerOutput;
-
-			if( IsEditingEnabled() &&  bMovingKeys)
-			{
-				MoveSelectedKeysByDelta(InputDelta);
-			}
-			else if( IsEditingEnabled() &&  bMovingTangent)
-			{
-				FVector2D MousePositionScreen = InMouseEvent.GetScreenSpacePosition();
-				MousePositionScreen -= InMyGeometry.AbsolutePosition;
-				FVector2D MousePositionCurve(ScaleInfo.LocalXToInput(MousePositionScreen.X), ScaleInfo.LocalYToOutput(MousePositionScreen.Y));
-				OnMoveTangent(MousePositionCurve);
-			}
-			else if(bPanning)
-			{
-				// we don't filter Output
-				ViewMinOutput = (ViewMinOutput-InputDelta.Y);
-				ViewMaxOutput = (ViewMaxOutput-InputDelta.Y);
-
-				// we clamped to Data in/output if set
-				float NewMinInput = ViewMinInput.Get() - InputDelta.X;
-				float NewMaxInput = ViewMaxInput.Get() - InputDelta.X;
-
-				// clamp to data input if set
-				ClampViewRangeToDataIfBound(NewMinInput, NewMaxInput, DataMinInput, DataMaxInput, ScaleInfo.ViewInputRange);
-				SetInputMinMax(NewMinInput, NewMaxInput);
-			}
-
-			return FReply::Handled();
+			TryStartDrag(InMyGeometry, InMouseEvent);
+		}
+		else if (DragState != EDragState::None)
+		{
+			ProcessDrag(InMyGeometry, InMouseEvent);
 		}
+		MouseMoveLocation = InMyGeometry.AbsoluteToLocal(InMouseEvent.GetScreenSpacePosition());
+		return FReply::Handled();
 	}
-
 	return FReply::Unhandled();
 }
 
@@ -1209,6 +1093,218 @@ FReply SCurveEditor::OnKeyDown( const FGeometry& MyGeometry, const FKeyboardEven
 	}
 }
 
+void SCurveEditor::TryStartDrag(const FGeometry& InMyGeometry, const FPointerEvent& InMouseEvent)
+{
+	const bool bLeftMouseButton = InMouseEvent.IsMouseButtonDown(EKeys::LeftMouseButton);
+	const bool bRightMouseButton = InMouseEvent.IsMouseButtonDown(EKeys::RightMouseButton);
+	const bool bControlDown = InMouseEvent.IsControlDown();
+
+	FVector2D DragVector = InMyGeometry.AbsoluteToLocal(InMouseEvent.GetScreenSpacePosition()) - MouseDownLocation;
+	if (DragVector.Size() >= DragThreshold)
+	{
+		if (IsEditingEnabled() && bLeftMouseButton)
+		{
+			// Check if we should start dragging keys.
+			FSelectedCurveKey HitKey = HitTestKeys(InMyGeometry, InMyGeometry.LocalToAbsolute(MouseDownLocation));
+			if (HitKey.IsValid())
+			{
+				if (IsKeySelected(HitKey) == false)
+				{
+					if (!bControlDown)
+					{
+						EmptySelection();
+					}
+					AddToSelection(HitKey);
+				}
+
+				BeginDragTransaction();
+				DragState = EDragState::DragKey;
+			}
+			else
+			{
+				// Check if we should start dragging a tangent.
+				FSelectedTangent Tangent = HitTestCubicTangents(InMyGeometry, InMyGeometry.LocalToAbsolute(MouseDownLocation));
+				if (Tangent.IsValid())
+				{
+					BeginDragTransaction();
+					SelectedTangent = Tangent;
+					DragState = EDragState::DragTangent;
+				}
+				else
+				{
+					// Otherwise if the user left clicked on nothing and start a marquee select.
+					DragState = EDragState::MarqueeSelect;
+				}
+			}
+		}
+		else if (bRightMouseButton)
+		{
+			DragState = EDragState::Pan;
+		}
+		else
+		{
+			DragState = EDragState::None;
+		}
+	}
+}
+
+void SCurveEditor::ProcessDrag(const FGeometry& InMyGeometry, const FPointerEvent& InMouseEvent)
+{
+	FTrackScaleInfo ScaleInfo(ViewMinInput.Get(), ViewMaxInput.Get(), ViewMinOutput, ViewMaxOutput, InMyGeometry.Size);
+	FVector2D ScreenDelta = InMouseEvent.GetCursorDelta();
+
+	FVector2D InputDelta;
+	InputDelta.X = ScreenDelta.X / ScaleInfo.PixelsPerInput;
+	InputDelta.Y = -ScreenDelta.Y / ScaleInfo.PixelsPerOutput;
+
+	if (DragState == EDragState::DragKey)
+	{
+		MoveSelectedKeysByDelta(InputDelta);
+	}
+	else if (DragState == EDragState::DragTangent)
+	{
+		FVector2D MousePositionScreen = InMouseEvent.GetScreenSpacePosition();
+		MousePositionScreen -= InMyGeometry.AbsolutePosition;
+		FVector2D MousePositionCurve(ScaleInfo.LocalXToInput(MousePositionScreen.X), ScaleInfo.LocalYToOutput(MousePositionScreen.Y));
+		OnMoveTangent(MousePositionCurve);
+	}
+	else if (DragState == EDragState::Pan)
+	{
+		// Output is not clamped.
+		ViewMinOutput = (ViewMinOutput - InputDelta.Y);
+		ViewMaxOutput = (ViewMaxOutput - InputDelta.Y);
+
+		// Input maybe clamped if DataMinInput or DataMaxOutput was set.
+		float NewMinInput = ViewMinInput.Get() - InputDelta.X;
+		float NewMaxInput = ViewMaxInput.Get() - InputDelta.X;
+		ClampViewRangeToDataIfBound(NewMinInput, NewMaxInput, DataMinInput, DataMaxInput, ScaleInfo.ViewInputRange);
+
+		SetInputMinMax(NewMinInput, NewMaxInput);
+	}
+}
+
+void SCurveEditor::EndDrag(const FGeometry& InMyGeometry, const FPointerEvent& InMouseEvent)
+{
+	const bool bControlDown = InMouseEvent.IsControlDown();
+
+	if (DragState == EDragState::DragKey || DragState == EDragState::DragTangent)
+	{
+		if (IsEditingEnabled())
+		{
+			EndDragTransaction();
+		}
+	}
+	else if (DragState == EDragState::MarqueeSelect)
+	{
+		FVector2D MarqueTopLeft(
+			FMath::Min(MouseDownLocation.X, MouseMoveLocation.X),
+			FMath::Min(MouseDownLocation.Y, MouseMoveLocation.Y)
+			);
+
+		FVector2D MarqueBottomRight(
+			FMath::Max(MouseDownLocation.X, MouseMoveLocation.X),
+			FMath::Max(MouseDownLocation.Y, MouseMoveLocation.Y)
+			);
+
+		TArray<FSelectedCurveKey> SelectedCurveKeys = GetKeysWithinMarquee(InMyGeometry, MarqueTopLeft, MarqueBottomRight);
+
+		if (!bControlDown)
+		{
+			EmptySelection();
+		}
+
+		for (auto SelectedCurveKey : SelectedCurveKeys)
+		{
+			if (IsKeySelected(SelectedCurveKey))
+			{
+				RemoveFromSelection(SelectedCurveKey);
+			}
+			else
+			{
+				AddToSelection(SelectedCurveKey);
+			}
+		}
+	}
+	DragState = EDragState::None;
+}
+
+void SCurveEditor::ProcessClick(const FGeometry& InMyGeometry, const FPointerEvent& InMouseEvent)
+{
+	const bool bLeftMouseButton = InMouseEvent.GetEffectingButton() == EKeys::LeftMouseButton;
+	const bool bRightMouseButton = InMouseEvent.GetEffectingButton() == EKeys::RightMouseButton;
+	const bool bControlDown = InMouseEvent.IsControlDown();
+	const bool bShiftDown = InMouseEvent.IsShiftDown();
+
+	FSelectedCurveKey HitKey = HitTestKeys(InMyGeometry, InMouseEvent.GetScreenSpacePosition());
+	if (IsEditingEnabled())
+	{
+		if (bLeftMouseButton)
+		{
+			// If the user left clicked a key, update selection based on modifier key state.
+			if (HitKey.IsValid())
+			{
+				if (!IsKeySelected(HitKey))
+				{
+					if (!bControlDown)
+					{
+						EmptySelection();
+					}
+					AddToSelection(HitKey);
+				}
+				else if (bControlDown)
+				{
+					RemoveFromSelection(HitKey);
+				}
+			}
+			else
+			{
+				// If the user didn't click a key, add a new one if shift is held down, or try to select a curve.
+				if (bShiftDown)
+				{
+					AddNewKey(InMyGeometry, InMouseEvent.GetScreenSpacePosition());
+				}
+				else
+				{
+					// clicking on background clears key selection
+					EmptySelection();
+
+					//select any curve that is nearby
+					if (Curves.Num() > 1)
+					{
+						HitTestCurves(InMyGeometry, InMouseEvent);
+					}
+				}
+			}
+		}
+		else if (bRightMouseButton)
+		{
+			// If the user right clicked, handle opening context menus.
+			if (HitKey.IsValid())
+			{
+				// Make sure key is selected in readiness for context menu
+				if (!IsKeySelected(HitKey))
+				{
+					EmptySelection();
+					AddToSelection(HitKey);
+				}
+				PushInterpolationMenu(InMouseEvent.GetScreenSpacePosition());
+			}
+			else if (HitTestCurves(InMyGeometry, InMouseEvent))
+			{
+				PushInterpolationMenu(InMouseEvent.GetScreenSpacePosition());
+			}
+			else
+			{
+				FVector2D ScreenPos = InMouseEvent.GetScreenSpacePosition();
+				if (!HitTestCubicTangents(InMyGeometry, ScreenPos).IsValid())
+				{
+					CreateContextMenu(InMyGeometry, ScreenPos);
+				}
+			}
+		}
+	}
+}
+
 TOptional<float> SCurveEditor::OnGetTime() const
 {
 	if ( SelectedKeys.Num() == 1 )
@@ -2040,6 +2136,37 @@ void SCurveEditor::GetTangentPoints(  FTrackScaleInfo &ScaleInfo, const FSelecte
 	Leave = KeyScreenPosition + ToLeave*CONST_KeyTangentOffset;
 }
 
+TArray<SCurveEditor::FSelectedCurveKey> SCurveEditor::GetKeysWithinMarquee(const FGeometry& InMyGeometry, FVector2D MarqueeTopLeft, FVector2D MarqueeBottomRight) const
+{
+	TArray<FSelectedCurveKey> KeysWithinMarquee;
+	if (AreCurvesVisible())
+	{
+		FTrackScaleInfo ScaleInfo(ViewMinInput.Get(), ViewMaxInput.Get(), ViewMinOutput, ViewMaxOutput, InMyGeometry.Size);
+		for (auto CurveIt = Curves.CreateConstIterator(); CurveIt; ++CurveIt)
+		{
+			FRichCurve* Curve = CurveIt->CurveToEdit;
+			if (Curve != NULL)
+			{
+				for (auto It(Curve->GetKeyHandleIterator()); It; ++It)
+				{
+					float KeyScreenX = ScaleInfo.InputToLocalX(Curve->GetKeyTime(It.Key()));
+					float KeyScreenY = ScaleInfo.OutputToLocalY(Curve->GetKeyValue(It.Key()));
+
+					if (KeyScreenX >= (MarqueeTopLeft.X - (0.5f * CONST_KeySize.X)) &&
+						KeyScreenX <= (MarqueeBottomRight.X + (0.5f * CONST_KeySize.X)) &&
+						KeyScreenY >= (MarqueeTopLeft.Y - (0.5f * CONST_KeySize.Y)) &&
+						KeyScreenY <= (MarqueeBottomRight.Y + (0.5f * CONST_KeySize.Y)))
+					{
+						KeysWithinMarquee.Add(FSelectedCurveKey(Curve, It.Key()));
+					}
+				}
+			}
+		}
+	}
+
+	return KeysWithinMarquee;
+}
+
 void SCurveEditor::BeginDragTransaction()
 {
 	TransactionIndex = GEditor->BeginTransaction( LOCTEXT("CurveEditor_Drag", "Mouse Drag") );
diff --git a/Engine/Source/Editor/UnrealEd/Public/SCurveEditor.h b/Engine/Source/Editor/UnrealEd/Public/SCurveEditor.h
--- a/Engine/Source/Editor/UnrealEd/Public/SCurveEditor.h
+++ b/Engine/Source/Editor/UnrealEd/Public/SCurveEditor.h
@@ -7,7 +7,9 @@
 //////////////////////////////////////////////////////////////////////////
 // FTrackScaleInfo
 
-/** Utility struct for converting between curve space and local/absolute screen space */
+/** Utility struct for converting between curve space and local/absolute screen space.  The input domain is traditionally
+ *  the time axis of the curve, and the output domain is traditionally the value axis.
+ */
 struct FTrackScaleInfo
 {
 	float ViewMinInput;
@@ -79,6 +81,25 @@ class SCurveEditor :
 	public FGCObject,
 	public FEditorUndoClient
 {
+private:
+
+	/** Represents the different states of a drag operation */
+	enum class EDragState
+	{
+		/** The user has clicked a mouse button, but hasn't moved more then the drag threshold. */
+		PreDrag,
+		/** The user is dragging the selected keys. */
+		DragKey,
+		/** The user is dragging a selected tangent handle. */
+		DragTangent,
+		/** The user is performing a marquee selection of keys. */
+		MarqueeSelect,
+		/** The user is panning the curve view. */
+		Pan,
+		/** There is no active drag operation. */
+		None
+	};
+
 public:
 	SLATE_BEGIN_ARGS( SCurveEditor )
 		: _ViewMinInput(0.0f)
@@ -233,6 +254,9 @@ private:
 	/** Get screen space tangent positions for a given key */
 	void GetTangentPoints( FTrackScaleInfo &ScaleInfo,const FSelectedCurveKey &Key, FVector2D& Arrive, FVector2D& Leave) const;
 
+	/** Get the set of keys within a rectangle in local space */
+	TArray<FSelectedCurveKey> GetKeysWithinMarquee(const FGeometry& InMyGeometry, FVector2D MarqueeTopLeft, FVector2D MarqueeBottomRight) const;
+
 	/** Empy key selecttion set */
 	void EmptySelection();
 	/** Add a key to the selection set */
@@ -315,6 +339,9 @@ private:
 	/** Paint Grid lines, these make it easier to visualize relative distance */
 	void PaintGridLines( const FGeometry &AllottedGeometry, FTrackScaleInfo &ScaleInfo, FSlateWindowElementList &OutDrawElements, int32 LayerId, const FSlateRect& MyClippingRect, ESlateDrawEffect::Type DrawEffects )const;
 
+	/** Paints the marquee for selection */
+	void PaintMarquee(const FGeometry& AllottedGeometry, const FSlateRect& MyClippingRect, FSlateWindowElementList& OutDrawElements, int32 LayerId) const;
+
 	// SWidget interface
 	UNREALED_API virtual FReply OnMouseButtonDown( const FGeometry& MyGeometry, const FPointerEvent& MouseEvent ) override;
 	UNREALED_API virtual FReply OnMouseButtonUp( const FGeometry& MyGeometry, const FPointerEvent& MouseEvent ) override;
@@ -323,6 +350,18 @@ private:
 	UNREALED_API virtual FReply OnKeyDown( const FGeometry& MyGeometry, const FKeyboardEvent& InKeyboardEvent ) override;
 	UNREALED_API virtual void   OnMouseCaptureLost() override;
 
+	/** Attempts to start a drag operation when the mouse moves. */
+	void TryStartDrag(const FGeometry& MyGeometry, const FPointerEvent& MouseEvent);
+
+	/** Processes an ongoing drag operation when the mouse moves. */
+	void ProcessDrag(const FGeometry& MyGeometry, const FPointerEvent& MouseEvent);
+
+	/** Completes an ongoing drag operation on mouse up. */
+	void EndDrag(const FGeometry& MyGeometry, const FPointerEvent& MouseEvent);
+
+	/** Handles a mouse click operation on mouse up. */
+	void ProcessClick(const FGeometry& MyGeometry, const FPointerEvent& MouseEvent);
+
 	virtual bool SupportsKeyboardFocus() const override
 	{
 		return true;
@@ -372,9 +411,6 @@ private:
 	/** End a transaction for dragging a key or tangent */
 	void	EndDragTransaction();
 
-	/** Resets values to indicate dragging is not occurring */
-	void	ResetDragValues();
-
 	/** Calculate the distance between grid lines: determines next lowest power of 2, works with fractional numbers */
 	static float CalcGridLineStepDistancePow2(double RawValue);
 
@@ -438,16 +474,6 @@ private:
 	/** Currently selected tangent */
 	FSelectedTangent	SelectedTangent;
 
-	/** If we are currently moving the selected key set */
-	bool				bMovingKeys;
-	/** If we are currently panning the timeline */
-	bool				bPanning;
-	/** If we are currently moving a tangent */
-	bool				bMovingTangent;
-
-	/** Total distance we moved the mouse while button down */
-	float				TotalDragDist;
-
 	/** Minimum input of data range  */
 	TAttribute< TOptional<float> >	DataMinInput;
 	/** Maximum input of data range  */
@@ -505,6 +531,18 @@ protected:
 
 	/** True if the internal zoom buttons should be visible. */
 	bool bShowZoomButtons;
+
+	/** The location of mouse during the last OnMouseButtonDown callback in widget local coordinates. */
+	FVector2D MouseDownLocation;
+
+	/** The location of the mouse during the last OnMouseMove callback in widget local coordinates. */
+	FVector2D MouseMoveLocation;
+
+	/** The state of the current drag operation happening with the widget, if any. */
+	EDragState DragState;
+
+	/** The number of pixels which the mouse must move before a drag operation starts. */
+	float DragThreshold;
 };
 
 #endif // __SCurveEditor_h__
