UMG/Slate - Fixing a few areas where DPI calculations were based on the wrong size.

diff --git a/Engine/Source/Editor/UnrealEd/Classes/Editor/EditorEngine.h b/Engine/Source/Editor/UnrealEd/Classes/Editor/EditorEngine.h
--- a/Engine/Source/Editor/UnrealEd/Classes/Editor/EditorEngine.h
+++ b/Engine/Source/Editor/UnrealEd/Classes/Editor/EditorEngine.h
@@ -2334,7 +2334,7 @@ public:
 
 private:
 	/** Gets the DPI Scale for the game viewport in the editor. */
-	float GetGameViewportDPIScale() const;
+	float GetGameViewportDPIScale(UGameViewportClient* ViewportClient) const;
 
 	/**
 	 * Non Online PIE creation flow, creates all instances of PIE at once when online isn't requested/required
diff --git a/Engine/Source/Editor/UnrealEd/Private/PlayLevel.cpp b/Engine/Source/Editor/UnrealEd/Private/PlayLevel.cpp
--- a/Engine/Source/Editor/UnrealEd/Private/PlayLevel.cpp
+++ b/Engine/Source/Editor/UnrealEd/Private/PlayLevel.cpp
@@ -2397,19 +2397,22 @@ UGameInstance* UEditorEngine::CreatePIEGameInstance(int32 PIEInstance, bool bInS
 
 				FSlateApplication::Get().AddWindow( PieWindow );
 
+				TSharedPtr<SDPIScaler> PIEDPIScaler;
+
 				TSharedRef<SOverlay> ViewportOverlayWidgetRef = SNew( SOverlay );
 				PieViewportWidget = 
 					SNew( SViewport )
 						.IsEnabled(FSlateApplication::Get().GetNormalExecutionAttribute())
 						.EnableGammaCorrection( false )// Gamma correction in the game is handled in post processing in the scene renderer
 						[
-							SNew(SDPIScaler)
-							.DPIScale(TAttribute<float>::Create(TAttribute<float>::FGetter::CreateUObject(this, &UEditorEngine::GetGameViewportDPIScale)))
+							SAssignNew(PIEDPIScaler, SDPIScaler)
 							[
 								ViewportOverlayWidgetRef
 							]
 						];
 
+				PIEDPIScaler->SetDPIScale(TAttribute<float>::Create(TAttribute<float>::FGetter::CreateUObject(this, &UEditorEngine::GetGameViewportDPIScale, ViewportClient)));
+
 				// Create a viewport widget for the game to render in.
 				PieWindow->SetContent( PieViewportWidget.ToSharedRef() );
 
@@ -2528,10 +2531,11 @@ UGameInstance* UEditorEngine::CreatePIEGameInstance(int32 PIEInstance, bool bInS
 	return GameInstance;
 }
 
-float UEditorEngine::GetGameViewportDPIScale() const
+float UEditorEngine::GetGameViewportDPIScale(UGameViewportClient* ViewportClient) const
 {
-	UGameUserSettings* UserSettings = GEngine->GetGameUserSettings();
-	return GetDefault<URendererSettings>(URendererSettings::StaticClass())->GetDPIScaleBasedOnSize( UserSettings->GetScreenResolution() );
+	FVector2D ViewportSize;
+	ViewportClient->GetViewportSize(ViewportSize);
+	return GetDefault<URendererSettings>(URendererSettings::StaticClass())->GetDPIScaleBasedOnSize(FIntPoint(ViewportSize.X, ViewportSize.Y));
 }
 
 FViewport* UEditorEngine::GetActiveViewport()
diff --git a/Engine/Source/Runtime/Engine/Classes/Engine/GameEngine.h b/Engine/Source/Runtime/Engine/Classes/Engine/GameEngine.h
--- a/Engine/Source/Runtime/Engine/Classes/Engine/GameEngine.h
+++ b/Engine/Source/Runtime/Engine/Classes/Engine/GameEngine.h
@@ -137,6 +137,6 @@ public:
 	UWorld* GetGameWorld();
 
 protected:
-	float GetGameViewportDPIScale() const;
+	float GetGameViewportDPIScale(UGameViewportClient* GameViewportClient) const;
 };
 
diff --git a/Engine/Source/Runtime/Engine/Private/GameEngine.cpp b/Engine/Source/Runtime/Engine/Private/GameEngine.cpp
--- a/Engine/Source/Runtime/Engine/Private/GameEngine.cpp
+++ b/Engine/Source/Runtime/Engine/Private/GameEngine.cpp
@@ -114,7 +114,7 @@ void UGameEngine::CreateGameViewportWidget( UGameViewportClient* GameViewportCli
 			.RenderDirectlyToWindow( !GEngine->bStartWithMatineeCapture && GIsDumpingMovie == 0 )
 			[
 				SNew(SDPIScaler)
-				.DPIScale(TAttribute<float>::Create(TAttribute<float>::FGetter::CreateUObject(this, &UGameEngine::GetGameViewportDPIScale)) )
+				.DPIScale(TAttribute<float>::Create(TAttribute<float>::FGetter::CreateUObject(this, &UGameEngine::GetGameViewportDPIScale, GameViewportClient)))
 				[
 					ViewportOverlayWidgetRef
 				]
@@ -168,11 +168,11 @@ void UGameEngine::CreateGameViewport( UGameViewportClient* GameViewportClient )
 	GameViewport->SetViewportFrame(ViewportFrame);
 }
 
-float UGameEngine::GetGameViewportDPIScale() const
+float UGameEngine::GetGameViewportDPIScale(UGameViewportClient* ViewportClient) const
 {
-	UGameUserSettings* UserSettings = GEngine->GetGameUserSettings();
-	FIntPoint ScreenRes = UserSettings->GetScreenResolution();
-	return GetDefault<URendererSettings>(URendererSettings::StaticClass())->GetDPIScaleBasedOnSize( ScreenRes );
+	FVector2D ViewportSize;
+	ViewportClient->GetViewportSize(ViewportSize);
+	return GetDefault<URendererSettings>(URendererSettings::StaticClass())->GetDPIScaleBasedOnSize(FIntPoint(ViewportSize.X, ViewportSize.Y));
 }
 
 void UGameEngine::ConditionallyOverrideSettings(int32& ResolutionX, int32& ResolutionY, EWindowMode::Type& WindowMode)
