Fixed an issue where searching for assets in a project would crash because of an object class was null.
Moved some of the code from the header file to the cpp.

diff --git a/Engine/Source/Editor/UnrealEd/Private/Tests/AutomationEditorCommon.cpp b/Engine/Source/Editor/UnrealEd/Private/Tests/AutomationEditorCommon.cpp
--- a/Engine/Source/Editor/UnrealEd/Private/Tests/AutomationEditorCommon.cpp
+++ b/Engine/Source/Editor/UnrealEd/Private/Tests/AutomationEditorCommon.cpp
@@ -360,4 +360,176 @@ bool FEndPlayMapCommand::Update()
 {
 	GUnrealEd->RequestEndPlayMap();
 	return true;
+}
+
+/**
+* Generates a list of assets from the ENGINE and the GAME by a specific type.
+* This is to be used by the GetTest() function.
+*/
+void FEditorAutomationTestUtilities::CollectTestsByClass(UClass * Class, TArray<FString>& OutBeautifiedNames, TArray <FString>& OutTestCommands)
+{
+	FAssetRegistryModule& AssetRegistryModule = FModuleManager::LoadModuleChecked<FAssetRegistryModule>(TEXT("AssetRegistry"));
+	TArray<FAssetData> ObjectList;
+	AssetRegistryModule.Get().GetAssetsByClass(Class->GetFName(), ObjectList);
+
+	for (TObjectIterator<UClass> AllClassesIt; AllClassesIt; ++AllClassesIt)
+	{
+		UClass* ClassList = *AllClassesIt;
+		FName ClassName = ClassList->GetFName();
+	}
+
+	for (auto ObjIter = ObjectList.CreateConstIterator(); ObjIter; ++ObjIter)
+	{
+		const FAssetData & Asset = *ObjIter;
+		FString Filename = Asset.ObjectPath.ToString();
+		//convert to full paths
+		Filename = FPackageName::LongPackageNameToFilename(Filename);
+		if (FAutomationTestFramework::GetInstance().ShouldTestContent(Filename))
+		{
+			FString BeautifiedFilename = Asset.AssetName.ToString();
+			OutBeautifiedNames.Add(BeautifiedFilename);
+			OutTestCommands.Add(Asset.ObjectPath.ToString());
+		}
+	}
+}
+
+/**
+* Generates a list of assets from the GAME by a specific type.
+* This is to be used by the GetTest() function.
+*/
+void FEditorAutomationTestUtilities::CollectGameContentTestsByClass(UClass * Class, bool bRecursiveClass, TArray<FString>& OutBeautifiedNames, TArray <FString>& OutTestCommands)
+{
+	//Setting the Asset Registry
+
+	FAssetRegistryModule& AssetRegistryModule = FModuleManager::LoadModuleChecked<FAssetRegistryModule>(TEXT("AssetRegistry"));
+
+	//Variable setups
+	TArray<FAssetData> ObjectList;
+	FARFilter AssetFilter;
+
+	//Generating the list of assets.
+	//This list is being filtered by the game folder and class type.  The results are placed into the ObjectList variable.
+	AssetFilter.ClassNames.Add(Class->GetFName());
+	AssetFilter.PackagePaths.Add("/Game");
+	AssetFilter.bRecursiveClasses = bRecursiveClass;
+	AssetFilter.bRecursivePaths = true;
+	AssetRegistryModule.Get().GetAssets(AssetFilter, ObjectList);
+
+	//Loop through the list of assets, make their path full and a string, then add them to the test.
+	for (auto ObjIter = ObjectList.CreateConstIterator(); ObjIter; ++ObjIter)
+	{
+		const FAssetData & Asset = *ObjIter;
+		FString Filename = Asset.ObjectPath.ToString();
+		//convert to full paths
+		Filename = FPackageName::LongPackageNameToFilename(Filename);
+		if (FAutomationTestFramework::GetInstance().ShouldTestContent(Filename))
+		{
+			FString BeautifiedFilename = Asset.AssetName.ToString();
+			OutBeautifiedNames.Add(BeautifiedFilename);
+			OutTestCommands.Add(Asset.ObjectPath.ToString());
+		}
+	}
+}
+
+/**
+* Generates a list of misc. assets from the GAME.
+* This is to be used by the GetTest() function.
+*/
+void FEditorAutomationTestUtilities::CollectMiscGameContentTestsByClass(TArray<FString>& OutBeautifiedNames, TArray <FString>& OutTestCommands)
+{
+	//Setting the Asset Registry
+	FAssetRegistryModule& AssetRegistryModule = FModuleManager::LoadModuleChecked<FAssetRegistryModule>(TEXT("AssetRegistry"));
+	TArray<FAssetData> ObjectList;
+	FARFilter AssetFilter;
+
+	//This is a list of classes that we don't want to be in the misc category.
+	TArray<FName> ExcludeClassesList;
+	ExcludeClassesList.Add(UTexture::StaticClass()->GetFName());
+	ExcludeClassesList.Add(UTexture2D::StaticClass()->GetFName());
+	ExcludeClassesList.Add(UTextureRenderTarget::StaticClass()->GetFName());
+	ExcludeClassesList.Add(UTextureRenderTarget2D::StaticClass()->GetFName());
+	ExcludeClassesList.Add(UTextureCube::StaticClass()->GetFName());
+	ExcludeClassesList.Add(UAimOffsetBlendSpace::StaticClass()->GetFName());
+	ExcludeClassesList.Add(UAimOffsetBlendSpace1D::StaticClass()->GetFName());
+	ExcludeClassesList.Add(UAnimMontage::StaticClass()->GetFName());
+	ExcludeClassesList.Add(UAnimSequence::StaticClass()->GetFName());
+	ExcludeClassesList.Add(UBlendSpace::StaticClass()->GetFName());
+	ExcludeClassesList.Add(UBlendSpace1D::StaticClass()->GetFName());
+	ExcludeClassesList.Add(UDestructibleMesh::StaticClass()->GetFName());
+	ExcludeClassesList.Add(UDialogueVoice::StaticClass()->GetFName());
+	ExcludeClassesList.Add(UDialogueWave::StaticClass()->GetFName());
+	ExcludeClassesList.Add(UMaterial::StaticClass()->GetFName());
+	ExcludeClassesList.Add(UMaterialFunction::StaticClass()->GetFName());
+	ExcludeClassesList.Add(UMaterialInstanceConstant::StaticClass()->GetFName());
+	ExcludeClassesList.Add(UMaterialParameterCollection::StaticClass()->GetFName());
+	ExcludeClassesList.Add(UParticleSystem::StaticClass()->GetFName());
+	ExcludeClassesList.Add(UPhysicalMaterial::StaticClass()->GetFName());
+	ExcludeClassesList.Add(UReverbEffect::StaticClass()->GetFName());
+	ExcludeClassesList.Add(USlateWidgetStyleAsset::StaticClass()->GetFName());
+	ExcludeClassesList.Add(USlateBrushAsset::StaticClass()->GetFName());
+	ExcludeClassesList.Add(USoundAttenuation::StaticClass()->GetFName());
+	ExcludeClassesList.Add(UStaticMesh::StaticClass()->GetFName());
+	ExcludeClassesList.Add(USkeletalMesh::StaticClass()->GetFName());
+	ExcludeClassesList.Add(USkeleton::StaticClass()->GetFName());
+	ExcludeClassesList.Add(USoundClass::StaticClass()->GetFName());
+	ExcludeClassesList.Add(USoundCue::StaticClass()->GetFName());
+	ExcludeClassesList.Add(USoundMix::StaticClass()->GetFName());
+	ExcludeClassesList.Add(USoundWave::StaticClass()->GetFName());
+	ExcludeClassesList.Add(USubsurfaceProfile::StaticClass()->GetFName());
+	ExcludeClassesList.Add(UFont::StaticClass()->GetFName());
+	ExcludeClassesList.Add(UBlueprint::StaticClass()->GetFName());
+	ExcludeClassesList.Add(UAnimBlueprint::StaticClass()->GetFName());
+
+	//Generating the list of assets.
+	//This list isn't expected to have anything that may be obtained with another function.
+	//This list is being filtered by the game folder and class type.  The results are placed into the ObjectList variable.
+	AssetFilter.PackagePaths.Add("/Game");
+	AssetFilter.bRecursivePaths = true;
+	AssetRegistryModule.Get().GetAssets(AssetFilter, ObjectList);
+
+	//Loop through the list of assets, make their path full and a string, then add them to the test.
+	for (auto ObjIter = ObjectList.CreateConstIterator(); ObjIter; ++ObjIter)
+	{
+		const FAssetData & Asset = *ObjIter;
+		//First we check if the class is valid.  If not then we move onto the next object.
+		if (Asset.GetClass()->IsValidLowLevel())
+		{
+			//Variable that holds the class FName for the current asset iteration.
+			FName AssetClassFName = Asset.GetClass()->GetFName();
+
+			//Counter used to keep track for the following for loop.
+			float ExcludedClassesCounter = 1;
+
+			for (auto ExcludeIter = ExcludeClassesList.CreateConstIterator(); ExcludeIter; ++ExcludeIter)
+			{
+				FName ExludedName = *ExcludeIter;
+
+				//If the classes are the same then we don't want this asset. So we move onto the next one instead.
+				if (AssetClassFName == ExludedName)
+				{
+					break;
+				}
+
+				//We run out of class names in our Excluded list then we want the current ObjectList asset.
+				if ((ExcludedClassesCounter + 1) > ExcludeClassesList.Num())
+				{
+					FString Filename = Asset.ObjectPath.ToString();
+					//convert to full paths
+					Filename = FPackageName::LongPackageNameToFilename(Filename);
+
+					if (FAutomationTestFramework::GetInstance().ShouldTestContent(Filename))
+					{
+						FString BeautifiedFilename = Asset.AssetName.ToString();
+						OutBeautifiedNames.Add(BeautifiedFilename);
+						OutTestCommands.Add(Asset.ObjectPath.ToString());
+					}
+
+					break;
+				}
+
+				//increment the counter.
+				ExcludedClassesCounter++;
+			}
+		}
+	}
 }
\ No newline at end of file
diff --git a/Engine/Source/Editor/UnrealEd/Private/Tests/EditorOpenAssetAutomationTests.cpp b/Engine/Source/Editor/UnrealEd/Private/Tests/EditorOpenAssetAutomationTests.cpp
--- a/Engine/Source/Editor/UnrealEd/Private/Tests/EditorOpenAssetAutomationTests.cpp
+++ b/Engine/Source/Editor/UnrealEd/Private/Tests/EditorOpenAssetAutomationTests.cpp
@@ -140,7 +140,7 @@ bool FOpenAssetEditors::RunTest(const FString& Parameters)
 /**
 * This test opens each AimOffsetBlendSpace into its sub-editor.
 */
-IMPLEMENT_COMPLEX_AUTOMATION_TEST(FOpenAimOffsetBlendSpaceAssetEditors, "Editor.Content.Skeletal Mesh.Animation.BlendSpaces.Open Aim Offset BlendSpace Assets", (EAutomationTestFlags::ATF_Editor | EAutomationTestFlags::ATF_RequiresUser));
+IMPLEMENT_COMPLEX_AUTOMATION_TEST(FOpenAimOffsetBlendSpaceAssetEditors, "Tools.Skeletal Mesh.Animation.BlendSpaces.Open Aim Offset BlendSpace Assets", (EAutomationTestFlags::ATF_Editor | EAutomationTestFlags::ATF_RequiresUser));
 
 void FOpenAimOffsetBlendSpaceAssetEditors::GetTests(TArray<FString>& OutBeautifiedNames, TArray<FString>& OutTestCommands) const
 {
@@ -159,7 +159,7 @@ bool FOpenAimOffsetBlendSpaceAssetEditors::RunTest(const FString& Parameters)
 /**
 * This test opens each AimOffsetBlendSpace1D into its sub-editor.
 */
-IMPLEMENT_COMPLEX_AUTOMATION_TEST(FOpenAimOffsetBlendSpace1DAssetEditors, "Editor.Content.Skeletal Mesh.Animation.BlendSpaces.Open Aim Offset BlendSpace1D Assets", (EAutomationTestFlags::ATF_Editor | EAutomationTestFlags::ATF_RequiresUser));
+IMPLEMENT_COMPLEX_AUTOMATION_TEST(FOpenAimOffsetBlendSpace1DAssetEditors, "Tools.Skeletal Mesh.Animation.BlendSpaces.Open Aim Offset BlendSpace1D Assets", (EAutomationTestFlags::ATF_Editor | EAutomationTestFlags::ATF_RequiresUser));
 
 void FOpenAimOffsetBlendSpace1DAssetEditors::GetTests(TArray<FString>& OutBeautifiedNames, TArray<FString>& OutTestCommands) const
 {
@@ -178,7 +178,7 @@ bool FOpenAimOffsetBlendSpace1DAssetEditors::RunTest(const FString& Parameters)
 /**
 * This test opens each AnimMontage into its sub-editor.
 */
-IMPLEMENT_COMPLEX_AUTOMATION_TEST(FOpenAnimMontageAssetEditors, "Editor.Content.Skeletal Mesh.Animation.Open AnimMontage Assets", (EAutomationTestFlags::ATF_Editor | EAutomationTestFlags::ATF_RequiresUser));
+IMPLEMENT_COMPLEX_AUTOMATION_TEST(FOpenAnimMontageAssetEditors, "Tools.Skeletal Mesh.Animation.Open AnimMontage Assets", (EAutomationTestFlags::ATF_Editor | EAutomationTestFlags::ATF_RequiresUser));
 
 void FOpenAnimMontageAssetEditors::GetTests(TArray<FString>& OutBeautifiedNames, TArray<FString>& OutTestCommands) const
 {
@@ -197,7 +197,7 @@ bool FOpenAnimMontageAssetEditors::RunTest(const FString& Parameters)
 /**
 * This test opens each AnimSequence into its sub-editor.
 */
-IMPLEMENT_COMPLEX_AUTOMATION_TEST(FOpenAnimSequenceAssetEditors, "Editor.Content.Skeletal Mesh.Animation.Open AnimSequence Assets", (EAutomationTestFlags::ATF_Editor | EAutomationTestFlags::ATF_RequiresUser));
+IMPLEMENT_COMPLEX_AUTOMATION_TEST(FOpenAnimSequenceAssetEditors, "Tools.Skeletal Mesh.Animation.Open AnimSequence Assets", (EAutomationTestFlags::ATF_Editor | EAutomationTestFlags::ATF_RequiresUser));
 
 void FOpenAnimSequenceAssetEditors::GetTests(TArray<FString>& OutBeautifiedNames, TArray<FString>& OutTestCommands) const
 {
@@ -216,7 +216,7 @@ bool FOpenAnimSequenceAssetEditors::RunTest(const FString& Parameters)
 /**
 * This test opens each BlendSpace into its sub-editor.
 */
-IMPLEMENT_COMPLEX_AUTOMATION_TEST(FOpenBlendSpaceAssetEditors, "Editor.Content.Skeletal Mesh.Animation.BlendSpaces.Open BlendSpace Assets", (EAutomationTestFlags::ATF_Editor | EAutomationTestFlags::ATF_RequiresUser));
+IMPLEMENT_COMPLEX_AUTOMATION_TEST(FOpenBlendSpaceAssetEditors, "Tools.Skeletal Mesh.Animation.BlendSpaces.Open BlendSpace Assets", (EAutomationTestFlags::ATF_Editor | EAutomationTestFlags::ATF_RequiresUser));
 
 void FOpenBlendSpaceAssetEditors::GetTests(TArray<FString>& OutBeautifiedNames, TArray<FString>& OutTestCommands) const
 {
@@ -235,7 +235,7 @@ bool FOpenBlendSpaceAssetEditors::RunTest(const FString& Parameters)
 /**
 * This test opens each BlendSpace1D into its sub-editor.
 */
-IMPLEMENT_COMPLEX_AUTOMATION_TEST(FOpenBlendSpace1DAssetEditors, "Editor.Content.Skeletal Mesh.Animation.BlendSpaces.Open BlendSpace1D Assets", (EAutomationTestFlags::ATF_Editor | EAutomationTestFlags::ATF_RequiresUser));
+IMPLEMENT_COMPLEX_AUTOMATION_TEST(FOpenBlendSpace1DAssetEditors, "Tools.Skeletal Mesh.Animation.BlendSpaces.Open BlendSpace1D Assets", (EAutomationTestFlags::ATF_Editor | EAutomationTestFlags::ATF_RequiresUser));
 
 void FOpenBlendSpace1DAssetEditors::GetTests(TArray<FString>& OutBeautifiedNames, TArray<FString>& OutTestCommands) const
 {
@@ -255,7 +255,7 @@ bool FOpenBlendSpace1DAssetEditors::RunTest(const FString& Parameters)
 /**
 * This test opens each BLUEPRINT into its sub-editor.
 */
-IMPLEMENT_COMPLEX_AUTOMATION_TEST(FOpenBlueprintAssetEditors, "Editor.Content.Blueprint.Open Blueprint Assets", (EAutomationTestFlags::ATF_Editor | EAutomationTestFlags::ATF_RequiresUser));
+IMPLEMENT_COMPLEX_AUTOMATION_TEST(FOpenBlueprintAssetEditors, "Tools.Blueprint.Open Blueprint Assets", (EAutomationTestFlags::ATF_Editor | EAutomationTestFlags::ATF_RequiresUser));
 
 void FOpenBlueprintAssetEditors::GetTests(TArray<FString>& OutBeautifiedNames, TArray<FString>& OutTestCommands) const
 {
@@ -274,7 +274,7 @@ bool FOpenBlueprintAssetEditors::RunTest(const FString& Parameters)
 /**
 * This test opens each DESTRUCTIBLE MESH into its sub-editor.
 */
-IMPLEMENT_COMPLEX_AUTOMATION_TEST(FOpenDestructibleMeshAssetEditors, "Editor.Content.Skeletal Mesh.Open Destructible Mesh Assets", (EAutomationTestFlags::ATF_Editor | EAutomationTestFlags::ATF_RequiresUser));
+IMPLEMENT_COMPLEX_AUTOMATION_TEST(FOpenDestructibleMeshAssetEditors, "Tools.Skeletal Mesh.Open Destructible Mesh Assets", (EAutomationTestFlags::ATF_Editor | EAutomationTestFlags::ATF_RequiresUser));
 
 void FOpenDestructibleMeshAssetEditors::GetTests(TArray<FString>& OutBeautifiedNames, TArray<FString>& OutTestCommands) const
 {
@@ -293,7 +293,7 @@ bool FOpenDestructibleMeshAssetEditors::RunTest(const FString& Parameters)
 /**
 * This test opens each DialogueVoice into its sub-editor.
 */
-IMPLEMENT_COMPLEX_AUTOMATION_TEST(FOpenDialogueVoiceAssetEditors, "Editor.Content.Sounds.Open Dialogue Voice Assets", (EAutomationTestFlags::ATF_Editor | EAutomationTestFlags::ATF_RequiresUser));
+IMPLEMENT_COMPLEX_AUTOMATION_TEST(FOpenDialogueVoiceAssetEditors, "Tools.Sounds.Open Dialogue Voice Assets", (EAutomationTestFlags::ATF_Editor | EAutomationTestFlags::ATF_RequiresUser));
 
 void FOpenDialogueVoiceAssetEditors::GetTests(TArray<FString>& OutBeautifiedNames, TArray<FString>& OutTestCommands) const
 {
@@ -312,7 +312,7 @@ bool FOpenDialogueVoiceAssetEditors::RunTest(const FString& Parameters)
 /**
 * This test opens each DialogueWave into its sub-editor.
 */
-IMPLEMENT_COMPLEX_AUTOMATION_TEST(FOpenDialogueWaveAssetEditors, "Editor.Content.Sounds.Open Dialogue Wave Assets", (EAutomationTestFlags::ATF_Editor | EAutomationTestFlags::ATF_RequiresUser));
+IMPLEMENT_COMPLEX_AUTOMATION_TEST(FOpenDialogueWaveAssetEditors, "Tools.Sounds.Open Dialogue Wave Assets", (EAutomationTestFlags::ATF_Editor | EAutomationTestFlags::ATF_RequiresUser));
 
 void FOpenDialogueWaveAssetEditors::GetTests(TArray<FString>& OutBeautifiedNames, TArray<FString>& OutTestCommands) const
 {
@@ -331,7 +331,7 @@ bool FOpenDialogueWaveAssetEditors::RunTest(const FString& Parameters)
 /**
 * This test opens each Material into its sub-editor.
 */
-IMPLEMENT_COMPLEX_AUTOMATION_TEST(FOpenMaterialAssetEditors, "Editor.Content.Materials.Open Material Assets", (EAutomationTestFlags::ATF_Editor | EAutomationTestFlags::ATF_RequiresUser));
+IMPLEMENT_COMPLEX_AUTOMATION_TEST(FOpenMaterialAssetEditors, "Tools.Materials.Open Material Assets", (EAutomationTestFlags::ATF_Editor | EAutomationTestFlags::ATF_RequiresUser));
 
 void FOpenMaterialAssetEditors::GetTests(TArray<FString>& OutBeautifiedNames, TArray<FString>& OutTestCommands) const
 {
@@ -350,7 +350,7 @@ bool FOpenMaterialAssetEditors::RunTest(const FString& Parameters)
 ///**
 //* This test opens each MaterialFunction into its sub-editor.
 //*/
-IMPLEMENT_COMPLEX_AUTOMATION_TEST(FOpenMaterialFunctionAssetEditors, "Editor.Content.Materials.Open Material Function Assets", (EAutomationTestFlags::ATF_Editor | EAutomationTestFlags::ATF_RequiresUser));
+IMPLEMENT_COMPLEX_AUTOMATION_TEST(FOpenMaterialFunctionAssetEditors, "Tools.Materials.Open Material Function Assets", (EAutomationTestFlags::ATF_Editor | EAutomationTestFlags::ATF_RequiresUser));
 
 void FOpenMaterialFunctionAssetEditors::GetTests(TArray<FString>& OutBeautifiedNames, TArray<FString>& OutTestCommands) const
 {
@@ -369,7 +369,7 @@ bool FOpenMaterialFunctionAssetEditors::RunTest(const FString& Parameters)
 /**
 * This test opens each MaterialInstanceConstant into its sub-editor.
 */
-IMPLEMENT_COMPLEX_AUTOMATION_TEST(FOpenMaterialInstanceConstantAssetEditors, "Editor.Content.Materials.Open Material Instance Constant Assets", (EAutomationTestFlags::ATF_Editor | EAutomationTestFlags::ATF_RequiresUser));
+IMPLEMENT_COMPLEX_AUTOMATION_TEST(FOpenMaterialInstanceConstantAssetEditors, "Tools.Materials.Open Material Instance Constant Assets", (EAutomationTestFlags::ATF_Editor | EAutomationTestFlags::ATF_RequiresUser));
 
 void FOpenMaterialInstanceConstantAssetEditors::GetTests(TArray<FString>& OutBeautifiedNames, TArray<FString>& OutTestCommands) const
 {
@@ -388,7 +388,7 @@ bool FOpenMaterialInstanceConstantAssetEditors::RunTest(const FString& Parameter
 /**
 * This test opens each MaterialParameterCollection into its sub-editor.
 */
-IMPLEMENT_COMPLEX_AUTOMATION_TEST(FOpenMaterialParameterCollectionAssetEditors, "Editor.Content.Materials.Open Material Parameter Collection Assets", (EAutomationTestFlags::ATF_Editor | EAutomationTestFlags::ATF_RequiresUser));
+IMPLEMENT_COMPLEX_AUTOMATION_TEST(FOpenMaterialParameterCollectionAssetEditors, "Tools.Materials.Open Material Parameter Collection Assets", (EAutomationTestFlags::ATF_Editor | EAutomationTestFlags::ATF_RequiresUser));
 
 void FOpenMaterialParameterCollectionAssetEditors::GetTests(TArray<FString>& OutBeautifiedNames, TArray<FString>& OutTestCommands) const
 {
@@ -407,7 +407,7 @@ bool FOpenMaterialParameterCollectionAssetEditors::RunTest(const FString& Parame
 /**
 * This test opens each PARTICLE SYSTEM into its sub-editor.
 */
-IMPLEMENT_COMPLEX_AUTOMATION_TEST(FOpenParticleSystemAssetEditors, "Editor.Content.Particle System.Open Particle System Assets", (EAutomationTestFlags::ATF_Editor | EAutomationTestFlags::ATF_RequiresUser));
+IMPLEMENT_COMPLEX_AUTOMATION_TEST(FOpenParticleSystemAssetEditors, "Tools.Particle System.Open Particle System Assets", (EAutomationTestFlags::ATF_Editor | EAutomationTestFlags::ATF_RequiresUser));
 
 void FOpenParticleSystemAssetEditors::GetTests(TArray<FString>& OutBeautifiedNames, TArray<FString>& OutTestCommands) const
 {
@@ -426,7 +426,7 @@ bool FOpenParticleSystemAssetEditors::RunTest(const FString& Parameters)
 /**
 * This test opens each PhysicalMaterial into its sub-editor.
 */
-IMPLEMENT_COMPLEX_AUTOMATION_TEST(FOpenPhysicalMaterialAssetEditors, "Editor.Content.Physics.Open Physical Material Assets", (EAutomationTestFlags::ATF_Editor | EAutomationTestFlags::ATF_RequiresUser));
+IMPLEMENT_COMPLEX_AUTOMATION_TEST(FOpenPhysicalMaterialAssetEditors, "Tools.Physics.Open Physical Material Assets", (EAutomationTestFlags::ATF_Editor | EAutomationTestFlags::ATF_RequiresUser));
 
 void FOpenPhysicalMaterialAssetEditors::GetTests(TArray<FString>& OutBeautifiedNames, TArray<FString>& OutTestCommands) const
 {
@@ -446,7 +446,7 @@ bool FOpenPhysicalMaterialAssetEditors::RunTest(const FString& Parameters)
 /**
 * This test opens each PhysicsAsset into its sub-editor.
 */
-IMPLEMENT_COMPLEX_AUTOMATION_TEST(FOpenPhysicsAssetAssetEditors, "Editor.Content.Physics.Open Physics Asset Assets", (EAutomationTestFlags::ATF_Editor | EAutomationTestFlags::ATF_RequiresUser));
+IMPLEMENT_COMPLEX_AUTOMATION_TEST(FOpenPhysicsAssetAssetEditors, "Tools.Physics.Open Physics Asset Assets", (EAutomationTestFlags::ATF_Editor | EAutomationTestFlags::ATF_RequiresUser));
 
 void FOpenPhysicsAssetAssetEditors::GetTests(TArray<FString>& OutBeautifiedNames, TArray<FString>& OutTestCommands) const
 {
@@ -466,7 +466,7 @@ bool FOpenPhysicsAssetAssetEditors::RunTest(const FString& Parameters)
 /**
 * This test opens each ReverbEffect into its sub-editor.
 */
-IMPLEMENT_COMPLEX_AUTOMATION_TEST(FOpenReverbEffectAssetEditors, "Editor.Content.Sounds.Open Reverb Effect Assets", (EAutomationTestFlags::ATF_Editor | EAutomationTestFlags::ATF_RequiresUser));
+IMPLEMENT_COMPLEX_AUTOMATION_TEST(FOpenReverbEffectAssetEditors, "Tools.Sounds.Open Reverb Effect Assets", (EAutomationTestFlags::ATF_Editor | EAutomationTestFlags::ATF_RequiresUser));
 
 void FOpenReverbEffectAssetEditors::GetTests(TArray<FString>& OutBeautifiedNames, TArray<FString>& OutTestCommands) const
 {
@@ -487,7 +487,7 @@ bool FOpenReverbEffectAssetEditors::RunTest(const FString& Parameters)
 /**
 * This test opens each SlateWidgetStyleAsset into its sub-editor.
 */
-IMPLEMENT_COMPLEX_AUTOMATION_TEST(FOpenSlateWidgetStyleAssetAssetEditors, "Editor.Content.Slate.Open Slate Widget Style Assets", (EAutomationTestFlags::ATF_Editor | EAutomationTestFlags::ATF_RequiresUser));
+IMPLEMENT_COMPLEX_AUTOMATION_TEST(FOpenSlateWidgetStyleAssetAssetEditors, "Tools.Slate.Open Slate Widget Style Assets", (EAutomationTestFlags::ATF_Editor | EAutomationTestFlags::ATF_RequiresUser));
 
 void FOpenSlateWidgetStyleAssetAssetEditors::GetTests(TArray<FString>& OutBeautifiedNames, TArray<FString>& OutTestCommands) const
 {
@@ -506,7 +506,7 @@ bool FOpenSlateWidgetStyleAssetAssetEditors::RunTest(const FString& Parameters)
 /**
 * This test opens each Slate Brush Asset into its sub-editor.
 */
-IMPLEMENT_COMPLEX_AUTOMATION_TEST(FOpenSlateBrushAssetAssetEditors, "Editor.Content.Slate.Open Slate Brush Assets", (EAutomationTestFlags::ATF_Editor | EAutomationTestFlags::ATF_RequiresUser));
+IMPLEMENT_COMPLEX_AUTOMATION_TEST(FOpenSlateBrushAssetAssetEditors, "Tools.Slate.Open Slate Brush Assets", (EAutomationTestFlags::ATF_Editor | EAutomationTestFlags::ATF_RequiresUser));
 
 void FOpenSlateBrushAssetAssetEditors::GetTests(TArray<FString>& OutBeautifiedNames, TArray<FString>& OutTestCommands) const
 {
@@ -524,7 +524,7 @@ bool FOpenSlateBrushAssetAssetEditors::RunTest(const FString& Parameters)
 /**
 * This test opens each SoundAttenuation into its sub-editor.
 */
-IMPLEMENT_COMPLEX_AUTOMATION_TEST(FOpenSoundAttenuationAssetEditors, "Editor.Content.Sounds.Open Sound Attenuation Assets", (EAutomationTestFlags::ATF_Editor | EAutomationTestFlags::ATF_RequiresUser));
+IMPLEMENT_COMPLEX_AUTOMATION_TEST(FOpenSoundAttenuationAssetEditors, "Tools.Sounds.Open Sound Attenuation Assets", (EAutomationTestFlags::ATF_Editor | EAutomationTestFlags::ATF_RequiresUser));
 
 void FOpenSoundAttenuationAssetEditors::GetTests(TArray<FString>& OutBeautifiedNames, TArray<FString>& OutTestCommands) const
 {
@@ -543,7 +543,7 @@ bool FOpenSoundAttenuationAssetEditors::RunTest(const FString& Parameters)
 /**
 * This test opens each Static Mesh into their sub-editor.
 */
-IMPLEMENT_COMPLEX_AUTOMATION_TEST(FOpenStaticMeshAssetEditors, "Editor.Content.Static Mesh.Open Static Mesh Assets", (EAutomationTestFlags::ATF_Editor | EAutomationTestFlags::ATF_RequiresUser));
+IMPLEMENT_COMPLEX_AUTOMATION_TEST(FOpenStaticMeshAssetEditors, "Tools.Static Mesh.Open Static Mesh Assets", (EAutomationTestFlags::ATF_Editor | EAutomationTestFlags::ATF_RequiresUser));
 
 void FOpenStaticMeshAssetEditors::GetTests(TArray<FString>& OutBeautifiedNames, TArray<FString>& OutTestCommands) const
 {
@@ -562,7 +562,7 @@ bool FOpenStaticMeshAssetEditors::RunTest(const FString& Parameters)
 /**
 * This test opens each Skeletal Mesh into their sub-editor.
 */
-IMPLEMENT_COMPLEX_AUTOMATION_TEST(FOpenSkeletalMeshAssetEditors, "Editor.Content.Skeletal Mesh.Open Skeletal Mesh Assets", (EAutomationTestFlags::ATF_Editor | EAutomationTestFlags::ATF_RequiresUser));
+IMPLEMENT_COMPLEX_AUTOMATION_TEST(FOpenSkeletalMeshAssetEditors, "Tools.Skeletal Mesh.Open Skeletal Mesh Assets", (EAutomationTestFlags::ATF_Editor | EAutomationTestFlags::ATF_RequiresUser));
 
 void FOpenSkeletalMeshAssetEditors::GetTests(TArray<FString>& OutBeautifiedNames, TArray<FString>& OutTestCommands) const
 {
@@ -581,7 +581,7 @@ bool FOpenSkeletalMeshAssetEditors::RunTest(const FString& Parameters)
 /**
 * This test opens each Skeleton into its sub-editor.
 */
-IMPLEMENT_COMPLEX_AUTOMATION_TEST(FOpenSkeletonAssetEditors, "Editor.Content.Skeletal Mesh.Open Skeleton Assets", (EAutomationTestFlags::ATF_Editor | EAutomationTestFlags::ATF_RequiresUser));
+IMPLEMENT_COMPLEX_AUTOMATION_TEST(FOpenSkeletonAssetEditors, "Tools.Skeletal Mesh.Open Skeleton Assets", (EAutomationTestFlags::ATF_Editor | EAutomationTestFlags::ATF_RequiresUser));
 
 void FOpenSkeletonAssetEditors::GetTests(TArray<FString>& OutBeautifiedNames, TArray<FString>& OutTestCommands) const
 {
@@ -600,7 +600,7 @@ bool FOpenSkeletonAssetEditors::RunTest(const FString& Parameters)
 /**
 * This test opens each SoundClass into its sub-editor.
 */
-IMPLEMENT_COMPLEX_AUTOMATION_TEST(FOpenSoundClassAssetEditors, "Editor.Content.Sounds.Open SoundClass Assets", (EAutomationTestFlags::ATF_Editor | EAutomationTestFlags::ATF_RequiresUser));
+IMPLEMENT_COMPLEX_AUTOMATION_TEST(FOpenSoundClassAssetEditors, "Tools.Sounds.Open SoundClass Assets", (EAutomationTestFlags::ATF_Editor | EAutomationTestFlags::ATF_RequiresUser));
 
 void FOpenSoundClassAssetEditors::GetTests(TArray<FString>& OutBeautifiedNames, TArray<FString>& OutTestCommands) const
 {
@@ -620,7 +620,7 @@ bool FOpenSoundClassAssetEditors::RunTest(const FString& Parameters)
 /**
 * This test opens each SoundCue into its sub-editor.
 */
-IMPLEMENT_COMPLEX_AUTOMATION_TEST(FOpenSoundCueAssetEditors, "Editor.Content.Sounds.Open Sound Cue Assets", (EAutomationTestFlags::ATF_Editor | EAutomationTestFlags::ATF_RequiresUser));
+IMPLEMENT_COMPLEX_AUTOMATION_TEST(FOpenSoundCueAssetEditors, "Tools.Sounds.Open Sound Cue Assets", (EAutomationTestFlags::ATF_Editor | EAutomationTestFlags::ATF_RequiresUser));
 
 void FOpenSoundCueAssetEditors::GetTests(TArray<FString>& OutBeautifiedNames, TArray<FString>& OutTestCommands) const
 {
@@ -640,7 +640,7 @@ bool FOpenSoundCueAssetEditors::RunTest(const FString& Parameters)
 /**
 * This test opens each SoundMix into its sub-editor.
 */
-IMPLEMENT_COMPLEX_AUTOMATION_TEST(FOpenSoundMixAssetEditors, "Editor.Content.Sounds.Open Sound Mix Assets", (EAutomationTestFlags::ATF_Editor | EAutomationTestFlags::ATF_RequiresUser));
+IMPLEMENT_COMPLEX_AUTOMATION_TEST(FOpenSoundMixAssetEditors, "Tools.Sounds.Open Sound Mix Assets", (EAutomationTestFlags::ATF_Editor | EAutomationTestFlags::ATF_RequiresUser));
 
 void FOpenSoundMixAssetEditors::GetTests(TArray<FString>& OutBeautifiedNames, TArray<FString>& OutTestCommands) const
 {
@@ -660,7 +660,7 @@ bool FOpenSoundMixAssetEditors::RunTest(const FString& Parameters)
 /**
 * This test opens each SoundWave into its sub-editor.
 */
-IMPLEMENT_COMPLEX_AUTOMATION_TEST(FOpenSoundWaveAssetEditors, "Editor.Content.Sounds.Open Sound Wave Assets", (EAutomationTestFlags::ATF_Editor | EAutomationTestFlags::ATF_RequiresUser));
+IMPLEMENT_COMPLEX_AUTOMATION_TEST(FOpenSoundWaveAssetEditors, "Tools.Sounds.Open Sound Wave Assets", (EAutomationTestFlags::ATF_Editor | EAutomationTestFlags::ATF_RequiresUser));
 
 void FOpenSoundWaveAssetEditors::GetTests(TArray<FString>& OutBeautifiedNames, TArray<FString>& OutTestCommands) const
 {
@@ -678,7 +678,7 @@ bool FOpenSoundWaveAssetEditors::RunTest(const FString& Parameters)
 /**
 * This test opens each SubsurfaceProfile into its sub-editor.
 */
-IMPLEMENT_COMPLEX_AUTOMATION_TEST(FOpenSubsurfaceProfileAssetEditors, "Editor.Content.Materials.Open Subsurface Profile Assets", (EAutomationTestFlags::ATF_Editor | EAutomationTestFlags::ATF_RequiresUser));
+IMPLEMENT_COMPLEX_AUTOMATION_TEST(FOpenSubsurfaceProfileAssetEditors, "Tools.Materials.Open Subsurface Profile Assets", (EAutomationTestFlags::ATF_Editor | EAutomationTestFlags::ATF_RequiresUser));
 
 void FOpenSubsurfaceProfileAssetEditors::GetTests(TArray<FString>& OutBeautifiedNames, TArray<FString>& OutTestCommands) const
 {
@@ -699,7 +699,7 @@ bool FOpenSubsurfaceProfileAssetEditors::RunTest(const FString& Parameters)
 /**
 * This test opens each TEXTURE asset into its sub-editor.
 */
-IMPLEMENT_COMPLEX_AUTOMATION_TEST(FOpenTextureAssetEditors, "Editor.Content.Textures.Open Texture Assets", (EAutomationTestFlags::ATF_Editor | EAutomationTestFlags::ATF_RequiresUser));
+IMPLEMENT_COMPLEX_AUTOMATION_TEST(FOpenTextureAssetEditors, "Tools.Textures.Open Texture Assets", (EAutomationTestFlags::ATF_Editor | EAutomationTestFlags::ATF_RequiresUser));
 
 void FOpenTextureAssetEditors::GetTests(TArray<FString>& OutBeautifiedNames, TArray<FString>& OutTestCommands) const
 {
@@ -719,7 +719,7 @@ bool FOpenTextureAssetEditors::RunTest(const FString& Parameters)
 /**
 * This test opens each Font into its sub-editor.
 */
-IMPLEMENT_COMPLEX_AUTOMATION_TEST(FOpenFontAssetEditors, "Editor.Content.Textures.Open Font Assets", (EAutomationTestFlags::ATF_Editor | EAutomationTestFlags::ATF_RequiresUser));
+IMPLEMENT_COMPLEX_AUTOMATION_TEST(FOpenFontAssetEditors, "Tools.Textures.Open Font Assets", (EAutomationTestFlags::ATF_Editor | EAutomationTestFlags::ATF_RequiresUser));
 
 void FOpenFontAssetEditors::GetTests(TArray<FString>& OutBeautifiedNames, TArray<FString>& OutTestCommands) const
 {
@@ -738,7 +738,7 @@ bool FOpenFontAssetEditors::RunTest(const FString& Parameters)
 /**
 * This test opens each misc asset into its sub-editor.
 */
-IMPLEMENT_COMPLEX_AUTOMATION_TEST(FOpenMiscAssetEditors, "Editor.Content.Misc.Open Misc Assets", (EAutomationTestFlags::ATF_Editor | EAutomationTestFlags::ATF_RequiresUser));
+IMPLEMENT_COMPLEX_AUTOMATION_TEST(FOpenMiscAssetEditors, "Tools.Misc.Open Misc Assets", (EAutomationTestFlags::ATF_Editor | EAutomationTestFlags::ATF_RequiresUser));
 
 void FOpenMiscAssetEditors::GetTests(TArray<FString>& OutBeautifiedNames, TArray<FString>& OutTestCommands) const
 {
diff --git a/Engine/Source/Editor/UnrealEd/Public/Tests/AutomationEditorCommon.h b/Engine/Source/Editor/UnrealEd/Public/Tests/AutomationEditorCommon.h
--- a/Engine/Source/Editor/UnrealEd/Public/Tests/AutomationEditorCommon.h
+++ b/Engine/Source/Editor/UnrealEd/Public/Tests/AutomationEditorCommon.h
@@ -165,174 +165,25 @@ public:
 		ADD_LATENT_AUTOMATION_COMMAND(FWaitLatentCommand(3.0f));
 		ADD_LATENT_AUTOMATION_COMMAND(FEndPlayMapCommand());
 	}
-
+	
 	/**
 	* Generates a list of assets from the ENGINE and the GAME by a specific type.
 	* This is to be used by the GetTest() function.
 	*/
-	static void CollectTestsByClass(UClass * Class, TArray<FString>& OutBeautifiedNames, TArray <FString>& OutTestCommands)
-	{
-		FAssetRegistryModule& AssetRegistryModule = FModuleManager::LoadModuleChecked<FAssetRegistryModule>(TEXT("AssetRegistry"));
-		TArray<FAssetData> ObjectList;
-		AssetRegistryModule.Get().GetAssetsByClass(Class->GetFName(), ObjectList);
-
-		for (TObjectIterator<UClass> AllClassesIt; AllClassesIt; ++AllClassesIt)
-		{
-			UClass* ClassList = *AllClassesIt;
-			FName ClassName = ClassList->GetFName();
-		}
-		
-		for (auto ObjIter = ObjectList.CreateConstIterator(); ObjIter; ++ObjIter)
-		{
-			const FAssetData & Asset = *ObjIter;
-			FString Filename = Asset.ObjectPath.ToString();
-			//convert to full paths
-			Filename = FPackageName::LongPackageNameToFilename(Filename);
-			if (FAutomationTestFramework::GetInstance().ShouldTestContent(Filename))
-			{
-				FString BeautifiedFilename = Asset.AssetName.ToString();
-				OutBeautifiedNames.Add(BeautifiedFilename);
-				OutTestCommands.Add(Asset.ObjectPath.ToString());
-			}
-		}
-	}
-
+	static void CollectTestsByClass(UClass * Class, TArray<FString>& OutBeautifiedNames, TArray <FString>& OutTestCommands);
+	
 	/**
 	* Generates a list of assets from the GAME by a specific type.
 	* This is to be used by the GetTest() function.
 	*/
-	static void CollectGameContentTestsByClass(UClass * Class, bool bRecursiveClass, TArray<FString>& OutBeautifiedNames, TArray <FString>& OutTestCommands)
-	{
-		//Setting the Asset Registry
-
-		FAssetRegistryModule& AssetRegistryModule = FModuleManager::LoadModuleChecked<FAssetRegistryModule>(TEXT("AssetRegistry"));
+	static void CollectGameContentTestsByClass(UClass * Class, bool bRecursiveClass, TArray<FString>& OutBeautifiedNames, TArray <FString>& OutTestCommands);
 		
-		//Variable setups
-		TArray<FAssetData> ObjectList;
-		FARFilter AssetFilter;
-
-		//Generating the list of assets.
-		//This list is being filtered by the game folder and class type.  The results are placed into the ObjectList variable.
-		AssetFilter.ClassNames.Add(Class->GetFName());
-		AssetFilter.PackagePaths.Add("/Game");
-		AssetFilter.bRecursiveClasses = bRecursiveClass;
-		AssetFilter.bRecursivePaths = true;
-		AssetRegistryModule.Get().GetAssets(AssetFilter, ObjectList);
-
-		//Loop through the list of assets, make their path full and a string, then add them to the test.
-		for (auto ObjIter = ObjectList.CreateConstIterator(); ObjIter; ++ObjIter)
-		{
-			const FAssetData & Asset = *ObjIter;
-			FString Filename = Asset.ObjectPath.ToString();
-			//convert to full paths
-			Filename = FPackageName::LongPackageNameToFilename(Filename);
-			if (FAutomationTestFramework::GetInstance().ShouldTestContent(Filename))
-			{
-				FString BeautifiedFilename = Asset.AssetName.ToString();
-				OutBeautifiedNames.Add(BeautifiedFilename);
-				OutTestCommands.Add(Asset.ObjectPath.ToString());
-			}
-		}
-	}
-
 	/**
 	* Generates a list of misc. assets from the GAME.
 	* This is to be used by the GetTest() function.
 	*/
-	static void CollectMiscGameContentTestsByClass(TArray<FString>& OutBeautifiedNames, TArray <FString>& OutTestCommands)
-	{
-		//Setting the Asset Registry
-		FAssetRegistryModule& AssetRegistryModule = FModuleManager::LoadModuleChecked<FAssetRegistryModule>(TEXT("AssetRegistry"));
-		TArray<FAssetData> ObjectList;
-		FARFilter AssetFilter;
-
-		//This is a list of classes that we don't want to be in the misc category.
-		TArray<FName> ExcludeClassesList;
-		ExcludeClassesList.Add(UTexture::StaticClass()->GetFName());
-		ExcludeClassesList.Add(UTexture2D::StaticClass()->GetFName());
-		ExcludeClassesList.Add(UTextureRenderTarget::StaticClass()->GetFName());
-		ExcludeClassesList.Add(UTextureRenderTarget2D::StaticClass()->GetFName());
-		ExcludeClassesList.Add(UTextureCube::StaticClass()->GetFName());
-		ExcludeClassesList.Add(UAimOffsetBlendSpace::StaticClass()->GetFName());
-		ExcludeClassesList.Add(UAimOffsetBlendSpace1D::StaticClass()->GetFName());
-		ExcludeClassesList.Add(UAnimMontage::StaticClass()->GetFName());
-		ExcludeClassesList.Add(UAnimSequence::StaticClass()->GetFName());
-		ExcludeClassesList.Add(UBlendSpace::StaticClass()->GetFName());
-		ExcludeClassesList.Add(UBlendSpace1D::StaticClass()->GetFName());
-		ExcludeClassesList.Add(UDestructibleMesh::StaticClass()->GetFName());
-		ExcludeClassesList.Add(UDialogueVoice::StaticClass()->GetFName());
-		ExcludeClassesList.Add(UDialogueWave::StaticClass()->GetFName());
-		ExcludeClassesList.Add(UMaterial::StaticClass()->GetFName());
-		ExcludeClassesList.Add(UMaterialFunction::StaticClass()->GetFName());
-		ExcludeClassesList.Add(UMaterialInstanceConstant::StaticClass()->GetFName());
-		ExcludeClassesList.Add(UMaterialParameterCollection::StaticClass()->GetFName());
-		ExcludeClassesList.Add(UParticleSystem::StaticClass()->GetFName());
-		ExcludeClassesList.Add(UPhysicalMaterial::StaticClass()->GetFName());
-		ExcludeClassesList.Add(UReverbEffect::StaticClass()->GetFName());
-		ExcludeClassesList.Add(USlateWidgetStyleAsset::StaticClass()->GetFName());
-		ExcludeClassesList.Add(USlateBrushAsset::StaticClass()->GetFName());
-		ExcludeClassesList.Add(USoundAttenuation::StaticClass()->GetFName());
-		ExcludeClassesList.Add(UStaticMesh::StaticClass()->GetFName());
-		ExcludeClassesList.Add(USkeletalMesh::StaticClass()->GetFName());
-		ExcludeClassesList.Add(USkeleton::StaticClass()->GetFName());
-		ExcludeClassesList.Add(USoundClass::StaticClass()->GetFName());
-		ExcludeClassesList.Add(USoundCue::StaticClass()->GetFName());
-		ExcludeClassesList.Add(USoundMix::StaticClass()->GetFName());
-		ExcludeClassesList.Add(USoundWave::StaticClass()->GetFName());
-		ExcludeClassesList.Add(USubsurfaceProfile::StaticClass()->GetFName());
-		ExcludeClassesList.Add(UFont::StaticClass()->GetFName());
-		ExcludeClassesList.Add(UBlueprint::StaticClass()->GetFName());
-		ExcludeClassesList.Add(UAnimBlueprint::StaticClass()->GetFName());
-		
-		//Generating the list of assets.
-		//This list isn't expected to have anything that may be obtained with another function.
-		//This list is being filtered by the game folder and class type.  The results are placed into the ObjectList variable.
-		AssetFilter.PackagePaths.Add("/Game");
-		AssetFilter.bRecursivePaths = true;
-		AssetRegistryModule.Get().GetAssets(AssetFilter, ObjectList);
-		
-		//Loop through the list of assets, make their path full and a string, then add them to the test.
-		for (auto ObjIter = ObjectList.CreateConstIterator(); ObjIter; ++ObjIter)
-		{
-			const FAssetData & Asset = *ObjIter;
-			//Variable that holds the class FName for the current asset iteration.
-			FName AssetClassFName = Asset.GetClass()->GetFName();
-
-			//Counter used to keep track for the following for loop.
-			float ExcludedClassesCounter = 1;
-
-			for (auto ExcludeIter = ExcludeClassesList.CreateConstIterator(); ExcludeIter; ++ExcludeIter)
-			{
-				FName ExludedName = *ExcludeIter;
-				
-				//If the classes are the same then we don't want this asset. So we move onto the next one instead.
-				if (AssetClassFName == ExludedName)
-				{
-					break;
-				}
-
-				//We run out of class names in our Excluded list then we want the current ObjectList asset.
-				if ((ExcludedClassesCounter + 1) > ExcludeClassesList.Num())
-				{
-					FString Filename = Asset.ObjectPath.ToString();
-					//convert to full paths
-					Filename = FPackageName::LongPackageNameToFilename(Filename);
-
-					if (FAutomationTestFramework::GetInstance().ShouldTestContent(Filename))
-					{
-						FString BeautifiedFilename = Asset.AssetName.ToString();
-						OutBeautifiedNames.Add(BeautifiedFilename);
-						OutTestCommands.Add(Asset.ObjectPath.ToString());
-					}
-
-					break;
-				}
-				
-				//increment the counter.
-				ExcludedClassesCounter++;
-			}
-		}
-	}
+	static void CollectMiscGameContentTestsByClass(TArray<FString>& OutBeautifiedNames, TArray <FString>& OutTestCommands);
+	
 };
 
 
