BP: more informative warning. "The following member variables in blueprint  have invalid type"

diff --git a/Engine/Source/Editor/UnrealEd/Private/Kismet2/StructureEditorUtils.cpp b/Engine/Source/Editor/UnrealEd/Private/Kismet2/StructureEditorUtils.cpp
--- a/Engine/Source/Editor/UnrealEd/Private/Kismet2/StructureEditorUtils.cpp
+++ b/Engine/Source/Editor/UnrealEd/Private/Kismet2/StructureEditorUtils.cpp
@@ -526,7 +526,8 @@ void FStructureEditorUtils::RemoveInvalidStructureMemberVariableFromBlueprint(UB
 			auto Response = FMessageDialog::Open( 
 				EAppMsgType::OkCancel,
 				FText::Format(
-					LOCTEXT("RemoveInvalidStructureMemberVariable_Msg", "The following member variables have invalid type. Would you like to remove them? \n\n{0}"), 
+					LOCTEXT("RemoveInvalidStructureMemberVariable_Msg", "The following member variables in blueprint '{0}' have invalid type. Would you like to remove them? \n\n{1}"), 
+					FText::FromString(Blueprint->GetFullName()),
 					FText::FromString(DislpayList)
 				));
 			check((EAppReturnType::Ok == Response) || (EAppReturnType::Cancel == Response));
