Changed the Static Mesh UV's test to be in the Tools category under Static Mesh

diff --git a/Engine/Source/Editor/UnrealEd/Private/Tests/EditorAutomationTests.cpp b/Engine/Source/Editor/UnrealEd/Private/Tests/EditorAutomationTests.cpp
--- a/Engine/Source/Editor/UnrealEd/Private/Tests/EditorAutomationTests.cpp
+++ b/Engine/Source/Editor/UnrealEd/Private/Tests/EditorAutomationTests.cpp
@@ -1973,7 +1973,7 @@ bool FPerformUVTestCommand::Update()
 /**
  * StaticMeshUVsTest
  */
-IMPLEMENT_SIMPLE_AUTOMATION_TEST( FStaticMeshUVsTest, "Editor.Content.Static Mesh UVs", EAutomationTestFlags::ATF_Editor )
+IMPLEMENT_SIMPLE_AUTOMATION_TEST( FStaticMeshUVsTest, "Tools.Static Mesh.Static Mesh UVs Check", EAutomationTestFlags::ATF_Editor )
 
 /**
  * Find all static meshes and check the lightmap UVs
