Fix crash when retargetting assets that have no skeleton

diff --git a/Engine/Source/Developer/AssetTools/Private/AssetTypeActions/AssetTypeActions_AnimBlueprint.cpp b/Engine/Source/Developer/AssetTools/Private/AssetTypeActions/AssetTypeActions_AnimBlueprint.cpp
--- a/Engine/Source/Developer/AssetTools/Private/AssetTypeActions/AssetTypeActions_AnimBlueprint.cpp
+++ b/Engine/Source/Developer/AssetTools/Private/AssetTypeActions/AssetTypeActions_AnimBlueprint.cpp
@@ -150,18 +150,20 @@ void FAssetTypeActions_AnimBlueprint::ExecuteFindSkeleton(TArray<TWeakObjectPtr<
 void FAssetTypeActions_AnimBlueprint::RetargetAssets(const TArray<UObject*> InAnimBlueprints, bool bDuplicateAssets)
 {
 	bool bRemapReferencedAssets = false;
-	bool bConvertSpaces = true;
 	USkeleton * NewSkeleton = NULL;
 	USkeleton * OldSkeleton = NULL;
+
 	if ( InAnimBlueprints.Num() > 0 )
 	{
 		UAnimBlueprint * AnimBP = CastChecked<UAnimBlueprint>(InAnimBlueprints[0]);
 		OldSkeleton = AnimBP->TargetSkeleton;
 	}
 
+	bool bConvertSpaces = OldSkeleton != NULL;
+
 	const FText Message = LOCTEXT("RemapSkeleton_Warning", "Select the skeleton to remap this asset to.");
 
-	if (SAnimationRemapSkeleton::ShowModal(OldSkeleton, NewSkeleton, Message, bConvertSpaces, &bRemapReferencedAssets))
+	if (SAnimationRemapSkeleton::ShowModal(OldSkeleton, NewSkeleton, Message, OldSkeleton ? &bConvertSpaces : NULL, &bRemapReferencedAssets))
 	{
 		EditorAnimUtils::RetargetAnimations(OldSkeleton, NewSkeleton, InAnimBlueprints, bRemapReferencedAssets, bDuplicateAssets, bConvertSpaces);
 	}
diff --git a/Engine/Source/Developer/AssetTools/Private/AssetTypeActions/AssetTypeActions_AnimationAsset.cpp b/Engine/Source/Developer/AssetTools/Private/AssetTypeActions/AssetTypeActions_AnimationAsset.cpp
--- a/Engine/Source/Developer/AssetTools/Private/AssetTypeActions/AssetTypeActions_AnimationAsset.cpp
+++ b/Engine/Source/Developer/AssetTools/Private/AssetTypeActions/AssetTypeActions_AnimationAsset.cpp
@@ -156,7 +156,6 @@ void FAssetTypeActions_AnimationAsset::ExecuteFindSkeleton(TArray<TWeakObjectPtr
 void FAssetTypeActions_AnimationAsset::RetargetAssets(const TArray<UObject*> InAnimAssets, bool bDuplicateAssets)
 {
 	bool bRemapReferencedAssets = false;
-	bool bConvertSpaces = true;
 	USkeleton * NewSkeleton = NULL;
 	USkeleton * OldSkeleton = NULL;
 	if(InAnimAssets.Num() > 0)
@@ -164,9 +163,12 @@ void FAssetTypeActions_AnimationAsset::RetargetAssets(const TArray<UObject*> InA
 		UAnimationAsset * AnimAsset = CastChecked<UAnimationAsset>(InAnimAssets[0]);
 		OldSkeleton = AnimAsset->GetSkeleton();
 	}
+
+	bool bConvertSpaces = OldSkeleton != NULL;
+
 	const FText Message = LOCTEXT("SelectSkeletonToRemap", "Select the skeleton to remap this asset to.");
 
-	if (SAnimationRemapSkeleton::ShowModal(OldSkeleton, NewSkeleton, Message, bConvertSpaces, &bRemapReferencedAssets))
+	if (SAnimationRemapSkeleton::ShowModal(OldSkeleton, NewSkeleton, Message, OldSkeleton ? &bConvertSpaces : NULL, &bRemapReferencedAssets))
 	{
 		EditorAnimUtils::RetargetAnimations(OldSkeleton, NewSkeleton, InAnimAssets, bRemapReferencedAssets, bDuplicateAssets, bConvertSpaces);
 	}
diff --git a/Engine/Source/Developer/AssetTools/Private/AssetTypeActions/AssetTypeActions_Skeleton.cpp b/Engine/Source/Developer/AssetTools/Private/AssetTypeActions/AssetTypeActions_Skeleton.cpp
--- a/Engine/Source/Developer/AssetTools/Private/AssetTypeActions/AssetTypeActions_Skeleton.cpp
+++ b/Engine/Source/Developer/AssetTools/Private/AssetTypeActions/AssetTypeActions_Skeleton.cpp
@@ -179,9 +179,10 @@ void FAssetTypeActions_Skeleton::ExecuteRetargetSkeleton(TArray<TWeakObjectPtr<U
 
 			const FText Message = LOCTEXT("RetargetSkeleton_Warning", "This only converts animation data -i.e. animation assets and Anim Blueprints. \nIf you'd like to convert SkeletalMesh, use the context menu (Assign Skeleton) for each mesh. \n\nIf you'd like to convert mesh as well, please do so before converting animation data. \nOtherwise you will lose any extra track that is in the new mesh.");
 
-			bool bConvertSpaces = true;
+			bool bConvertSpaces = OldSkeleton != NULL;
+
 			// ask user what they'd like to change to 
-			if (SAnimationRemapSkeleton::ShowModal(OldSkeleton, NewSkeleton, Message, bConvertSpaces ))
+			if (SAnimationRemapSkeleton::ShowModal(OldSkeleton, NewSkeleton, Message, OldSkeleton ? &bConvertSpaces : NULL))
 			{
 				// find all assets who references old skeleton
 				TArray<FName> Packages;
diff --git a/Engine/Source/Editor/Persona/Private/SAnimationSequenceBrowser.cpp b/Engine/Source/Editor/Persona/Private/SAnimationSequenceBrowser.cpp
--- a/Engine/Source/Editor/Persona/Private/SAnimationSequenceBrowser.cpp
+++ b/Engine/Source/Editor/Persona/Private/SAnimationSequenceBrowser.cpp
@@ -297,7 +297,7 @@ void SAnimationSequenceBrowser::OnCreateCopy(TArray<FAssetData> Selected)
 		const FText Message = LOCTEXT("RemapSkeleton_Warning", "This will duplicate the asset and convert to new skeleton.");
 
 		// ask user what they'd like to change to 
-		if (SAnimationRemapSkeleton::ShowModal(OldSkeleton, NewSkeleton, Message, bConvertSpaces, &bRemapReferencedAssets))
+		if (SAnimationRemapSkeleton::ShowModal(OldSkeleton, NewSkeleton, Message, &bConvertSpaces, &bRemapReferencedAssets))
 		{
 			UObject* AssetToOpen = EditorAnimUtils::RetargetAnimations(OldSkeleton, NewSkeleton, Selected, bRemapReferencedAssets, bDuplicateAssets, bConvertSpaces);
 
diff --git a/Engine/Source/Editor/UnrealEd/Private/EditorAnimUtils.cpp b/Engine/Source/Editor/UnrealEd/Private/EditorAnimUtils.cpp
--- a/Engine/Source/Editor/UnrealEd/Private/EditorAnimUtils.cpp
+++ b/Engine/Source/Editor/UnrealEd/Private/EditorAnimUtils.cpp
@@ -141,7 +141,7 @@ namespace EditorAnimUtils
 
 	void FAnimationRetargetContext::RetargetAnimations(USkeleton * OldSkeleton, USkeleton* NewSkeleton)
 	{
-		check (OldSkeleton);
+		check (!bConvertAnimationDataInComponentSpaces || OldSkeleton);
 		check (NewSkeleton);
 
 		if (bConvertAnimationDataInComponentSpaces)
diff --git a/Engine/Source/Editor/UnrealEd/Private/SSkeletonWidget.cpp b/Engine/Source/Editor/UnrealEd/Private/SSkeletonWidget.cpp
--- a/Engine/Source/Editor/UnrealEd/Private/SSkeletonWidget.cpp
+++ b/Engine/Source/Editor/UnrealEd/Private/SSkeletonWidget.cpp
@@ -355,18 +355,21 @@ void SAnimationRemapSkeleton::Construct( const FArguments& InArgs )
 		];
 	}
 
-	TSharedPtr<SToolTip> ConvertSpaceTooltip = IDocumentation::Get()->CreateToolTip(FText::FromString("Check if you'd like to convert animation data to new skeleton space. If this is false, it won't convert any animation data to new space."), 
-											NULL, FString("Shared/Editors/Persona"), FString("AnimRemapSkeleton_ConvertSpace"));
-	Widget->AddSlot()
-	[
-		SNew(SCheckBox)
-		.IsChecked(this, &SAnimationRemapSkeleton::IsConvertSpacesChecked)
-		.OnCheckStateChanged(this, &SAnimationRemapSkeleton::OnConvertSpacesCheckChanged)
-		[
-			SNew(STextBlock).Text(LOCTEXT("RemapSkeleton_ConvertSpaces", "Convert Spaces to new Skeleton"))
-			.ToolTip(ConvertSpaceTooltip)
-		]
-	];
+	if (InArgs._ShowConvertSpacesOption)
+	{
+		TSharedPtr<SToolTip> ConvertSpaceTooltip = IDocumentation::Get()->CreateToolTip(FText::FromString("Check if you'd like to convert animation data to new skeleton space. If this is false, it won't convert any animation data to new space."),
+			NULL, FString("Shared/Editors/Persona"), FString("AnimRemapSkeleton_ConvertSpace"));
+		Widget->AddSlot()
+			[
+				SNew(SCheckBox)
+				.IsChecked(this, &SAnimationRemapSkeleton::IsConvertSpacesChecked)
+				.OnCheckStateChanged(this, &SAnimationRemapSkeleton::OnConvertSpacesCheckChanged)
+				[
+					SNew(STextBlock).Text(LOCTEXT("RemapSkeleton_ConvertSpaces", "Convert Spaces to new Skeleton"))
+					.ToolTip(ConvertSpaceTooltip)
+				]
+			];
+	}
 
 	TSharedPtr<SToolTip> SkeletonTooltip = IDocumentation::Get()->CreateToolTip(FText::FromString("Pick a skeleton for this mesh"), NULL, FString("Shared/Editors/Persona"), FString("Skeleton"));
 
@@ -542,7 +545,7 @@ void SAnimationRemapSkeleton::CloseWindow()
 	}
 }
 
-bool SAnimationRemapSkeleton::ShowModal(USkeleton * OldSkeleton, USkeleton * & NewSkeleton, const FText& WarningMessage, bool & bConvertSpace, bool * bRemapReferencedAssets)
+bool SAnimationRemapSkeleton::ShowModal(USkeleton * OldSkeleton, USkeleton * & NewSkeleton, const FText& WarningMessage, bool * bConvertSpace, bool * bRemapReferencedAssets)
 {
 	TSharedPtr<class SAnimationRemapSkeleton> DialogWidget;
 
@@ -561,6 +564,7 @@ bool SAnimationRemapSkeleton::ShowModal(USkeleton * OldSkeleton, USkeleton * & N
 			.WidgetWindow(DialogWindow)
 			.WarningMessage(WarningMessage)
 			.ShowRemapOption(bRemapReferencedAssets != NULL)
+			.ShowConvertSpacesOption(bConvertSpace != NULL)
 		];
 
 	DialogWindow->SetContent(DialogWrapper.ToSharedRef());
@@ -574,7 +578,10 @@ bool SAnimationRemapSkeleton::ShowModal(USkeleton * OldSkeleton, USkeleton * & N
 		*bRemapReferencedAssets = DialogWidget.Get()->bRemapReferencedAssets;
 	}
 
-	bConvertSpace = DialogWidget.Get()->bConvertSpaces;
+	if (bConvertSpace)
+	{
+		*bConvertSpace = DialogWidget.Get()->bConvertSpaces;
+	}
 
 	return (NewSkeleton != NULL && NewSkeleton != OldSkeleton);
 }
diff --git a/Engine/Source/Editor/UnrealEd/Public/SSkeletonWidget.h b/Engine/Source/Editor/UnrealEd/Public/SSkeletonWidget.h
--- a/Engine/Source/Editor/UnrealEd/Public/SSkeletonWidget.h
+++ b/Engine/Source/Editor/UnrealEd/Public/SSkeletonWidget.h
@@ -270,6 +270,7 @@ public:
 		SLATE_ARGUMENT( TSharedPtr<SWindow>, WidgetWindow )
 		SLATE_ARGUMENT( FText, WarningMessage )
 		SLATE_ARGUMENT( bool, ShowRemapOption )
+		SLATE_ARGUMENT( bool, ShowConvertSpacesOption )
 
 	SLATE_END_ARGS()	
 
@@ -333,7 +334,7 @@ public:
 	 *
 	 * @return true if successfully selected new skeleton
 	 */
-	static UNREALED_API bool ShowModal(USkeleton * OldSkeleton, USkeleton * & NewSkeleton, const FText& WarningMessage, bool & bConvertSpace, bool * bRemapReferencedAssets=NULL);
+	static UNREALED_API bool ShowModal(USkeleton * OldSkeleton, USkeleton * & NewSkeleton, const FText& WarningMessage, bool * bConvertSpace = NULL, bool * bRemapReferencedAssets=NULL);
 };
 
 ////////////////////////////////////////////////////
