AutomationEditorCommon.cpp
Improved the speed of the class check in the function that iterates through the list of known objects.

diff --git a/Engine/Source/Editor/UnrealEd/Private/Tests/AutomationEditorCommon.cpp b/Engine/Source/Editor/UnrealEd/Private/Tests/AutomationEditorCommon.cpp
--- a/Engine/Source/Editor/UnrealEd/Private/Tests/AutomationEditorCommon.cpp
+++ b/Engine/Source/Editor/UnrealEd/Private/Tests/AutomationEditorCommon.cpp
@@ -486,50 +486,54 @@ void FEditorAutomationTestUtilities::CollectMiscGameContentTestsByClass(TArray<F
 	AssetFilter.PackagePaths.Add("/Game");
 	AssetFilter.bRecursivePaths = true;
 	AssetRegistryModule.Get().GetAssets(AssetFilter, ObjectList);
+	
+	//if (ObjectList.Num() < 25000)
+	//{
 
-	//Loop through the list of assets, make their path full and a string, then add them to the test.
-	for (auto ObjIter = ObjectList.CreateConstIterator(); ObjIter; ++ObjIter)
-	{
-		const FAssetData & Asset = *ObjIter;
-		//First we check if the class is valid.  If not then we move onto the next object.
-		if (Asset.GetClass()->IsValidLowLevel())
+		//Loop through the list of assets, make their path full and a string, then add them to the test.
+		for (auto ObjIter = ObjectList.CreateConstIterator(); ObjIter; ++ObjIter)
 		{
-			//Variable that holds the class FName for the current asset iteration.
-			FName AssetClassFName = Asset.GetClass()->GetFName();
-
-			//Counter used to keep track for the following for loop.
-			float ExcludedClassesCounter = 1;
-
-			for (auto ExcludeIter = ExcludeClassesList.CreateConstIterator(); ExcludeIter; ++ExcludeIter)
+			const FAssetData & Asset = *ObjIter;
+			//First we check if the class is valid.  If not then we move onto the next object.
+			if (Asset.GetClass() != NULL)
 			{
-				FName ExludedName = *ExcludeIter;
+				//Variable that holds the class FName for the current asset iteration.
+				FName AssetClassFName = Asset.GetClass()->GetFName();
 
-				//If the classes are the same then we don't want this asset. So we move onto the next one instead.
-				if (AssetClassFName == ExludedName)
-				{
-					break;
-				}
+				//Counter used to keep track for the following for loop.
+				float ExcludedClassesCounter = 1;
 
-				//We run out of class names in our Excluded list then we want the current ObjectList asset.
-				if ((ExcludedClassesCounter + 1) > ExcludeClassesList.Num())
+				for (auto ExcludeIter = ExcludeClassesList.CreateConstIterator(); ExcludeIter; ++ExcludeIter)
 				{
-					FString Filename = Asset.ObjectPath.ToString();
-					//convert to full paths
-					Filename = FPackageName::LongPackageNameToFilename(Filename);
+					FName ExludedName = *ExcludeIter;
 
-					if (FAutomationTestFramework::GetInstance().ShouldTestContent(Filename))
+					//If the classes are the same then we don't want this asset. So we move onto the next one instead.
+					if (AssetClassFName == ExludedName)
 					{
-						FString BeautifiedFilename = Asset.AssetName.ToString();
-						OutBeautifiedNames.Add(BeautifiedFilename);
-						OutTestCommands.Add(Asset.ObjectPath.ToString());
+						break;
 					}
 
-					break;
-				}
+					//We run out of class names in our Excluded list then we want the current ObjectList asset.
+					if ((ExcludedClassesCounter + 1) > ExcludeClassesList.Num())
+					{
+						FString Filename = Asset.ObjectPath.ToString();
+						//convert to full paths
+						Filename = FPackageName::LongPackageNameToFilename(Filename);
+
+						if (FAutomationTestFramework::GetInstance().ShouldTestContent(Filename))
+						{
+							FString BeautifiedFilename = Asset.AssetName.ToString();
+							OutBeautifiedNames.Add(BeautifiedFilename);
+							OutTestCommands.Add(Asset.ObjectPath.ToString());
+						}
+
+						break;
+					}
 
-				//increment the counter.
-				ExcludedClassesCounter++;
+					//increment the counter.
+					ExcludedClassesCounter++;
+				}
 			}
 		}
-	}
+	//}
 }
\ No newline at end of file
