Analytics: Tracks the number of times assets of each class type are saved in a given session

diff --git a/Engine/Source/Editor/UnrealEd/Private/UnrealEdMisc.cpp b/Engine/Source/Editor/UnrealEd/Private/UnrealEdMisc.cpp
--- a/Engine/Source/Editor/UnrealEd/Private/UnrealEdMisc.cpp
+++ b/Engine/Source/Editor/UnrealEd/Private/UnrealEdMisc.cpp
@@ -84,6 +84,9 @@ void FUnrealEdMisc::OnInit()
 	FEditorSupportDelegates::RedrawAllViewports.AddRaw(this, &FUnrealEdMisc::CB_RedrawAllViewports);
 	GEngine->OnLevelActorAdded().AddRaw( this, &FUnrealEdMisc::CB_LevelActorsAdded );
 
+	FCoreDelegates::OnObjectSaved.AddRaw(this, &FUnrealEdMisc::OnObjectSaved);
+	FEditorDelegates::PreSaveWorld.AddRaw(this, &FUnrealEdMisc::OnWorldSaved);
+
 #if USE_UNIT_TESTS
 	FAutomationTestFramework::GetInstance().PreTestingEvent.AddRaw(this, &FUnrealEdMisc::CB_PreAutomationTesting);
 	FAutomationTestFramework::GetInstance().PostTestingEvent.AddRaw(this, &FUnrealEdMisc::CB_PostAutomationTesting);
@@ -533,6 +536,14 @@ void FUnrealEdMisc::OnExit()
 		
 		FSlateApplication::Get().GetPlatformApplication()->SendAnalytics(&FEngineAnalytics::GetProvider());
 
+		// Report asset updates (to reflect forward progress made by the user)
+		TArray<FAnalyticsEventAttribute> AssetUpdateCountAttribs;
+		for (auto& UpdatedAssetPair : NumUpdatesByAssetName)
+		{
+			AssetUpdateCountAttribs.Add(FAnalyticsEventAttribute(UpdatedAssetPair.Key.ToString(), UpdatedAssetPair.Value));
+		}
+		FEngineAnalytics::GetProvider().RecordEvent(FString("Editor.Usage.AssetsSaved"), AssetUpdateCountAttribs);
+
 		FEditorViewportStats::SendUsageData();
 	}
 
@@ -607,7 +618,7 @@ void FUnrealEdMisc::OnExit()
 		if( !Handle.IsValid() )
 		{
 			// We were not able to spawn the new project exe.
-			// Its likely that the exe doesnt exist.
+			// Its likely that the exe doesn't exist.
 			// Skip shutting down the editor if this happens
 			UE_LOG(LogUnrealEdMisc, Warning, TEXT("Could not restart the editor") );
 
@@ -1132,6 +1143,33 @@ void FUnrealEdMisc::OnGotoAsset(const FString& InAssetPath) const
 	}	
 }
 
+void FUnrealEdMisc::OnObjectSaved(UObject* SavedObject)
+{
+	// Ensure the saved object is a non-UWorld asset (UWorlds are handled separately)
+	if (!SavedObject->IsA<UWorld>() && SavedObject->IsAsset())
+	{
+		LogAssetUpdate(SavedObject);
+	}
+}
+
+void FUnrealEdMisc::OnWorldSaved(uint32 SaveFlags, UWorld* SavedWorld)
+{
+	LogAssetUpdate(SavedWorld);
+}
+
+void FUnrealEdMisc::LogAssetUpdate(UObject* UpdatedAsset)
+{
+	UPackage* AssetPackage = UpdatedAsset->GetOutermost();
+	const bool bIsPIESave = AssetPackage->RootPackageHasAnyFlags(PKG_PlayInEditor);
+	const bool bIsAutosave = GUnrealEd->GetPackageAutoSaver().IsAutoSaving();
+
+	if (!bIsPIESave && !bIsAutosave && !GIsAutomationTesting)
+	{
+		uint32& NumUpdates = NumUpdatesByAssetName.FindOrAdd(UpdatedAsset->GetClass()->GetFName());
+		NumUpdates++;
+	}
+}
+
 void FUnrealEdMisc::SwitchProject(const FString& GameOrProjectFileName, bool bWarn)
 {
 	if (GUnrealEd->WarnIfLightingBuildIsCurrentlyRunning())
diff --git a/Engine/Source/Editor/UnrealEd/Public/UnrealEdMisc.h b/Engine/Source/Editor/UnrealEd/Public/UnrealEdMisc.h
--- a/Engine/Source/Editor/UnrealEd/Public/UnrealEdMisc.h
+++ b/Engine/Source/Editor/UnrealEd/Public/UnrealEdMisc.h
@@ -260,6 +260,15 @@ private:
 	/** Delegate used to go to assets in the content browser */
 	void OnGotoAsset(const FString& InAssetPath) const;
 
+	/** Delegate used to update the map of asset update counts */
+	void OnObjectSaved(UObject* SavedObject);
+
+	/** Delegate used to update the map of asset update counts (for UWorlds specifically) */
+	void OnWorldSaved(uint32 SaveFlags, UWorld* SavedWorld);
+
+	/** Logs an update to an asset */
+	void LogAssetUpdate(UObject* UpdatedAsset);
+
 	/** Initialize engine analytics */
 	void InitEngineAnalytics();
 
@@ -303,4 +312,7 @@ private:
 
 	/** handler to notify about navigation building process */
 	TSharedPtr<FTickableEditorObject> NavigationBuildingNotificationHandler;
+
+	/** Package names and the number of times they have been updated */
+	TMap<FName, uint32> NumUpdatesByAssetName;
 };
diff --git a/Engine/Source/Runtime/Core/Private/Misc/CallbackDevice.cpp b/Engine/Source/Runtime/Core/Private/Misc/CallbackDevice.cpp
--- a/Engine/Source/Runtime/Core/Private/Misc/CallbackDevice.cpp
+++ b/Engine/Source/Runtime/Core/Private/Misc/CallbackDevice.cpp
@@ -33,6 +33,7 @@ FCoreDelegates::FOnActorLabelChanged FCoreDelegates::OnActorLabelChanged;
 #if WITH_EDITOR
 	FCoreDelegates::FOnObjectModified  FCoreDelegates::OnObjectModified;
 	FCoreDelegates::FOnAssetLoaded FCoreDelegates::OnAssetLoaded;
+	FCoreDelegates::FOnObjectSaved FCoreDelegates::OnObjectSaved;
 	FSimpleMulticastDelegate FCoreDelegates::PreModal;
 	FSimpleMulticastDelegate FCoreDelegates::PostModal;
 #endif	//WITH_EDITOR
diff --git a/Engine/Source/Runtime/Core/Public/Misc/CallbackDevice.h b/Engine/Source/Runtime/Core/Public/Misc/CallbackDevice.h
--- a/Engine/Source/Runtime/Core/Public/Misc/CallbackDevice.h
+++ b/Engine/Source/Runtime/Core/Public/Misc/CallbackDevice.h
@@ -59,6 +59,9 @@ public:
 
 	// Callback for when an asset is loaded (Editor)
 	DECLARE_MULTICAST_DELEGATE_OneParam(FOnAssetLoaded, UObject*);
+	
+	// Callback for when an asset is saved (Editor)
+	DECLARE_MULTICAST_DELEGATE_OneParam(FOnObjectSaved, UObject*);
 	#endif	//WITH_EDITOR
 
 	// delegate type for prompting the pak system to mount a new pak
@@ -132,6 +135,9 @@ public:
 	// Called when an asset is loaded
 	static FOnAssetLoaded OnAssetLoaded;
 
+	// Called when an asset is saved
+	static FOnObjectSaved OnObjectSaved;
+
 	// Called before the editor displays a modal window, allowing other windows the opportunity to disable themselves to avoid reentrant calls
 	static FSimpleMulticastDelegate PreModal;
 
@@ -142,7 +148,7 @@ public:
 	// Called when SetCurrentCulture is called.
 	static FSimpleMulticastDelegate OnCultureChanged;
 
-	// Called when an error occured.
+	// Called when an error occurred.
 	static FSimpleMulticastDelegate OnShutdownAfterError;
 
 	// Called when appInit is called.
@@ -181,7 +187,7 @@ public:
 	/** requests to open a message box */
 	static FOnModalMessageBox ModalErrorMessage;
 
-	/** querries whether an object should be loaded on top ( replace ) an already existing one */
+	/** Queries whether an object should be loaded on top ( replace ) an already existing one */
 	static FOnLoadObjectsOnTop ShouldLoadOnTop;
 
 	/** called when loading a string asset reference */
diff --git a/Engine/Source/Runtime/CoreUObject/Private/UObject/Obj.cpp b/Engine/Source/Runtime/CoreUObject/Private/UObject/Obj.cpp
--- a/Engine/Source/Runtime/CoreUObject/Private/UObject/Obj.cpp
+++ b/Engine/Source/Runtime/CoreUObject/Private/UObject/Obj.cpp
@@ -683,6 +683,13 @@ void UObject::ConditionalPostLoadSubobjects( FObjectInstancingGraph* OuterInstan
 	CheckDefaultSubobjects();
 }
 
+void UObject::PreSave()
+{
+#if WITH_EDITOR
+	FCoreDelegates::OnObjectSaved.Broadcast(this);
+#endif
+}
+
 bool UObject::Modify( bool bAlwaysMarkDirty/*=true*/ )
 {
 	bool bSavedToTransactionBuffer = false;
diff --git a/Engine/Source/Runtime/CoreUObject/Public/UObject/UObject.h b/Engine/Source/Runtime/CoreUObject/Public/UObject/UObject.h
--- a/Engine/Source/Runtime/CoreUObject/Public/UObject/UObject.h
+++ b/Engine/Source/Runtime/CoreUObject/Public/UObject/UObject.h
@@ -92,7 +92,7 @@ public:
 	 *
 	 * @warning: Objects created from within PreSave will NOT have PreSave called on them!!!
 	 */
-	virtual void PreSave() {}
+	virtual void PreSave();
 
 	/**
 	 * Note that the object will be modified.  If we are currently recording into the 
