Keeping invalid blueprint breakpoints from blocking the loading of a blueprint (to accomidate blocked users on answerhub).

diff --git a/Engine/Source/Editor/UnrealEd/Private/Kismet2/KismetDebugUtilities.cpp b/Engine/Source/Editor/UnrealEd/Private/Kismet2/KismetDebugUtilities.cpp
--- a/Engine/Source/Editor/UnrealEd/Private/Kismet2/KismetDebugUtilities.cpp
+++ b/Engine/Source/Editor/UnrealEd/Private/Kismet2/KismetDebugUtilities.cpp
@@ -639,15 +639,28 @@ void FKismetDebugUtilities::GetValidBreakpointLocations(const UK2Node_MacroInsta
 }
 
 // Finds a breakpoint for a given node if it exists, or returns NULL
-UBreakpoint* FKismetDebugUtilities::FindBreakpointForNode(const UBlueprint* Blueprint, const UEdGraphNode* Node, bool bCheckSubLocations)
+UBreakpoint* FKismetDebugUtilities::FindBreakpointForNode(UBlueprint* Blueprint, const UEdGraphNode* Node, bool bCheckSubLocations)
 {
-	for (int32 Index = 0; Index < Blueprint->Breakpoints.Num(); ++Index)
+	// iterate backwards so we can remove invalid breakpoints as we go
+	for (int32 Index = Blueprint->Breakpoints.Num()-1; Index >= 0; --Index)
 	{
 		UBreakpoint* Breakpoint = Blueprint->Breakpoints[Index];
-		check(Breakpoint);
+		if (Breakpoint == nullptr)
+		{
+			Blueprint->Breakpoints.RemoveAtSwap(Index);
+			Blueprint->MarkPackageDirty();
+			UE_LOG(LogBlueprintDebug, Warning, TEXT("Encountered an invalid blueprint breakpoint (this should not happen... if you know how your blueprint got in this state, then please notify the Engine-Blueprints team)"));
+			continue;
+		}
 
 		const UEdGraphNode* BreakpointLocation = Breakpoint->GetLocation();
-		check(BreakpointLocation);
+		if (BreakpointLocation == nullptr)
+		{
+			Blueprint->Breakpoints.RemoveAtSwap(Index);
+			Blueprint->MarkPackageDirty();
+			UE_LOG(LogBlueprintDebug, Warning, TEXT("Encountered a blueprint breakpoint without an associated node (this should not happen... if you know how your blueprint got in this state, then please notify the Engine-Blueprints team)"));
+			continue;
+		}
 
 		// Return this breakpoint if the location matches the given node
 		if (BreakpointLocation == Node)
diff --git a/Engine/Source/Editor/UnrealEd/Public/Kismet2/KismetDebugUtilities.h b/Engine/Source/Editor/UnrealEd/Public/Kismet2/KismetDebugUtilities.h
--- a/Engine/Source/Editor/UnrealEd/Public/Kismet2/KismetDebugUtilities.h
+++ b/Engine/Source/Editor/UnrealEd/Public/Kismet2/KismetDebugUtilities.h
@@ -186,7 +186,7 @@ public:
 	static bool HasDebuggingData(const UBlueprint* Blueprint);
 
 	/** Returns the breakpoint associated with a node, or NULL */
-	static UBreakpoint* FindBreakpointForNode(const UBlueprint* Blueprint, const UEdGraphNode* Node, bool bCheckSubLocations = false);
+	static UBreakpoint* FindBreakpointForNode(UBlueprint* Blueprint, const UEdGraphNode* Node, bool bCheckSubLocations = false);
 
 	/** Deletes all breakpoints in this blueprint */
 	static void ClearBreakpoints(UBlueprint* Blueprint);
