[338826] fixed - Compile Blueprints automation test dirties all blueprints

diff --git a/Engine/Source/Editor/UnrealEd/Private/Tests/BlueprintAutomationTests.cpp b/Engine/Source/Editor/UnrealEd/Private/Tests/BlueprintAutomationTests.cpp
--- a/Engine/Source/Editor/UnrealEd/Private/Tests/BlueprintAutomationTests.cpp
+++ b/Engine/Source/Editor/UnrealEd/Private/Tests/BlueprintAutomationTests.cpp
@@ -464,11 +464,21 @@ public:
 			return false;
 		}
 
+		UPackage* const BlueprintPackage = Cast<UPackage>(BlueprintObj->GetOutermost());
+		// compiling the blueprint will inherently dirty the package, but if there 
+		// weren't any changes to save before, there shouldn't be after
+		bool const bStartedWithUnsavedChanges = (BlueprintPackage != nullptr) ? BlueprintPackage->IsDirty() : true;
+
 		bool bIsRegeneratingOnLoad = false;
 		bool bSkipGarbageCollection = true;
 		FBlueprintEditorUtils::RefreshAllNodes(BlueprintObj);
 		FKismetEditorUtilities::CompileBlueprint(BlueprintObj, bIsRegeneratingOnLoad, bSkipGarbageCollection);
 
+		if (BlueprintPackage != nullptr)
+		{
+			BlueprintPackage->SetDirtyFlag(bStartedWithUnsavedChanges);
+		}
+
 		return true;
 	}
 
