EditorBuildPromotionTests.cpp
--Added a fix so that the build lighting will be skipped when running with NullRHI.

diff --git a/Engine/Source/Editor/UnrealEd/Private/Tests/EditorBuildPromotionTests.cpp b/Engine/Source/Editor/UnrealEd/Private/Tests/EditorBuildPromotionTests.cpp
--- a/Engine/Source/Editor/UnrealEd/Private/Tests/EditorBuildPromotionTests.cpp
+++ b/Engine/Source/Editor/UnrealEd/Private/Tests/EditorBuildPromotionTests.cpp
@@ -55,6 +55,7 @@
 #include "ScopedTransaction.h"
 #include "LevelEditor.h"
 #include "ObjectTools.h"
+#include "RHI.h"
 
 #define LOCTEXT_NAMESPACE "EditorBuildPromotionTests"
 
@@ -902,6 +903,13 @@ namespace EditorBuildPromotionTestUtils
 	*/
 	static void BuildLighting()
 	{
+		//If we are running with -NullRHI then we have to skip this step.
+		if (GUsingNullRHI)
+		{
+			UE_LOG(LogEditorBuildPromotionTests, Warning, TEXT("SKIPPED Build Lighting Step.  You're currently running with -NullRHI."));
+			return;
+		}
+
 		FLightingBuildOptions LightingBuildOptions;
 
 		// Retrieve settings from ini.
