Add map load times logging for the Load Levels test.
EditorAutomationTests.cpp

diff --git a/Engine/Source/Editor/UnrealEd/Private/Tests/EditorAutomationTests.cpp b/Engine/Source/Editor/UnrealEd/Private/Tests/EditorAutomationTests.cpp
--- a/Engine/Source/Editor/UnrealEd/Private/Tests/EditorAutomationTests.cpp
+++ b/Engine/Source/Editor/UnrealEd/Private/Tests/EditorAutomationTests.cpp
@@ -831,6 +831,7 @@ void FLoadAllMapsInEditorTest::GetTests(TArray<FString>& OutBeautifiedNames, TAr
 bool FLoadAllMapsInEditorTest::RunTest(const FString& Parameters)
 {
 	FString MapName = Parameters;
+	double MapLoadStartTime = 0;
 
 	const bool bTakeScreenshots = FAutomationTestFramework::GetInstance().IsScreenshotAllowed();
 	if( bTakeScreenshots )
@@ -854,8 +855,13 @@ bool FLoadAllMapsInEditorTest::RunTest(const FString& Parameters)
 		const FString TestName = FString::Printf(TEXT("LoadAllMaps_Editor/%s"), *FPaths::GetBaseFilename(MapName));
 		AutomationCommon::GetScreenshotPath(TestName, WindowParameters.ScreenshotName, true);
 
+		//Get the current number of seconds.  This will be used to track how long it took to load the map.
+		MapLoadStartTime = FPlatformTime::Seconds();
 		//Load the map
 		FEditorAutomationTestUtilities::LoadMap(MapName);
+		//Log how long it took to launch the map.
+		UE_LOG(LogEditorAutomationTests, Display, TEXT("Map '%s' took %.3f to load"), *MapName, FPlatformTime::Seconds() - MapLoadStartTime);
+
 
 		//If we don't have NoTextureStreaming enabled, give the textures some time to load.
 		if( !FParse::Param( FCommandLine::Get(), TEXT( "NoTextureStreaming" ) ) )
@@ -868,8 +874,12 @@ bool FLoadAllMapsInEditorTest::RunTest(const FString& Parameters)
 	}
 	else
 	{
-		//Load the map
+		//Get the current number of seconds.  This will be used to track how long it took to load the map.
+		MapLoadStartTime = FPlatformTime::Seconds();
+				//Load the map
 		FEditorAutomationTestUtilities::LoadMap(MapName);
+		//Log how long it took to launch the map.
+		UE_LOG(LogEditorAutomationTests, Display, TEXT("Map '%s' took %.3f to load"), *MapName, FPlatformTime::Seconds() - MapLoadStartTime);
 	}
 
 	return true;
