Added an editor notification for when texture streaming is taking a while

TTP# 340822 - Would like a message similar to "Compiling Materials" for "Textures Streaming"

This message works similar to the "Compiling Materials" message, and only appears if more than 30 textures have been streaming for more than 1 second (moving around the viewport doesn't generally stream enough textures to trigger this notice).

This change moves the logic for displaying a global notification after a certain delay into the new FGlobalEditorNotification class. This allows the logic to be shared between FShaderCompilingNotificationImpl and FTextureStreamingNotificationImpl.

Tweaked the notification delay rules to make them less intrusive, and made sure that they had the HitTestInvisible visibility set (this stops them stealing focus from the editor viewport as you were flying around).

diff --git a/Engine/Source/Editor/UnrealEd/Private/GlobalEditorNotification.cpp b/Engine/Source/Editor/UnrealEd/Private/GlobalEditorNotification.cpp
new file mode 100644
--- /dev/null
+++ b/Engine/Source/Editor/UnrealEd/Private/GlobalEditorNotification.cpp
@@ -0,0 +1,94 @@
+// Copyright 1998-2014 Epic Games, Inc. All Rights Reserved.
+
+#include "UnrealEd.h"
+#include "GlobalEditorNotification.h"
+
+TSharedPtr<SNotificationItem> FGlobalEditorNotification::BeginNotification()
+{
+	TSharedPtr<SNotificationItem> NotificationItem = NotificationPtr.Pin();
+	
+	if (NotificationItem.IsValid())
+	{
+		NotificationItem->ExpireAndFadeout();
+		NotificationPtr.Reset();
+	}
+
+	FNotificationInfo Info(FText::GetEmpty());
+	Info.bFireAndForget = false;
+	
+	// Setting fade out and expire time to 0 as the expire message is currently very obnoxious
+	Info.FadeOutDuration = 0.0f;
+	Info.ExpireDuration = 0.0f;
+
+	NotificationItem = FSlateNotificationManager::Get().AddNotification(Info);
+	NotificationPtr = NotificationItem;
+
+	if (NotificationItem.IsValid())
+	{
+		NotificationItem->SetCompletionState(SNotificationItem::CS_Pending);
+		NotificationItem->SetVisibility(EVisibility::HitTestInvisible);
+	}
+
+	return NotificationItem;
+}
+
+void FGlobalEditorNotification::EndNotification()
+{
+	TSharedPtr<SNotificationItem> NotificationItem = NotificationPtr.Pin();
+
+	if (NotificationItem.IsValid())
+	{
+		NotificationItem->SetText(FText::GetEmpty());
+		NotificationItem->SetCompletionState(SNotificationItem::CS_Success);
+		NotificationItem->ExpireAndFadeout();
+
+		NotificationPtr.Reset();
+	}
+}
+
+void FGlobalEditorNotification::Tick(float DeltaTime)
+{
+	TSharedPtr<SNotificationItem> NotificationItem = NotificationPtr.Pin();
+	const bool bIsNotificationActive = NotificationItem.IsValid() && NotificationItem->GetCompletionState() == SNotificationItem::CS_Pending;
+	const bool bShouldShowNotification = ShouldShowNotification(bIsNotificationActive);
+
+	// Trigger a new notification if we should be active and we haven't displayed the notification recently
+	const double TimeNowInSeconds = FPlatformTime::Seconds();
+	if (bShouldShowNotification && !bIsNotificationActive)
+	{
+		if (NextEnableTimeInSeconds == 0.0)
+		{
+			NextEnableTimeInSeconds = TimeNowInSeconds + EnableDelayInSeconds;
+		}
+
+		if (TimeNowInSeconds >= NextEnableTimeInSeconds)
+		{
+			NotificationItem = BeginNotification();
+		}
+	}
+	// Disable the notification when we are no longer streaming textures
+	else if (!bShouldShowNotification)
+	{
+		if (bIsNotificationActive)
+		{
+			EndNotification();
+		}
+
+		NextEnableTimeInSeconds = 0.0;
+	}
+	
+	if (bShouldShowNotification && NotificationItem.IsValid())
+	{
+		SetNotificationText(NotificationItem);
+	}
+}
+
+bool FGlobalEditorNotification::IsTickable() const
+{
+	return true;
+}
+
+TStatId FGlobalEditorNotification::GetStatId() const
+{
+	RETURN_QUICK_DECLARE_CYCLE_STAT(FGlobalEditorNotification, STATGROUP_Tickables);
+}
diff --git a/Engine/Source/Editor/UnrealEd/Private/GlobalEditorNotification.h b/Engine/Source/Editor/UnrealEd/Private/GlobalEditorNotification.h
new file mode 100644
--- /dev/null
+++ b/Engine/Source/Editor/UnrealEd/Private/GlobalEditorNotification.h
@@ -0,0 +1,52 @@
+// Copyright 1998-2014 Epic Games, Inc. All Rights Reserved.
+
+#pragma once
+
+/**
+ * Class used to provide simple global editor notifications (for things like shader compilation and texture streaming) 
+ */
+class FGlobalEditorNotification : public FTickableEditorObject
+{
+
+public:
+	FGlobalEditorNotification(const double InEnableDelayInSeconds = 1.0)
+		: EnableDelayInSeconds(InEnableDelayInSeconds)
+		, NextEnableTimeInSeconds(0)
+	{
+	}
+
+	virtual ~FGlobalEditorNotification()
+	{
+	}
+
+protected:
+	/** 
+	 * Used to work out whether the notification should currently be visible
+	 * (causes BeginNotification, EndNotification, and SetNotificationText to be called at appropriate points) 
+	 */
+	virtual bool ShouldShowNotification(const bool bIsNotificationAlreadyActive) const = 0;
+
+	/** Called to update the text on the given notification */
+	virtual void SetNotificationText(const TSharedPtr<SNotificationItem>& InNotificationItem) const = 0;
+
+private:
+	/** Begin the notification (make it visible, and mark it as pending) */
+	TSharedPtr<SNotificationItem> BeginNotification();
+
+	/** End the notification (hide it, and mark is at complete) */
+	void EndNotification();
+
+	/** FTickableEditorObject interface */
+	virtual void Tick(float DeltaTime) override;
+	virtual bool IsTickable() const override;
+	virtual TStatId GetStatId() const override;
+
+	/** Delay (in seconds) before we should actually show the notification (default is 1) */
+	double EnableDelayInSeconds;
+
+	/** Minimum time before the notification should actually appear (used to avoid spamming), or zero if no next time has been set */
+	double NextEnableTimeInSeconds;
+
+	/** The actual item being shown */
+	TWeakPtr<SNotificationItem> NotificationPtr;
+};
diff --git a/Engine/Source/Editor/UnrealEd/Private/ShaderCompilingNotification.cpp b/Engine/Source/Editor/UnrealEd/Private/ShaderCompilingNotification.cpp
--- a/Engine/Source/Editor/UnrealEd/Private/ShaderCompilingNotification.cpp
+++ b/Engine/Source/Editor/UnrealEd/Private/ShaderCompilingNotification.cpp
@@ -2,118 +2,36 @@
 
 
 #include "UnrealEd.h"
-#include "MainFrame.h"
+#include "GlobalEditorNotification.h"
 #include "ShaderCompiler.h"
 
 /** Notification class for asynchronous shader compiling. */
-class FShaderCompilingNotificationImpl
-	: public FTickableEditorObject
+class FShaderCompilingNotificationImpl : public FGlobalEditorNotification
 {
-
-public:
-
-	FShaderCompilingNotificationImpl()
-	{
-		LastEnableTime = 0;
-	}
-
-	/** Starts the notification. */
-	void ShaderCompileStarted();
-
-	/** Ends the notification. */
-	void ShaderCompileFinished();
-
 protected:
-	/** FTickableEditorObject interface */
-	virtual void Tick(float DeltaTime);
-	virtual bool IsTickable() const
-	{
-		return true;
-	}
-	virtual TStatId GetStatId() const override;
-
-private:
-
-	/** Tracks the last time the notification was started, used to avoid spamming. */
-	double LastEnableTime;
-
-	/** The source code symbol query in progress message */
-	TWeakPtr<SNotificationItem> ShaderCompileNotificationPtr;
+	/** FGlobalEditorNotification interface */
+	virtual bool ShouldShowNotification(const bool bIsNotificationAlreadyActive) const override;
+	virtual void SetNotificationText(const TSharedPtr<SNotificationItem>& InNotificationItem) const override;
 };
 
 /** Global notification object. */
 FShaderCompilingNotificationImpl GShaderCompilingNotification;
 
-void FShaderCompilingNotificationImpl::ShaderCompileStarted()
-{
-	LastEnableTime = FPlatformTime::Seconds();
-
-	// Starting a new request! Notify the UI.
-	if (ShaderCompileNotificationPtr.IsValid())
-	{
-		ShaderCompileNotificationPtr.Pin()->ExpireAndFadeout();
-	}
-	
-	FNotificationInfo Info( NSLOCTEXT("ShaderCompile", "ShaderCompileInProgress", "Compiling Shaders") );
-	Info.bFireAndForget = false;
-	
-	// Setting fade out and expire time to 0 as the expire message is currently very obnoxious
-	Info.FadeOutDuration = 0.0f;
-	Info.ExpireDuration = 0.0f;
-
-	ShaderCompileNotificationPtr = FSlateNotificationManager::Get().AddNotification(Info);
-
-	if (ShaderCompileNotificationPtr.IsValid())
-	{
-		ShaderCompileNotificationPtr.Pin()->SetCompletionState(SNotificationItem::CS_Pending);
-	}
-}
-
-void FShaderCompilingNotificationImpl::ShaderCompileFinished()
+bool FShaderCompilingNotificationImpl::ShouldShowNotification(const bool bIsNotificationAlreadyActive) const
 {
-	// Finished all requests! Notify the UI.
-	TSharedPtr<SNotificationItem> NotificationItem = ShaderCompileNotificationPtr.Pin();
-
-	if (NotificationItem.IsValid())
-	{
-		NotificationItem->SetText( NSLOCTEXT("ShaderCompile", "ShaderCompileComplete", "Shaders finished compiling!") );
-		NotificationItem->SetCompletionState(SNotificationItem::CS_Success);
-		NotificationItem->ExpireAndFadeout();
-
-		ShaderCompileNotificationPtr.Reset();
-	}
+	// ShouldDisplayCompilingNotification is more a hint, and may start returning false while there's still work being done
+	// If we're already showing the notification, we should continue to do so until all the jobs have finished
+	return GShaderCompilingManager->ShouldDisplayCompilingNotification() || (bIsNotificationAlreadyActive && GShaderCompilingManager->IsCompiling());
 }
 
-void FShaderCompilingNotificationImpl::Tick(float DeltaTime)
+void FShaderCompilingNotificationImpl::SetNotificationText(const TSharedPtr<SNotificationItem>& InNotificationItem) const
 {
-	// Trigger a new notification if we are doing an async shader compile, and we haven't displayed the notification recently
-	if (GShaderCompilingManager->ShouldDisplayCompilingNotification() 
-		&& !ShaderCompileNotificationPtr.IsValid()
-		&& (FPlatformTime::Seconds() - LastEnableTime) > 5)
+	if (GShaderCompilingManager->IsCompiling())
 	{
-		ShaderCompileStarted();
-	}
-	// Disable the notification when we are no longer doing an async compile
-	else if (!GShaderCompilingManager->IsCompiling() && ShaderCompileNotificationPtr.IsValid())
-	{
-		ShaderCompileFinished();
-	}
-	else if (GShaderCompilingManager->IsCompiling() && ShaderCompileNotificationPtr.IsValid())
-	{
-		TSharedPtr<SNotificationItem> NotificationItem = ShaderCompileNotificationPtr.Pin();
+		FFormatNamedArguments Args;
+		Args.Add(TEXT("ShaderJobs"), FText::AsNumber(GShaderCompilingManager->GetNumRemainingJobs()));
+		const FText ProgressMessage = FText::Format(NSLOCTEXT("ShaderCompile", "ShaderCompileInProgressFormat", "Compiling Shaders ({ShaderJobs})"), Args);
 
-		if (NotificationItem.IsValid())
-		{
-			FFormatNamedArguments Args;
-			Args.Add( TEXT("ShaderJobs"), FText::AsNumber( GShaderCompilingManager->GetNumRemainingJobs() ) );
-			FText ProgressMessage = FText::Format(NSLOCTEXT("ShaderCompile", "ShaderCompileInProgressFormat", "Compiling Shaders ({ShaderJobs})"), Args);
-
-			NotificationItem->SetText( ProgressMessage );
-		}
+		InNotificationItem->SetText(ProgressMessage);
 	}
 }
-
-TStatId FShaderCompilingNotificationImpl::GetStatId() const
-{
-	RETURN_QUICK_DECLARE_CYCLE_STAT(FShaderCompilingNotificationImpl, STATGROUP_Tickables);
-}
diff --git a/Engine/Source/Editor/UnrealEd/Private/TextureStreamingNotification.cpp b/Engine/Source/Editor/UnrealEd/Private/TextureStreamingNotification.cpp
new file mode 100644
--- /dev/null
+++ b/Engine/Source/Editor/UnrealEd/Private/TextureStreamingNotification.cpp
@@ -0,0 +1,56 @@
+// Copyright 1998-2014 Epic Games, Inc. All Rights Reserved.
+
+
+#include "UnrealEd.h"
+#include "GlobalEditorNotification.h"
+#include "ContentStreaming.h"
+
+/** Notification class for texture streaming. */
+class FTextureStreamingNotificationImpl : public FGlobalEditorNotification
+{
+protected:
+	/** FGlobalEditorNotification interface */
+	virtual bool ShouldShowNotification(const bool bIsNotificationAlreadyActive) const override;
+	virtual void SetNotificationText(const TSharedPtr<SNotificationItem>& InNotificationItem) const override;
+
+private:
+	static int32 GetNumStreamingTextures();
+};
+
+/** Global notification object. */
+FTextureStreamingNotificationImpl GTextureStreamingNotification;
+
+bool FTextureStreamingNotificationImpl::ShouldShowNotification(const bool bIsNotificationAlreadyActive) const
+{
+	// We only want to show the notification initially if there's still enough work left to do to warrant it
+	// However, if we're already showing the notification, we should continue to do so until all the streaming has finished
+	// Never show these notifications during PIE
+	const int32 NumStreamingTextures = GetNumStreamingTextures();
+	return !GEditor->PlayWorld && ((NumStreamingTextures > 30) || (bIsNotificationAlreadyActive && NumStreamingTextures > 0));
+}
+
+void FTextureStreamingNotificationImpl::SetNotificationText(const TSharedPtr<SNotificationItem>& InNotificationItem) const
+{
+	const int32 NumStreamingTextures = GetNumStreamingTextures();
+	if (NumStreamingTextures > 0)
+	{
+		FFormatNamedArguments Args;
+		Args.Add(TEXT("NumTextures"), FText::AsNumber(NumStreamingTextures));
+		const FText ProgressMessage = FText::Format(NSLOCTEXT("StreamingTextures", "StreamingTexturesInProgressFormat", "Streaming Textures ({NumTextures})"), Args);
+
+		InNotificationItem->SetText(ProgressMessage);
+	}
+}
+
+int32 FTextureStreamingNotificationImpl::GetNumStreamingTextures()
+{
+	FStreamingManagerCollection& StreamingManagerCollection = IStreamingManager::Get();
+
+	if (StreamingManagerCollection.IsTextureStreamingEnabled())
+	{
+		ITextureStreamingManager& TextureStreamingManager = StreamingManagerCollection.GetTextureStreamingManager();
+		return TextureStreamingManager.GetNumWantingResources();
+	}
+
+	return 0;
+}
diff --git a/Engine/Source/Runtime/Engine/Public/ContentStreaming.h b/Engine/Source/Runtime/Engine/Public/ContentStreaming.h
--- a/Engine/Source/Runtime/Engine/Public/ContentStreaming.h
+++ b/Engine/Source/Runtime/Engine/Public/ContentStreaming.h
@@ -579,7 +579,7 @@ struct FStreamingManagerCollection : public IStreamingManager
 	/**
 	 * Checks whether any kind of streaming is active
 	 */
-	bool IsStreamingEnabled() const;
+	ENGINE_API bool IsStreamingEnabled() const;
 
 	/**
 	 * Checks whether texture streaming is active
@@ -589,12 +589,12 @@ struct FStreamingManagerCollection : public IStreamingManager
 	/**
 	 * Gets a reference to the Texture Streaming Manager interface
 	 */
-	ITextureStreamingManager& GetTextureStreamingManager() const;
+	ENGINE_API ITextureStreamingManager& GetTextureStreamingManager() const;
 
 	/**
 	 * Gets a reference to the Audio Streaming Manager interface
 	 */
-	IAudioStreamingManager& GetAudioStreamingManager() const;
+	ENGINE_API IAudioStreamingManager& GetAudioStreamingManager() const;
 
 	/**
 	 * Adds a streaming manager to the array of managers to route function calls to.
