[339682] Fix for a copy/paste crash in Components mode in the Blueprint editor.

- Modified FCustomizableTextObjectFactory::ClearObjectNameUsage() to inlcude the 'REN_DontCreateRedirectors' flag in the Rename() call.

diff --git a/Engine/Source/Editor/UnrealEd/Private/Factories/EditorFactories.cpp b/Engine/Source/Editor/UnrealEd/Private/Factories/EditorFactories.cpp
--- a/Engine/Source/Editor/UnrealEd/Private/Factories/EditorFactories.cpp
+++ b/Engine/Source/Editor/UnrealEd/Private/Factories/EditorFactories.cpp
@@ -5283,7 +5283,7 @@ void FCustomizableTextObjectFactory::ClearObjectNameUsage(UObject* InParent, FNa
 	{
 		check(Found->GetOuter() == InParent);
 
-		Found->Rename(nullptr, nullptr, REN_None);
+		Found->Rename(nullptr, nullptr, REN_DontCreateRedirectors);
 	}
 }
 
