Enable "throttle CPU usage when not in foreground" by default

diff --git a/Engine/Config/BaseEditorUserSettings.ini b/Engine/Config/BaseEditorUserSettings.ini
--- a/Engine/Config/BaseEditorUserSettings.ini
+++ b/Engine/Config/BaseEditorUserSettings.ini
@@ -42,7 +42,7 @@ bDisplayMemorySizeDataInLevelBrowser=False
 bShowFrameRateAndMemory=False
 
 ; Lowers CPU usage when the editor is in the background and not the active application
-bThrottleWhenNotForeground=False
+bThrottleCPUWhenNotForeground=True
 
 ; True if we should automatically reimport animset assets when a change to source content is detected
 bAutoReimportAnimSets=False
diff --git a/Engine/Source/Editor/UnrealEd/Classes/Editor/EditorUserSettings.h b/Engine/Source/Editor/UnrealEd/Classes/Editor/EditorUserSettings.h
--- a/Engine/Source/Editor/UnrealEd/Classes/Editor/EditorUserSettings.h
+++ b/Engine/Source/Editor/UnrealEd/Classes/Editor/EditorUserSettings.h
@@ -33,7 +33,7 @@ class UEditorUserSettings : public UObject
 
 	/** Lowers CPU usage when the editor is in the background and not the active application */
 	UPROPERTY(EditAnywhere, config, Category=Performance, meta=(DisplayName="Use Less CPU when in Background") )
-	uint32 bThrottleWhenNotForeground:1;
+	uint32 bThrottleCPUWhenNotForeground:1;
 
 	/** When turned on, the editor will constantly monitor performance and adjust scalability settings for you when performance drops (disabled in debug) */
 	UPROPERTY(EditAnywhere, config, Category=Performance)
diff --git a/Engine/Source/Editor/UnrealEd/Private/Editor.cpp b/Engine/Source/Editor/UnrealEd/Private/Editor.cpp
--- a/Engine/Source/Editor/UnrealEd/Private/Editor.cpp
+++ b/Engine/Source/Editor/UnrealEd/Private/Editor.cpp
@@ -5964,7 +5964,7 @@ bool UEditorEngine::ShouldThrottleCPUUsage() const
 	if( !bIsForeground )
 	{
 		const UEditorUserSettings* Settings = GetDefault<UEditorUserSettings>();
-		bShouldThrottle = Settings->bThrottleWhenNotForeground;
+		bShouldThrottle = Settings->bThrottleCPUWhenNotForeground;
 
 		// Check if we should throttle due to all windows being minimized
 		if ( !bShouldThrottle )
