CurveEditor - Small visual updates.
+ Make the tangent icons smaller and more distinct.
+ Make the tangent handles longer.

diff --git a/Engine/Source/Editor/EditorStyle/Private/SlateEditorStyle.cpp b/Engine/Source/Editor/EditorStyle/Private/SlateEditorStyle.cpp
--- a/Engine/Source/Editor/EditorStyle/Private/SlateEditorStyle.cpp
+++ b/Engine/Source/Editor/EditorStyle/Private/SlateEditorStyle.cpp
@@ -784,7 +784,7 @@ void FSlateEditorStyle::FStyle::SetupGeneralStyles()
 		Set( "CurveEd.CurveKey", new IMAGE_BRUSH( "Common/Key", FVector2D(11.0f, 11.0f) ) );
 		Set( "CurveEd.CurveKeySelected", new IMAGE_BRUSH( "Common/Key", FVector2D(11.0f, 11.0f), SelectionColor ) );
 		Set( "CurveEd.InfoFont", TTF_CORE_FONT( "Fonts/Roboto-Regular", 8 ) );
-		Set( "CurveEd.Tangent", new IMAGE_BRUSH( "Common/Key", FVector2D(11.0f, 11.0f), FLinearColor(0.0f, 0.66f, 0.7f) ) );
+		Set( "CurveEd.Tangent", new IMAGE_BRUSH( "Common/Tangent", FVector2D(7.0f, 7.0f), FLinearColor(0.0f, 0.66f, 0.7f) ) );
 		Set( "CurveEd.TangentColor", FLinearColor(0.0f, 0.66f, 0.7f) );
 	}
 	
diff --git a/Engine/Source/Editor/UnrealEd/Private/SCurveEditor.cpp b/Engine/Source/Editor/UnrealEd/Private/SCurveEditor.cpp
--- a/Engine/Source/Editor/UnrealEd/Private/SCurveEditor.cpp
+++ b/Engine/Source/Editor/UnrealEd/Private/SCurveEditor.cpp
@@ -10,12 +10,13 @@
 #define LOCTEXT_NAMESPACE "SCurveEditor"
 
 const static FVector2D	CONST_KeySize		= FVector2D(11,11);
+const static FVector2D	CONST_TangentSize	= FVector2D(7,7);
 const static FVector2D	CONST_CurveSize		= FVector2D(12,12);
 
 const static float		CONST_FitMargin		= 0.05f;
 const static float		CONST_MinViewRange	= 0.01f;
 const static float		CONST_DefaultZoomRange	= 1.0f;
-const static float      CONST_KeyTangentOffset = 20.0f;
+const static float      CONST_KeyTangentOffset = 60.0f;
 
 
 //////////////////////////////////////////////////////////////////////////
@@ -621,8 +622,10 @@ int32 SCurveEditor::PaintKeys( FRichCurve* Curve, FTrackScaleInfo &ScaleInfo, FS
 		FKeyHandle KeyHandle = It.Key();
 
 		// Work out where it is
-		const float KeyX = ScaleInfo.InputToLocalX( Curve->GetKeyTime(KeyHandle) ) - CONST_KeySize.X/2;
-		const float KeyY = ScaleInfo.OutputToLocalY( Curve->GetKeyValue(KeyHandle) ) - CONST_KeySize.Y/2;
+		FVector2D KeyLocation(
+			ScaleInfo.InputToLocalX(Curve->GetKeyTime(KeyHandle)),
+			ScaleInfo.OutputToLocalY(Curve->GetKeyValue(KeyHandle)));
+		FVector2D KeyIconLocation = KeyLocation - (CONST_KeySize / 2);
 
 		// Get brush
 		bool IsSelected = IsKeySelected(FSelectedCurveKey(Curve,KeyHandle));
@@ -632,7 +635,7 @@ int32 SCurveEditor::PaintKeys( FRichCurve* Curve, FTrackScaleInfo &ScaleInfo, FS
 		FSlateDrawElement::MakeBox(
 			OutDrawElements,
 			LayerToUse,
-			AllottedGeometry.ToPaintGeometry( FVector2D(KeyX, KeyY), CONST_KeySize ),
+			AllottedGeometry.ToPaintGeometry( KeyIconLocation, CONST_KeySize ),
 			KeyBrush,
 			MyClippingRect,
 			DrawEffects,
@@ -642,7 +645,7 @@ int32 SCurveEditor::PaintKeys( FRichCurve* Curve, FTrackScaleInfo &ScaleInfo, FS
 		//Handle drawing the tangent controls for curve
 		if(IsSelected && (Curve->GetKeyInterpMode(KeyHandle) == RCIM_Cubic || LastInterpMode == RCIM_Cubic))
 		{
-			PaintTangent(ScaleInfo, Curve, KeyHandle, KeyX, KeyY, OutDrawElements, LayerId, AllottedGeometry, MyClippingRect, DrawEffects, LayerToUse, InWidgetStyle);
+			PaintTangent(ScaleInfo, Curve, KeyHandle, KeyLocation, OutDrawElements, LayerId, AllottedGeometry, MyClippingRect, DrawEffects, LayerToUse, InWidgetStyle);
 		}
 
 		LastInterpMode = Curve->GetKeyInterpMode(KeyHandle);
@@ -651,26 +654,21 @@ int32 SCurveEditor::PaintKeys( FRichCurve* Curve, FTrackScaleInfo &ScaleInfo, FS
 }
 
 
-void SCurveEditor::PaintTangent( FTrackScaleInfo &ScaleInfo, FRichCurve* Curve, FKeyHandle KeyHandle, const float KeyX, const float KeyY, FSlateWindowElementList &OutDrawElements, int32 LayerId, const FGeometry &AllottedGeometry, const FSlateRect& MyClippingRect, ESlateDrawEffect::Type DrawEffects, int32 LayerToUse, const FWidgetStyle &InWidgetStyle ) const
+void SCurveEditor::PaintTangent( FTrackScaleInfo &ScaleInfo, FRichCurve* Curve, FKeyHandle KeyHandle, FVector2D KeyLocation, FSlateWindowElementList &OutDrawElements, int32 LayerId, const FGeometry &AllottedGeometry, const FSlateRect& MyClippingRect, ESlateDrawEffect::Type DrawEffects, int32 LayerToUse, const FWidgetStyle &InWidgetStyle ) const
 {
-	FVector2D ArriveTangentPos, LeaveTangentPos;
-	GetTangentPoints(ScaleInfo, FSelectedCurveKey(Curve,KeyHandle), ArriveTangentPos, LeaveTangentPos);
+	FVector2D ArriveTangentLocation, LeaveTangentLocation;
+	GetTangentPoints(ScaleInfo, FSelectedCurveKey(Curve,KeyHandle), ArriveTangentLocation, LeaveTangentLocation);
 
-	FVector2D OffsetKeySize(CONST_KeySize.X,CONST_KeySize.Y);
-	OffsetKeySize *= 0.5f;
-
-	ArriveTangentPos -= OffsetKeySize;
-	LeaveTangentPos  -= OffsetKeySize;
+	FVector2D ArriveTangentIconLocation = ArriveTangentLocation - (CONST_TangentSize / 2);
+	FVector2D LeaveTangentIconLocation = LeaveTangentLocation - (CONST_TangentSize / 2);
 
 	const FSlateBrush* TangentBrush = FEditorStyle::GetBrush("CurveEd.Tangent");
 	const FLinearColor TangentColor = FEditorStyle::GetColor("CurveEd.TangentColor");
 
-	FVector2D KeyOffset(CONST_KeySize.X*0.5f, CONST_KeySize.Y*0.5f);
-
 	//Add lines from tangent control point to 'key'
 	TArray<FVector2D> LinePoints;
-	LinePoints.Add(FVector2D(KeyX, KeyY)+KeyOffset );
-	LinePoints.Add(ArriveTangentPos+KeyOffset);
+	LinePoints.Add(KeyLocation);
+	LinePoints.Add(ArriveTangentLocation);
 	FSlateDrawElement::MakeLines(
 		OutDrawElements,
 		LayerId,
@@ -682,8 +680,8 @@ void SCurveEditor::PaintTangent( FTrackScaleInfo &ScaleInfo, FRichCurve* Curve,
 		);
 
 	LinePoints.Empty();
-	LinePoints.Add(FVector2D(KeyX, KeyY)+KeyOffset);
-	LinePoints.Add(LeaveTangentPos+KeyOffset);
+	LinePoints.Add(KeyLocation);
+	LinePoints.Add(LeaveTangentLocation);
 	FSlateDrawElement::MakeLines(
 		OutDrawElements,
 		LayerId,
@@ -698,7 +696,7 @@ void SCurveEditor::PaintTangent( FTrackScaleInfo &ScaleInfo, FRichCurve* Curve,
 	FSlateDrawElement::MakeBox(
 		OutDrawElements,
 		LayerToUse,
-		AllottedGeometry.ToPaintGeometry(ArriveTangentPos, CONST_KeySize ),
+		AllottedGeometry.ToPaintGeometry(ArriveTangentIconLocation, CONST_TangentSize ),
 		TangentBrush,
 		MyClippingRect,
 		DrawEffects,
@@ -708,7 +706,7 @@ void SCurveEditor::PaintTangent( FTrackScaleInfo &ScaleInfo, FRichCurve* Curve,
 	FSlateDrawElement::MakeBox(
 		OutDrawElements,
 		LayerToUse,
-		AllottedGeometry.ToPaintGeometry(LeaveTangentPos, CONST_KeySize ),
+		AllottedGeometry.ToPaintGeometry(LeaveTangentIconLocation, CONST_TangentSize ),
 		TangentBrush,
 		MyClippingRect,
 		DrawEffects,
@@ -2001,8 +1999,8 @@ SCurveEditor::FSelectedTangent SCurveEditor::HitTestCubicTangents( const FGeomet
 
 				if( HitPosition.Y > (Arrive.Y - (0.5f * CONST_CurveSize.Y)) &&
 					HitPosition.Y < (Arrive.Y + (0.5f * CONST_CurveSize.Y)) &&
-					HitPosition.X > (Arrive.X - (0.5f * CONST_KeySize.X)) && 
-					HitPosition.X < (Arrive.X + (0.5f * CONST_KeySize.X)))
+					HitPosition.X > (Arrive.X - (0.5f * CONST_TangentSize.X)) && 
+					HitPosition.X < (Arrive.X + (0.5f * CONST_TangentSize.X)))
 				{
 					Tangent.Key = Key;
 					Tangent.bIsArrival = true;
@@ -2010,8 +2008,8 @@ SCurveEditor::FSelectedTangent SCurveEditor::HitTestCubicTangents( const FGeomet
 				}
 				if( HitPosition.Y > (Leave.Y - (0.5f * CONST_CurveSize.Y)) &&
 					HitPosition.Y < (Leave.Y  + (0.5f * CONST_CurveSize.Y)) &&
-					HitPosition.X > (Leave.X - (0.5f * CONST_KeySize.X)) && 
-					HitPosition.X < (Leave.X + (0.5f * CONST_KeySize.X)))
+					HitPosition.X > (Leave.X - (0.5f * CONST_TangentSize.X)) && 
+					HitPosition.X < (Leave.X + (0.5f * CONST_TangentSize.X)))
 				{
 					Tangent.Key = Key;
 					Tangent.bIsArrival = false;
diff --git a/Engine/Source/Editor/UnrealEd/Public/SCurveEditor.h b/Engine/Source/Editor/UnrealEd/Public/SCurveEditor.h
--- a/Engine/Source/Editor/UnrealEd/Public/SCurveEditor.h
+++ b/Engine/Source/Editor/UnrealEd/Public/SCurveEditor.h
@@ -342,7 +342,7 @@ private:
 	int32 PaintKeys( FRichCurve* Curve, FTrackScaleInfo &ScaleInfo, FSlateWindowElementList &OutDrawElements, int32 LayerId, const FGeometry &AllottedGeometry, const FSlateRect& MyClippingRect, ESlateDrawEffect::Type DrawEffects, const FWidgetStyle &InWidgetStyle ) const;
 
 	/** Paint the tangent for a key with cubic curves */
-	void PaintTangent( FTrackScaleInfo &ScaleInfo, FRichCurve* Curve, FKeyHandle KeyHandle, const float KeyX, const float KeyY, FSlateWindowElementList &OutDrawElements, int32 LayerId, 
+	void PaintTangent( FTrackScaleInfo &ScaleInfo, FRichCurve* Curve, FKeyHandle KeyHandle, FVector2D KeyLocation, FSlateWindowElementList &OutDrawElements, int32 LayerId, 
 					   const FGeometry &AllottedGeometry, const FSlateRect& MyClippingRect, ESlateDrawEffect::Type DrawEffects, int32 LayerToUse, const FWidgetStyle &InWidgetStyle ) const;
 
 	/** Paint Grid lines, these make it easier to visualize relative distance */
