Fixed mutation of SNullWidget in the SceneOutliner

diff --git a/Engine/Source/Editor/SceneOutliner/Private/SceneOutlinerActorInfoColumn.cpp b/Engine/Source/Editor/SceneOutliner/Private/SceneOutlinerActorInfoColumn.cpp
--- a/Engine/Source/Editor/SceneOutliner/Private/SceneOutlinerActorInfoColumn.cpp
+++ b/Engine/Source/Editor/SceneOutliner/Private/SceneOutlinerActorInfoColumn.cpp
@@ -67,7 +67,7 @@ TSharedRef< SWidget > FSceneOutlinerActorInfoColumn::ConstructClassHyperlink( co
 	const AActor* Actor = InActor.Get();
 	if ( Actor == nullptr )
 	{
-		return SNullWidget::NullWidget;
+		return SNew( SSpacer );
 	}
 
 	UClass* ObjectClass = Actor->GetClass();
diff --git a/Engine/Source/Editor/UnrealEd/Private/EditorClassUtils.cpp b/Engine/Source/Editor/UnrealEd/Private/EditorClassUtils.cpp
--- a/Engine/Source/Editor/UnrealEd/Private/EditorClassUtils.cpp
+++ b/Engine/Source/Editor/UnrealEd/Private/EditorClassUtils.cpp
@@ -67,7 +67,7 @@ TSharedRef<SWidget> FEditorClassUtils::GetDocumentationLinkWidget(const UClass*
 
 TSharedRef<SWidget> FEditorClassUtils::GetSourceLink(const UClass* Class, const TWeakObjectPtr<UObject> ObjectWeakPtr)
 {
-	TSharedRef<SWidget> SourceHyperlink = SNullWidget::NullWidget;
+	TSharedRef<SWidget> SourceHyperlink = SNew( SSpacer );
 	UBlueprint* Blueprint = (Class ? Cast<UBlueprint>(Class->ClassGeneratedBy) : nullptr);
 
 	if (Blueprint)
