Macros will allow latent functions to be added and functions will no longer allow macros with latent functions to be added.

#ttp 332475 - BLUEPRINTS: Allow macros to support latent actions

diff --git a/Engine/Source/Editor/BlueprintGraph/Classes/K2Node_EditablePinBase.h b/Engine/Source/Editor/BlueprintGraph/Classes/K2Node_EditablePinBase.h
--- a/Engine/Source/Editor/BlueprintGraph/Classes/K2Node_EditablePinBase.h
+++ b/Engine/Source/Editor/BlueprintGraph/Classes/K2Node_EditablePinBase.h
@@ -69,9 +69,14 @@ struct FKismetUserDeclaredFunctionMetadata
 	UPROPERTY()
 	bool bCallInEditor;
 
+	/** Cached value for whether or not the graph has latent functions, positive for TRUE, zero for FALSE, and INDEX_None for undetermined */
+	UPROPERTY(transient)
+	int8 HasLatentFunctions;
+
 public:
 	FKismetUserDeclaredFunctionMetadata()
 		: InstanceTitleColor(FLinearColor::White)
+		, HasLatentFunctions(INDEX_NONE)
 	{
 	}
 };
diff --git a/Engine/Source/Editor/BlueprintGraph/Private/EdGraphSchema_K2.cpp b/Engine/Source/Editor/BlueprintGraph/Private/EdGraphSchema_K2.cpp
--- a/Engine/Source/Editor/BlueprintGraph/Private/EdGraphSchema_K2.cpp
+++ b/Engine/Source/Editor/BlueprintGraph/Private/EdGraphSchema_K2.cpp
@@ -328,7 +328,7 @@ bool UEdGraphSchema_K2::CanFunctionBeUsedInClass(const UClass* InClass, UFunctio
 
 	if(InDestGraph != NULL)
 	{
-		bLatentFuncs = (GetGraphType(InDestGraph) == GT_Ubergraph);
+		bLatentFuncs = (GetGraphType(InDestGraph) == GT_Ubergraph || (GetGraphType(InDestGraph) == GT_Macro));
 		bIsConstructionScript = IsConstructionScript(InDestGraph);
 	}
 
diff --git a/Engine/Source/Editor/BlueprintGraph/Private/K2ActionMenuBuilder.cpp b/Engine/Source/Editor/BlueprintGraph/Private/K2ActionMenuBuilder.cpp
--- a/Engine/Source/Editor/BlueprintGraph/Private/K2ActionMenuBuilder.cpp
+++ b/Engine/Source/Editor/BlueprintGraph/Private/K2ActionMenuBuilder.cpp
@@ -2419,6 +2419,12 @@ void FK2ActionMenuBuilder::GetMacroTools(FBlueprintPaletteListBuilder& ActionMen
 				//If so, then don't allow adding action menu to instance the macro graph to its own definition.
 				if (InCurrentGraph != MacroGraph)
 				{
+					// Do not allow any macros with latent functions to appear in the list
+					if(InCurrentGraph && InCurrentGraph->GetSchema()->GetGraphType(InCurrentGraph) == GT_Function && FBlueprintEditorUtils::CheckIfGraphHasLatentFunctions(MacroGraph))
+					{
+						continue;
+					}
+
 					UK2Node_Tunnel* FunctionEntryNode = NULL;
 					UK2Node_Tunnel* FunctionResultNode = NULL;
 					bool bIsMacroPure = false;
diff --git a/Engine/Source/Editor/BlueprintGraph/Private/K2Node_BaseAsyncTask.cpp b/Engine/Source/Editor/BlueprintGraph/Private/K2Node_BaseAsyncTask.cpp
--- a/Engine/Source/Editor/BlueprintGraph/Private/K2Node_BaseAsyncTask.cpp
+++ b/Engine/Source/Editor/BlueprintGraph/Private/K2Node_BaseAsyncTask.cpp
@@ -46,7 +46,8 @@ bool UK2Node_BaseAsyncTask::IsCompatibleWithGraph(const UEdGraph* TargetGraph) c
 void UK2Node_BaseAsyncTask::GetMenuEntries(FGraphContextMenuBuilder& ContextMenuBuilder) const
 {
 	const UEdGraphSchema_K2* K2Schema = GetDefault<UEdGraphSchema_K2>();
-	const bool bAllowLatentFuncs = (K2Schema->GetGraphType(ContextMenuBuilder.CurrentGraph) == GT_Ubergraph);
+	EGraphType GraphType = K2Schema->GetGraphType(ContextMenuBuilder.CurrentGraph);
+	const bool bAllowLatentFuncs = (GraphType == GT_Ubergraph || GraphType == GT_Macro);
 
 	if (bAllowLatentFuncs)
 {
diff --git a/Engine/Source/Editor/Kismet/Private/BPFunctionDragDropAction.cpp b/Engine/Source/Editor/Kismet/Private/BPFunctionDragDropAction.cpp
--- a/Engine/Source/Editor/Kismet/Private/BPFunctionDragDropAction.cpp
+++ b/Engine/Source/Editor/Kismet/Private/BPFunctionDragDropAction.cpp
@@ -60,9 +60,10 @@ static bool CanFunctionBeDropped(TSharedPtr<FEdGraphSchemaAction> /*DropActionIn
  * @param  HoveredGraphIn	A pointer to the graph that the user currently has the item dragged over.
  * @param  ImpededReasonOut	If this returns false, this will be filled out with a reason to present the user with.
  * @param  MacroGraph		The macro graph that the node would be executing.
+ * @param  bInIsLatentMacro	True if the macro has latent functions in it
  * @return True is the dragged palette item can be dropped where it currently is, false if not.
  */
-static bool CanMacroBeDropped(TSharedPtr<FEdGraphSchemaAction> /*DropActionIn*/, UEdGraph* HoveredGraphIn, FText& ImpededReasonOut, UEdGraph* MacroGraph)
+static bool CanMacroBeDropped(TSharedPtr<FEdGraphSchemaAction> /*DropActionIn*/, UEdGraph* HoveredGraphIn, FText& ImpededReasonOut, UEdGraph* MacroGraph, bool bIsLatentMacro)
 {
 	bool bCanBePlaced = true;
 	if (HoveredGraphIn == NULL)
@@ -80,6 +81,12 @@ static bool CanMacroBeDropped(TSharedPtr<FEdGraphSchemaAction> /*DropActionIn*/,
 		bCanBePlaced = false;
 		ImpededReasonOut = LOCTEXT("CannotRecurseMacro", "Cannot place a macro instance in its own graph");
 	}
+	else if( bIsLatentMacro && HoveredGraphIn->GetSchema()->GetGraphType(HoveredGraphIn) == GT_Function)
+	{
+		bCanBePlaced = false;
+		ImpededReasonOut = LOCTEXT("CannotPlaceLatentMacros", "Cannot place a macro instance with latent functions in function graphs!");
+	}
+
 	return bCanBePlaced;
 }
 
@@ -292,7 +299,10 @@ TSharedRef<FKismetMacroDragDropAction> FKismetMacroDragDropAction::New(
 	Operation->Blueprint = InBlueprint;
 	Operation->AnalyticCallback = AnalyticCallback;
 
-	Operation->CanBeDroppedDelegate = FCanBeDroppedDelegate::CreateStatic(&CanMacroBeDropped, InMacro);
+	// Check to see if the macro has any latent functions in it, some graph types do not allow for latent functions
+	bool bIsLatentMacro = FBlueprintEditorUtils::CheckIfGraphHasLatentFunctions(Operation->Macro);
+
+	Operation->CanBeDroppedDelegate = FCanBeDroppedDelegate::CreateStatic(&CanMacroBeDropped, InMacro, bIsLatentMacro);
 
 	Operation->Construct();
 	return Operation;
diff --git a/Engine/Source/Editor/Kismet/Private/BlueprintEditor.cpp b/Engine/Source/Editor/Kismet/Private/BlueprintEditor.cpp
--- a/Engine/Source/Editor/Kismet/Private/BlueprintEditor.cpp
+++ b/Engine/Source/Editor/Kismet/Private/BlueprintEditor.cpp
@@ -3730,19 +3730,6 @@ bool FBlueprintEditor::CanCollapseSelectionToMacro(TSet<class UEdGraphNode*>& In
 			LogResults.Error(*LOCTEXT("CannotPasteNodeMacro_Error", "@@ cannot be placed in macro graph").ToString(), Node);
 			bCollapseAllowed = false;
 		}
-
-		if (Schema->CanEncapuslateNode(*Node))
-		{
-			if(UK2Node_CallFunction* FunctionNode = Cast<UK2Node_CallFunction>(Node))
-			{
-				// Latent functions are not allowed in macros.
-				if(FunctionNode->IsLatentFunction())
-				{
-					LogResults.Error(*LOCTEXT("LatentFunctionNotAllowed_Error", "Latent functions (@@) are not allowed in macros!").ToString(), FunctionNode);
-					bCollapseAllowed = false;
-				}
-			}
-		}
 	}
 
 	FMessageLog MessageLog("BlueprintLog");
diff --git a/Engine/Source/Editor/OnlineBlueprintSupport/Private/K2Node_LatentOnlineCall.cpp b/Engine/Source/Editor/OnlineBlueprintSupport/Private/K2Node_LatentOnlineCall.cpp
--- a/Engine/Source/Editor/OnlineBlueprintSupport/Private/K2Node_LatentOnlineCall.cpp
+++ b/Engine/Source/Editor/OnlineBlueprintSupport/Private/K2Node_LatentOnlineCall.cpp
@@ -17,7 +17,8 @@ UK2Node_LatentOnlineCall::UK2Node_LatentOnlineCall(const FPostConstructInitializ
 void UK2Node_LatentOnlineCall::GetMenuEntries(FGraphContextMenuBuilder& ContextMenuBuilder) const
 {
 	const UEdGraphSchema_K2* K2Schema = GetDefault<UEdGraphSchema_K2>();
-	const bool bAllowLatentFuncs = (K2Schema->GetGraphType(ContextMenuBuilder.CurrentGraph) == GT_Ubergraph);
+	EGraphType GraphType = K2Schema->GetGraphType(ContextMenuBuilder.CurrentGraph);
+	const bool bAllowLatentFuncs = (GraphType == GT_Ubergraph || GraphType == GT_Macro);
 
 
 	if (bAllowLatentFuncs)
diff --git a/Engine/Source/Editor/UnrealEd/Private/Kismet2/BlueprintEditorUtils.cpp b/Engine/Source/Editor/UnrealEd/Private/Kismet2/BlueprintEditorUtils.cpp
--- a/Engine/Source/Editor/UnrealEd/Private/Kismet2/BlueprintEditorUtils.cpp
+++ b/Engine/Source/Editor/UnrealEd/Private/Kismet2/BlueprintEditorUtils.cpp
@@ -1515,6 +1515,21 @@ void FBlueprintEditorUtils::MarkBlueprintAsStructurallyModified(UBlueprint* Blue
 	Blueprint->bCachedDependenciesUpToDate = false;
 	if (Blueprint->Status != BS_BeingCreated && !Blueprint->bBeingCompiled)
 	{
+		TArray<UEdGraph*> AllGraphs;
+		Blueprint->GetAllGraphs(AllGraphs);
+		for (int32 i = 0; i < AllGraphs.Num(); i++)
+		{
+			TWeakObjectPtr<UK2Node_EditablePinBase> EntryNode;
+			TWeakObjectPtr<UK2Node_EditablePinBase> ResultNode;
+			GetEntryAndResultNodes(AllGraphs[i], EntryNode, ResultNode);
+
+			if(UK2Node_Tunnel* TunnelNode = ExactCast<UK2Node_Tunnel>(EntryNode.Get()))
+			{
+				// Remove data marking graphs as latent, this will be re-cache'd as needed
+				TunnelNode->MetaData.HasLatentFunctions = INDEX_NONE;
+			}
+		}
+
 		{
 			// Invoke the compiler to update the skeleton class definition
 			IKismetCompilerInterface& Compiler = FModuleManager::LoadModuleChecked<IKismetCompilerInterface>(KISMET_COMPILER_MODULENAME);
@@ -1542,6 +1557,21 @@ void FBlueprintEditorUtils::MarkBlueprintAsModified(UBlueprint* Blueprint)
 	Blueprint->bCachedDependenciesUpToDate = false;
 	if (Blueprint->Status != BS_BeingCreated)
 	{
+		TArray<UEdGraph*> AllGraphs;
+		Blueprint->GetAllGraphs(AllGraphs);
+		for (int32 i = 0; i < AllGraphs.Num(); i++)
+		{
+			TWeakObjectPtr<UK2Node_EditablePinBase> EntryNode;
+			TWeakObjectPtr<UK2Node_EditablePinBase> ResultNode;
+			GetEntryAndResultNodes(AllGraphs[i], EntryNode, ResultNode);
+
+			if(UK2Node_Tunnel* TunnelNode = ExactCast<UK2Node_Tunnel>(EntryNode.Get()))
+			{
+				// Remove data marking graphs as latent, this will be re-cache'd as needed
+				TunnelNode->MetaData.HasLatentFunctions = INDEX_NONE;
+			}
+		}
+
 		Blueprint->Status = BS_Dirty;
 		Blueprint->MarkPackageDirty();
 		Blueprint->PostEditChange();
@@ -6263,5 +6293,78 @@ FText FBlueprintEditorUtils::GetGraphDescription(const UEdGraph* InGraph)
 	return LOCTEXT( "NoGraphTooltip", "(None)" );
 }
 
+bool FBlueprintEditorUtils::CheckIfGraphHasLatentFunctions(UEdGraph* InGraph)
+{
+	struct Local
+	{
+		static bool CheckIfGraphHasLatentFunctions(UEdGraph* InGraph, TArray<UEdGraph*>& InspectedGraphList)
+		{
+			TWeakObjectPtr<UK2Node_EditablePinBase> EntryNode;
+			TWeakObjectPtr<UK2Node_EditablePinBase> ResultNode;
+			GetEntryAndResultNodes(InGraph, EntryNode, ResultNode);
+
+			UK2Node_Tunnel* TunnelNode = ExactCast<UK2Node_Tunnel>(EntryNode.Get());
+			if(!TunnelNode)
+			{
+				// No tunnel, no metadata.
+				return false;
+			}
+
+			if(TunnelNode->MetaData.HasLatentFunctions != INDEX_NONE)
+			{
+				return TunnelNode->MetaData.HasLatentFunctions > 0;
+			}
+			else
+			{
+				// Add all graphs to the list of already inspected, this prevents circular inclusion issues.
+				InspectedGraphList.Add(InGraph);
+
+				for( const UEdGraphNode* Node : InGraph->Nodes )
+				{
+					if(const UK2Node_CallFunction* CallFunctionNode = Cast<UK2Node_CallFunction>(Node))
+					{
+						// Check any function call nodes to see if they are latent.
+						if(CallFunctionNode->GetTargetFunction()->HasMetaData(FBlueprintMetadata::MD_Latent))
+						{
+							TunnelNode->MetaData.HasLatentFunctions = 1;
+							return true;
+						}
+					}
+					else if(const UK2Node_MacroInstance* MacroInstanceNode = Cast<UK2Node_MacroInstance>(Node))
+					{
+						// Any macro graphs that haven't already been checked need to be checked for latent function calls
+						if(InspectedGraphList.Find(MacroInstanceNode->GetMacroGraph()) == INDEX_NONE)
+						{
+							if(CheckIfGraphHasLatentFunctions(MacroInstanceNode->GetMacroGraph(), InspectedGraphList))
+							{
+								TunnelNode->MetaData.HasLatentFunctions = 1;
+								return true;
+							}
+						}
+					}
+					else if(const UK2Node_Composite* CompositeNode = Cast<UK2Node_Composite>(Node))
+					{
+						// Any collapsed graphs that haven't already been checked need to be checked for latent function calls
+						if(InspectedGraphList.Find(CompositeNode->BoundGraph) == INDEX_NONE)
+						{
+							if(CheckIfGraphHasLatentFunctions(CompositeNode->BoundGraph, InspectedGraphList))
+							{
+								TunnelNode->MetaData.HasLatentFunctions = 1;
+								return true;
+							}
+						}
+					}
+				}
+
+				TunnelNode->MetaData.HasLatentFunctions = 0;
+				return false;
+			}
+		}
+	};
+
+	TArray<UEdGraph*> InspectedGraphList;
+	return Local::CheckIfGraphHasLatentFunctions(InGraph, InspectedGraphList);
+}
+
 #undef LOCTEXT_NAMESPACE
 
diff --git a/Engine/Source/Editor/UnrealEd/Public/Kismet2/BlueprintEditorUtils.h b/Engine/Source/Editor/UnrealEd/Public/Kismet2/BlueprintEditorUtils.h
--- a/Engine/Source/Editor/UnrealEd/Public/Kismet2/BlueprintEditorUtils.h
+++ b/Engine/Source/Editor/UnrealEd/Public/Kismet2/BlueprintEditorUtils.h
@@ -953,6 +953,10 @@ public:
 	 * @return					The description of the graph
 	 */
 	static FText GetGraphDescription(const UEdGraph* InGraph);
+
+	/** Checks if a graph (or any sub-graphs or referenced graphs) have latent function nodes */
+	static bool CheckIfGraphHasLatentFunctions(UEdGraph* InGraph);
+
 protected:
 	// Removes all NULL graph references from the SubGraphs array and recurses thru the non-NULL ones
 	static void CleanNullGraphReferencesRecursive(UEdGraph* Graph);
