DataTable can use regular UsedDefinedStruct.
Experimnetal code, to base UDS on another struct, was removed.

diff --git a/Engine/Config/BaseEditor.ini b/Engine/Config/BaseEditor.ini
--- a/Engine/Config/BaseEditor.ini
+++ b/Engine/Config/BaseEditor.ini
@@ -262,7 +262,6 @@ SupportWebsiteURL="http://unrealengine.com/support"
 [UserDefinedStructure]
 bUseUserDefinedStructure=true
 bShowDefaultValuePropertyEditor=true
-bCanBeBasedOnTableRow=false
 
 [EpicSurvey]
 Source="CDN"
diff --git a/Engine/Source/Editor/UnrealEd/Classes/Factories/StructureFactory.h b/Engine/Source/Editor/UnrealEd/Classes/Factories/StructureFactory.h
--- a/Engine/Source/Editor/UnrealEd/Classes/Factories/StructureFactory.h
+++ b/Engine/Source/Editor/UnrealEd/Classes/Factories/StructureFactory.h
@@ -9,11 +9,7 @@ class UStructureFactory : public UFactory
 {
 	GENERATED_UCLASS_BODY()
 
-	UPROPERTY()
-	class UScriptStruct* Base;
-
 	// Begin UFactory Interface
-	virtual bool ConfigureProperties() override;
 	virtual UObject* FactoryCreateNew(UClass* Class,UObject* InParent,FName Name,EObjectFlags Flags,UObject* Context,FFeedbackContext* Warn) override;
 	// Begin UFactory Interface
 };
diff --git a/Engine/Source/Editor/UnrealEd/Private/Factories/CSVImportFactory.cpp b/Engine/Source/Editor/UnrealEd/Private/Factories/CSVImportFactory.cpp
--- a/Engine/Source/Editor/UnrealEd/Private/Factories/CSVImportFactory.cpp
+++ b/Engine/Source/Editor/UnrealEd/Private/Factories/CSVImportFactory.cpp
@@ -104,9 +104,10 @@ public:
 			{
 				UScriptStruct* Struct = *It;
 				// If a child of the table row struct base, but not itself
-				if (Struct->IsChildOf(TableRowStruct) 
-					&& (Struct != TableRowStruct) 
-					&& (Struct->GetOutermost() != GetTransientPackage()))
+				const bool bBasedOnTableRowBase = Struct->IsChildOf(TableRowStruct) && (Struct != TableRowStruct);
+				const bool bUDStruct = Struct->IsA<UUserDefinedStruct>();
+				const bool bValidStruct = (Struct->GetOutermost() != GetTransientPackage());
+				if ((bBasedOnTableRowBase || bUDStruct) && bValidStruct)
 				{
 					RowStructs.Add(Struct);
 				}
diff --git a/Engine/Source/Editor/UnrealEd/Private/Factories/EditorFactories.cpp b/Engine/Source/Editor/UnrealEd/Private/Factories/EditorFactories.cpp
--- a/Engine/Source/Editor/UnrealEd/Private/Factories/EditorFactories.cpp
+++ b/Engine/Source/Editor/UnrealEd/Private/Factories/EditorFactories.cpp
@@ -7056,21 +7056,10 @@ UStructureFactory::UStructureFactory(const class FPostConstructInitializePropert
 	bEditAfterNew = true;
 }
 
-bool UStructureFactory::ConfigureProperties()
-{
-	static FBoolConfigValueHelper CanBeBasedOnTableRow(TEXT("UserDefinedStructure"), TEXT("bCanBeBasedOnTableRow"));
-	if (CanBeBasedOnTableRow)
-	{
-		const auto ReturnCode = FMessageDialog::Open(EAppMsgType::YesNo, LOCTEXT("BasedOnTableRowBase", "Is this structure based on TableRowBase?"));
-		Base = (EAppReturnType::Yes == ReturnCode) ? FTableRowBase::StaticStruct() : NULL;
-	}
-	return true;
-}
-
 UObject* UStructureFactory::FactoryCreateNew(UClass* Class, UObject* InParent, FName Name, EObjectFlags Flags, UObject* Context, FFeedbackContext* Warn)
 {
 	ensure(UUserDefinedStruct::StaticClass() == Class);
-	return FStructureEditorUtils::CreateUserDefinedStruct(InParent, Name, Flags, Base);
+	return FStructureEditorUtils::CreateUserDefinedStruct(InParent, Name, Flags);
 }
 
 /*-----------------------------------------------------------------------------
diff --git a/Engine/Source/Editor/UnrealEd/Private/Kismet2/StructureEditorUtils.cpp b/Engine/Source/Editor/UnrealEd/Private/Kismet2/StructureEditorUtils.cpp
--- a/Engine/Source/Editor/UnrealEd/Private/Kismet2/StructureEditorUtils.cpp
+++ b/Engine/Source/Editor/UnrealEd/Private/Kismet2/StructureEditorUtils.cpp
@@ -28,7 +28,7 @@ FStructureEditorUtils::FStructEditorManager& FStructureEditorUtils::FStructEdito
 
 //////////////////////////////////////////////////////////////////////////
 // FStructureEditorUtils
-UUserDefinedStruct* FStructureEditorUtils::CreateUserDefinedStruct(UObject* InParent, FName Name, EObjectFlags Flags, UScriptStruct* BaseStruct)
+UUserDefinedStruct* FStructureEditorUtils::CreateUserDefinedStruct(UObject* InParent, FName Name, EObjectFlags Flags)
 {
 	UUserDefinedStruct* Struct = NULL;
 	
@@ -39,10 +39,6 @@ UUserDefinedStruct* FStructureEditorUtils::CreateUserDefinedStruct(UObject* InPa
 		Struct->EditorData = NewNamedObject<UUserDefinedStructEditorData>(Struct, NAME_None, RF_Transactional);
 		check(Struct->EditorData);
 
-		check(IsPropertBaseStruct(BaseStruct));
-		CastChecked<UUserDefinedStructEditorData>(Struct->EditorData)->NativeBase = BaseStruct;
-		Struct->SetSuperStruct(BaseStruct);
-
 		Struct->SetMetaData(TEXT("BlueprintType"), TEXT("true"));
 		Struct->Bind();
 		Struct->StaticLink(true);
@@ -695,30 +691,13 @@ FGuid FStructureEditorUtils::GetGuidFromPropertyName(const FName Name)
 	return FMemberVariableNameHelper::GetGuidFromName(Name);
 }
 
-bool FStructureEditorUtils::IsPropertBaseStruct(UScriptStruct* Base)
-{
-	return !Base || (FTableRowBase::StaticStruct() == Base);
-}
-
-bool FStructureEditorUtils::SetNativeBase(UUserDefinedStruct* Struct, UScriptStruct* NativeBase)
-{
-	bool bResult = false;
-	if (Struct && IsPropertBaseStruct(NativeBase))
-	{
-		CastChecked<UUserDefinedStructEditorData>(Struct->EditorData)->NativeBase = NativeBase;
-		OnStructureChanged(Struct);
-		bResult = true;
-	}
-	return bResult;
-}
-
 struct FReinstanceDataTableHelper
 {
 	// TODO: shell we cache the dependency?
 	static TArray<UDataTable*> GetTablesDependentOnStruct(UUserDefinedStruct* Struct)
 	{
 		TArray<UDataTable*> Result;
-		if (Struct && (Struct->GetSuperStruct() == FTableRowBase::StaticStruct()))
+		if (Struct)
 		{
 			TArray<UObject*> DataTables;
 			GetObjectsOfClass(UDataTable::StaticClass(), DataTables);
diff --git a/Engine/Source/Editor/UnrealEd/Public/Kismet2/StructureEditorUtils.h b/Engine/Source/Editor/UnrealEd/Public/Kismet2/StructureEditorUtils.h
--- a/Engine/Source/Editor/UnrealEd/Public/Kismet2/StructureEditorUtils.h
+++ b/Engine/Source/Editor/UnrealEd/Public/Kismet2/StructureEditorUtils.h
@@ -44,7 +44,7 @@ public:
 	};
 
 	//STRUCTURE
-	static UUserDefinedStruct* CreateUserDefinedStruct(UObject* InParent, FName Name, EObjectFlags Flags, UScriptStruct* BaseStruct = NULL);
+	static UUserDefinedStruct* CreateUserDefinedStruct(UObject* InParent, FName Name, EObjectFlags Flags);
 
 	static void CompileStructure(UUserDefinedStruct* Struct);
 
@@ -52,10 +52,6 @@ public:
 
 	static bool ChangeTooltip(UUserDefinedStruct* Struct, const FString& InTooltip);
 
-	static bool SetNativeBase(UUserDefinedStruct* Struct, UScriptStruct* NativeBase);
-
-	static bool IsPropertBaseStruct(UScriptStruct* Base);
-
 	//VARIABLE
 	static bool AddVariable(UUserDefinedStruct* Struct, const FEdGraphPinType& VarType);
 
diff --git a/Engine/Source/Runtime/Engine/Classes/Engine/UserDefinedStruct.h b/Engine/Source/Runtime/Engine/Classes/Engine/UserDefinedStruct.h
--- a/Engine/Source/Runtime/Engine/Classes/Engine/UserDefinedStruct.h
+++ b/Engine/Source/Runtime/Engine/Classes/Engine/UserDefinedStruct.h
@@ -46,6 +46,8 @@ public:
 	virtual void PostDuplicate(bool bDuplicateForPIE) override;
 	virtual void GetAssetRegistryTags(TArray<FAssetRegistryTag>& OutTags) const override;
 	// End of UObject interface.
+
+	void InitializeDefaultValue(uint8* StructData) const;
 #endif	// WITH_EDITOR
 
 	virtual void SerializeTaggedProperties(FArchive& Ar, uint8* Data, UStruct* DefaultsStruct, uint8* Defaults, const UObject* BreakRecursionIfFullyLoad=NULL) const override;
diff --git a/Engine/Source/Runtime/Engine/Private/DataTable.cpp b/Engine/Source/Runtime/Engine/Private/DataTable.cpp
--- a/Engine/Source/Runtime/Engine/Private/DataTable.cpp
+++ b/Engine/Source/Runtime/Engine/Private/DataTable.cpp
@@ -3,10 +3,6 @@
 #include "EnginePrivate.h"
 #include "Json.h"
 
-#if WITH_EDITOR
-#include "StructureEditorUtils.h"
-#endif //WITH_EDITOR
-
 DEFINE_LOG_CATEGORY(LogDataTable);
 
 ENGINE_API const FString FDataTableRowHandle::Unknown(TEXT("UNKNOWN"));
@@ -540,7 +536,10 @@ TArray<FString> UDataTable::CreateTableFromCSVString(const FString& InString)
 		// And be sure to call DestroyScriptStruct later
 
 #if WITH_EDITOR
-		FStructureEditorUtils::Fill_MakeStructureDefaultValue(Cast<const UUserDefinedStruct>(RowStruct), RowData);
+		if (auto UDStruct = Cast<const UUserDefinedStruct>(RowStruct))
+		{
+			UDStruct->InitializeDefaultValue(RowData);
+		}
 #endif // WITH_EDITOR
 
 		// Add to row map
diff --git a/Engine/Source/Runtime/Engine/Private/UserDefinedStruct.cpp b/Engine/Source/Runtime/Engine/Private/UserDefinedStruct.cpp
--- a/Engine/Source/Runtime/Engine/Private/UserDefinedStruct.cpp
+++ b/Engine/Source/Runtime/Engine/Private/UserDefinedStruct.cpp
@@ -43,9 +43,6 @@ void UUserDefinedStruct::GetAssetRegistryTags(TArray<FAssetRegistryTag>& OutTags
 	Super::GetAssetRegistryTags(OutTags);
 
 	OutTags.Add(FAssetRegistryTag(TEXT("Tooltip"), FStructureEditorUtils::GetTooltip(this), FAssetRegistryTag::TT_Hidden));
-
-	const FName SuperStructName = SuperStruct ? SuperStruct->GetFName() : NAME_None;
-	OutTags.Add(FAssetRegistryTag(TEXT("BaseStruct"), SuperStructName.ToString(), FAssetRegistryTag::TT_Alphabetical));
 }
 
 UProperty* UUserDefinedStruct::CustomFindProperty(const FName Name) const
@@ -56,6 +53,11 @@ UProperty* UUserDefinedStruct::CustomFindProperty(const FName Name) const
 	return Property;
 }
 
+void UUserDefinedStruct::InitializeDefaultValue(uint8* StructData) const
+{
+	FStructureEditorUtils::Fill_MakeStructureDefaultValue(this, StructData);
+}
+
 #endif	// WITH_EDITOR
 
 void UUserDefinedStruct::SerializeTaggedProperties(FArchive& Ar, uint8* Data, UStruct* DefaultsStruct, uint8* Defaults, const UObject* BreakRecursionIfFullyLoad) const
