Fixed typo

diff --git a/Engine/Source/Editor/UnrealEd/Private/Dialogs/SDeleteAssetsDialog.cpp b/Engine/Source/Editor/UnrealEd/Private/Dialogs/SDeleteAssetsDialog.cpp
--- a/Engine/Source/Editor/UnrealEd/Private/Dialogs/SDeleteAssetsDialog.cpp
+++ b/Engine/Source/Editor/UnrealEd/Private/Dialogs/SDeleteAssetsDialog.cpp
@@ -407,7 +407,7 @@ TSharedRef<SWidget> SDeleteAssetsDialog::BuildCantUseReplaceReferencesWidget()
 {
 	return SNew( STextBlock )
 		.AutoWrapText( true )
-		.Text( LOCTEXT( "ReplaceReferencesNotAvailabeText", "Not all objects are compatable, so Replace References is unavailable." ) );
+		.Text( LOCTEXT( "ReplaceReferencesNotAvailabeText", "Not all objects are compatible, so Replace References is unavailable." ) );
 }
 
 TSharedRef<SWidget> SDeleteAssetsDialog::BuildReplaceReferencesWidget()
