Fixed not being able to import FBX animations into camera anims

The problem was that importing an FBX into a camera anim was adding new groups for the imported tracks if the FBX name differed from the camera anim group name. Camera Anims can only have 1 group so I changed the import code to always import into that group. If there are multiple animations in the FBX, they will import over the top of each other.

This addresses TTP#339544 - LICENSEE: LIVE: CAMERA ANIM: Cannot import a Camera Animation into a Camera Anim Asset

diff --git a/Engine/Source/Editor/UnrealEd/Private/Fbx/FbxMatineeImport.cpp b/Engine/Source/Editor/UnrealEd/Private/Fbx/FbxMatineeImport.cpp
--- a/Engine/Source/Editor/UnrealEd/Private/Fbx/FbxMatineeImport.cpp
+++ b/Engine/Source/Editor/UnrealEd/Private/Fbx/FbxMatineeImport.cpp
@@ -30,6 +30,7 @@
 
 #include "UnrealEd.h"
 #include "Matinee/MatineeActor.h"
+#include "Matinee/MatineeActorCameraAnim.h"
 #include "Matinee/InterpGroup.h"
 #include "Matinee/InterpGroupInst.h"
 #include "Matinee/InterpData.h"
@@ -211,9 +212,21 @@ bool FFbxImporter::ImportMatineeSequence(AMatineeActor* InMatineeActor)
 
 		AActor* Actor = NULL;
 
-		// First check to see if the scene node name matches a Matinee group name
-		UInterpGroupInst* FoundGroupInst =
+		const bool bIsCameraAnim = InMatineeActor->IsA<AMatineeActorCameraAnim>();
+
+		// Find a group instance to import into
+		UInterpGroupInst* FoundGroupInst = nullptr;
+		if ( bIsCameraAnim )
+		{
+			// We can only ever import into the camera anim group
+			FoundGroupInst = InMatineeActor->GroupInst[0];
+		}
+		else
+		{
+			// Check to see if the scene node name matches a Matinee group name
 			InMatineeActor->FindFirstGroupInstByName( FString( Node->GetName() ) );
+		}
+
 		if( FoundGroupInst != NULL )
 		{
 			// OK, we found an actor bound to a Matinee group that matches this scene node name
@@ -253,6 +266,10 @@ bool FFbxImporter::ImportMatineeSequence(AMatineeActor* InMatineeActor)
 			{
 				MatineeGroup = CreateMatineeGroup(InMatineeActor, Actor, FString(Node->GetName()));
 			}
+			else if ( bIsCameraAnim )
+			{
+				MatineeGroup->Group->GroupName = Node->GetName();
+			}
 
 			float TimeLength = ImportMatineeActor(Node, MatineeGroup);
 			InterpLength = FMath::Max(InterpLength, TimeLength);
@@ -273,6 +290,10 @@ bool FFbxImporter::ImportMatineeSequence(AMatineeActor* InMatineeActor)
 				{
 					MatineeGroup = CreateMatineeGroup(InMatineeActor, Actor, FString(Node->GetName()));
 				}
+				else if ( bIsCameraAnim )
+				{
+					MatineeGroup->Group->GroupName = Node->GetName();
+				}
 
 				ImportCamera((ACameraActor*) Actor, MatineeGroup, CameraNode);
 			}
