Curve Editor - Small enhancements
+ Add commands for zoom all, and zoom selected, and add buttons to the curve asset editor toolbar.
+ Add commands for setting interpolation modes and add buttons to the curve asset editor toolbar.
+ Add an "Add Key" menu option to the right click menu when clicking on a curve.

diff --git a/Engine/Source/Editor/CurveAssetEditor/Private/CurveAssetEditor.cpp b/Engine/Source/Editor/CurveAssetEditor/Private/CurveAssetEditor.cpp
--- a/Engine/Source/Editor/CurveAssetEditor/Private/CurveAssetEditor.cpp
+++ b/Engine/Source/Editor/CurveAssetEditor/Private/CurveAssetEditor.cpp
@@ -270,8 +270,21 @@ TSharedPtr<FExtender> FCurveAssetEditor::GetToolbarExtender()
 			{
 				ToolbarBuilder.AddToolBarButton(FRichCurveEditorCommands::Get().ZoomToFitHorizontal);
 				ToolbarBuilder.AddToolBarButton(FRichCurveEditorCommands::Get().ZoomToFitVertical);
+				ToolbarBuilder.AddToolBarButton(FRichCurveEditorCommands::Get().ZoomToFitAll);
+				ToolbarBuilder.AddToolBarButton(FRichCurveEditorCommands::Get().ZoomToFitSelected);
 			}
 			ToolbarBuilder.EndSection();
+
+			ToolbarBuilder.BeginSection("Interpolation");
+			{
+				ToolbarBuilder.AddToolBarButton(FRichCurveEditorCommands::Get().InterpolationCubicAuto);
+				ToolbarBuilder.AddToolBarButton(FRichCurveEditorCommands::Get().InterpolationCubicUser);
+				ToolbarBuilder.AddToolBarButton(FRichCurveEditorCommands::Get().InterpolationCubicBreak);
+				ToolbarBuilder.AddToolBarButton(FRichCurveEditorCommands::Get().InterpolationLinear);
+				ToolbarBuilder.AddToolBarButton(FRichCurveEditorCommands::Get().InterpolationConstant);
+			}
+			ToolbarBuilder.EndSection();
+
 			ToolbarBuilder.BeginSection("Snap");
 			{
 				ToolbarBuilder.AddToolBarButton(FRichCurveEditorCommands::Get().ToggleSnapping);
diff --git a/Engine/Source/Editor/EditorStyle/Private/SlateEditorStyle.cpp b/Engine/Source/Editor/EditorStyle/Private/SlateEditorStyle.cpp
--- a/Engine/Source/Editor/EditorStyle/Private/SlateEditorStyle.cpp
+++ b/Engine/Source/Editor/EditorStyle/Private/SlateEditorStyle.cpp
@@ -3660,6 +3660,9 @@ void FSlateEditorStyle::FStyle::SetupLevelEditorStyle()
 		Set( "LevelEditor.Recompile", new IMAGE_BRUSH( "Icons/icon_compile_40x", Icon40x40 ) );
 		Set( "LevelEditor.Recompile.Small", new IMAGE_BRUSH( "Icons/icon_compile_40x", Icon20x20 ) );
 
+		Set( "LevelEditor.ViewOptions", new IMAGE_BRUSH( "Icons/icon_view_40x", Icon40x40 ) );
+		Set( "LevelEditor.ViewOptions.Small", new IMAGE_BRUSH( "Icons/icon_view_40x", Icon20x20 ) );
+
 		Set( "LevelEditor.GameSettings", new IMAGE_BRUSH( "Icons/icon_game_settings_40x", Icon40x40 ) );
 		Set( "LevelEditor.GameSettings.Small", new IMAGE_BRUSH( "Icons/icon_game_settings_40x", Icon20x20 ) );
 
@@ -5205,7 +5208,16 @@ void FSlateEditorStyle::FStyle::SetupToolkitStyles()
 	{
 		Set("RichCurveEditor.ZoomToFitHorizontal", new IMAGE_BRUSH("Icons/icon_CurveEditor_Horizontal_40x", Icon40x40));
 		Set("RichCurveEditor.ZoomToFitVertical", new IMAGE_BRUSH("Icons/icon_CurveEditor_Vertical_40x", Icon40x40));
+		Set("RichCurveEditor.ZoomToFitAll", new IMAGE_BRUSH("Icons/icon_CurveEditor_All_40x", Icon40x40));
+		Set("RichCurveEditor.ZoomToFitSelected", new IMAGE_BRUSH("Icons/icon_CurveEditor_Selected_40x", Icon40x40));
+
 		Set("RichCurveEditor.ToggleSnapping", new IMAGE_BRUSH("Icons/icon_CurveEditor_ToggleSnap_40x", Icon40x40));
+
+		Set("RichCurveEditor.InterpolationCubicAuto", new IMAGE_BRUSH("Icons/icon_CurveEditor_Auto_40x", Icon40x40));
+		Set("RichCurveEditor.InterpolationCubicUser", new IMAGE_BRUSH("Icons/icon_CurveEditor_User_40x", Icon40x40));
+		Set("RichCurveEditor.InterpolationCubicBreak", new IMAGE_BRUSH("Icons/icon_CurveEditor_Break_40x", Icon40x40));
+		Set("RichCurveEditor.InterpolationLinear", new IMAGE_BRUSH("Icons/icon_CurveEditor_Linear_40x", Icon40x40));
+		Set("RichCurveEditor.InterpolationConstant", new IMAGE_BRUSH("Icons/icon_CurveEditor_Constant_40x", Icon40x40));
 	}
 
 	// PhAT
diff --git a/Engine/Source/Editor/UnrealEd/Private/RichCurveEditorCommands.cpp b/Engine/Source/Editor/UnrealEd/Private/RichCurveEditorCommands.cpp
--- a/Engine/Source/Editor/UnrealEd/Private/RichCurveEditorCommands.cpp
+++ b/Engine/Source/Editor/UnrealEd/Private/RichCurveEditorCommands.cpp
@@ -5,5 +5,14 @@ void FRichCurveEditorCommands::RegisterCommands()
 {
 	UI_COMMAND(ZoomToFitHorizontal, "Fit Horizontal", "Zoom to Fit - Horizontal", EUserInterfaceActionType::Button, FInputGesture());
 	UI_COMMAND(ZoomToFitVertical, "Fit Vertical", "Zoom to Fit - Vertical", EUserInterfaceActionType::Button, FInputGesture());
+	UI_COMMAND(ZoomToFitAll, "Fit All", "Zoom to Fit", EUserInterfaceActionType::Button, FInputGesture());
+	UI_COMMAND(ZoomToFitSelected, "Fit Selected", "Zoom to Fit - Selected", EUserInterfaceActionType::Button, FInputGesture(EKeys::F));
+
 	UI_COMMAND(ToggleSnapping, "Snapping", "Toggle Snapping", EUserInterfaceActionType::ToggleButton, FInputGesture());
+
+	UI_COMMAND(InterpolationConstant, "Constant", "Constant interpolation", EUserInterfaceActionType::ToggleButton, FInputGesture());
+	UI_COMMAND(InterpolationLinear, "Linear", "Linear interpolation", EUserInterfaceActionType::ToggleButton, FInputGesture());
+	UI_COMMAND(InterpolationCubicAuto, "Auto", "Cubic interpolation - Automatic tangents", EUserInterfaceActionType::ToggleButton, FInputGesture());
+	UI_COMMAND(InterpolationCubicUser, "User", "Cubic interpolation - User flat tangents", EUserInterfaceActionType::ToggleButton, FInputGesture());
+	UI_COMMAND(InterpolationCubicBreak, "Break", "Cubic interpolation - User broken tangents", EUserInterfaceActionType::ToggleButton, FInputGesture());
 }
diff --git a/Engine/Source/Editor/UnrealEd/Private/SCurveEditor.cpp b/Engine/Source/Editor/UnrealEd/Private/SCurveEditor.cpp
--- a/Engine/Source/Editor/UnrealEd/Private/SCurveEditor.cpp
+++ b/Engine/Source/Editor/UnrealEd/Private/SCurveEditor.cpp
@@ -81,16 +81,47 @@ void SCurveEditor::Construct(const FArguments& InArgs)
 		FExecuteAction::CreateSP(this, &SCurveEditor::RedoAction));
 
 	Commands->MapAction(FRichCurveEditorCommands::Get().ZoomToFitHorizontal,
-		FExecuteAction::CreateSP(this, &SCurveEditor::ZoomToFitHorizontal));
+		FExecuteAction::CreateSP(this, &SCurveEditor::ZoomToFitHorizontal, false));
 
 	Commands->MapAction(FRichCurveEditorCommands::Get().ZoomToFitVertical,
-		FExecuteAction::CreateSP(this, &SCurveEditor::ZoomToFitVertical));
+		FExecuteAction::CreateSP(this, &SCurveEditor::ZoomToFitVertical, false));
+
+	Commands->MapAction(FRichCurveEditorCommands::Get().ZoomToFitAll,
+		FExecuteAction::CreateSP(this, &SCurveEditor::ZoomToFitAll, false));
+
+	Commands->MapAction(FRichCurveEditorCommands::Get().ZoomToFitSelected,
+		FExecuteAction::CreateSP(this, &SCurveEditor::ZoomToFitAll, true));
 
 	Commands->MapAction(FRichCurveEditorCommands::Get().ToggleSnapping,
 		FExecuteAction::CreateSP(this, &SCurveEditor::ToggleSnapping),
 		FCanExecuteAction(),
 		FIsActionChecked::CreateSP(this, &SCurveEditor::IsSnappingEnabled));
 
+	Commands->MapAction(FRichCurveEditorCommands::Get().InterpolationConstant,
+		FExecuteAction::CreateSP(this, &SCurveEditor::OnSelectInterpolationMode, RCIM_Constant, RCTM_Auto),
+		FCanExecuteAction(),
+		FIsActionChecked::CreateSP(this, &SCurveEditor::IsInterpolationModeSelected, RCIM_Constant, RCTM_Auto));
+
+	Commands->MapAction(FRichCurveEditorCommands::Get().InterpolationLinear,
+		FExecuteAction::CreateSP(this, &SCurveEditor::OnSelectInterpolationMode, RCIM_Linear, RCTM_Auto),
+		FCanExecuteAction(),
+		FIsActionChecked::CreateSP(this, &SCurveEditor::IsInterpolationModeSelected, RCIM_Linear, RCTM_Auto));
+
+	Commands->MapAction(FRichCurveEditorCommands::Get().InterpolationCubicAuto,
+		FExecuteAction::CreateSP(this, &SCurveEditor::OnSelectInterpolationMode, RCIM_Cubic, RCTM_Auto),
+		FCanExecuteAction(),
+		FIsActionChecked::CreateSP(this, &SCurveEditor::IsInterpolationModeSelected, RCIM_Cubic, RCTM_Auto));
+
+	Commands->MapAction(FRichCurveEditorCommands::Get().InterpolationCubicUser,
+		FExecuteAction::CreateSP(this, &SCurveEditor::OnSelectInterpolationMode, RCIM_Cubic, RCTM_User),
+		FCanExecuteAction(),
+		FIsActionChecked::CreateSP(this, &SCurveEditor::IsInterpolationModeSelected, RCIM_Cubic, RCTM_User));
+
+	Commands->MapAction(FRichCurveEditorCommands::Get().InterpolationCubicBreak,
+		FExecuteAction::CreateSP(this, &SCurveEditor::OnSelectInterpolationMode, RCIM_Cubic, RCTM_Break),
+		FCanExecuteAction(),
+		FIsActionChecked::CreateSP(this, &SCurveEditor::IsInterpolationModeSelected, RCIM_Cubic, RCTM_Break));
+
 	SAssignNew(WarningMessageText, SErrorText);
 
 	TSharedRef<SBox> CurveSelector = SNew(SBox)
@@ -304,47 +335,33 @@ void SCurveEditor::PushWarningMenu( FVector2D Position, const FText& Message )
 }
 
 
-void SCurveEditor::PushInterpolationMenu(FVector2D Position)
+void SCurveEditor::PushCurveMenu(const FGeometry& InMyGeometry, FVector2D Position)
 {
-	FText Heading = (SelectedKeys.Num() > 0) ? FText::Format( LOCTEXT("KeyInterpolationMode", "Key Interpolation({0})"), GetInterpolationModeText()) : LOCTEXT("CurveInterpolationMode", "Curve Interpolation");
+	FText Heading = (SelectedKeys.Num() > 0) ? LOCTEXT("KeyInterpolationMode", "Key Interpolation") : LOCTEXT("CurveInterpolationMode", "Curve Interpolation");
 
-	FMenuBuilder MenuBuilder(true, NULL);
-	MenuBuilder.BeginSection("CurveEditorInterpolation", Heading);
-	{
-		MenuBuilder.AddMenuEntry(
-			LOCTEXT("InterpolateLinearMode", "Linear"),
-			LOCTEXT("LinearInterpolationToolTip", "Use Linear interpolation"),
-			FSlateIcon(),
-			FUIAction(FExecuteAction::CreateSP(this, &SCurveEditor::OnSelectInterpolationMode, RCIM_Linear, RCTM_Auto) ,
-			FCanExecuteAction()));
+	FMenuBuilder MenuBuilder(true, Commands.ToSharedRef());
 
-		MenuBuilder.AddMenuEntry(
-			LOCTEXT("InterpolateClampMode", "Clamp"),
-			LOCTEXT("ClampInterpolationToolTip", "Clamp to value"),
-			FSlateIcon(),
-			FUIAction(FExecuteAction::CreateSP(this, &SCurveEditor::OnSelectInterpolationMode, RCIM_Constant, RCTM_Auto) ,
-			FCanExecuteAction()));
-
-		MenuBuilder.AddMenuEntry(
-			LOCTEXT("InterpolateCubicAutoMode", "Cubic-auto"),
-			LOCTEXT("CubicAutoInterpolationToolTip", "Cubic interpolation with automatic tangents"),
-			FSlateIcon(),
-			FUIAction(FExecuteAction::CreateSP(this, &SCurveEditor::OnSelectInterpolationMode, RCIM_Cubic, RCTM_Auto) ,
-			FCanExecuteAction()));
-
-		MenuBuilder.AddMenuEntry(
-			LOCTEXT("InterpolateCubicLockedMode", "Cubic-locked"),
-			LOCTEXT("CubicLockedInterpolationToolTip", "Cubic interpolation, with user controlled tangent"),
-			FSlateIcon(),
-			FUIAction(FExecuteAction::CreateSP(this, &SCurveEditor::OnSelectInterpolationMode, RCIM_Cubic, RCTM_User) ,
-			FCanExecuteAction()));
+	if (SelectedKeys.Num() == 0)
+	{
+		MenuBuilder.BeginSection("EditCurveEditorActions", LOCTEXT("Actions", "Actions"));
+		{
+			FUIAction Action = FUIAction(FExecuteAction::CreateSP(this, &SCurveEditor::AddNewKey, InMyGeometry, Position));
+			MenuBuilder.AddMenuEntry(
+				LOCTEXT("AddKey", "Add Key"),
+				LOCTEXT("AddKey_ToolTip", "Create a new point in the curve.  You can also use (Shift + Left Mouse Button)"),
+				FSlateIcon(),
+				Action);
+		}
+		MenuBuilder.EndSection();
+	}
 
-		MenuBuilder.AddMenuEntry(
-			LOCTEXT("InterpolateCubicBreakMode", "Cubic-break"),
-			LOCTEXT("CubicBreakInterpolationToolTip", "Cubic interpolation with seperate user controlled tangents for arrive and leave"),
-			FSlateIcon(),
-			FUIAction(FExecuteAction::CreateSP(this, &SCurveEditor::OnSelectInterpolationMode, RCIM_Cubic, RCTM_Break) ,
-			FCanExecuteAction()));
+	MenuBuilder.BeginSection("CurveEditorInterpolation", Heading);
+	{
+		MenuBuilder.AddMenuEntry(FRichCurveEditorCommands::Get().InterpolationCubicAuto);
+		MenuBuilder.AddMenuEntry(FRichCurveEditorCommands::Get().InterpolationCubicUser);
+		MenuBuilder.AddMenuEntry(FRichCurveEditorCommands::Get().InterpolationCubicBreak);
+		MenuBuilder.AddMenuEntry(FRichCurveEditorCommands::Get().InterpolationLinear);
+		MenuBuilder.AddMenuEntry(FRichCurveEditorCommands::Get().InterpolationConstant);
 	}
 	MenuBuilder.EndSection(); //CurveEditorInterpolation
 
@@ -879,12 +896,12 @@ void SCurveEditor::SetCurveOwner(FCurveOwnerInterface* InCurveOwner, bool bCanEd
 
 	if( bZoomToFitVertical )
 	{
-		ZoomToFitVertical();
+		ZoomToFitVertical(false);
 	}
 
 	if ( bZoomToFitHorizontal )
 	{
-		ZoomToFitHorizontal();
+		ZoomToFitHorizontal(false);
 	}
 
 	CurveSelectionWidget.Pin()->SetContent(CreateCurveSelectionWidget());
@@ -1305,11 +1322,13 @@ void SCurveEditor::ProcessClick(const FGeometry& InMyGeometry, const FPointerEve
 					EmptySelection();
 					AddToSelection(HitKey);
 				}
-				PushInterpolationMenu(InMouseEvent.GetScreenSpacePosition());
+				PushCurveMenu(InMyGeometry, InMouseEvent.GetScreenSpacePosition());
 			}
 			else if (HitTestCurves(InMyGeometry, InMouseEvent))
 			{
-				PushInterpolationMenu(InMouseEvent.GetScreenSpacePosition());
+				// Empty the selection if a key wasn't clicked so that menu has the correct context.
+				EmptySelection();
+				PushCurveMenu(InMyGeometry, InMouseEvent.GetScreenSpacePosition());
 			}
 			else
 			{
@@ -1613,29 +1632,44 @@ TArray<FRichCurve*> SCurveEditor::GetCurvesToFit() const
 }
 
 
-void SCurveEditor::ZoomToFitHorizontal()
+void SCurveEditor::ZoomToFitHorizontal(bool OnlySelected)
 {
 	TArray<FRichCurve*> CurvesToFit = GetCurvesToFit();
 
 	if(Curves.Num() > 0)
 	{
-		float InMin = FLT_MAX, InMax = -FLT_MAX;
+		float InMin = FLT_MAX;
+		float InMax = -FLT_MAX;
 		int32 TotalKeys = 0;
-		for(auto It = CurvesToFit.CreateConstIterator();It;++It)
+
+		if (OnlySelected)
 		{
-			FRichCurve* Curve = *It;
-			float MinTime, MaxTime;
-			Curve->GetTimeRange(MinTime, MaxTime);
-			InMin = FMath::Min(MinTime, InMin);
-			InMax = FMath::Max(MaxTime, InMax);
-			TotalKeys += Curve->GetNumKeys();
+			for (auto SelectedKey : SelectedKeys)
+			{
+				TotalKeys++;
+				float KeyTime = SelectedKey.Curve->GetKeyTime(SelectedKey.KeyHandle);
+				InMin = FMath::Min(KeyTime, InMin);
+				InMax = FMath::Max(KeyTime, InMax);
+			}
 		}
-		
-		if(TotalKeys > 0)
+		else
+		{
+			for(auto It = CurvesToFit.CreateConstIterator();It;++It)
+			{
+				FRichCurve* Curve = *It;
+				float MinTime, MaxTime;
+				Curve->GetTimeRange(MinTime, MaxTime);
+				InMin = FMath::Min(MinTime, InMin);
+				InMax = FMath::Max(MaxTime, InMax);
+				TotalKeys += Curve->GetNumKeys();
+			}
+		}
+
+		if (TotalKeys > 0)
 		{
 			// Clamp the minimum size
 			float Size = InMax - InMin;
-			if( Size < CONST_MinViewRange )
+			if (Size < CONST_MinViewRange)
 			{
 				InMin -= (0.5f*CONST_MinViewRange);
 				InMax += (0.5f*CONST_MinViewRange);
@@ -1658,7 +1692,7 @@ void SCurveEditor::ZoomToFitHorizontal()
 
 FReply SCurveEditor::ZoomToFitHorizontalClicked()
 {
-	ZoomToFitHorizontal();
+	ZoomToFitHorizontal(false);
 	return FReply::Handled();
 }
 
@@ -1669,23 +1703,39 @@ void SCurveEditor::SetDefaultOutput(const float MinZoomRange)
 	ViewMaxOutput = (ViewMaxOutput + (0.5f*MinZoomRange));
 }
 
-void SCurveEditor::ZoomToFitVertical()
+void SCurveEditor::ZoomToFitVertical(bool OnlySelected)
 {
 	TArray<FRichCurve*> CurvesToFit = GetCurvesToFit();
 
 	if(CurvesToFit.Num() > 0)
 	{
-		float InMin = FLT_MAX, InMax = -FLT_MAX;
+		float InMin = FLT_MAX;
+		float InMax = -FLT_MAX;
 		int32 TotalKeys = 0;
-		for(auto It = CurvesToFit.CreateConstIterator();It;++It)
+
+		if (OnlySelected)
 		{
-			FRichCurve* Curve = *It;
-			float MinVal, MaxVal;
-			Curve->GetValueRange(MinVal, MaxVal);
-			InMin = FMath::Min(MinVal, InMin);
-			InMax = FMath::Max(MaxVal, InMax);
-			TotalKeys += Curve->GetNumKeys();
+			for (auto SelectedKey : SelectedKeys)
+			{
+				TotalKeys++;
+				float KeyValue = SelectedKey.Curve->GetKeyValue(SelectedKey.KeyHandle);
+				InMin = FMath::Min(KeyValue, InMin);
+				InMax = FMath::Max(KeyValue, InMax);
+			}
+		}
+		else
+		{
+			for(auto It = CurvesToFit.CreateConstIterator();It;++It)
+			{
+				FRichCurve* Curve = *It;
+				float MinVal, MaxVal;
+				Curve->GetValueRange(MinVal, MaxVal);
+				InMin = FMath::Min(MinVal, InMin);
+				InMax = FMath::Max(MaxVal, InMax);
+				TotalKeys += Curve->GetNumKeys();
+			}
 		}
+
 		const float MinZoomRange = (TotalKeys > 0 ) ? CONST_MinViewRange: CONST_DefaultZoomRange;
 
 		ViewMinOutput = (InMin);
@@ -1707,10 +1757,21 @@ void SCurveEditor::ZoomToFitVertical()
 
 FReply SCurveEditor::ZoomToFitVerticalClicked()
 {
-	ZoomToFitVertical();
+	ZoomToFitVertical(false);
 	return FReply::Handled();
 }
 
+void SCurveEditor::ZoomToFitAll(bool OnlySelected)
+{
+	if (OnlySelected && SelectedKeys.Num() == 0)
+	{
+		// Don't zoom to selected if there is no selection.
+		return;
+	}
+	ZoomToFitHorizontal(OnlySelected);
+	ZoomToFitVertical(OnlySelected);
+}
+
 void SCurveEditor::ToggleSnapping()
 {
 	bSnappingEnabled = !bSnappingEnabled;
@@ -2091,6 +2152,35 @@ void SCurveEditor::OnSelectInterpolationMode(ERichCurveInterpMode InterpMode, ER
 	CurveOwner->OnCurveChanged();
 }
 
+bool SCurveEditor::IsInterpolationModeSelected(ERichCurveInterpMode InterpMode, ERichCurveTangentMode TangentMode)
+{
+	if (SelectedKeys.Num() > 0)
+	{
+		for (auto SelectedKey : SelectedKeys)
+		{
+			if (SelectedKey.Curve->GetKeyInterpMode(SelectedKey.KeyHandle) != InterpMode || SelectedKey.Curve->GetKeyTangentMode(SelectedKey.KeyHandle) != TangentMode)
+			{
+				return false;
+			}
+		}
+	}
+	else
+	{
+		for (auto SelectedCurve : SelectedCurves)
+		{
+			check(IsValidCurve(SelectedCurve));
+			for (auto KeyHandleIterator(SelectedCurve->GetKeyHandleIterator()); KeyHandleIterator; ++KeyHandleIterator)
+			{
+				if (SelectedCurve->GetKeyInterpMode(KeyHandleIterator->Key) != InterpMode || SelectedCurve->GetKeyTangentMode(KeyHandleIterator->Key) != TangentMode)
+				{
+					return false;
+				}
+			}
+		}
+	}
+	return true;
+}
+
 /* Given a tangent value for a key, calculates the 2D delta vector from that key in curve space */
 static inline FVector2D CalcTangentDir(float Tangent)
 {
@@ -2263,51 +2353,4 @@ FVector2D SCurveEditor::SnapLocation(FVector2D InLocation)
 	return InLocation;
 }
 
-FText SCurveEditor::GetInterpolationModeText() const
-{
-	FString Results;
-
-	const int32 LastValidKeyIndex = SelectedKeys.Num()-1;
-	for (int32 i = 0; i<SelectedKeys.Num(); ++i)
-	{
-		FSelectedCurveKey Key = SelectedKeys[i];
-		switch(Key.Curve->GetKeyInterpMode(Key.KeyHandle))
-		{
-		case RCIM_Linear:
-			{
-				Results += LOCTEXT("InterpolateLinearMode", "Linear").ToString();
-			}break;
-		case RCIM_Constant:
-			{
-				Results += LOCTEXT("InterpolateClampMode", "Clamp").ToString();
-			}break;
-		case RCIM_Cubic:
-			{
-				switch(Key.Curve->GetKey(Key.KeyHandle).TangentMode)
-				{
-				case RCTM_Auto:
-					{
-						Results += LOCTEXT("InterpolateCubicAutoMode", "Cubic-auto").ToString();
-					}break;
-				case RCTM_User:
-					{
-						Results += LOCTEXT("InterpolateCubicLockedMode", "Cubic-locked").ToString();
-					}break;
-				case RCTM_Break:
-					{
-						Results += LOCTEXT("InterpolateCubicBreakMode", "Cubic-break").ToString();
-					}break;
-				}
-			}break;
-		}
-
-		if(i != LastValidKeyIndex)
-		{
-			Results += ", ";
-		}
-	}
-
-	return FText::FromString( Results );
-}
-
 #undef LOCTEXT_NAMESPACE
diff --git a/Engine/Source/Editor/UnrealEd/Public/RichCurveEditorCommands.h b/Engine/Source/Editor/UnrealEd/Public/RichCurveEditorCommands.h
--- a/Engine/Source/Editor/UnrealEd/Public/RichCurveEditorCommands.h
+++ b/Engine/Source/Editor/UnrealEd/Public/RichCurveEditorCommands.h
@@ -19,7 +19,14 @@ public:
 
 	TSharedPtr<FUICommandInfo> ZoomToFitHorizontal;
 	TSharedPtr<FUICommandInfo> ZoomToFitVertical;
+	TSharedPtr<FUICommandInfo> ZoomToFitAll;
+	TSharedPtr<FUICommandInfo> ZoomToFitSelected;
 	TSharedPtr<FUICommandInfo> ToggleSnapping;
+	TSharedPtr<FUICommandInfo> InterpolationConstant;
+	TSharedPtr<FUICommandInfo> InterpolationLinear;
+	TSharedPtr<FUICommandInfo> InterpolationCubicAuto;
+	TSharedPtr<FUICommandInfo> InterpolationCubicUser;
+	TSharedPtr<FUICommandInfo> InterpolationCubicBreak;
 
 public:
 	virtual void RegisterCommands() override;
diff --git a/Engine/Source/Editor/UnrealEd/Public/SCurveEditor.h b/Engine/Source/Editor/UnrealEd/Public/SCurveEditor.h
--- a/Engine/Source/Editor/UnrealEd/Public/SCurveEditor.h
+++ b/Engine/Source/Editor/UnrealEd/Public/SCurveEditor.h
@@ -288,8 +288,9 @@ private:
 	/* Get the curves to will be used during a fit operation */
 	TArray<FRichCurve*> GetCurvesToFit()const;
 
-	void ZoomToFitHorizontal();
-	void ZoomToFitVertical();
+	void ZoomToFitHorizontal(bool OnlySelected);
+	void ZoomToFitVertical(bool OnlySelected);
+	void ZoomToFitAll(bool OnlySelected);
 
 	FReply ZoomToFitHorizontalClicked();
 	FReply ZoomToFitVerticalClicked();
@@ -402,18 +403,17 @@ private:
 	/** Get the tooltip for a curve selection checkbox */
 	FText			GetCurveCheckBoxToolTip(FRichCurve* Curve) const;
 
-	/* Get text representation of selected keys interpolation modes */
-	FText			GetInterpolationModeText() const;
-
 	/** Create ontext Menu for waring menu*/
 	void	PushWarningMenu(FVector2D Position, const FText& Message);
 
 	/** Create context Menu for key interpolation settings*/
-	void	PushInterpolationMenu(FVector2D Position);
+	void	PushCurveMenu(const FGeometry& InMyGeometry, FVector2D Position);
 
 	/** Called when the user selects the interpolation mode */
 	void	OnSelectInterpolationMode(ERichCurveInterpMode InterpMode, ERichCurveTangentMode TangentMode);
 
+	bool IsInterpolationModeSelected(ERichCurveInterpMode InterpMode, ERichCurveTangentMode TangentMode);
+
 	/** Begin a transaction for dragging a key or tangent */
 	void	BeginDragTransaction();
 
diff --git a/Engine/Source/Runtime/Engine/Classes/Curves/CurveBase.h b/Engine/Source/Runtime/Engine/Classes/Curves/CurveBase.h
--- a/Engine/Source/Runtime/Engine/Classes/Curves/CurveBase.h
+++ b/Engine/Source/Runtime/Engine/Classes/Curves/CurveBase.h
@@ -339,9 +339,12 @@ public:
 	/** Set the tangent weight mode of the specified key */
 	void SetKeyTangentWeightMode(FKeyHandle KeyHandle, ERichCurveTangentWeightMode NewTangentWeightMode);
 
-	/** Set the interp mode of the specified key */
+	/** Get the interp mode of the specified key */
 	ERichCurveInterpMode GetKeyInterpMode(FKeyHandle KeyHandle) const;
 
+	/** Get the tangent mode of the specified key */
+	ERichCurveTangentMode GetKeyTangentMode(FKeyHandle KeyHandle) const;
+
 	/** Get range of input time values. Outside this region curve continues constantly the start/end values. */
 	void GetTimeRange(float& MinTime, float& MaxTime) const;
 
diff --git a/Engine/Source/Runtime/Engine/Private/CurveBase.cpp b/Engine/Source/Runtime/Engine/Private/CurveBase.cpp
--- a/Engine/Source/Runtime/Engine/Private/CurveBase.cpp
+++ b/Engine/Source/Runtime/Engine/Private/CurveBase.cpp
@@ -505,6 +505,13 @@ ERichCurveInterpMode FRichCurve::GetKeyInterpMode(FKeyHandle KeyHandle) const
 	return GetKey(KeyHandle).InterpMode;
 }
 
+ERichCurveTangentMode FRichCurve::GetKeyTangentMode(FKeyHandle KeyHandle) const
+{
+	if (!IsKeyHandleValid(KeyHandle)) { return RCTM_Auto; }
+
+	return GetKey(KeyHandle).TangentMode;
+}
+
 void FRichCurve::GetTimeRange(float& MinTime, float& MaxTime) const
 {
 	if(Keys.Num() == 0)
