UE4: You can now choose to open the same world in the level editor as the one that is currently there. It will reload the asset, just like how it used to. TTP#344522

diff --git a/Engine/Source/Developer/AssetTools/Private/AssetTypeActions/AssetTypeActions_World.cpp b/Engine/Source/Developer/AssetTools/Private/AssetTypeActions/AssetTypeActions_World.cpp
--- a/Engine/Source/Developer/AssetTools/Private/AssetTypeActions/AssetTypeActions_World.cpp
+++ b/Engine/Source/Developer/AssetTools/Private/AssetTypeActions/AssetTypeActions_World.cpp
@@ -9,7 +9,7 @@ void FAssetTypeActions_World::OpenAssetEditor( const TArray<UObject*>& InObjects
 	for (auto ObjIt = InObjects.CreateConstIterator(); ObjIt; ++ObjIt)
 	{
 		auto World = Cast<UWorld>(*ObjIt);
-		if (World != NULL && World != GEditor->GetEditorWorldContext().World())
+		if (World != NULL)
 		{
 			// If there are any unsaved changes to the current level, see if the user wants to save those first.
 			bool bPromptUserToSave = true;
diff --git a/Engine/Source/Editor/UnrealEd/Private/EditorServer.cpp b/Engine/Source/Editor/UnrealEd/Private/EditorServer.cpp
--- a/Engine/Source/Editor/UnrealEd/Private/EditorServer.cpp
+++ b/Engine/Source/Editor/UnrealEd/Private/EditorServer.cpp
@@ -2034,12 +2034,6 @@ bool UEditorEngine::Map_Load(const TCHAR* Str, FOutputDevice& Ar)
 		FString LongTempFname;
 		if ( FPackageName::TryConvertFilenameToLongPackageName(TempFname, LongTempFname) )
 		{
- 			if ( Context.World() && Context.World()->GetOutermost() && Context.World()->GetOutermost() == FindPackage(NULL, *LongTempFname) )
- 			{
- 				// This map is already loaded and in the editor.
- 				return true;
- 			}
-
 			// Is the new world already loaded?
 			UPackage* ExistingPackage = FindPackage(NULL, *LongTempFname);
 			UWorld* ExistingWorld = NULL;
@@ -2112,7 +2106,15 @@ bool UEditorEngine::Map_Load(const TCHAR* Str, FOutputDevice& Ar)
 						Arguments.Add(TEXT("MapFileName"), FText::FromString( MapFileName ));
 						FMessageLog("LoadErrors").NewPage( FText::Format( LOCTEXT("LoadMapLogPage", "Loading map: {MapFileName}"), Arguments ) );
 					}
-					EditorDestroyWorld( Context, LocalizedLoadingMap, ExistingWorld );
+
+					// If we are loading the same world again (reloading) then we must not specify that we want to keep this world in memory.
+					// Otherwise, try to keep the existing world in memory since there is not reason to reload it.
+					UWorld* NewWorld = nullptr;
+					if (ExistingWorld != nullptr && Context.World() != ExistingWorld)
+					{
+						NewWorld = ExistingWorld;
+					}
+					EditorDestroyWorld( Context, LocalizedLoadingMap, NewWorld );
 
 					// Refresh ExistingPackage and Existing World now that GC has occurred.
 					ExistingPackage = FindPackage(NULL, *LongTempFname);
diff --git a/Engine/Source/Editor/UnrealEd/Private/Toolkits/AssetEditorManager.cpp b/Engine/Source/Editor/UnrealEd/Private/Toolkits/AssetEditorManager.cpp
--- a/Engine/Source/Editor/UnrealEd/Private/Toolkits/AssetEditorManager.cpp
+++ b/Engine/Source/Editor/UnrealEd/Private/Toolkits/AssetEditorManager.cpp
@@ -280,7 +280,19 @@ bool FAssetEditorManager::OpenEditorForAsset(UObject* Asset, const EToolkitMode:
 	{
 		TArray<UObject*> AssetsToEdit;
 		AssetsToEdit.Add(Asset);
+
+		// Some assets (like UWorlds) may be destroyed and recreated as part of opening. To protect against this, keep the path to the asset and try to re-find it if it disappeared.
+		TWeakObjectPtr<UObject> WeakAsset = Asset;
+		const FString AssetPath = Asset->GetPathName();
+
 		AssetTypeActions.Pin()->OpenAssetEditor(AssetsToEdit, ActualToolkitMode == EToolkitMode::WorldCentric ? OpenedFromLevelEditor : TSharedPtr<IToolkitHost>());
+		
+		// If the Asset was destroyed, attempt to find it if it was recreated
+		if ( !WeakAsset.IsValid() && !AssetPath.IsEmpty() )
+		{
+			Asset = FindObject<UObject>(nullptr, *AssetPath);
+		}
+
 		AssetEditorOpenedEvent.Broadcast(Asset);
 	}
 	else if( bCanSummonSimpleAssetEditor )
