Can duplicate Blueprint variables using Ctrl+W or the right click context menu.

Fixes a crash when changing pin type from (specifically tested) a string to an enum and the set enum value would be out of range instead of cleared when it failed to parse the old default value.

Fixed bug wite duplicating functions not always getting a unique name.

#ttp 310184 - ROCKET: TASK: PUNTABLE: Duplicate Variable from RMB menu

diff --git a/Engine/Source/Editor/BlueprintGraph/Private/EdGraphSchema_K2.cpp b/Engine/Source/Editor/BlueprintGraph/Private/EdGraphSchema_K2.cpp
--- a/Engine/Source/Editor/BlueprintGraph/Private/EdGraphSchema_K2.cpp
+++ b/Engine/Source/Editor/BlueprintGraph/Private/EdGraphSchema_K2.cpp
@@ -4379,7 +4379,8 @@ UEdGraph* UEdGraphSchema_K2::DuplicateGraph(UEdGraph* GraphToDuplicate) const
 
 		if (NewGraph)
 		{
-			FEdGraphUtilities::RenameGraphCloseToName(NewGraph,GraphToDuplicate->GetFName().GetPlainNameString());
+			FName NewGraphName = FBlueprintEditorUtils::FindUniqueKismetName(Blueprint, GraphToDuplicate->GetFName().GetPlainNameString());
+			FEdGraphUtilities::RenameGraphCloseToName(NewGraph,NewGraphName.ToString());
 			// can't have two graphs with the same guid... that'd be silly!
 			NewGraph->GraphGuid = FGuid::NewGuid();
 
diff --git a/Engine/Source/Editor/Kismet/Private/SMyBlueprint.cpp b/Engine/Source/Editor/Kismet/Private/SMyBlueprint.cpp
--- a/Engine/Source/Editor/Kismet/Private/SMyBlueprint.cpp
+++ b/Engine/Source/Editor/Kismet/Private/SMyBlueprint.cpp
@@ -1916,7 +1916,10 @@ bool SMyBlueprint::IsDuplicateActionVisible() const
 			return GraphAction->EdGraph && GraphAction->EdGraph->GetSchema()->CanDuplicateGraph(GraphAction->EdGraph);
 		}
 	}
-
+	else if (SelectionAsVar() || SelectionAsLocalVar())
+	{
+		return true;
+	}
 	return false;
 }
 
@@ -1930,11 +1933,17 @@ bool SMyBlueprint::CanDuplicateAction() const
 			return GraphAction->EdGraph->GetSchema()->CanDuplicateGraph(GraphAction->EdGraph);
 		}
 	}
+	else if(SelectionAsVar() || SelectionAsLocalVar())
+	{
+		return true;
+	}
 	return false;
 }
 
 void SMyBlueprint::OnDuplicateAction()
 {
+	FName DuplicateActionName = NAME_None;
+
 	if (FEdGraphSchemaAction_K2Graph* GraphAction = SelectionAsGraph())
 	{
 		const FScopedTransaction Transaction( LOCTEXT( "DuplicateGraph", "Duplicate Graph" ) );
@@ -1951,6 +1960,29 @@ void SMyBlueprint::OnDuplicateAction()
 
 		GetBlueprintObj()->FunctionGraphs.Add(DuplicatedGraph);
 		FBlueprintEditorUtils::MarkBlueprintAsStructurallyModified(GetBlueprintObj());
+
+		DuplicateActionName = DuplicatedGraph->GetFName();
+	}
+	else if (FEdGraphSchemaAction_K2Var* VarAction = SelectionAsVar())
+	{
+		const FScopedTransaction Transaction( LOCTEXT( "DuplicateVariable", "Duplicate Variable" ) );
+		GetBlueprintObj()->Modify();
+
+		DuplicateActionName = FBlueprintEditorUtils::DuplicateVariable(GetBlueprintObj(), nullptr, VarAction->GetVariableName());
+	}
+	else if (FEdGraphSchemaAction_K2LocalVar* LocalVarAction = SelectionAsLocalVar())
+	{
+		const FScopedTransaction Transaction( LOCTEXT( "Duplicate Local Variable", "Duplicate Local Variable" ) );
+		GetBlueprintObj()->Modify();
+
+		DuplicateActionName = FBlueprintEditorUtils::DuplicateVariable(GetBlueprintObj(), LocalVarAction->GetVariableScope(), LocalVarAction->GetVariableName());
+	}
+
+	// Select and rename the duplicated action
+	if(DuplicateActionName != NAME_None)
+	{
+		SelectItemByName(DuplicateActionName);
+		OnRequestRenameOnActionNode();
 	}
 }
 
diff --git a/Engine/Source/Editor/UnrealEd/Private/Kismet2/BlueprintEditorUtils.cpp b/Engine/Source/Editor/UnrealEd/Private/Kismet2/BlueprintEditorUtils.cpp
--- a/Engine/Source/Editor/UnrealEd/Private/Kismet2/BlueprintEditorUtils.cpp
+++ b/Engine/Source/Editor/UnrealEd/Private/Kismet2/BlueprintEditorUtils.cpp
@@ -3474,6 +3474,93 @@ void FBlueprintEditorUtils::ChangeMemberVariableType(UBlueprint* Blueprint, cons
 	}
 }
 
+FName FBlueprintEditorUtils::DuplicateVariable(UBlueprint* InBlueprint, const UStruct* InScope, const FName& InVariableToDuplicate)
+{
+	FName DuplicatedVariableName = NAME_None;
+
+	if (InVariableToDuplicate != NAME_None)
+	{
+		const FScopedTransaction Transaction( LOCTEXT( "DuplicateVariable", "Duplicate Variable" ) );
+		InBlueprint->Modify();
+
+		FBPVariableDescription NewVar;
+
+		const int32 VarIndex = FBlueprintEditorUtils::FindNewVariableIndex(InBlueprint, InVariableToDuplicate);
+		if (VarIndex != INDEX_NONE)
+		{
+			FBPVariableDescription& Variable = InBlueprint->NewVariables[VarIndex];
+
+			NewVar = DuplicateVariableDescription(InBlueprint, Variable);
+
+			// We need to manually pull the DefaultValue from the UProperty to set it
+			void* OldPropertyAddr = NULL;
+
+			//Grab property of blueprint's current CDO
+			UClass* GeneratedClass = InBlueprint->GeneratedClass;
+			UObject* GeneratedCDO = GeneratedClass->GetDefaultObject();
+			UProperty* TargetProperty = FindField<UProperty>(GeneratedClass, Variable.VarName);
+
+			if( TargetProperty )
+			{
+				// Grab the address of where the property is actually stored (UObject* base, plus the offset defined in the property)
+				OldPropertyAddr = TargetProperty->ContainerPtrToValuePtr<void>(GeneratedCDO);
+				if(OldPropertyAddr)
+				{
+					// if there is a property for variable, it means the original default value was already copied, so it can be safely overridden
+					Variable.DefaultValue.Empty();
+					TargetProperty->ExportTextItem(NewVar.DefaultValue, OldPropertyAddr, OldPropertyAddr, NULL, PPF_None);
+				}
+			}
+
+			// Add the new variable
+			InBlueprint->NewVariables.Add(NewVar);
+		}
+		else
+		{
+			// It's probably a local variable
+
+			UK2Node_FunctionEntry* FunctionEntry = NULL;
+			FBPVariableDescription* LocalVariable = FBlueprintEditorUtils::FindLocalVariable(InBlueprint, InScope, InVariableToDuplicate, &FunctionEntry);
+
+			if (LocalVariable)
+			{
+				FunctionEntry->Modify();
+
+				FBPVariableDescription& Variable = *LocalVariable;
+
+				NewVar = DuplicateVariableDescription(InBlueprint, *LocalVariable);
+
+				// Add the new variable
+				FunctionEntry->LocalVariables.Add(NewVar);
+			}
+		}
+
+		if(NewVar.VarGuid.IsValid())
+		{
+			DuplicatedVariableName = NewVar.VarName;
+
+			// Potentially adjust variable names for any child blueprints
+			FBlueprintEditorUtils::ValidateBlueprintChildVariables(InBlueprint, NewVar.VarName);
+
+			FBlueprintEditorUtils::MarkBlueprintAsStructurallyModified(InBlueprint);
+		}
+	}
+
+	return DuplicatedVariableName;
+}
+
+FBPVariableDescription FBlueprintEditorUtils::DuplicateVariableDescription(UBlueprint* InBlueprint, FBPVariableDescription& InVariableDescription)
+{
+	FName DuplicatedVariableName = FindUniqueKismetName(InBlueprint, InVariableDescription.VarName.GetPlainNameString());
+
+	// Now create new variable
+	FBPVariableDescription NewVar = InVariableDescription;
+	NewVar.VarName = DuplicatedVariableName;
+	NewVar.VarGuid = FGuid::NewGuid();
+
+	return NewVar;
+}
+
 void FBlueprintEditorUtils::RemoveLocalVariable(UBlueprint* InBlueprint, const UStruct* InScope, const FName& InVarName)
 {
 	UEdGraph* ScopeGraph = FindScopeGraph(InBlueprint, InScope);
@@ -3562,7 +3649,8 @@ void FBlueprintEditorUtils::RenameLocalVariable(UBlueprint* InBlueprint, const U
 {
 	if ((InOldName != InNewName) && (InNewName != NAME_None))
 	{
-		FBPVariableDescription* LocalVariable = FindLocalVariable(InBlueprint, InScope, InOldName);
+		UK2Node_FunctionEntry* FunctionEntry = nullptr;
+		FBPVariableDescription* LocalVariable = FindLocalVariable(InBlueprint, InScope, InOldName, &FunctionEntry);
 		const auto ExistingProperty = FindField<const UProperty>(InScope, InNewName);
 		if (ExistingProperty)
 		{
@@ -3573,6 +3661,7 @@ void FBlueprintEditorUtils::RenameLocalVariable(UBlueprint* InBlueprint, const U
 		{
 			const FScopedTransaction Transaction( LOCTEXT("RenameVariable", "Rename Local Variable") );
 			InBlueprint->Modify();
+			FunctionEntry->Modify();
 
 			// Update the name
 			const UEdGraphSchema_K2* K2Schema = GetDefault<UEdGraphSchema_K2>();
@@ -3591,6 +3680,12 @@ void FBlueprintEditorUtils::RenameLocalVariable(UBlueprint* InBlueprint, const U
 	}
 }
 
+FBPVariableDescription* FBlueprintEditorUtils::FindLocalVariable(UBlueprint* InBlueprint, const UStruct* InScope, const FName& InVariableName)
+{
+	UK2Node_FunctionEntry* DummyFunctionEntry = nullptr;
+	return FindLocalVariable(InBlueprint, InScope, InVariableName, &DummyFunctionEntry);
+}
+
 FBPVariableDescription* FBlueprintEditorUtils::FindLocalVariable(const UBlueprint* InBlueprint, const UStruct* InScope, const FName& InVariableName, class UK2Node_FunctionEntry** OutFunctionEntry)
 {
 	UEdGraph* ScopeGraph = FindScopeGraph(InBlueprint, InScope);
@@ -3671,7 +3766,9 @@ void FBlueprintEditorUtils::ChangeLocalVariableType(UBlueprint* InBlueprint, con
 	{
 		FString ActionCategory;
 		const UEdGraphSchema_K2* K2Schema = GetDefault<UEdGraphSchema_K2>();
-		FBPVariableDescription* VariablePtr = FindLocalVariable(InBlueprint, InScope, InVariableName);
+		
+		UK2Node_FunctionEntry* FunctionEntry = NULL;
+		FBPVariableDescription* VariablePtr = FindLocalVariable(InBlueprint, InScope, InVariableName, &FunctionEntry);
 
 		if(VariablePtr)
 		{
@@ -3682,6 +3779,7 @@ void FBlueprintEditorUtils::ChangeLocalVariableType(UBlueprint* InBlueprint, con
 			{
 				const FScopedTransaction Transaction( LOCTEXT("ChangeLocalVariableType", "Change Local Variable Type") );
 				InBlueprint->Modify();
+				FunctionEntry->Modify();
 
 				Variable.VarType = NewPinType;
 
@@ -5779,19 +5877,19 @@ bool FBlueprintEditorUtils::IsObjectADebugCandidate( AActor* InActorObject, UBlu
 
 bool FBlueprintEditorUtils::PropertyValueFromString(const UProperty* Property, const FString& Value, uint8* DefaultObject)
 {
-	bool bParseSuccedded = true;
+	bool bParseSucceeded = true;
 	if( !Property->IsA(UStructProperty::StaticClass()) )
 	{
 		if( Property->IsA(UIntProperty::StaticClass()) )
 		{
 			int32 IntValue = 0;
-			bParseSuccedded = FDefaultValueHelper::ParseInt(Value, IntValue);
+			bParseSucceeded = FDefaultValueHelper::ParseInt(Value, IntValue);
 			CastChecked<UIntProperty>(Property)->SetPropertyValue_InContainer(DefaultObject, IntValue);
 		}
 		else if( Property->IsA(UFloatProperty::StaticClass()) )
 		{
 			float FloatValue = 0.0f;
-			bParseSuccedded = FDefaultValueHelper::ParseFloat(Value, FloatValue);
+			bParseSucceeded = FDefaultValueHelper::ParseFloat(Value, FloatValue);
 			CastChecked<UFloatProperty>(Property)->SetPropertyValue_InContainer(DefaultObject, FloatValue);
 		}
 		else if (const UByteProperty* ByteProperty = Cast<const UByteProperty>(Property))
@@ -5800,13 +5898,19 @@ bool FBlueprintEditorUtils::PropertyValueFromString(const UProperty* Property, c
 			if (const UEnum* Enum = ByteProperty->Enum)
 			{
 				IntValue = Enum->FindEnumIndex(FName(*Value));
-				bParseSuccedded = (INDEX_NONE != IntValue);
+				bParseSucceeded = (INDEX_NONE != IntValue);
+
+				// If the parse did not succeed, clear out the int to keep the enum value valid
+				if(!bParseSucceeded)
+				{
+					IntValue = 0;
+				}
 			}
 			else
 			{
-				bParseSuccedded = FDefaultValueHelper::ParseInt(Value, IntValue);
+				bParseSucceeded = FDefaultValueHelper::ParseInt(Value, IntValue);
 			}
-			bParseSuccedded = bParseSuccedded && ( IntValue <= 255 ) && ( IntValue >= 0 );
+			bParseSucceeded = bParseSucceeded && ( IntValue <= 255 ) && ( IntValue >= 0 );
 			ByteProperty->SetPropertyValue_InContainer(DefaultObject, IntValue);
 		}
 		else if( Property->IsA(UStrProperty::StaticClass()) )
@@ -5839,7 +5943,7 @@ bool FBlueprintEditorUtils::PropertyValueFromString(const UProperty* Property, c
 		{
 			FStringOutputDevice ImportError;
 			const auto EndOfParsedBuff = Property->ImportText(Value.IsEmpty() ? TEXT("()") : *Value, Property->ContainerPtrToValuePtr<uint8>(DefaultObject), 0, NULL, &ImportError);
-			bParseSuccedded = EndOfParsedBuff && ImportError.IsEmpty();
+			bParseSucceeded = EndOfParsedBuff && ImportError.IsEmpty();
 		}
 		else
 		{
@@ -5861,26 +5965,26 @@ bool FBlueprintEditorUtils::PropertyValueFromString(const UProperty* Property, c
 		if( StructProperty->Struct == VectorStruct )
 		{
 			FVector V = FVector::ZeroVector;
-			bParseSuccedded = FDefaultValueHelper::ParseVector(Value, V);
+			bParseSucceeded = FDefaultValueHelper::ParseVector(Value, V);
 			Property->CopyCompleteValue( Property->ContainerPtrToValuePtr<uint8>(DefaultObject), &V );
 		}
 		else if( StructProperty->Struct == RotatorStruct )
 		{
 			FRotator R = FRotator::ZeroRotator;
-			bParseSuccedded = FDefaultValueHelper::ParseRotator(Value, R);
+			bParseSucceeded = FDefaultValueHelper::ParseRotator(Value, R);
 			Property->CopyCompleteValue( Property->ContainerPtrToValuePtr<uint8>(DefaultObject), &R );
 		}
 		else if( StructProperty->Struct == TransformStruct )
 		{
 			FTransform T = FTransform::Identity;
-			bParseSuccedded = T.InitFromString( Value );
+			bParseSucceeded = T.InitFromString( Value );
 			Property->CopyCompleteValue( Property->ContainerPtrToValuePtr<uint8>(DefaultObject), &T );
 		}
 		else if( StructProperty->Struct == LinearColorStruct )
 		{
 			FLinearColor Color;
 			// Color form: "(R=%f,G=%f,B=%f,A=%f)"
-			bParseSuccedded = Color.InitFromString(Value);
+			bParseSucceeded = Color.InitFromString(Value);
 			Property->CopyCompleteValue( Property->ContainerPtrToValuePtr<uint8>(DefaultObject), &Color );
 		}
 		else if (StructProperty->Struct)
@@ -5890,15 +5994,15 @@ bool FBlueprintEditorUtils::PropertyValueFromString(const UProperty* Property, c
 			uint8* StructData = Property->ContainerPtrToValuePtr<uint8>(DefaultObject);
 			StructProperty->InitializeValue(StructData);
 			ensure(1 == StructProperty->ArrayDim);
-			bParseSuccedded = FStructureEditorUtils::Fill_MakeStructureDefaultValue(Cast<const UUserDefinedStruct>(Struct), StructData);
+			bParseSucceeded = FStructureEditorUtils::Fill_MakeStructureDefaultValue(Cast<const UUserDefinedStruct>(Struct), StructData);
 
 			FStringOutputDevice ImportError;
 			const auto EndOfParsedBuff = StructProperty->ImportText(Value.IsEmpty() ? TEXT("()") : *Value, StructData, 0, NULL, &ImportError);
-			bParseSuccedded &= EndOfParsedBuff && ImportError.IsEmpty();
+			bParseSucceeded &= EndOfParsedBuff && ImportError.IsEmpty();
 		}
 	}
 
-	return bParseSuccedded;
+	return bParseSucceeded;
 }
 
 bool FBlueprintEditorUtils::PropertyValueToString(const UProperty* Property, const uint8* Container, FString& OutForm)
diff --git a/Engine/Source/Editor/UnrealEd/Public/Kismet2/BlueprintEditorUtils.h b/Engine/Source/Editor/UnrealEd/Public/Kismet2/BlueprintEditorUtils.h
--- a/Engine/Source/Editor/UnrealEd/Public/Kismet2/BlueprintEditorUtils.h
+++ b/Engine/Source/Editor/UnrealEd/Public/Kismet2/BlueprintEditorUtils.h
@@ -490,6 +490,27 @@ public:
 	 */
 	static bool AddMemberVariable(UBlueprint* Blueprint, const FName& NewVarName, const FEdGraphPinType& NewVarType, const FString& DefaultValue = FString());
 
+	/**
+	 * Duplicates a variable given it's name and Blueprint
+	 *
+	 * @param InBlueprint					The Blueprint the variable can be found in
+	 * @paramInScope						Local variable's scope
+	 * @param InVariableToDuplicate			Variable name to be found and duplicated
+	 *
+	 * @return								Returns the name of the new variable or NAME_None is failed to duplicate
+	 */
+	static FName DuplicateVariable(UBlueprint* InBlueprint, const UStruct* InScope, const FName& InVariableToDuplicate);
+
+	/**
+	 * Internal function that deep copies a variable description
+	 *
+	 * @param InBlueprint					The blueprint to ensure a uniquely named variable in
+	 * @param InVariableToDuplicate			Variable description to duplicate
+	 *
+	 * @return								Returns a copy of the passed in variable description.
+	 */
+	static FBPVariableDescription DuplicateVariableDescription(UBlueprint* InBlueprint, FBPVariableDescription& InVariableDescription);
+
 	/**
 	 * Removes a member variable if it was declared in this blueprint and not in a base class.
 	 *
@@ -540,6 +561,15 @@ public:
 	 */
 	static void RemoveLocalVariable(UBlueprint* InBlueprint, const UStruct* InScope, const FName& InVarName);
 
+	/**
+	 * Returns a local variable with the function entry it was found in
+	 *
+	 * @param InBlueprint		Blueprint to search for the local variable
+	 * @param InVariableName	Name of the variable to search for
+	 * @return					The local variable description
+	 */
+	static FBPVariableDescription* FindLocalVariable(UBlueprint* InBlueprint, const UStruct* InScope, const FName& InVariableName);
+
 	/**
 	 * Returns a local variable
 	 *
