Fixed typo in comment automatically used for CulturesToStage tooltip in Project Packaging Settings.

diff --git a/Engine/Source/Editor/UnrealEd/Classes/Settings/ProjectPackagingSettings.h b/Engine/Source/Editor/UnrealEd/Classes/Settings/ProjectPackagingSettings.h
--- a/Engine/Source/Editor/UnrealEd/Classes/Settings/ProjectPackagingSettings.h
+++ b/Engine/Source/Editor/UnrealEd/Classes/Settings/ProjectPackagingSettings.h
@@ -66,7 +66,7 @@ public:
 	bool UseOBB_InAPK;
 
 	/**
-	 * ISO codes of cultures whose data should be packaged to be cooked, staged, and packaged.
+	 * ISO codes of cultures whose data should be cooked, staged, and packaged.
 	 * Note: These paths are relative to your project Content directory
 	 */
 	UPROPERTY(config, EditAnywhere, Category=Packaging, AdvancedDisplay, meta=(DisplayName="Cultures to Package"))
