Fix for consolidating shape for morphtarget

diff --git a/Engine/Source/Editor/UnrealEd/Private/Fbx/FbxSkeletalMeshImport.cpp b/Engine/Source/Editor/UnrealEd/Private/Fbx/FbxSkeletalMeshImport.cpp
--- a/Engine/Source/Editor/UnrealEd/Private/Fbx/FbxSkeletalMeshImport.cpp
+++ b/Engine/Source/Editor/UnrealEd/Private/Fbx/FbxSkeletalMeshImport.cpp
@@ -2659,15 +2659,7 @@ private:
 void UnFbx::FFbxImporter::ImportMorphTargetsInternal( TArray<FbxNode*>& SkelMeshNodeArray, USkeletalMesh* BaseSkelMesh, UObject * InParent, const FString& InFilename, int32 LODIndex )
 {
 	FbxString ShapeNodeName;
-	TArray<FbxShape*> FbxShapeArray;
-	FbxShapeArray.AddUninitialized(SkelMeshNodeArray.Num());
-
-	// Initialize the shape array.
-	// In the shape array, only one geometry has shape at a time. Other geometries has no shape
-	for (int32 NodeIndex = 0; NodeIndex < SkelMeshNodeArray.Num(); NodeIndex++)
-	{
-		FbxShapeArray[NodeIndex] = NULL;
-	}
+	TMap<FString, TArray<FbxShape*>> ShapeNameToShapeArray;
 
 	// Temp arrays to keep track of data being used by threads
 	TArray<USkeletalMesh*> TempMeshes;
@@ -2681,43 +2673,15 @@ void UnFbx::FFbxImporter::ImportMorphTargetsInternal( TArray<FbxNode*>& SkelMesh
 	// For each morph in FBX geometries, we create one morph target for the Unreal skeletal mesh
 	for (int32 NodeIndex = 0; NodeIndex < SkelMeshNodeArray.Num(); NodeIndex++)
 	{
-		// reset the shape array
-		if (NodeIndex > 0)
-		{
-			FbxShapeArray[NodeIndex-1] = NULL;
-		}
-
-		uint32 GlobalShapeCount = 0;
 		FbxGeometry* Geometry = (FbxGeometry*)SkelMeshNodeArray[NodeIndex]->GetNodeAttribute();
 		if (Geometry)
 		{
 			const int32 BlendShapeDeformerCount = Geometry->GetDeformerCount(FbxDeformer::eBlendShape);
 
 			/************************************************************************/
-			/* 1. count shapes to properly update progress                          */
+			/* collect all the shapes                                               */
 			/************************************************************************/
 			for(int32 BlendShapeIndex = 0; BlendShapeIndex<BlendShapeDeformerCount; ++BlendShapeIndex)
-			{
-				const FbxBlendShape* BlendShape = (FbxBlendShape*)Geometry->GetDeformer(BlendShapeIndex, FbxDeformer::eBlendShape);
-
-				const int32 BlendShapeChannelCount = BlendShape->GetBlendShapeChannelCount();
-				for(int32 ChannelIndex = 0; ChannelIndex<BlendShapeChannelCount; ++ChannelIndex)
-				{
-					const FbxBlendShapeChannel* Channel = BlendShape->GetBlendShapeChannel(ChannelIndex);
-					if(Channel)
-					{
-						//Find which shape should we use according to the weight.
-						const int32 CurrentChannelShapeCount = Channel->GetTargetShapeCount();
-						GlobalShapeCount += CurrentChannelShapeCount;
-					}
-				}
-			}
-
-			/************************************************************************/
-			/* 2. do the job actually                                               */
-			/************************************************************************/
-			uint32 CurrentShapeIndex = 0;
-			for(int32 BlendShapeIndex = 0; BlendShapeIndex<BlendShapeDeformerCount; ++BlendShapeIndex)
 			{
 				FbxBlendShape* BlendShape = (FbxBlendShape*)Geometry->GetDeformer(BlendShapeIndex, FbxDeformer::eBlendShape);
 				const int32 BlendShapeChannelCount = BlendShape->GetBlendShapeChannelCount();
@@ -2747,8 +2711,6 @@ void UnFbx::FFbxImporter::ImportMorphTargetsInternal( TArray<FbxNode*>& SkelMesh
 						{
 							FbxShape* Shape = Channel->GetTargetShape(ShapeIndex);
 			
-							FbxShapeArray[NodeIndex] = Shape;
-
 							FString ShapeName;
 							if( CurrentChannelShapeCount > 1 )
 							{
@@ -2767,47 +2729,13 @@ void UnFbx::FFbxImporter::ImportMorphTargetsInternal( TArray<FbxNode*>& SkelMesh
 								}
 							}
 
-							FFormatNamedArguments Args;
-							Args.Add( TEXT("ShapeName"), FText::FromString( ShapeName ) );
-							Args.Add( TEXT("CurrentShapeIndex"), CurrentShapeIndex+1 );
-							Args.Add( TEXT("TotalShapes"), GlobalShapeCount+1 );
-							const FText StatusUpate = FText::Format( NSLOCTEXT("FbxImporter", "GeneratingMorphTargetMeshStatus", "Generating morph target mesh {ShapeName} ({CurrentShapeIndex} of {TotalShapes})"), Args );
-
-							GWarn->StatusUpdate( CurrentShapeIndex+1, GlobalShapeCount, StatusUpate );
-
-							FSkeletalMeshImportData ImportData;
-
-							// See if this morph target already exists.
-							UMorphTarget * Result = FindObject<UMorphTarget>(BaseSkelMesh,*ShapeName);
-							// we only create new one for LOD0, otherwise don't create new one
-							if(!Result)
+							TArray<FbxShape*> & ShapeArray = ShapeNameToShapeArray.FindOrAdd(ShapeName);
+							if (ShapeArray.Num() == 0)
 							{
-								if (LODIndex == 0)
-								{
-									Result = NewNamedObject<UMorphTarget>(BaseSkelMesh, FName(*ShapeName));
-								}
-								else
-								{
-									AddTokenizedErrorMessage(FTokenizedMessage::Create(EMessageSeverity::Error, FText::Format(FText::FromString("Could not find the {0} morphtarget for LOD {1}. \
-																													Make sure the name for morphtarget matches with LOD 0"), FText::FromString(ShapeName), FText::FromString(FString::FromInt(LODIndex)))), 
-																													FFbxErrors::SkeletalMesh_LOD_MissingMorphTarget);
-								}
+								ShapeArray.AddZeroed(SkelMeshNodeArray.Num());
 							}
 
-							if (Result)
-							{
-								// now we get a shape for whole mesh, import to unreal as a morph target
-								// @todo AssetImportData do we need import data for this temp mesh?
-								UFbxSkeletalMeshImportData* TmpMeshImportData = NULL;
-								USkeletalMesh* TmpSkeletalMesh = (USkeletalMesh*)ImportSkeletalMesh(GetTransientPackage(), SkelMeshNodeArray, NAME_None, (EObjectFlags)0, TmpMeshImportData, FPaths::GetBaseFilename(InFilename), &FbxShapeArray, &ImportData, false);
-								TempMeshes.Add(TmpSkeletalMesh);
-								MorphTargets.Add(Result);
-
-								// Process the skeletal mesh on a separate thread
-								FAsyncTask<FAsyncImportMorphTargetWork>* NewWork = new (PendingWork)FAsyncTask<FAsyncImportMorphTargetWork>(TmpSkeletalMesh, LODIndex, ImportData, ImportOptions->bKeepOverlappingVertices);
-								NewWork->StartBackgroundTask();
-								++CurrentShapeIndex;
-							}
+							ShapeArray[NodeIndex] = Shape;
 						}
 					}
 				}
@@ -2815,6 +2743,57 @@ void UnFbx::FFbxImporter::ImportMorphTargetsInternal( TArray<FbxNode*>& SkelMesh
 		}
 	} // for NodeIndex
 
+	int32 ShapeIndex = 0;
+	int32 TotalShapeCount = ShapeNameToShapeArray.Num();
+	// iterate through shapename, and create morphtarget
+	for (auto Iter = ShapeNameToShapeArray.CreateIterator(); Iter; ++Iter)
+	{
+		FString ShapeName = Iter.Key();
+		TArray<FbxShape*> & ShapeArray = Iter.Value();
+
+		FFormatNamedArguments Args;
+		Args.Add(TEXT("ShapeName"), FText::FromString(ShapeName));
+		Args.Add(TEXT("CurrentShapeIndex"), ShapeIndex + 1);
+		Args.Add(TEXT("TotalShapes"), TotalShapeCount);
+		const FText StatusUpate = FText::Format(NSLOCTEXT("FbxImporter", "GeneratingMorphTargetMeshStatus", "Generating morph target mesh {ShapeName} ({CurrentShapeIndex} of {TotalShapes})"), Args);
+
+		GWarn->StatusUpdate(ShapeIndex + 1, TotalShapeCount, StatusUpate);
+
+		FSkeletalMeshImportData ImportData;
+
+		// See if this morph target already exists.
+		UMorphTarget * Result = FindObject<UMorphTarget>(BaseSkelMesh, *ShapeName);
+		// we only create new one for LOD0, otherwise don't create new one
+		if (!Result)
+		{
+			if (LODIndex == 0)
+			{
+				Result = NewNamedObject<UMorphTarget>(BaseSkelMesh, FName(*ShapeName));
+			}
+			else
+			{
+				AddTokenizedErrorMessage(FTokenizedMessage::Create(EMessageSeverity::Error, FText::Format(FText::FromString("Could not find the {0} morphtarget for LOD {1}. \
+																																																																																											Make sure the name for morphtarget matches with LOD 0"), FText::FromString(ShapeName), FText::FromString(FString::FromInt(LODIndex)))),
+																																																																																											FFbxErrors::SkeletalMesh_LOD_MissingMorphTarget);
+			}
+		}
+
+		if (Result)
+		{
+			// now we get a shape for whole mesh, import to unreal as a morph target
+			// @todo AssetImportData do we need import data for this temp mesh?
+			UFbxSkeletalMeshImportData* TmpMeshImportData = NULL;
+			USkeletalMesh* TmpSkeletalMesh = (USkeletalMesh*)ImportSkeletalMesh(GetTransientPackage(), SkelMeshNodeArray, NAME_None, (EObjectFlags)0, TmpMeshImportData, FPaths::GetBaseFilename(InFilename), &ShapeArray, &ImportData, false);
+			TempMeshes.Add(TmpSkeletalMesh);
+			MorphTargets.Add(Result);
+
+			// Process the skeletal mesh on a separate thread
+			FAsyncTask<FAsyncImportMorphTargetWork>* NewWork = new (PendingWork)FAsyncTask<FAsyncImportMorphTargetWork>(TmpSkeletalMesh, LODIndex, ImportData, ImportOptions->bKeepOverlappingVertices);
+			NewWork->StartBackgroundTask();
+			++ShapeIndex;
+		}
+	}
+
 	// Wait for all importing tasks to complete
 	int32 NumCompleted = 0;
 	int32 NumTasks = PendingWork.Num();
