Shootergame now allows Play From Here from the editor, and Click Plane is now set when Click Location is

#udn "Play from here" problems

diff --git a/Engine/Source/Editor/UnrealEd/Private/LevelEditorViewport.cpp b/Engine/Source/Editor/UnrealEd/Private/LevelEditorViewport.cpp
--- a/Engine/Source/Editor/UnrealEd/Private/LevelEditorViewport.cpp
+++ b/Engine/Source/Editor/UnrealEd/Private/LevelEditorViewport.cpp
@@ -2258,7 +2258,9 @@ bool FLevelEditorViewportClient::InputKey(FViewport* Viewport, int32 ControllerI
 	if ( InputState.IsAnyMouseButtonDown() )
 	{
 		const FViewportCursorLocation Cursor(View, this, HitX, HitY);
-		GEditor->ClickLocation = FActorPositioning::TraceWorldForPositionWithDefault(Cursor, *View).Location;
+		const FActorPositionTraceResult TraceResult = FActorPositioning::TraceWorldForPositionWithDefault(Cursor, *View);
+		GEditor->ClickLocation = TraceResult.Location;
+		GEditor->ClickPlane = FPlane(TraceResult.Location, TraceResult.SurfaceNormal);
 
 		// Snap the new location if snapping is enabled
 		FSnappingUtils::SnapPointToGrid(GEditor->ClickLocation, FVector::ZeroVector);
diff --git a/Engine/Source/Editor/UnrealEd/Private/UnrealEdSrv.cpp b/Engine/Source/Editor/UnrealEd/Private/UnrealEdSrv.cpp
--- a/Engine/Source/Editor/UnrealEd/Private/UnrealEdSrv.cpp
+++ b/Engine/Source/Editor/UnrealEd/Private/UnrealEdSrv.cpp
@@ -1650,6 +1650,7 @@ bool UUnrealEdEngine::Exec_Pivot( const TCHAR* Str, FOutputDevice& Ar )
 		if( Count > 0 )
 		{
 			ClickLocation = Center / Count;
+			ClickPlane = FPlane(0.f,0.f,0.f,0.f);
 
 			SetPivot( ClickLocation, false, false );
 			FinishAllSnaps();
