BP: ttp 343677 LIVE: CRASH: Playing project in Standalone Preview or Mobile preview is causing the Play window to crash

diff --git a/Engine/Source/Editor/UnrealEd/Private/Kismet2/KismetReinstanceUtilities.cpp b/Engine/Source/Editor/UnrealEd/Private/Kismet2/KismetReinstanceUtilities.cpp
--- a/Engine/Source/Editor/UnrealEd/Private/Kismet2/KismetReinstanceUtilities.cpp
+++ b/Engine/Source/Editor/UnrealEd/Private/Kismet2/KismetReinstanceUtilities.cpp
@@ -102,10 +102,17 @@ FBlueprintCompileReinstancer::FBlueprintCompileReinstancer(UClass* InClassToRein
 			{
 				UClass* ChildClass = *ClassIt;
 				UBlueprint* ChildBP = Cast<UBlueprint>(ChildClass->ClassGeneratedBy);
-				if( ChildBP && !ChildBP->HasAnyFlags(RF_BeingRegenerated))
+				if (ChildBP)
 				{
+					if (ChildBP->HasAnyFlags(RF_BeingRegenerated))
+					{
+						if (ChildClass->GetSuperClass() == ClassToReinstance)
+						{
+							ReparentChild(ChildClass);
+						}
+					}
 					// If this is a direct child, change the parent and relink so the property chain is valid for reinstancing
-					if( !ChildBP->HasAnyFlags(RF_NeedLoad) )
+					else if( !ChildBP->HasAnyFlags(RF_NeedLoad) )
 					{
 						if( ChildClass->GetSuperClass() == ClassToReinstance )
 						{
