Fixed a problem with CompilerResultLog incorrectly classifying some lines as errors. This code needs further fixes, though, as currently it only works OK with MSVC compiler output.

diff --git a/Engine/Source/Editor/UnrealEd/Private/Kismet2/CompilerResultsLog.cpp b/Engine/Source/Editor/UnrealEd/Private/Kismet2/CompilerResultsLog.cpp
--- a/Engine/Source/Editor/UnrealEd/Private/Kismet2/CompilerResultsLog.cpp
+++ b/Engine/Source/Editor/UnrealEd/Private/Kismet2/CompilerResultsLog.cpp
@@ -277,11 +277,12 @@ TArray< TSharedRef<FTokenizedMessage> > FCompilerResultsLog::ParseCompilerLogDum
 
 		// handle output line error message if applicable
 		// @todo Handle case where there are parenthesis in path names
+		// @todo Handle errors reported by Clang
 		FString LeftStr, RightStr;
 		FString FullPath, LineNumberString;
 		if (Line.Split(TEXT(")"), &LeftStr, &RightStr, ESearchCase::CaseSensitive) &&
 			LeftStr.Split(TEXT("("), &FullPath, &LineNumberString, ESearchCase::CaseSensitive) &&
-			(FCString::Strtoi(*LineNumberString, NULL, 10) > 0))
+			LineNumberString.IsNumeric() && (FCString::Strtoi(*LineNumberString, NULL, 10) > 0))
 		{
 			EMessageSeverity::Type Severity = EMessageSeverity::Error;
 			FString FullPathTrimmed = FullPath;
