Check preview mesh availability before retargeting

diff --git a/Engine/Source/Editor/UnrealEd/Private/EditorAnimUtils.cpp b/Engine/Source/Editor/UnrealEd/Private/EditorAnimUtils.cpp
--- a/Engine/Source/Editor/UnrealEd/Private/EditorAnimUtils.cpp
+++ b/Engine/Source/Editor/UnrealEd/Private/EditorAnimUtils.cpp
@@ -148,8 +148,17 @@ namespace EditorAnimUtils
 		{
 			// we need to update reference pose before retargeting. 
 			// this is to ensure the skeleton has the latest pose you're looking at. 
-			OldSkeleton->UpdateReferencePoseFromMesh(OldSkeleton->GetPreviewMesh(true));
-			NewSkeleton->UpdateReferencePoseFromMesh(NewSkeleton->GetPreviewMesh(true));
+			USkeletalMesh * PreviewMesh = OldSkeleton->GetPreviewMesh(true);
+			if (PreviewMesh)
+			{
+				OldSkeleton->UpdateReferencePoseFromMesh(PreviewMesh);
+			}
+			
+			PreviewMesh = NewSkeleton->GetPreviewMesh(true);
+			if (PreviewMesh)
+			{
+				NewSkeleton->UpdateReferencePoseFromMesh(PreviewMesh);
+			}
 		}
 
 		for(auto Iter = AnimSequencesToRetarget.CreateIterator(); Iter; ++Iter)
