--Adding the Editor Performance Capture test.
This test will grab the Average FPS and Working Set memory numbers over a period of time in seconds.
The map used and the time frame are set in the Engine.ini file for that project or using the Editor Preferences > Automation inside the editor.
The numbers are spit out to several files in the Saved\\Automation\\Logs\\Performance\\<MAPNAME>\\ folder
This test is currently flagged to RequireUser so it will not run on the automation farm at this time.

diff --git a/Engine/Source/Editor/UnrealEd/Private/Tests/AutomationEditorCommon.cpp b/Engine/Source/Editor/UnrealEd/Private/Tests/AutomationEditorCommon.cpp
--- a/Engine/Source/Editor/UnrealEd/Private/Tests/AutomationEditorCommon.cpp
+++ b/Engine/Source/Editor/UnrealEd/Private/Tests/AutomationEditorCommon.cpp
@@ -5,7 +5,9 @@
 #include "AssetEditorManager.h"
 #include "ModuleManager.h"
 #include "LevelEditor.h"
-#include "ModuleManager.h"
+#include "Editor/MainFrame/Public/MainFrame.h"
+
+
 
 DEFINE_LOG_CATEGORY_STATIC(LogAutomationEditorCommon, Log, All);
 
@@ -362,6 +364,106 @@ bool FEndPlayMapCommand::Update()
 	return true;
 }
 
+/**
+* This command grabs the FPS and Memory for the current editor session.
+*/
+bool FEditorPerformanceCommand::Update()
+{
+	//Get the base filename for the map that will be used.
+	FString ShortMapName = FPaths::GetBaseFilename(BaseMapName);	
+	
+	//Map Load time variables
+	double MapLoadStartTime = 0.0f;
+	double MapLoadTime = 0.0f;
+
+	//Gets the main frame module
+	const IMainFrameModule& MainFrameModule = FModuleManager::GetModuleChecked< IMainFrameModule >("MainFrame");
+
+	//Variables that hold the content from the text files
+	FString SavedMapLoadTime;
+	FString SavedAverageFPS;
+	FString SavedMemory;
+	static SIZE_T StaticLastTotalAllocated = 0;
+
+	//Performance file locations and setups.
+	FString FileSaveLocation = FPaths::AutomationLogDir() + TEXT("Performance/") + ShortMapName + TEXT("/");
+	FString FPSFileSaveLocation = FileSaveLocation + TEXT("RawFPS.txt");
+	FString MemoryFileSaveLocation = FileSaveLocation + TEXT("RawMemory.txt");
+	FString MapLoadTimeFileSaveLocation = FileSaveLocation + TEXT("RAWMapLoadTime.txt");
+	
+	//Open the map if it hasn't been loaded already.
+	if (MainFrameModule.GetLoadedLevelName() != ShortMapName)
+	{
+		UE_LOG(LogEditorAutomationTests, Log, TEXT("Intended map is not loaded in the editor.  Loading '%s' now."), *ShortMapName);
+		
+		//Get the current number of seconds.  This will be used to track how long it took to load the map.
+		MapLoadStartTime = FPlatformTime::Seconds();
+		//Load the map
+		FEditorAutomationTestUtilities::LoadMap(BaseMapName);
+		//This is the time it took to load the map in the editor.
+		MapLoadTime = FPlatformTime::Seconds() - MapLoadStartTime;
+
+		UE_LOG(LogEditorAutomationTests, Log, TEXT("%s has been loaded."), *ShortMapName);
+
+		//Log out to a text file the time it takes to load the map.
+		FString LoadTime = FString::Printf(TEXT("%.3f"), MapLoadTime);
+		FFileHelper::LoadFileToString(SavedMapLoadTime, *MapLoadTimeFileSaveLocation);
+		FString MapLoadTimeFileSetup = SavedMapLoadTime + LoadTime + TEXT(",");
+		FFileHelper::SaveStringToFile(MapLoadTimeFileSetup, *MapLoadTimeFileSaveLocation);
+
+		//Display the time it took to launch the map.
+		UE_LOG(LogEditorAutomationTests, Display, TEXT("Map '%s' took '%.3f' to load"), *BaseMapName, MapLoadTime);
+		return true;
+	}
+	
+	//FPS capture
+	const float CurrentFPS = 1.0f / FSlateApplication::Get().GetAverageDeltaTime();
+	const float ClampedFPS = (CurrentFPS < 0.3f || CurrentFPS > 4000.0f) ? 0.0f : CurrentFPS;
+	//Log out to a text file the AverageFPS
+	FString AverageFPS = FString::Printf(TEXT("%3.1f"), ClampedFPS);
+	FFileHelper::LoadFileToString(SavedAverageFPS, *FPSFileSaveLocation);
+	FString FPSFileSetup = SavedAverageFPS + AverageFPS + TEXT(",");
+	FFileHelper::SaveStringToFile(FPSFileSetup, *FPSFileSaveLocation);
+
+	//Working Set Memory capture
+	// Query OS for process memory used
+	FPlatformMemoryStats MemoryStats = FPlatformMemory::GetStats();
+	StaticLastTotalAllocated = MemoryStats.UsedPhysical;
+	//Log out to a text file the Memory in kilobytes.
+	FString Memory = FString::Printf(TEXT("%f"), ((float)StaticLastTotalAllocated / 1024));
+	FFileHelper::LoadFileToString(SavedMemory, *MemoryFileSaveLocation);
+	FString MemoryFileSetup = SavedMemory + Memory + TEXT(",");
+	FFileHelper::SaveStringToFile(MemoryFileSetup, *MemoryFileSaveLocation);
+	
+	return true;
+}
+
+/**
+* This this command loads a map into the editor.
+*/
+bool FEditorLoadMap::Update()
+{
+	double MapLoadStartTime = 0.0f;
+	double MapLoadTime = 0.0f;
+	
+	//Get the current number of seconds.  This will be used to track how long it took to load the map.
+	MapLoadStartTime = FPlatformTime::Seconds();
+
+	//Load the map
+	FEditorAutomationTestUtilities::LoadMap(MapName);
+
+	//This is the time it took to load the map in the editor.
+	MapLoadTime = FPlatformTime::Seconds() - MapLoadStartTime;
+
+	//Log how long it took to launch the map.
+	UE_LOG(LogEditorAutomationTests, Display, TEXT("Map '%s' took '%.3f' to load"), *MapName, MapLoadTime);
+
+	return true;
+}
+
+//////////////////////////////////////////////////////////////////////
+//Find Asset Commands
+
 /**
 * Generates a list of assets from the ENGINE and the GAME by a specific type.
 * This is to be used by the GetTest() function.
diff --git a/Engine/Source/Editor/UnrealEd/Private/Tests/EditorPerformanceTests.cpp b/Engine/Source/Editor/UnrealEd/Private/Tests/EditorPerformanceTests.cpp
new file mode 100644
--- /dev/null
+++ b/Engine/Source/Editor/UnrealEd/Private/Tests/EditorPerformanceTests.cpp
@@ -0,0 +1,106 @@
+// Copyright 1998-2014 Epic Games, Inc. All Rights Reserved.
+
+#include "UnrealEd.h"
+#include "Tests/AutomationTestSettings.h"
+#include "ModuleManager.h"
+#include "AutomationEditorCommon.h"
+#include "AutomationCommon.h"
+#include "IMainFrameModule.h"
+
+
+
+/**
+* Map Performance in Editor tests
+* Grabs certain performance numbers and saves it to a file.
+*/
+IMPLEMENT_COMPLEX_AUTOMATION_TEST(FMapPerformanceInEditor, "Performance.Map Perfomance in Editor", EAutomationTestFlags::ATF_Editor | EAutomationTestFlags::ATF_RequiresUser)
+
+/**
+* Requests a enumeration of all maps to be loaded
+*/
+void FMapPerformanceInEditor::GetTests(TArray<FString>& OutBeautifiedNames, TArray <FString>& OutTestCommands) const
+{
+	UAutomationTestSettings const* AutomationTestSettings = GetDefault<UAutomationTestSettings>();
+	check(AutomationTestSettings);
+
+	TArray<FString> PerformanceMapName;
+	for (auto Iter = AutomationTestSettings->EditorPerformanceTestMaps.CreateConstIterator(); Iter; ++Iter)
+	{
+		if (Iter->PerformanceTestmap.FilePath.Len() > 0)
+		{
+			PerformanceMapName.AddUnique(Iter->PerformanceTestmap.FilePath);
+		}
+		
+		for (int32 indexloop = 0; indexloop < PerformanceMapName.Num(); ++indexloop)
+		{
+			//Get the location of the map being used.
+			FString Filename = FPaths::ConvertRelativePathToFull(PerformanceMapName[indexloop]);
+			
+			//If true it will proceed to add the asset to the test list.
+			//This will be false if the map is on a different drive.
+			if (FPaths::MakePathRelativeTo(Filename, *FPaths::GameContentDir()))
+			{
+				FString ShortName = FPaths::GetBaseFilename(Filename);
+				FString PathName = FPaths::GetPath(Filename);
+				OutBeautifiedNames.Add(ShortName);
+				FString AssetName = FString::Printf(TEXT("/Game/%s/%s.%s"), *PathName, *ShortName, *ShortName);
+				OutTestCommands.Add(AssetName);
+			}
+			else
+			{
+				UE_LOG(LogEditorAutomationTests, Error, TEXT("Invalid asset path: %s."), *Filename);
+			}
+		}
+
+	}
+}
+
+/*
+*
+*
+*/ 
+bool FMapPerformanceInEditor::RunTest(const FString& Parameters)
+{
+	//Get the map name from the parameters.
+	FString MapName = Parameters;
+	//Get the base filename for the map that will be used.
+	FString ShortMapName = FPaths::GetBaseFilename(MapName);
+	//Duration variable will be used to indicate how long the test will run for.
+	float Duration = 0.0f;
+	
+	//This gets the info we need from the automation test settings in the engine.ini.
+	UAutomationTestSettings const* AutomationTestSettings = GetDefault<UAutomationTestSettings>();
+	check(AutomationTestSettings);
+
+	//Now we find the test timer(aka duration) for our test.
+	for (auto Iter = AutomationTestSettings->EditorPerformanceTestMaps.CreateConstIterator(); Iter; ++Iter)
+	{
+		FString IterMapName = FPaths::GetBaseFilename(Iter->PerformanceTestmap.FilePath);
+		if (IterMapName == ShortMapName)
+		{
+			Duration = ((Iter->TestTimer));
+			//If the duration is equal to 0 then we simply warn the user that they need to set the test timer option for the performance test.
+			//If the duration is less than 0 then we fail this test.
+			if (Duration == 0)
+			{
+				UE_LOG(LogEditorAutomationTests, Warning, TEXT("Please set the test timer for '%s' in the automation preferences or engine.ini."), *ShortMapName);
+			}
+			else 
+			if (Duration < 0)
+			{
+				UE_LOG(LogEditorAutomationTests, Error, TEXT("Test timer preference option '%s' is less than 0."), *ShortMapName);
+				return false; 
+			}
+		}
+	}
+	//Loop through the latent performance command equal to the duration time.
+	UE_LOG(LogEditorAutomationTests, Log, TEXT("Starting loop to run the latent FEditorPerformanceCommand on '%s'."), *ShortMapName);
+	for (int32 Counter = 0; Counter < Duration; Counter++)
+	{
+		ADD_LATENT_AUTOMATION_COMMAND(FWaitLatentCommand(1.0f));
+		ADD_LATENT_AUTOMATION_COMMAND(FEditorPerformanceCommand(MapName));
+	}
+	return true;
+}
+
+
diff --git a/Engine/Source/Editor/UnrealEd/Public/Tests/AutomationEditorCommon.h b/Engine/Source/Editor/UnrealEd/Public/Tests/AutomationEditorCommon.h
--- a/Engine/Source/Editor/UnrealEd/Public/Tests/AutomationEditorCommon.h
+++ b/Engine/Source/Editor/UnrealEd/Public/Tests/AutomationEditorCommon.h
@@ -127,6 +127,15 @@ DEFINE_LATENT_AUTOMATION_COMMAND_ONE_PARAMETER(FStartPIECommand, bool, bSimulate
 */
 DEFINE_LATENT_AUTOMATION_COMMAND(FEndPlayMapCommand);
 
+/**
+* This will display the average FPS and Memory numbers over a duration of time.
+*/
+DEFINE_LATENT_AUTOMATION_COMMAND_ONE_PARAMETER(FEditorPerformanceCommand, FString, BaseMapName);
+
+/**
+* Loads a map
+*/
+DEFINE_LATENT_AUTOMATION_COMMAND_ONE_PARAMETER(FEditorLoadMap, FString, MapName);
 
 //////////////////////////////////////////////////////////////////////////
 // FEditorAutomationTestUtilities
diff --git a/Engine/Source/Runtime/Engine/Classes/Tests/AutomationTestSettings.h b/Engine/Source/Runtime/Engine/Classes/Tests/AutomationTestSettings.h
--- a/Engine/Source/Runtime/Engine/Classes/Tests/AutomationTestSettings.h
+++ b/Engine/Source/Runtime/Engine/Classes/Tests/AutomationTestSettings.h
@@ -274,6 +274,23 @@ struct FBuildPromotionTestSettings
 	FFilePath	SourceControlMaterial;
 };
 
+/**
+* Holds settings for the asset import / export automation test
+*/
+USTRUCT()
+struct FEditorMapPerformanceTestDefinition
+{
+	GENERATED_USTRUCT_BODY()
+	
+	/** Map to be used for the Performance Capture **/
+	UPROPERTY(config, EditAnywhere, Category = Automation, meta = (FilePathFilter = "umap"))
+	FFilePath PerformanceTestmap;
+
+	/** How long is this test expected to run before stopping **/
+	UPROPERTY(config, EditAnywhere, Category = Automation, meta = (ToolTip = "This is the length of time in seconds that this test will run for before stopping."))
+	int32 TestTimer;
+};
+
 /**
  * Implements the Editor's user settings.
  */
@@ -285,38 +302,50 @@ class ENGINE_API UAutomationTestSettings : public UObject
 public:
 
 	/**
-	 * The directory to which the packaged project will be copied.
+	 * The automation test map to be used for several of the automation tests.
 	 */
 	UPROPERTY(config, EditAnywhere, Category=Automation, meta=(FilePathFilter = "umap"))
 	FFilePath AutomationTestmap;
 
 	/**
-	 * Asset to test for open in automation process
-	 */
-	UPROPERTY(EditAnywhere, config, Category=Automation)
-	TArray<FOpenTestAsset> TestAssetsToOpen;
+	* The map to be used for the editor performance capture tool.
+	*/
+	UPROPERTY(config, EditAnywhere, Category = Automation, meta = (FilePathFilter = "umap"))
+	TArray<FEditorMapPerformanceTestDefinition> EditorPerformanceTestMaps;
 
 	/**
-	 * Modules to load that have engine tests
-	 */
-	UPROPERTY(EditAnywhere, config, Category=Automation)
+	* Editor build promotion test settings
+	*/
+	UPROPERTY(EditAnywhere, config, Category = Automation)
+	FBuildPromotionTestSettings BuildPromotionTest;
+
+	/**
+	* Modules to load that have engine tests
+	*/
+	UPROPERTY(EditAnywhere, config, Category = MiscAutomationSetups)
 	TArray<FString> EngineTestModules;
 
 	/**
-	 * Modules to load that have editor tests
-	 */
-	UPROPERTY(EditAnywhere, config, Category=Automation)
+	* Modules to load that have editor tests
+	*/
+	UPROPERTY(EditAnywhere, config, Category = MiscAutomationSetups)
 	TArray<FString> EditorTestModules;
 
 	/**
-	 * Folders containing levels to exclude from automated tests
-	 */
-	UPROPERTY(EditAnywhere, config, Category=Automation)
+	* Folders containing levels to exclude from automated tests
+	*/
+	UPROPERTY(EditAnywhere, config, Category = MiscAutomationSetups)
 	TArray<FString> TestLevelFolders;
 
 	/**
 	 * Asset to test for open in automation process
 	 */
+	UPROPERTY(EditAnywhere, config, Category=Automation)
+	TArray<FOpenTestAsset> TestAssetsToOpen;
+	
+	/**
+	 * External executables and scripts to run as part of automation.
+	 */
 	UPROPERTY(EditAnywhere, config, Category=ExternalTools)
 	TArray<FExternalToolDefinition> ExternalTools;
 
@@ -326,9 +355,7 @@ public:
 	UPROPERTY(EditAnywhere, config, Category = Automation)
 	TArray<FEditorImportExportTestDefinition> ImportExportTestDefinitions;
 
-	/**
-	* Editor build promotion test settings
-	*/
-	UPROPERTY(EditAnywhere, config, Category = Automation)
-	FBuildPromotionTestSettings BuildPromotionTest;
+
+
+
 };
