Increased texture streaming notification threshold to reduce spam

diff --git a/Engine/Source/Editor/UnrealEd/Private/TextureStreamingNotification.cpp b/Engine/Source/Editor/UnrealEd/Private/TextureStreamingNotification.cpp
--- a/Engine/Source/Editor/UnrealEd/Private/TextureStreamingNotification.cpp
+++ b/Engine/Source/Editor/UnrealEd/Private/TextureStreamingNotification.cpp
@@ -26,7 +26,7 @@ bool FTextureStreamingNotificationImpl::ShouldShowNotification(const bool bIsNot
 	// However, if we're already showing the notification, we should continue to do so until all the streaming has finished
 	// Never show these notifications during PIE
 	const int32 NumStreamingTextures = GetNumStreamingTextures();
-	return !GEditor->PlayWorld && ((NumStreamingTextures > 30) || (bIsNotificationAlreadyActive && NumStreamingTextures > 0));
+	return !GEditor->PlayWorld && ((NumStreamingTextures > 300) || (bIsNotificationAlreadyActive && NumStreamingTextures > 0));
 }
 
 void FTextureStreamingNotificationImpl::SetNotificationText(const TSharedPtr<SNotificationItem>& InNotificationItem) const
