Consistent localization-related commandlet path-handling with regards to relative paths. From now on, please specify paths relative to to the project's directory or the Engine directory.

diff --git a/Engine/Config/Localization/Editor.ini b/Engine/Config/Localization/Editor.ini
--- a/Engine/Config/Localization/Editor.ini
+++ b/Engine/Config/Localization/Editor.ini
@@ -1,20 +1,21 @@
 [CommonSettings]
-ManifestDependencies=../../../Engine/Content/Localization/Engine/Engine.manifest
-SourcePath=../../../Engine/Content/Localization/Editor
-DestinationPath=../../../Engine/Content/Localization/Editor
+ManifestDependencies=./Content/Localization/Engine/Engine.manifest
+SourcePath=./Content/Localization/Editor
+DestinationPath=./Content/Localization/Editor
 ManifestName=Editor.manifest
 ArchiveName=Editor.archive
+PortableObjectName=Editor.po
 SourceCulture=en
 CulturesToGenerate=en
+CulturesToGenerate=es
 CulturesToGenerate=ja
 CulturesToGenerate=ko
 CulturesToGenerate=zh-Hans-CN
-CulturesToGenerate=es
 
 ;Gather text from source code
 [GatherTextStep0]
 CommandletClass=GatherTextFromSource
-IncludePaths=Engine/Source/Editor/
+IncludePaths=./Source/Editor/
 SourceFileSearchFilters=*.cpp
 SourceFileSearchFilters=*.h
 SourceFileSearchFilters=*.c
@@ -24,13 +25,13 @@ SourceFileSearchFilters=*.mm
 ;Gather text from metadata
 [GatherTextStep1]
 CommandletClass=GatherTextFromMetaData
-IncludePaths=*/Engine/Source/Editor/*
-IncludePaths=*/Engine/Source/Runtime/*
-IncludePaths=*/Engine/Source/Developer/*
-ExcludePaths=*/Engine/Config/Localization/*
-ExcludePaths=*/Engine/Source/Developer/NoRedist/CommunityPortalServices/*
-ExcludePaths=*/Engine/Source/Developer/NoRedist/UnrealEngineLauncherServices/*
-ExcludePaths=*/Engine/Source/Developer/NoRedist/BuildPatchServices/*
+IncludePaths=./Source/Editor/*
+IncludePaths=./Source/Runtime/*
+IncludePaths=./Source/Developer/*
+ExcludePaths=*/Config/Localization/*
+ExcludePaths=*/Source/Developer/NoRedist/CommunityPortalServices/*
+ExcludePaths=*/Source/Developer/NoRedist/UnrealEngineLauncherServices/*
+ExcludePaths=*/Source/Developer/NoRedist/BuildPatchServices/*
 InputKeys=Category
 OutputNamespaces=UObjectCategories
 OutputKeys="{MetaDataValue}"
@@ -38,16 +39,11 @@ OutputKeys="{MetaDataValue}"
 ;Gather text from assets
 [GatherTextStep2]
 CommandletClass=GatherTextFromAssets
-IncludePaths=*/Engine/Content/Editor/*
-IncludePaths=*/Engine/Content/Editor*
+IncludePaths=./Content/Editor/*
+IncludePaths=./Content/Editor*
 PackageExtensions=*.umap
 PackageExtensions=*.upk
 PackageExtensions=*.uasset
-ExcludeClasses=Texture
-ExcludeClasses=TextureCube
-ExcludeClasses=Texture2D
-ExcludeClasses=StaticMesh
-ExcludeClasses=Material
 bFixBroken=false
 
 ;Write Manifest
@@ -63,9 +59,6 @@ bPurgeOldEmptyEntries=true
 [GatherTextStep5]
 CommandletClass=InternationalizationExport
 bImportLoc=true
-ManifestName=Editor.manifest
-ArchiveName=Editor.archive
-PortableObjectName=Editor.po
 
 ;Write Localized Text Resource
 [GatherTextStep6]
@@ -75,7 +68,4 @@ ResourceName=Editor.locres
 ;Import localized PO files
 [GatherTextStep7]
 CommandletClass=InternationalizationExport
-bExportLoc=true
-ManifestName=Editor.manifest
-ArchiveName=Editor.archive
-PortableObjectName=Editor.po
\ No newline at end of file
+bExportLoc=true
\ No newline at end of file
diff --git a/Engine/Config/Localization/Engine.ini b/Engine/Config/Localization/Engine.ini
--- a/Engine/Config/Localization/Engine.ini
+++ b/Engine/Config/Localization/Engine.ini
@@ -1,26 +1,27 @@
 [CommonSettings]
-SourcePath=../../../Engine/Content/Localization/Engine
-DestinationPath=../../../Engine/Content/Localization/Engine
+SourcePath=./Content/Localization/Engine
+DestinationPath=./Content/Localization/Engine
 ManifestName=Engine.manifest
 ArchiveName=Engine.archive
+PortableObjectName=Engine.po
 SourceCulture=en
 CulturesToGenerate=en
+CulturesToGenerate=es
 CulturesToGenerate=ja
 CulturesToGenerate=ko
 CulturesToGenerate=zh-Hans-CN
-CulturesToGenerate=es
 
 ;Gather text from source code
 [GatherTextStep0]
 CommandletClass=GatherTextFromSource
-IncludePaths=Engine/Source/Runtime
-IncludePaths=Engine/Source/Rocket
-IncludePaths=Engine/Source/Developer
-IncludePaths=Engine/Config/
-ExcludePaths=*/Engine/Config/Localization/*
-ExcludePaths=*/Engine/Source/Developer/NoRedist/CommunityPortalServices/*
-ExcludePaths=*/Engine/Source/Developer/NoRedist/UnrealEngineLauncherServices/*
-ExcludePaths=*/Engine/Source/Developer/NoRedist/BuildPatchServices/*
+IncludePaths=./Source/Runtime/
+IncludePaths=./Source/Rocket/
+IncludePaths=./Source/Developer/
+IncludePaths=./Config/
+ExcludePaths=*/Config/Localization/*
+ExcludePaths=*/Source/Developer/NoRedist/CommunityPortalServices/*
+ExcludePaths=*/Source/Developer/NoRedist/UnrealEngineLauncherServices/*
+ExcludePaths=*/Source/Developer/NoRedist/BuildPatchServices/*
 SourceFileSearchFilters=*.cpp
 SourceFileSearchFilters=*.h
 SourceFileSearchFilters=*.c
@@ -31,22 +32,17 @@ SourceFileSearchFilters=*.ini
 ;Gather text from assets
 [GatherTextStep1]
 CommandletClass=GatherTextFromAssets
-IncludePaths=*/Engine/Content/*
-ExcludePaths=*/Engine/Content/Editor/*
-ExcludePaths=*/Engine/Content/Developers/*
-ExcludePaths=*/Engine/Content/Localization/*
-ExcludePaths=*/Engine/Content/TestPackages/*
-ExcludePaths=*/Engine/Content/QA_Assets/*
-ExcludePaths=*/Engine/Content/EngineSounds/*
-ExcludePaths=*/Engine/Content/Maps/Automation/*
+IncludePaths=./Content/*
+ExcludePaths=*/Content/Editor/*
+ExcludePaths=*/Content/Developers/*
+ExcludePaths=*/Content/Localization/*
+ExcludePaths=*/Content/TestPackages/*
+ExcludePaths=*/Content/QA_Assets/*
+ExcludePaths=*/Content/EngineSounds/*
+ExcludePaths=*/Content/Maps/Automation/*
 PackageExtensions=*.umap
 PackageExtensions=*.upk
 PackageExtensions=*.uasset
-ExcludeClasses=Texture
-ExcludeClasses=TextureCube
-ExcludeClasses=Texture2D
-ExcludeClasses=StaticMesh
-ExcludeClasses=Material
 bFixBroken=false
 
 ;Write Manifest
@@ -58,14 +54,17 @@ CommandletClass=GenerateGatherManifest
 CommandletClass=GenerateGatherArchive
 bPurgeOldEmptyEntries=true
 
-;Write Localized Text Resource
+;Import localized PO files
 [GatherTextStep4]
+CommandletClass=InternationalizationExport
+bImportLoc=true
+
+;Write Localized Text Resource
+[GatherTextStep5]
 CommandletClass=GenerateTextLocalizationResource
 ResourceName=Engine.locres
 
-;Write Text Localization Report
-[GatherTextStep5]
-CommandletClass=GenerateTextLocalizationReport
-DestinationPath=\\epicgames.net\root\UE3\Localization\WordCounts
-bConflictReport=true
-ConflictReportName=EngineConflicts.txt
\ No newline at end of file
+;Export  PO files
+[GatherTextStep6]
+CommandletClass=InternationalizationExport
+bExportLoc=true
\ No newline at end of file
diff --git a/Engine/Config/Localization/PortableObjectExport.ini b/Engine/Config/Localization/PortableObjectExport.ini
--- a/Engine/Config/Localization/PortableObjectExport.ini
+++ b/Engine/Config/Localization/PortableObjectExport.ini
@@ -1,4 +1,5 @@
 [CommonSettings]
+bExportLoc=true
 CulturesToGenerate=en
 CulturesToGenerate=es
 CulturesToGenerate=ja
@@ -7,36 +8,32 @@ CulturesToGenerate=zh-Hans-CN
 
 [GatherTextStep0]
 CommandletClass=InternationalizationExport
-bExportLoc=true
-SourcePath=../../../Engine/Content/Localization/Engine
-DestinationPath=../../../Engine/Content/Localization/Engine
+SourcePath=./Content/Localization/Engine
+DestinationPath=./Content/Localization/Engine
 ManifestName=Engine.manifest
 ArchiveName=Engine.archive
 PortableObjectName=Engine.po
 
 [GatherTextStep1]
 CommandletClass=InternationalizationExport
-bExportLoc=true
-SourcePath=../../../Engine/Content/Localization/Editor
-DestinationPath=../../../Engine/Content/Localization/Editor
+SourcePath=./Content/Localization/Editor
+DestinationPath=./Content/Localization/Editor
 ManifestName=Editor.manifest
 ArchiveName=Editor.archive
 PortableObjectName=Editor.po
 
 [GatherTextStep2]
 CommandletClass=InternationalizationExport
-bExportLoc=true
-SourcePath=../../../Engine/Content/Localization/ToolTips
-DestinationPath=../../../Engine/Content/Localization/ToolTips
+SourcePath=./Content/Localization/ToolTips
+DestinationPath=./Content/Localization/ToolTips
 ManifestName=ToolTips.manifest
 ArchiveName=ToolTips.archive
 PortableObjectName=ToolTips.po
 
 [GatherTextStep3]
 CommandletClass=InternationalizationExport
-bExportLoc=true
-SourcePath=../../../Engine/Content/Localization/PropertyNames
-DestinationPath=../../../Engine/Content/Localization/PropertyNames
+SourcePath=./Content/Localization/PropertyNames
+DestinationPath=./Content/Localization/PropertyNames
 ManifestName=PropertyNames.manifest
 ArchiveName=PropertyNames.archive
 PortableObjectName=PropertyNames.po
\ No newline at end of file
diff --git a/Engine/Config/Localization/PortableObjectImport.ini b/Engine/Config/Localization/PortableObjectImport.ini
--- a/Engine/Config/Localization/PortableObjectImport.ini
+++ b/Engine/Config/Localization/PortableObjectImport.ini
@@ -1,4 +1,5 @@
 [CommonSettings]
+bImportLoc=true
 CulturesToGenerate=en
 CulturesToGenerate=es
 CulturesToGenerate=ja
@@ -7,16 +8,16 @@ CulturesToGenerate=zh-Hans-CN
 
 [GatherTextStep0]
 CommandletClass=InternationalizationExport
-bImportLoc=true
-SourcePath=../../../Engine/Content/Localization/Engine
-DestinationPath=../../../Engine/Content/Localization/Engine
+SourcePath=./Content/Localization/Engine
+DestinationPath=./Content/Localization/Engine
 ManifestName=Engine.manifest
 ArchiveName=Engine.archive
+PortableObjectName=Engine.po
 
 [GatherTextStep1]
 CommandletClass=InternationalizationExport
-bImportLoc=true
-SourcePath=../../../Engine/Content/Localization/Editor
-DestinationPath=../../../Engine/Content/Localization/Editor
+SourcePath=./Content/Localization/Editor
+DestinationPath=./Content/Localization/Editor
 ManifestName=Editor.manifest
-ArchiveName=Editor.archive
\ No newline at end of file
+ArchiveName=Editor.archive
+PortableObjectName=Editor.po
diff --git a/Engine/Config/Localization/PropertyNames.ini b/Engine/Config/Localization/PropertyNames.ini
--- a/Engine/Config/Localization/PropertyNames.ini
+++ b/Engine/Config/Localization/PropertyNames.ini
@@ -1,6 +1,6 @@
 [CommonSettings]
-SourcePath=../../../Engine/Content/Localization/PropertyNames
-DestinationPath=../../../Engine/Content/Localization/PropertyNames
+SourcePath=./Content/Localization/PropertyNames
+DestinationPath=./Content/Localization/PropertyNames
 ManifestName=PropertyNames.manifest
 ArchiveName=PropertyNames.archive
 SourceCulture=en
@@ -13,11 +13,11 @@ CulturesToGenerate=zh-Hans-CN
 ;Gather text from metadata
 [GatherTextStep0]
 CommandletClass=GatherTextFromMetaData
-IncludePaths=*/Engine/Source/Editor/*
-IncludePaths=*/Engine/Source/Runtime/*
-IncludePaths=*/Engine/Source/Developer/*
-ExcludePaths=*/Engine/Config/Localization/*
-ExcludePaths=*/Engine/Source/Developer/NoRedist/CommunityPortalServices/*
+IncludePaths=./Source/Editor/*
+IncludePaths=./Source/Runtime/*
+IncludePaths=./Source/Developer/*
+ExcludePaths=*/Config/Localization/*
+ExcludePaths=*/Source/Developer/NoRedist/CommunityPortalServices/*
 InputKeys=DisplayName
 OutputNamespaces=UObjectDisplayNames
 OutputKeys="{FieldPath}"
diff --git a/Engine/Config/Localization/RegenerateEditor.ini b/Engine/Config/Localization/RegenerateEditor.ini
--- a/Engine/Config/Localization/RegenerateEditor.ini
+++ b/Engine/Config/Localization/RegenerateEditor.ini
@@ -1,5 +1,5 @@
 [RegenerateResources]
-SourcePath=../../../Engine/Content/Localization/Editor
-DestinationPath=../../../Engine/Content/Localization/Editor
+SourcePath=./Content/Localization/Editor
+DestinationPath=./Content/Localization/Editor
 ManifestName=Editor.manifest
 ResourceName=Editor.locres
\ No newline at end of file
diff --git a/Engine/Config/Localization/RegenerateEngine.ini b/Engine/Config/Localization/RegenerateEngine.ini
--- a/Engine/Config/Localization/RegenerateEngine.ini
+++ b/Engine/Config/Localization/RegenerateEngine.ini
@@ -1,5 +1,5 @@
 [RegenerateResources]
-SourcePath=../../../Engine/Content/Localization/Engine
-DestinationPath=../../../Engine/Content/Localization/Engine
+SourcePath=./Content/Localization/Engine
+DestinationPath=./Content/Localization/Engine
 ManifestName=Engine.manifest
 ResourceName=Engine.locres
\ No newline at end of file
diff --git a/Engine/Config/Localization/ResourceFileGen.ini b/Engine/Config/Localization/ResourceFileGen.ini
--- a/Engine/Config/Localization/ResourceFileGen.ini
+++ b/Engine/Config/Localization/ResourceFileGen.ini
@@ -1,51 +1,38 @@
-;Write Localized Text Resource for Engine
-[GatherTextStep0]
-CommandletClass=GenerateTextLocalizationResource
-SourcePath=../../../Engine/Content/Localization/Engine
-DestinationPath=../../../Engine/Content/Localization/Engine
-ManifestName=Engine.manifest
+[CommonSettings]
 CulturesToGenerate=en
 CulturesToGenerate=es
 CulturesToGenerate=ja
 CulturesToGenerate=ko
 CulturesToGenerate=zh-Hans-CN
+
+;Write Localized Text Resource for Engine
+[GatherTextStep0]
+CommandletClass=GenerateTextLocalizationResource
+SourcePath=./Content/Localization/Engine
+DestinationPath=./Content/Localization/Engine
+ManifestName=Engine.manifest
 ResourceName=Engine.locres
 
 ;Write Localized Text Resource for Editor
 [GatherTextStep1]
 CommandletClass=GenerateTextLocalizationResource
-SourcePath=../../../Engine/Content/Localization/Editor
-DestinationPath=../../../Engine/Content/Localization/Editor
+SourcePath=./Content/Localization/Editor
+DestinationPath=./Content/Localization/Editor
 ManifestName=Editor.manifest
-CulturesToGenerate=en
-CulturesToGenerate=es
-CulturesToGenerate=ja
-CulturesToGenerate=ko
-CulturesToGenerate=zh-Hans-CN
 ResourceName=Editor.locres
 
 ;Write Localized Text Resource for Editor Tool Tips
-[GatherTextStep1]
+[GatherTextStep2]
 CommandletClass=GenerateTextLocalizationResource
-SourcePath=../../../Engine/Content/Localization/Editor/ToolTips
-DestinationPath=../../../Engine/Content/Localization/Editor/ToolTips
+SourcePath=./Content/Localization/ToolTips
+DestinationPath=./Content/Localization/ToolTips
 ManifestName=ToolTips.manifest
-CulturesToGenerate=en
-CulturesToGenerate=es
-CulturesToGenerate=ja
-CulturesToGenerate=ko
-CulturesToGenerate=zh-Hans-CN
 ResourceName=ToolTips.locres
 
 ;Write Localized Text Resource for Editor Property Names
-[GatherTextStep1]
+[GatherTextStep3]
 CommandletClass=GenerateTextLocalizationResource
-SourcePath=../../../Engine/Content/Localization/Editor/PropertyNames
-DestinationPath=../../../Engine/Content/Localization/Editor/PropertyNames
+SourcePath=./Content/Localization/PropertyNames
+DestinationPath=./Content/Localization/PropertyNames
 ManifestName=PropertyNames.manifest
-CulturesToGenerate=en
-CulturesToGenerate=es
-CulturesToGenerate=ja
-CulturesToGenerate=ko
-CulturesToGenerate=zh-Hans-CN
 ResourceName=PropertyNames.locres
\ No newline at end of file
diff --git a/Engine/Config/Localization/ToolTips.ini b/Engine/Config/Localization/ToolTips.ini
--- a/Engine/Config/Localization/ToolTips.ini
+++ b/Engine/Config/Localization/ToolTips.ini
@@ -1,6 +1,6 @@
 [CommonSettings]
-SourcePath=../../../Engine/Content/Localization/ToolTips
-DestinationPath=../../../Engine/Content/Localization/ToolTips
+SourcePath=./Content/Localization/ToolTips
+DestinationPath=./Content/Localization/ToolTips
 ManifestName=ToolTips.manifest
 ArchiveName=ToolTips.archive
 SourceCulture=en
@@ -13,11 +13,11 @@ CulturesToGenerate=zh-Hans-CN
 ;Gather text from metadata
 [GatherTextStep0]
 CommandletClass=GatherTextFromMetaData
-IncludePaths=*/Engine/Source/Editor/*
-IncludePaths=*/Engine/Source/Runtime/*
-IncludePaths=*/Engine/Source/Developer/*
-ExcludePaths=*/Engine/Config/Localization/*
-ExcludePaths=*/Engine/Source/Developer/NoRedist/CommunityPortalServices/*
+IncludePaths=./Source/Editor/*
+IncludePaths=./Source/Runtime/*
+IncludePaths=./Source/Developer/*
+ExcludePaths=*/Config/Localization/*
+ExcludePaths=*/Source/Developer/NoRedist/CommunityPortalServices/*
 InputKeys=ToolTip
 OutputNamespaces=UObjectToolTips
 OutputKeys="{FieldPath}"
diff --git a/Engine/Config/Localization/WordCount.ini b/Engine/Config/Localization/WordCount.ini
--- a/Engine/Config/Localization/WordCount.ini
+++ b/Engine/Config/Localization/WordCount.ini
@@ -1,50 +1,35 @@
 [CommonSettings]
+bWordCountReport=true
 DestinationPath=\\epicgames.net\root\UE3\Localization\WordCounts
+CulturesToGenerate=es
+CulturesToGenerate=ja
+CulturesToGenerate=ko
+CulturesToGenerate=zh-Hans-CN
 
 ;Write Text Localization Report for Engine
 [GatherTextStep0]
 CommandletClass=GenerateTextLocalizationReport
-bWordCountReport=true
-SourcePath=../../../Engine/Content/Localization/Engine
+SourcePath=./Content/Localization/Engine
 ManifestName=Engine.manifest
-CulturesToGenerate=es
-CulturesToGenerate=ja
-CulturesToGenerate=ko
-CulturesToGenerate=zh-Hans-CN
 WordCountReportName=Engine.csv
 
 ;Write Text Localization Report for Editor
 [GatherTextStep1]
 CommandletClass=GenerateTextLocalizationReport
-bWordCountReport=true
-SourcePath=../../../Engine/Content/Localization/Editor
+SourcePath=./Content/Localization/Editor
 ManifestName=Editor.manifest
-CulturesToGenerate=es
-CulturesToGenerate=ja
-CulturesToGenerate=ko
-CulturesToGenerate=zh-Hans-CN
 WordCountReportName=Editor.csv
 
 ;Write Text Localization Report for Editor Tool Tips
 [GatherTextStep2]
 CommandletClass=GenerateTextLocalizationReport
-bWordCountReport=true
-SourcePath=../../../Engine/Content/Localization/ToolTips
+SourcePath=./Content/Localization/ToolTips
 ManifestName=ToolTips.manifest
-CulturesToGenerate=es
-CulturesToGenerate=ja
-CulturesToGenerate=ko
-CulturesToGenerate=zh-Hans-CN
 WordCountReportName=ToolTips.csv
 
 ;Write Text Localization Report for Editor Property Names
 [GatherTextStep3]
 CommandletClass=GenerateTextLocalizationReport
-bWordCountReport=true
-SourcePath=../../../Engine/Content/Localization/PropertyNames
+SourcePath=./Content/Localization/PropertyNames
 ManifestName=PropertyNames.manifest
-CulturesToGenerate=es
-CulturesToGenerate=ja
-CulturesToGenerate=ko
-CulturesToGenerate=zh-Hans-CN
 WordCountReportName=PropertyNames.csv
diff --git a/Engine/Source/Editor/UnrealEd/Private/Commandlets/GatherTextCommandlet.cpp b/Engine/Source/Editor/UnrealEd/Private/Commandlets/GatherTextCommandlet.cpp
--- a/Engine/Source/Editor/UnrealEd/Private/Commandlets/GatherTextCommandlet.cpp
+++ b/Engine/Source/Editor/UnrealEd/Private/Commandlets/GatherTextCommandlet.cpp
@@ -44,6 +44,20 @@ int32 UGatherTextCommandlet::Main( const FString& Params )
 		return -1;
 	}
 
+	if(FPaths::IsRelative(GatherTextConfigPath))
+	{
+		FString ProjectBasePath;
+		if (!FPaths::GameDir().IsEmpty())
+		{
+			ProjectBasePath = FPaths::GameDir();
+		}
+		else
+		{
+			ProjectBasePath = FPaths::EngineDir();
+		}
+		GatherTextConfigPath = FPaths::Combine( *ProjectBasePath, *GatherTextConfigPath );
+	}
+
 	GConfig->LoadFile(*GatherTextConfigPath);
 
 	FConfigFile* ConfigFile = GConfig->FindConfigFile(*GatherTextConfigPath);
diff --git a/Engine/Source/Editor/UnrealEd/Private/Commandlets/GatherTextFromAssetsCommandlet.cpp b/Engine/Source/Editor/UnrealEd/Private/Commandlets/GatherTextFromAssetsCommandlet.cpp
--- a/Engine/Source/Editor/UnrealEd/Private/Commandlets/GatherTextFromAssetsCommandlet.cpp
+++ b/Engine/Source/Editor/UnrealEd/Private/Commandlets/GatherTextFromAssetsCommandlet.cpp
@@ -692,6 +692,39 @@ int32 UGatherTextFromAssetsCommandlet::Main(const FString& Params)
 		return -1;
 	}
 
+	for (FString& IncludePath : IncludePaths)
+	{
+		FPaths::NormalizeDirectoryName(IncludePath);
+		if (FPaths::IsRelative(IncludePath))
+		{
+			if (!FPaths::GameDir().IsEmpty())
+			{
+				IncludePath = FPaths::Combine( *( FPaths::GameDir() ), *IncludePath );
+			}
+			else
+			{
+				IncludePath = FPaths::Combine( *( FPaths::EngineDir() ), *IncludePath );
+			}
+		}
+
+		IncludePath = FPaths::ConvertRelativePathToFull(IncludePath);
+
+		// All paths must ends with "/*"
+		if ( !IncludePath.EndsWith(TEXT("/*")) )
+		{
+			// If it ends in a slash, add the star.
+			if ( IncludePath.EndsWith(TEXT("/")) )
+			{
+				IncludePath.AppendChar(TEXT('*'));
+			}
+			// If it doesn't end in a slash or slash star, just add slash star.
+			else
+			{
+				IncludePath.Append(TEXT("/*"));
+			}
+		}
+	}
+
 	//Exclude paths
 	TArray<FString> ExcludePaths;
 	GetConfigArray(*SectionName, TEXT("ExcludePaths"), ExcludePaths, GatherTextConfigPath);
@@ -746,6 +779,21 @@ int32 UGatherTextFromAssetsCommandlet::Main(const FString& Params)
 	TArray<FString> ManifestDependenciesList;
 	GetConfigArray(*SectionName, TEXT("ManifestDependencies"), ManifestDependenciesList, GatherTextConfigPath);
 	
+	for (FString& ManifestDependency : ManifestDependenciesList)
+	{
+		if (FPaths::IsRelative(ManifestDependency))
+		{
+			if (!FPaths::GameDir().IsEmpty())
+			{
+				ManifestDependency = FPaths::Combine( *( FPaths::GameDir() ), *ManifestDependency );
+			}
+			else
+			{
+				ManifestDependency = FPaths::Combine( *( FPaths::EngineDir() ), *ManifestDependency );
+			}
+		}
+	}
+
 	if( !ManifestInfo->AddManifestDependencies( ManifestDependenciesList ) )
 	{
 		UE_LOG(LogGatherTextFromAssetsCommandlet, Error, TEXT("The GatherTextFromAssets commandlet couldn't find all the specified manifest dependencies."));
@@ -779,14 +827,16 @@ int32 UGatherTextFromAssetsCommandlet::Main(const FString& Params)
 		}
 
 		//Run through all the files found and add any that pass the include, exclude and filter constraints to OrderedPackageFilesToLoad
-		for( int32 PackageFileIdx=0; PackageFileIdx<PackageFiles.Num(); ++PackageFileIdx )
+		for (FString& PackageFile : PackageFiles)
 		{
+			PackageFile = FPaths::ConvertRelativePathToFull(PackageFile);
+
 			bool bExclude = false;
 			//Ensure it matches the include paths if there are some.
-			for( int32 IncludePathIdx=0; IncludePathIdx<IncludePaths.Num() ; ++IncludePathIdx )
+			for (FString& IncludePath : IncludePaths)
 			{
 				bExclude = true;
-				if( PackageFiles[PackageFileIdx].MatchesWildcard(IncludePaths[IncludePathIdx]) )
+				if( PackageFile.MatchesWildcard(IncludePath) )
 				{
 					bExclude = false;
 					break;
@@ -795,31 +845,31 @@ int32 UGatherTextFromAssetsCommandlet::Main(const FString& Params)
 
 			if ( bExclude )
 			{
-				PackageFilesNotInIncludePath.Add(PackageFiles[PackageFileIdx]);
+				PackageFilesNotInIncludePath.Add(PackageFile);
 			}
 
 			//Ensure it does not match the exclude paths if there are some.
 			for( int32 ExcludePathIdx=0; !bExclude && ExcludePathIdx<ExcludePaths.Num() ; ++ExcludePathIdx )
 			{
-				if( PackageFiles[PackageFileIdx].MatchesWildcard(ExcludePaths[ExcludePathIdx]) )
+				if( PackageFile.MatchesWildcard(ExcludePaths[ExcludePathIdx]) )
 				{
 					bExclude = true;
-					PackageFilesInExcludePath.Add(PackageFiles[PackageFileIdx]);
+					PackageFilesInExcludePath.Add(PackageFile);
 					break;
 				}
 			}
 
 			//Check that this is not on the list of packages that we don't care about e.g. textures.
-			if ( !bExclude && IsAssetPackage && LongPackageNamesToExclude.Contains( PackageFiles[PackageFileIdx] ) )
+			if ( !bExclude && IsAssetPackage && LongPackageNamesToExclude.Contains( PackageFile ) )
 			{
 				bExclude = true;
-				PackageFilesExcludedByClass.Add(PackageFiles[PackageFileIdx]);
+				PackageFilesExcludedByClass.Add(PackageFile);
 			}
 
 			//If we haven't failed one of the above checks, add it to the array of packages to process.
 			if(!bExclude)
 			{
-				TScopedPointer< FArchive > FileReader( IFileManager::Get().CreateFileReader( *PackageFiles[PackageFileIdx] ) );
+				TScopedPointer< FArchive > FileReader( IFileManager::Get().CreateFileReader( *PackageFile ) );
 				if( FileReader )
 				{
 					// Read package file summary from the file
@@ -829,7 +879,7 @@ int32 UGatherTextFromAssetsCommandlet::Main(const FString& Params)
 					// Early out check if the package has been flagged as needing localization gathering
 					if( PackageSummary.PackageFlags & PKG_RequiresLocalizationGather || PackageSummary.GetFileVersionUE4() < VER_UE4_PACKAGE_REQUIRES_LOCALIZATION_GATHER_FLAGGING )
 					{
-						PackageFileNamesToLoad.Add( PackageFiles[PackageFileIdx] );
+						PackageFileNamesToLoad.Add( PackageFile );
 					}
 				}
 			}
diff --git a/Engine/Source/Editor/UnrealEd/Private/Commandlets/GatherTextFromMetadataCommandlet.cpp b/Engine/Source/Editor/UnrealEd/Private/Commandlets/GatherTextFromMetadataCommandlet.cpp
--- a/Engine/Source/Editor/UnrealEd/Private/Commandlets/GatherTextFromMetadataCommandlet.cpp
+++ b/Engine/Source/Editor/UnrealEd/Private/Commandlets/GatherTextFromMetadataCommandlet.cpp
@@ -59,6 +59,39 @@ int32 UGatherTextFromMetaDataCommandlet::Main( const FString& Params )
 		return -1;
 	}
 
+	for (FString& IncludePath : IncludePaths)
+	{
+		FPaths::NormalizeDirectoryName(IncludePath);
+		if (FPaths::IsRelative(IncludePath))
+		{
+			if (!FPaths::GameDir().IsEmpty())
+			{
+				IncludePath = FPaths::Combine( *( FPaths::GameDir() ), *IncludePath );
+			}
+			else
+			{
+				IncludePath = FPaths::Combine( *( FPaths::EngineDir() ), *IncludePath );
+			}
+		}
+
+		IncludePath = FPaths::ConvertRelativePathToFull(IncludePath);
+
+		// All paths must ends with "/*"
+		if ( !IncludePath.EndsWith(TEXT("/*")) )
+		{
+			// If it ends in a slash, add the star.
+			if ( IncludePath.EndsWith(TEXT("/")) )
+			{
+				IncludePath.AppendChar(TEXT('*'));
+			}
+			// If it doesn't end in a slash or slash star, just add slash star.
+			else
+			{
+				IncludePath.Append(TEXT("/*"));
+			}
+		}
+	}
+
 	//Exclude paths
 	TArray<FString> ExcludePaths;
 	GetConfigArray(*SectionName, TEXT("ExcludePaths"), ExcludePaths, GatherTextConfigPath);
@@ -90,6 +123,21 @@ int32 UGatherTextFromMetaDataCommandlet::Main( const FString& Params )
 	TArray<FString> ManifestDependenciesList;
 	GetConfigArray(*SectionName, TEXT("ManifestDependencies"), ManifestDependenciesList, GatherTextConfigPath);
 	
+	for (FString& ManifestDependency : ManifestDependenciesList)
+	{
+		if (FPaths::IsRelative(ManifestDependency))
+		{
+			if (!FPaths::GameDir().IsEmpty())
+			{
+				ManifestDependency = FPaths::Combine( *( FPaths::GameDir() ), *ManifestDependency );
+			}
+			else
+			{
+				ManifestDependency = FPaths::Combine( *( FPaths::EngineDir() ), *ManifestDependency );
+			}
+		}
+	}
+
 	if( !ManifestInfo->AddManifestDependencies( ManifestDependenciesList ) )
 	{
 		UE_LOG(LogGatherTextFromMetaDataCommandlet, Error, TEXT("The GatherTextFromMetaData commandlet couldn't find all the specified manifest dependencies."));
@@ -105,6 +153,7 @@ void UGatherTextFromMetaDataCommandlet::GatherTextFromUObjects(const TArray<FStr
 	{
 		FString SourceFilePath;
 		FSourceCodeNavigation::FindClassHeaderPath(*It, SourceFilePath);
+		SourceFilePath = FPaths::ConvertRelativePathToFull(SourceFilePath);
 
 		// Returns true if in an include path. False otherwise.
 		auto IncludePathLogic = [&]() -> bool
diff --git a/Engine/Source/Editor/UnrealEd/Private/Commandlets/GatherTextFromSourceCommandlet.cpp b/Engine/Source/Editor/UnrealEd/Private/Commandlets/GatherTextFromSourceCommandlet.cpp
--- a/Engine/Source/Editor/UnrealEd/Private/Commandlets/GatherTextFromSourceCommandlet.cpp
+++ b/Engine/Source/Editor/UnrealEd/Private/Commandlets/GatherTextFromSourceCommandlet.cpp
@@ -93,13 +93,28 @@ int32 UGatherTextFromSourceCommandlet::Main( const FString& Params )
 	// Search in the root folder for each of the wildcard filters specified and build a list of files
 	TArray<FString> AllFoundFiles;
 
-	for(int32 i = 0; i < IncludePaths.Num(); i++)
+	for (FString& IncludePath : IncludePaths)
 	{
+		FString ProjectBasePath;
+		if (!FPaths::GameDir().IsEmpty())
+		{
+			ProjectBasePath = FPaths::GameDir();
+		}
+		else
+		{
+			ProjectBasePath = FPaths::EngineDir();
+		}
+
+		if(FPaths::IsRelative(IncludePath))
+		{
+			IncludePath = FPaths::Combine( *ProjectBasePath, *IncludePath );
+		}
+
 		for (int32 SourceFileSearchIdx=0; SourceFileSearchIdx < UniqueSourceFileSearchFilters.Num(); SourceFileSearchIdx++)
 		{
 			TArray<FString> RootSourceFiles;
 
-			IFileManager::Get().FindFilesRecursive(RootSourceFiles, *(FPaths::RootDir() + IncludePaths[i]), *UniqueSourceFileSearchFilters[SourceFileSearchIdx], true, false,false);
+			IFileManager::Get().FindFilesRecursive(RootSourceFiles, *IncludePath, *UniqueSourceFileSearchFilters[SourceFileSearchIdx], true, false,false);
 
 			AllFoundFiles.Append(RootSourceFiles);
 		}
@@ -142,6 +157,21 @@ int32 UGatherTextFromSourceCommandlet::Main( const FString& Params )
 	TArray<FString> ManifestDependenciesList;
 	GetConfigArray(*SectionName, TEXT("ManifestDependencies"), ManifestDependenciesList, GatherTextConfigPath);
 	
+	for (FString& ManifestDependency : ManifestDependenciesList)
+	{
+		if (FPaths::IsRelative(ManifestDependency))
+		{
+			if (!FPaths::GameDir().IsEmpty())
+			{
+				ManifestDependency = FPaths::Combine( *( FPaths::GameDir() ), *ManifestDependency );
+			}
+			else
+			{
+				ManifestDependency = FPaths::Combine( *( FPaths::EngineDir() ), *ManifestDependency );
+			}
+		}
+	}
+
 	if( !ManifestInfo->AddManifestDependencies( ManifestDependenciesList ) )
 	{
 		UE_LOG(LogGatherTextFromSourceCommandlet, Error, TEXT("The GatherTextFromSource commandlet couldn't find all the specified manifest dependencies."));
@@ -175,9 +205,20 @@ int32 UGatherTextFromSourceCommandlet::Main( const FString& Params )
 	ParseCtxt.ManifestInfo = ManifestInfo;
 
 	// Parse all source files for macros and add entries to SourceParsedEntries
-	for (int32 SourceFileIdx=0; SourceFileIdx < FilesToProcess.Num(); SourceFileIdx++)
+	for ( FString& SourceFile : FilesToProcess)
 	{
-		ParseCtxt.Filename = FilesToProcess[SourceFileIdx];
+		FString ProjectBasePath;
+		if (!FPaths::GameDir().IsEmpty())
+		{
+			ProjectBasePath = FPaths::GameDir();
+		}
+		else
+		{
+			ProjectBasePath = FPaths::EngineDir();
+		}
+
+		ParseCtxt.Filename = SourceFile;
+		FPaths::MakePathRelativeTo(ParseCtxt.Filename, *ProjectBasePath);
 		ParseCtxt.LineNumber = 0;
 		ParseCtxt.LineText.Empty();
 		ParseCtxt.Namespace.Empty();
@@ -189,7 +230,7 @@ int32 UGatherTextFromSourceCommandlet::Main( const FString& Params )
 		ParseCtxt.WithinStartingLine.Empty();
 
 		FString SourceFileText;
-		if (!FFileHelper::LoadFileToString(SourceFileText, *ParseCtxt.Filename))
+		if (!FFileHelper::LoadFileToString(SourceFileText, *SourceFile))
 		{
 			UE_LOG(LogGatherTextFromSourceCommandlet, Error, TEXT("GatherTextSource failed to open file %s"), *ParseCtxt.Filename);
 		}
diff --git a/Engine/Source/Editor/UnrealEd/Private/Commandlets/GenerateGatherArchiveCommandlet.cpp b/Engine/Source/Editor/UnrealEd/Private/Commandlets/GenerateGatherArchiveCommandlet.cpp
--- a/Engine/Source/Editor/UnrealEd/Private/Commandlets/GenerateGatherArchiveCommandlet.cpp
+++ b/Engine/Source/Editor/UnrealEd/Private/Commandlets/GenerateGatherArchiveCommandlet.cpp
@@ -58,14 +58,6 @@ int32 UGenerateGatherArchiveCommandlet::Main( const FString& Params )
 		return -1;
 	}
 
-	// Get source path.
-	FString SourcePath;
-	if( !( GetConfigString( *SectionName, TEXT("SourcePath"), SourcePath, GatherTextConfigPath ) ) )
-	{
-		UE_LOG(LogGenerateArchiveCommandlet, Error, TEXT("No source path specified."));
-		return -1;
-	}
-
 	// Get manifest name.
 	FString ManifestName;
 	if( !GetConfigString( *SectionName, TEXT("ManifestName"), ManifestName, GatherTextConfigPath ) )
@@ -110,6 +102,18 @@ int32 UGenerateGatherArchiveCommandlet::Main( const FString& Params )
 		return -1;
 	}
 
+	if (FPaths::IsRelative(DestinationPath))
+	{
+		if (!FPaths::GameDir().IsEmpty())
+		{
+			DestinationPath = FPaths::Combine( *( FPaths::GameDir() ), *DestinationPath );
+		}
+		else
+		{
+			DestinationPath = FPaths::Combine( *( FPaths::EngineDir() ), *DestinationPath );
+		}
+	}
+
 	// Get archive name.
 	FString ArchiveName;
 	if( !( GetConfigString(* SectionName, TEXT("ArchiveName"), ArchiveName, GatherTextConfigPath ) ) )
diff --git a/Engine/Source/Editor/UnrealEd/Private/Commandlets/GenerateGatherManifestCommandlet.cpp b/Engine/Source/Editor/UnrealEd/Private/Commandlets/GenerateGatherManifestCommandlet.cpp
--- a/Engine/Source/Editor/UnrealEd/Private/Commandlets/GenerateGatherManifestCommandlet.cpp
+++ b/Engine/Source/Editor/UnrealEd/Private/Commandlets/GenerateGatherManifestCommandlet.cpp
@@ -60,6 +60,18 @@ int32 UGenerateGatherManifestCommandlet::Main( const FString& Params )
 		return -1;
 	}
 
+	if (FPaths::IsRelative(DestinationPath))
+	{
+		if (!FPaths::GameDir().IsEmpty())
+		{
+			DestinationPath = FPaths::Combine( *( FPaths::GameDir() ), *DestinationPath );
+		}
+		else
+		{
+			DestinationPath = FPaths::Combine( *( FPaths::EngineDir() ), *DestinationPath );
+		}
+	}
+
 	// Get manifest name.
 	FString ManifestName;
 	if( !GetConfigString( *SectionName, TEXT("ManifestName"), ManifestName, GatherTextConfigPath ) )
@@ -72,6 +84,21 @@ int32 UGenerateGatherManifestCommandlet::Main( const FString& Params )
 	TArray<FString> ManifestDependenciesList;
 	GetConfigArray(*SectionName, TEXT("ManifestDependencies"), ManifestDependenciesList, GatherTextConfigPath);
 
+	for (FString& ManifestDependency : ManifestDependenciesList)
+	{
+		if (FPaths::IsRelative(ManifestDependency))
+		{
+			if (!FPaths::GameDir().IsEmpty())
+			{
+				ManifestDependency = FPaths::Combine( *( FPaths::GameDir() ), *ManifestDependency );
+			}
+			else
+			{
+				ManifestDependency = FPaths::Combine( *( FPaths::EngineDir() ), *ManifestDependency );
+			}
+		}
+	}
+
 	if( ManifestDependenciesList.Num() > 0 )
 	{
 		if( !ManifestInfo->AddManifestDependencies( ManifestDependenciesList ) )
diff --git a/Engine/Source/Editor/UnrealEd/Private/Commandlets/GenerateTextLocalizationReportCommandlet.cpp b/Engine/Source/Editor/UnrealEd/Private/Commandlets/GenerateTextLocalizationReportCommandlet.cpp
--- a/Engine/Source/Editor/UnrealEd/Private/Commandlets/GenerateTextLocalizationReportCommandlet.cpp
+++ b/Engine/Source/Editor/UnrealEd/Private/Commandlets/GenerateTextLocalizationReportCommandlet.cpp
@@ -89,6 +89,18 @@ int32 UGenerateTextLocalizationReportCommandlet::Main(const FString& Params)
 		return -1;
 	}
 
+	if (FPaths::IsRelative(SourcePath))
+	{
+		if (!FPaths::GameDir().IsEmpty())
+		{
+			SourcePath = FPaths::Combine( *( FPaths::GameDir() ), *SourcePath );
+		}
+		else
+		{
+			SourcePath = FPaths::Combine( *( FPaths::EngineDir() ), *SourcePath );
+		}
+	}
+
 	// Get destination path.
 	if( !( GetConfigString( *SectionName, TEXT("DestinationPath"), DestinationPath, GatherTextConfigPath ) ) )
 	{
@@ -96,6 +108,18 @@ int32 UGenerateTextLocalizationReportCommandlet::Main(const FString& Params)
 		return -1;
 	}
 
+	if (FPaths::IsRelative(DestinationPath))
+	{
+		if (!FPaths::GameDir().IsEmpty())
+		{
+			DestinationPath = FPaths::Combine( *( FPaths::GameDir() ), *DestinationPath );
+		}
+		else
+		{
+			DestinationPath = FPaths::Combine( *( FPaths::EngineDir() ), *DestinationPath );
+		}
+	}
+
 	// Get the timestamp from the commandline, if not provided we will use the current time.
 	const FString* TimeStampParamVal = ParamVals.Find(FString(TEXT("TimeStamp")));
 	if ( TimeStampParamVal && !TimeStampParamVal->IsEmpty() )
diff --git a/Engine/Source/Editor/UnrealEd/Private/Commandlets/GenerateTextLocalizationResourceCommandlet.cpp b/Engine/Source/Editor/UnrealEd/Private/Commandlets/GenerateTextLocalizationResourceCommandlet.cpp
--- a/Engine/Source/Editor/UnrealEd/Private/Commandlets/GenerateTextLocalizationResourceCommandlet.cpp
+++ b/Engine/Source/Editor/UnrealEd/Private/Commandlets/GenerateTextLocalizationResourceCommandlet.cpp
@@ -59,6 +59,17 @@ int32 UGenerateTextLocalizationResourceCommandlet::Main(const FString& Params)
 		return -1;
 	}
 
+	if (FPaths::IsRelative(SourcePath))
+	{
+		if (!FPaths::GameDir().IsEmpty())
+		{
+			SourcePath = FPaths::Combine( *( FPaths::GameDir() ), *SourcePath );
+		}
+		else
+		{
+			SourcePath = FPaths::Combine( *( FPaths::EngineDir() ), *SourcePath );
+		}
+	}
 	// Get manifest name.
 	FString ManifestName;
 	if( !( GetConfigString( *SectionName, TEXT("ManifestName"), ManifestName, GatherTextConfigPath ) ) )
@@ -93,6 +104,18 @@ int32 UGenerateTextLocalizationResourceCommandlet::Main(const FString& Params)
 		return -1;
 	}
 
+	if (FPaths::IsRelative(DestinationPath))
+	{
+		if (!FPaths::GameDir().IsEmpty())
+		{
+			DestinationPath = FPaths::Combine( *( FPaths::GameDir() ), *DestinationPath );
+		}
+		else
+		{
+			DestinationPath = FPaths::Combine( *( FPaths::EngineDir() ), *DestinationPath );
+		}
+	}
+
 	// Get resource name.
 	FString ResourceName;
 	if( !( GetConfigString( *SectionName, TEXT("ResourceName"), ResourceName, GatherTextConfigPath ) ) )
@@ -136,11 +159,14 @@ int32 UGenerateTextLocalizationResourceCommandlet::Main(const FString& Params)
 		}
 
 		TAutoPtr<FArchive> TextLocalizationResourceArchive(IFileManager::Get().CreateFileWriter(*TextLocalizationResourcePath));
-		if( !(FTextLocalizationResourceGenerator::Generate(SourcePath, InternationalizationManifest, CultureName, TextLocalizationResourceArchive)) )
+		if (TextLocalizationResourceArchive)
 		{
-			IFileManager::Get().Delete( *TextLocalizationResourcePath );
+			if( !(FTextLocalizationResourceGenerator::Generate(SourcePath, InternationalizationManifest, CultureName, TextLocalizationResourceArchive)) )
+			{
+				IFileManager::Get().Delete( *TextLocalizationResourcePath );
+			}
+			TextLocalizationResourceArchive->Close();
 		}
-		TextLocalizationResourceArchive->Close();
 	}
 
 	return 0;
diff --git a/Engine/Source/Editor/UnrealEd/Private/Commandlets/InternationalizationExportCommandlet.cpp b/Engine/Source/Editor/UnrealEd/Private/Commandlets/InternationalizationExportCommandlet.cpp
--- a/Engine/Source/Editor/UnrealEd/Private/Commandlets/InternationalizationExportCommandlet.cpp
+++ b/Engine/Source/Editor/UnrealEd/Private/Commandlets/InternationalizationExportCommandlet.cpp
@@ -1269,6 +1269,18 @@ int32 UInternationalizationExportCommandlet::Main( const FString& Params )
 		return -1;
 	}
 
+	if (FPaths::IsRelative(SourcePath))
+	{
+		if (!FPaths::GameDir().IsEmpty())
+		{
+			SourcePath = FPaths::Combine( *( FPaths::GameDir() ), *SourcePath );
+		}
+		else
+		{
+			SourcePath = FPaths::Combine( *( FPaths::EngineDir() ), *SourcePath );
+		}
+	}
+
 	FString DestinationPath; // Destination path that we will write files to.
 	if( !GetConfigString( *SectionName, TEXT("DestinationPath"), DestinationPath, ConfigPath ) )
 	{
@@ -1276,6 +1288,18 @@ int32 UInternationalizationExportCommandlet::Main( const FString& Params )
 		return -1;
 	}
 
+	if (FPaths::IsRelative(DestinationPath))
+	{
+		if (!FPaths::GameDir().IsEmpty())
+		{
+			DestinationPath = FPaths::Combine( *( FPaths::GameDir() ), *DestinationPath );
+		}
+		else
+		{
+			DestinationPath = FPaths::Combine( *( FPaths::EngineDir() ), *DestinationPath );
+		}
+	}
+
 	FString Filename; // Name of the file to read or write from
 	if (!GetConfigString(*SectionName, TEXT("PortableObjectName"), Filename, ConfigPath))
 	{
diff --git a/Engine/Source/Programs/AutomationTool/Scripts/LauncherLocalization.Automation.cs b/Engine/Source/Programs/AutomationTool/Scripts/LauncherLocalization.Automation.cs
--- a/Engine/Source/Programs/AutomationTool/Scripts/LauncherLocalization.Automation.cs
+++ b/Engine/Source/Programs/AutomationTool/Scripts/LauncherLocalization.Automation.cs
@@ -66,8 +66,8 @@ class LauncherLocalization : BuildCommand
 		// Setup commandlet arguments with configurations.
 		var CommandletArgumentSets = new string[] 
 			{
-				String.Format("-config={0}", @"../../../Engine/Programs/NoRedist/UnrealEngineLauncher/Config/Localization/App.ini") + (string.IsNullOrEmpty(CommandletSCCArguments) ? "" : " " + CommandletSCCArguments),
-				String.Format("-config={0}", @"../../../Engine/Programs/NoRedist/UnrealEngineLauncher/Config/Localization/WordCount.ini"),
+				String.Format("-config={0}", @"./Programs/NoRedist/UnrealEngineLauncher/Config/Localization/App.ini") + (string.IsNullOrEmpty(CommandletSCCArguments) ? "" : " " + CommandletSCCArguments),
+				String.Format("-config={0}", @"./Programs/NoRedist/UnrealEngineLauncher/Config/Localization/WordCount.ini"),
 			};
 
 		// Execute commandlet for each set of arguments.
diff --git a/Engine/Source/Programs/AutomationTool/Scripts/Localisation.Automation.cs b/Engine/Source/Programs/AutomationTool/Scripts/Localisation.Automation.cs
--- a/Engine/Source/Programs/AutomationTool/Scripts/Localisation.Automation.cs
+++ b/Engine/Source/Programs/AutomationTool/Scripts/Localisation.Automation.cs
@@ -40,11 +40,11 @@ class Localise : BuildCommand
 		// Setup commandlet arguments with configurations.
 		var CommandletArgumentSets = new string[] 
 			{
-				String.Format("-config={0}", @"../../../Engine/Config/Localization/Engine.ini") + (string.IsNullOrEmpty(CommandletSCCArguments) ? "" : " " + CommandletSCCArguments),
-				String.Format("-config={0}", @"../../../Engine/Config/Localization/Editor.ini") + (string.IsNullOrEmpty(CommandletSCCArguments) ? "" : " " + CommandletSCCArguments),
-				String.Format("-config={0}", @"../../../Engine/Config/Localization/PropertyNames.ini") + (string.IsNullOrEmpty(CommandletSCCArguments) ? "" : " " + CommandletSCCArguments),
-				String.Format("-config={0}", @"../../../Engine/Config/Localization/ToolTips.ini") + (string.IsNullOrEmpty(CommandletSCCArguments) ? "" : " " + CommandletSCCArguments),
-				String.Format("-config={0}", @"../../../Engine/Config/Localization/WordCount.ini"),
+				String.Format("-config={0}", @"./Config/Localization/Engine.ini") + (string.IsNullOrEmpty(CommandletSCCArguments) ? "" : " " + CommandletSCCArguments),
+				String.Format("-config={0}", @"./Config/Localization/Editor.ini") + (string.IsNullOrEmpty(CommandletSCCArguments) ? "" : " " + CommandletSCCArguments),
+				String.Format("-config={0}", @"./Config/Localization/PropertyNames.ini") + (string.IsNullOrEmpty(CommandletSCCArguments) ? "" : " " + CommandletSCCArguments),
+				String.Format("-config={0}", @"./Config/Localization/ToolTips.ini") + (string.IsNullOrEmpty(CommandletSCCArguments) ? "" : " " + CommandletSCCArguments),
+				String.Format("-config={0}", @"./Config/Localization/WordCount.ini"),
 			};
 
 		// Execute commandlet for each set of arguments.
