Attempt to fix assert dependency reporting for chunk manifests when nothing is initially assigned to a chunk (e.g. chunk 0)

diff --git a/Engine/Source/Editor/UnrealEd/Private/Commandlets/ChunkManifestGenerator.cpp b/Engine/Source/Editor/UnrealEd/Private/Commandlets/ChunkManifestGenerator.cpp
--- a/Engine/Source/Editor/UnrealEd/Private/Commandlets/ChunkManifestGenerator.cpp
+++ b/Engine/Source/Editor/UnrealEd/Private/Commandlets/ChunkManifestGenerator.cpp
@@ -403,22 +403,22 @@ bool FChunkManifestGenerator::GenerateAssetChunkInformationCSV(const FString& Ou
 						FileSize = 0;
 					}
 
+					FString SoftChain;
 					bool bHardChunk = false;
 					if (ChunkID < ChunkManifests.Num())
 					{
 						bHardChunk = ChunkManifests[ChunkID] && ChunkManifests[ChunkID]->Contains(AssetData.PackageName);
-					}
-
-					FString SoftChain;
-					if (!bHardChunk)
-					{
-						//
-						SoftChain = GetShortestReferenceChain(AssetData.PackageName, ChunkID);
-						if (SoftChain.IsEmpty())
+					
+						if (!bHardChunk)
 						{
-							SoftChain = TEXT("Soft: Possibly Unassigned Asset");
+							//
+							SoftChain = GetShortestReferenceChain(AssetData.PackageName, ChunkID);
 						}
 					}
+					if (SoftChain.IsEmpty())
+					{
+						SoftChain = TEXT("Soft: Possibly Unassigned Asset");
+					}
 
 					TmpString = FString::Printf(TEXT("%d,%s,%s,%s,%lld,"), ChunkID, *AssetData.PackageName.ToString(), *AssetData.AssetClass.ToString(), bHardChunk ? TEXT("Hard") : *SoftChain, FileSize);
 					CSVString += TmpString;
