Fixed measure tool calculating the wrong size

The problem here was that the world start position of the tool was being compared to a world-projected pixel position, which had vastly different depth values. This has been fixed by storing the start and end position in pixel space and using that for the distance calculation, rather than a world-space position. This would need re-thinking if we wanted the tool to work in perspective viewports, but that was a problem with the old code as well.

Also made the measurement display in more relevant units, rather than just showing raw uu.

This addresses TTP#344821 - CRITICAL: Regression: EDITOR: Ruler gives incorrect measurements.

diff --git a/Engine/Source/Editor/UnrealEd/Private/DragTool_Measure.cpp b/Engine/Source/Editor/UnrealEd/Private/DragTool_Measure.cpp
--- a/Engine/Source/Editor/UnrealEd/Private/DragTool_Measure.cpp
+++ b/Engine/Source/Editor/UnrealEd/Private/DragTool_Measure.cpp
@@ -18,11 +18,8 @@ FDragTool_Measure::FDragTool_Measure(FEditorViewportClient* InViewportClient)
 	bConvertDelta = false;
 }
 
-void FDragTool_Measure::SetEndWorldPositionFromCursor()
+FVector2D FDragTool_Measure::GetSnappedPixelPos(FVector2D PixelPos)
 {
-	FIntPoint MousePos;
-	ViewportClient->Viewport->GetMousePos(MousePos);
-
 	FSceneViewFamilyContext ViewFamily(FSceneViewFamily::ConstructionValues(
 		ViewportClient->Viewport,
 		ViewportClient->GetScene(),
@@ -31,47 +28,50 @@ void FDragTool_Measure::SetEndWorldPositionFromCursor()
 
 	FSceneView* View = ViewportClient->CalcSceneView(&ViewFamily);
 
-	End = View->ScreenToWorld(View->PixelToScreen(MousePos.X, MousePos.Y, 0.5f));
+	// Put the mouse pos in world space
+	FVector WorldPos = View->ScreenToWorld(View->PixelToScreen(PixelPos.X, PixelPos.Y, 0.5f));;
 
+	// Snap the world position
 	const float GridSize = GEditor->GetGridSize();
 	const FVector GridBase( GridSize, GridSize, GridSize );
-	FSnappingUtils::SnapPointToGrid( End, GridBase );
+	FSnappingUtils::SnapPointToGrid( WorldPos, GridBase );
+
+	// And back into pixel-space (might fail, in which case we return the original)
+	View->WorldToPixel(WorldPos, PixelPos);
+
+	return PixelPos;
 }
 
 void FDragTool_Measure::StartDrag(FEditorViewportClient* InViewportClient, const FVector& InStart, const FVector2D& InStartScreen)
 {
 	FDragTool::StartDrag(InViewportClient, InStart, InStartScreen);
-	SetEndWorldPositionFromCursor();
+	PixelStart = GetSnappedPixelPos(InStartScreen);
+	PixelEnd = PixelStart;
 }
 
 void FDragTool_Measure::AddDelta(const FVector& InDelta)
 {
-	SetEndWorldPositionFromCursor();
+	FDragTool::AddDelta(InDelta);
+	
+	FIntPoint MousePos;
+	ViewportClient->Viewport->GetMousePos(MousePos);
+	PixelEnd = GetSnappedPixelPos(FVector2D(MousePos));
 }
 
 void FDragTool_Measure::Render(const FSceneView* View, FCanvas* Canvas)
 {
-	FVector2D PixelStart;
-	FVector2D PixelEnd;
-	if (View != nullptr && Canvas != nullptr && View->WorldToPixel(Start, PixelStart) && View->WorldToPixel(End, PixelEnd))
+	const float Length = FMath::RoundToFloat((PixelEnd - PixelStart).Size() * ViewportClient->GetOrthoUnitsPerPixel(ViewportClient->Viewport));
+
+	if (View != nullptr && Canvas != nullptr && Length >= 1.f)
 	{
 		FCanvasLineItem LineItem( PixelStart, PixelEnd );
 		Canvas->DrawItem( LineItem );
 
-		const int32 Length = FMath::CeilToInt((End - Start).Size());
-		if( Length == 0 )
-		{
-			return;
-		}
-
-		const FVector WorldMid = Start + ((End - Start) / 2);
-		FVector2D PixelMid;
-		if (View->WorldToPixel(WorldMid, PixelMid))
-		{
-			FString LengthStr = FString::Printf(TEXT("%d"), Length);
-			FCanvasTextItem TextItem( FVector2D( FMath::FloorToFloat(PixelMid.X), FMath::FloorToFloat(PixelMid.Y) ), FText::FromString( LengthStr ), GEngine->GetSmallFont(), FLinearColor::White );
-			TextItem.bCentreX = true;
-			Canvas->DrawItem( TextItem );
-		}
+		const FVector2D PixelMid = FVector2D(PixelStart + ((PixelEnd - PixelStart) / 2));
+
+		const FString LengthStr = FEditorViewportClient::UnrealUnitsToSiUnits(Length);
+		FCanvasTextItem TextItem( FVector2D( FMath::FloorToFloat(PixelMid.X), FMath::FloorToFloat(PixelMid.Y) ), FText::FromString( LengthStr ), GEngine->GetSmallFont(), FLinearColor::White );
+		TextItem.bCentreX = true;
+		Canvas->DrawItem( TextItem );
 	}
 }
diff --git a/Engine/Source/Editor/UnrealEd/Private/EditorViewportClient.cpp b/Engine/Source/Editor/UnrealEd/Private/EditorViewportClient.cpp
--- a/Engine/Source/Editor/UnrealEd/Private/EditorViewportClient.cpp
+++ b/Engine/Source/Editor/UnrealEd/Private/EditorViewportClient.cpp
@@ -2149,7 +2149,7 @@ void FEditorViewportClient::DrawAxes(FViewport* InViewport, FCanvas* Canvas, con
 }
 
 /** Convert the specified number (in cm or unreal units) into a readable string with relevant si units */
-FString UnrealUnitsToSiUnits(float UnrealUnits)
+FString FEditorViewportClient::UnrealUnitsToSiUnits(float UnrealUnits)
 {
 	// Put it in mm to start off with
 	UnrealUnits *= 10.f;
diff --git a/Engine/Source/Editor/UnrealEd/Public/DragTool_Measure.h b/Engine/Source/Editor/UnrealEd/Public/DragTool_Measure.h
--- a/Engine/Source/Editor/UnrealEd/Public/DragTool_Measure.h
+++ b/Engine/Source/Editor/UnrealEd/Public/DragTool_Measure.h
@@ -35,10 +35,12 @@ public:
 
 private:
 	/**
-	 * Sets the End member explicitly from the current cursor position, overriding the calculation based on
-	 * the deltas passed to AddDelta.
+	 * Gets the world-space snapped position of the specified pixel position
 	 */
-	void SetEndWorldPositionFromCursor();
+	FVector2D GetSnappedPixelPos(FVector2D PixelPos);
+
+	/** Pixel-space positions for start and end */
+	FVector2D PixelStart, PixelEnd;
 
 	FEditorViewportClient* ViewportClient;
 };
diff --git a/Engine/Source/Editor/UnrealEd/Public/EditorViewportClient.h b/Engine/Source/Editor/UnrealEd/Public/EditorViewportClient.h
--- a/Engine/Source/Editor/UnrealEd/Public/EditorViewportClient.h
+++ b/Engine/Source/Editor/UnrealEd/Public/EditorViewportClient.h
@@ -399,6 +399,9 @@ public:
 	/** @return The number of units per pixel displayed in this viewport */
 	float GetOrthoUnitsPerPixel(const FViewport* Viewport) const;
 
+	/** Get a prettified string representation of the specified unreal units */
+	static FString UnrealUnitsToSiUnits(float UnrealUnits);
+
 	void RemoveCameraRoll()
 	{
 		FRotator Rotation = GetViewRotation();
