Fix for crash when copying animations between skeletons if they have curve data attached

diff --git a/Engine/Source/Editor/UnrealEd/Private/EditorAnimUtils.cpp b/Engine/Source/Editor/UnrealEd/Private/EditorAnimUtils.cpp
--- a/Engine/Source/Editor/UnrealEd/Private/EditorAnimUtils.cpp
+++ b/Engine/Source/Editor/UnrealEd/Private/EditorAnimUtils.cpp
@@ -164,6 +164,17 @@ namespace EditorAnimUtils
 		for(auto Iter = AnimSequencesToRetarget.CreateIterator(); Iter; ++Iter)
 		{
 			UAnimSequence* AssetToRetarget = (*Iter);
+
+			// Copy curve data from source asset, preserving data in the target if present.
+			FSmartNameMapping* OldNameMapping = OldSkeleton->SmartNames.GetContainer(USkeleton::AnimCurveMappingName);
+			FSmartNameMapping* NewNameMapping = NewSkeleton->SmartNames.GetContainer(USkeleton::AnimCurveMappingName);
+			AssetToRetarget->RawCurveData.UpdateLastObservedNames(OldNameMapping);
+
+			for(FFloatCurve& Curve : AssetToRetarget->RawCurveData.FloatCurves)
+			{
+				NewNameMapping->AddName(Curve.LastObservedName, Curve.CurveUid);
+			}
+
 			AssetToRetarget->ReplaceSkeleton(NewSkeleton, bConvertAnimationDataInComponentSpaces);
 		}
 
diff --git a/Engine/Source/Runtime/Engine/Classes/Animation/AnimSequenceBase.h b/Engine/Source/Runtime/Engine/Classes/Animation/AnimSequenceBase.h
--- a/Engine/Source/Runtime/Engine/Classes/Animation/AnimSequenceBase.h
+++ b/Engine/Source/Runtime/Engine/Classes/Animation/AnimSequenceBase.h
@@ -302,7 +302,7 @@ struct FRawCurveTracks
 	/**
 	 * Updates the LastObservedName field of the curves from the provided name container
 	 */
-	void UpdateLastObservedNames(FSmartNameMapping* NameMapping);
+	ENGINE_API void UpdateLastObservedNames(FSmartNameMapping* NameMapping);
 
 	void Serialize(FArchive& Ar);
 };
diff --git a/Engine/Source/Runtime/Engine/Private/Animation/AnimSequenceBase.cpp b/Engine/Source/Runtime/Engine/Private/Animation/AnimSequenceBase.cpp
--- a/Engine/Source/Runtime/Engine/Private/Animation/AnimSequenceBase.cpp
+++ b/Engine/Source/Runtime/Engine/Private/Animation/AnimSequenceBase.cpp
@@ -165,7 +165,7 @@ void FRawCurveTracks::Serialize(FArchive& Ar)
 	}
 }
 
-void FRawCurveTracks::UpdateLastObservedNames(FSmartNameMapping* NameMapping)
+ENGINE_API void FRawCurveTracks::UpdateLastObservedNames(FSmartNameMapping* NameMapping)
 {
 	if(NameMapping)
 	{
