Support tooltips for categories.
Fix localization for category registration

diff --git a/Engine/Source/Editor/BlueprintGraph/Private/K2Node_AIMoveTo.cpp b/Engine/Source/Editor/BlueprintGraph/Private/K2Node_AIMoveTo.cpp
--- a/Engine/Source/Editor/BlueprintGraph/Private/K2Node_AIMoveTo.cpp
+++ b/Engine/Source/Editor/BlueprintGraph/Private/K2Node_AIMoveTo.cpp
@@ -20,7 +20,7 @@ UK2Node_AIMoveTo::UK2Node_AIMoveTo(const class FPostConstructInitializePropertie
 
 FText UK2Node_AIMoveTo::GetMenuCategory() const
 {
-	return FEditorCategoryUtils::GetCommonCategory(FCommonEditorCategory::Ai);
+	return FEditorCategoryUtils::GetCommonCategory(FCommonEditorCategory::AI);
 }
 
 FString UK2Node_AIMoveTo::GetTooltip() const
diff --git a/Engine/Source/Editor/GraphEditor/Private/SGraphActionMenu.cpp b/Engine/Source/Editor/GraphEditor/Private/SGraphActionMenu.cpp
--- a/Engine/Source/Editor/GraphEditor/Private/SGraphActionMenu.cpp
+++ b/Engine/Source/Editor/GraphEditor/Private/SGraphActionMenu.cpp
@@ -6,6 +6,7 @@
 #include "GraphActionNode.h"
 #include "SScrollBorder.h"
 #include "IDocumentation.h"
+#include "EditorCategoryUtils.h"
 
 #define LOCTEXT_NAMESPACE "GraphActionMenu"
 
@@ -87,8 +88,11 @@ public:
 	{
 		ActionNode = InActionNode;
 
-		FText CategoryText = FText::FromString(InActionNode->Category);
-		TSharedRef<SToolTip> ToolTipWidget = IDocumentation::Get()->CreateToolTip(CategoryText, NULL, TEXT("Shared/GraphNodes/Blueprint/NodeCategories"), InActionNode->Category);
+		FText CategoryTooltip;
+		FString CategoryLink, CategoryExcerpt;
+		FEditorCategoryUtils::GetCategoryTooltipInfo(*InActionNode->Category, CategoryTooltip, CategoryLink, CategoryExcerpt);
+
+		TSharedRef<SToolTip> ToolTipWidget = IDocumentation::Get()->CreateToolTip(CategoryTooltip, NULL, CategoryLink, CategoryExcerpt);
 
 		this->ChildSlot
 		[
@@ -98,7 +102,7 @@ public:
 			[
 				SAssignNew(InlineWidget, SInlineEditableTextBlock)
 				.Font( FSlateFontInfo( FPaths::EngineContentDir() / TEXT("Slate/Fonts/Roboto-Bold.ttf"), 9 )  )
-				.Text( CategoryText )
+				.Text( FEditorCategoryUtils::GetCategoryDisplayString(FText::FromString(InActionNode->Category)) )
 				.ToolTip( ToolTipWidget )
 				.HighlightText( InArgs._HighlightText )
 				.OnVerifyTextChanged( this, &SGraphActionCategoryWidget::OnVerifyTextChanged )
diff --git a/Engine/Source/Editor/UnrealEd/Private/EditorCategoryUtils.cpp b/Engine/Source/Editor/UnrealEd/Private/EditorCategoryUtils.cpp
--- a/Engine/Source/Editor/UnrealEd/Private/EditorCategoryUtils.cpp
+++ b/Engine/Source/Editor/UnrealEd/Private/EditorCategoryUtils.cpp
@@ -2,15 +2,28 @@
 
 #include "UnrealEd.h"
 #include "EditorCategoryUtils.h"
+#include "IDocumentation.h"
 
 #define LOCTEXT_NAMESPACE "EditorCategoryUtils"
 
 /*******************************************************************************
- * Static FEditorCategoryUtils Helpers
+ * FEditorCategoryUtils Helpers
  ******************************************************************************/
 
 namespace FEditorCategoryUtilsImpl
 {
+	using namespace FEditorCategoryUtils;
+
+	struct FCategoryInfo
+	{
+		FText DisplayName;
+		FText Tooltip;
+		FString DocLink;
+		FString DocExcerpt;
+	};
+
+	typedef TMap<FString, FCategoryInfo> FCategoryInfoMap;
+
 	/**
 	 * Gets the table that tracks mappings from string keys to qualified 
 	 * category paths. Inits the structure if it hadn't been before (adds 
@@ -18,7 +31,7 @@ namespace FEditorCategoryUtilsImpl
 	 * 
 	 * @return A map from string category keys, to fully qualified category paths.
 	 */
-	static FFormatNamedArguments& GetCategoryTable();
+	FCategoryInfoMap& GetCategoryTable();
 
 	/**
 	 * Performs a lookup into the category key table, retrieving a fully 
@@ -27,20 +40,28 @@ namespace FEditorCategoryUtilsImpl
 	 * @param  Key	The key you want a category path for.
 	 * @return The category display string associated with the specified key (an empty string if an entry wasn't found).
 	 */
-	static FText const& GetCategory(FString const& Key);
+	FText const& GetCategory(const FString& Key);
 
+	/**
+	 * Performs a lookup into the category key table, retrieving a fully 
+	 * qualified category path for the specified key.
+	 * 
+	 * @param  CategoryDisplayName	Display name for the category, will be used if a tooltip can not be found in the Documentation Page
+	 * @param  DocLink				Path to the documentation page that contains the excerpt for this category
+	 * @param  DocExcerpt			Name of the excerpt within the document page for this category
+	 * @return						The tooltip (if any) stored at the doc path
+	 */
+	FText GetTooltipForCategory(FString const& CategoryDisplayName, FString const& DocLink, FString const& DocExcerpt);
 
 	/** Metadata tags */
-	static const FName ClassHideCategoriesMetaKey(TEXT("HideCategories"));
-	static const FName ClassShowCategoriesMetaKey(TEXT("ShowCategories"));
+	const FName ClassHideCategoriesMetaKey(TEXT("HideCategories"));
+	const FName ClassShowCategoriesMetaKey(TEXT("ShowCategories"));
 }
 
-#define ENUM_STRING(EnumVal) FString(#EnumVal)
-
 //------------------------------------------------------------------------------
-static FFormatNamedArguments& FEditorCategoryUtilsImpl::GetCategoryTable()
+FEditorCategoryUtilsImpl::FCategoryInfoMap& FEditorCategoryUtilsImpl::GetCategoryTable()
 {
-	static FFormatNamedArguments CategoryLookup;
+	static FCategoryInfoMap CategoryLookup;
 
 	static bool bInitialized = false;
 	// this function is reentrant, so we have to guard against recursion
@@ -48,103 +69,139 @@ static FFormatNamedArguments& FEditorCategoryUtilsImpl::GetCategoryTable()
 	{
 		bInitialized = true;
 
-#define ASSIGN_ROOT_CATEGORY(EnumVal, Text) FCommonEditorCategory::EnumVal; \
-		FEditorCategoryUtils::RegisterCategoryKey(ENUM_STRING(EnumVal), LOCTEXT(#EnumVal L"Category", Text))
-#define REGISTER_ROOT_CATEGORY(EnumVal) ASSIGN_ROOT_CATEGORY(EnumVal, #EnumVal)
-
-		ASSIGN_ROOT_CATEGORY(Ai, "AI");
-		REGISTER_ROOT_CATEGORY(Animation);
-		REGISTER_ROOT_CATEGORY(Audio);
-		REGISTER_ROOT_CATEGORY(Development);
-		REGISTER_ROOT_CATEGORY(Effects);
-		ASSIGN_ROOT_CATEGORY(Gameplay, "Game");
-		REGISTER_ROOT_CATEGORY(Input);
-		REGISTER_ROOT_CATEGORY(Math);
-		REGISTER_ROOT_CATEGORY(Networking);
-		REGISTER_ROOT_CATEGORY(Pawn);
-		REGISTER_ROOT_CATEGORY(Rendering);
-		REGISTER_ROOT_CATEGORY(Utilities);
-		ASSIGN_ROOT_CATEGORY(Delegates, "Event Dispatchers");
-		REGISTER_ROOT_CATEGORY(Variables);
-
-#undef REGISTER_ROOT_CATEGORY
-#undef ASSIGN_ROOT_CATEGORY
-
-		// NOTE: the root category has to be registered prior to this
-#define REGISTER_SUB_CATEGORY(RootId, EnumVal) FCommonEditorCategory::EnumVal; \
-		FText EnumVal##_Category = FEditorCategoryUtils::BuildCategoryString(FCommonEditorCategory::RootId, LOCTEXT(#EnumVal L"Category", #EnumVal)); \
-		FEditorCategoryUtils::RegisterCategoryKey(ENUM_STRING(EnumVal), EnumVal##_Category)
-
-		REGISTER_SUB_CATEGORY(Utilities, FlowControl);
-		REGISTER_SUB_CATEGORY(Utilities, Transformation);
-		REGISTER_SUB_CATEGORY(Utilities, String);
-		REGISTER_SUB_CATEGORY(Utilities, Text);
-		REGISTER_SUB_CATEGORY(Utilities, Name);
-		REGISTER_SUB_CATEGORY(Utilities, Enum);
-		REGISTER_SUB_CATEGORY(Utilities, Struct);
-		REGISTER_SUB_CATEGORY(Utilities, Macro);
-
-#undef REGISTER_SUB_CATEGORY
+		RegisterCategoryKey("AI", LOCTEXT("AICategory", "AI"));
+		RegisterCategoryKey("Animation", LOCTEXT("AnimationCategory", "Animation"));
+		RegisterCategoryKey("Audio", LOCTEXT("AudioCategory", "Audio"));
+		RegisterCategoryKey("Development", LOCTEXT("DevelopmentCategory", "Development"));
+		RegisterCategoryKey("Effects", LOCTEXT("EffectsCategory", "Effects"));
+		RegisterCategoryKey("Gameplay", LOCTEXT("GameplayCategory", "Game"));
+		RegisterCategoryKey("Input", LOCTEXT("InputCategory", "Input"));
+		RegisterCategoryKey("Math", LOCTEXT("MathCategory", "Math"));
+		RegisterCategoryKey("Networking", LOCTEXT("NetworkingCategory", "Networking"));
+		RegisterCategoryKey("Pawn", LOCTEXT("PawnCategory", "Pawn"));
+		RegisterCategoryKey("Rendering", LOCTEXT("RenderingCategory", "Rendering"));
+		RegisterCategoryKey("Utilities", LOCTEXT("UtilitiesCategory", "Utilities"));
+		RegisterCategoryKey("Delegates", LOCTEXT("DelegatesCategory", "Event Dispatchers"));
+		RegisterCategoryKey("Variables", LOCTEXT("VariablesCategory", "Variables"));
+
+		// Utilities sub categories
+		RegisterCategoryKey("FlowControl", BuildCategoryString(FCommonEditorCategory::Utilities, LOCTEXT("FlowControlCategory", "Flow Control")));
+		RegisterCategoryKey("Transformation", BuildCategoryString(FCommonEditorCategory::Utilities, LOCTEXT("TransformationCategory", "Transformation")));
+		RegisterCategoryKey("String", BuildCategoryString(FCommonEditorCategory::Utilities, LOCTEXT("StringCategory", "String")));
+		RegisterCategoryKey("Text", BuildCategoryString(FCommonEditorCategory::Utilities, LOCTEXT("TextCategory", "Text")));
+		RegisterCategoryKey("Name", BuildCategoryString(FCommonEditorCategory::Utilities, LOCTEXT("NameCategory", "Name")));
+		RegisterCategoryKey("Enum", BuildCategoryString(FCommonEditorCategory::Utilities, LOCTEXT("EnumCategory", "Enum")));
+		RegisterCategoryKey("Struct", BuildCategoryString(FCommonEditorCategory::Utilities, LOCTEXT("StructCategory", "Struct")));
+		RegisterCategoryKey("Macro", BuildCategoryString(FCommonEditorCategory::Utilities, LOCTEXT("MacroCategory", "Macro")));
 	}
 
 	return CategoryLookup;
 }
 
 //------------------------------------------------------------------------------
-static FText const& FEditorCategoryUtilsImpl::GetCategory(FString const& Key)
+FText const& FEditorCategoryUtilsImpl::GetCategory(const FString& Key)
 {
-	FFormatNamedArguments const& CategoryLookup = GetCategoryTable();
-	if (FFormatArgumentValue const* FoundCategory = CategoryLookup.Find(Key.ToUpper()))
+	if (FEditorCategoryUtilsImpl::FCategoryInfo const* FoundCategory = GetCategoryTable().Find(Key))
 	{
-		return *FoundCategory->TextValue;
+		return FoundCategory->DisplayName;
 	}
 	return FText::GetEmpty();
 }
 
+//------------------------------------------------------------------------------
+FText FEditorCategoryUtilsImpl::GetTooltipForCategory(FString const& CategoryDisplayName, FString const& DocLink, FString const& DocExcerpt)
+{
+	FText Tooltip;
+
+	TSharedRef<IDocumentation> Documentation = IDocumentation::Get();
+	if (Documentation->PageExists(DocLink))
+	{
+		TSharedRef<IDocumentationPage> DocPage = Documentation->GetPage(DocLink, NULL);
+	
+		const FString TooltipExcerptSuffix(TEXT("__Tooltip"));
+
+		FExcerpt Excerpt;
+		if (DocPage->GetExcerpt(DocExcerpt + TooltipExcerptSuffix, Excerpt))
+		{
+			static const FString TooltipVarKey(TEXT("Tooltip"));
+			if (FString* TooltipValue = Excerpt.Variables.Find(TooltipVarKey))
+			{
+				Tooltip = FText::FromString(TooltipValue->Replace(TEXT("\\n"),TEXT("\n")));
+			}
+		}
+	}
+
+	if (Tooltip.IsEmpty())
+	{
+		FString CategoryTooltip;
+
+		if (CategoryDisplayName.Split(TEXT("|"), nullptr, &CategoryTooltip, ESearchCase::CaseSensitive, ESearchDir::FromEnd))
+		{
+			Tooltip = FText::FromString(CategoryTooltip);
+		}
+		else
+		{
+			Tooltip = FText::FromString(CategoryDisplayName);
+		}
+	}
+
+	return Tooltip;
+}
+
 /*******************************************************************************
  * FEditorCategoryUtils
  ******************************************************************************/
 
 //------------------------------------------------------------------------------
-void FEditorCategoryUtils::RegisterCategoryKey(FString const& Key, FText const& Category)
+void FEditorCategoryUtils::RegisterCategoryKey(FString const& Key, FText const& Category, FText const& Tooltip)
+{
+	FEditorCategoryUtilsImpl::FCategoryInfo& CategoryInfo = FEditorCategoryUtilsImpl::GetCategoryTable().Add(Key);
+
+	CategoryInfo.DisplayName = GetCategoryDisplayString(Category);
+	CategoryInfo.DocLink = TEXT("Shared/GraphNodes/Blueprint/NodeCategories");
+	CategoryInfo.DocExcerpt = Key;
+	CategoryInfo.Tooltip = (Tooltip.IsEmpty() ? FEditorCategoryUtilsImpl::GetTooltipForCategory(CategoryInfo.DisplayName.ToString(), CategoryInfo.DocLink, CategoryInfo.DocExcerpt) : Tooltip);
+}
+
+void FEditorCategoryUtils::RegisterCategoryKey(FString const& Key, FText const& Category, FString const& DocLink, FString const& DocExcerpt)
 {
-	FEditorCategoryUtilsImpl::GetCategoryTable().Add(Key.ToUpper(), FEditorCategoryUtils::GetCategoryDisplayString(Category));
+	FEditorCategoryUtilsImpl::FCategoryInfo& CategoryInfo = FEditorCategoryUtilsImpl::GetCategoryTable().Add(Key);
+
+	CategoryInfo.DisplayName = GetCategoryDisplayString(Category);
+	CategoryInfo.DocLink = DocLink;
+	CategoryInfo.DocExcerpt = DocExcerpt;
+	CategoryInfo.Tooltip = FEditorCategoryUtilsImpl::GetTooltipForCategory(CategoryInfo.DisplayName.ToString(), CategoryInfo.DocLink, CategoryInfo.DocExcerpt);
 }
 
 //------------------------------------------------------------------------------
-FText const& FEditorCategoryUtils::GetCommonCategory(FCommonEditorCategory::EValue CategoryId)
+FText const& FEditorCategoryUtils::GetCommonCategory(const FCommonEditorCategory::EValue CategoryId)
 {
 	static TMap<FCommonEditorCategory::EValue, FString> CommonCategoryKeys;
 	if (CommonCategoryKeys.Num() == 0)
 	{
-#define REGISTER_COMMON_CATEGORY(EnumVal) \
-		CommonCategoryKeys.Add(FCommonEditorCategory::EnumVal, ENUM_STRING(EnumVal));
-
-		REGISTER_COMMON_CATEGORY(Ai);
-		REGISTER_COMMON_CATEGORY(Animation);
-		REGISTER_COMMON_CATEGORY(Audio);
-		REGISTER_COMMON_CATEGORY(Development);
-		REGISTER_COMMON_CATEGORY(Effects);
-		REGISTER_COMMON_CATEGORY(Gameplay);
-		REGISTER_COMMON_CATEGORY(Input);
-		REGISTER_COMMON_CATEGORY(Math);
-		REGISTER_COMMON_CATEGORY(Networking);
-		REGISTER_COMMON_CATEGORY(Pawn);
-		REGISTER_COMMON_CATEGORY(Rendering);
-		REGISTER_COMMON_CATEGORY(Utilities);
-		REGISTER_COMMON_CATEGORY(Delegates);
-		REGISTER_COMMON_CATEGORY(Variables);
-
-		REGISTER_COMMON_CATEGORY(FlowControl);
-		REGISTER_COMMON_CATEGORY(Transformation);
-		REGISTER_COMMON_CATEGORY(String);
-		REGISTER_COMMON_CATEGORY(Text);
-		REGISTER_COMMON_CATEGORY(Name);
-		REGISTER_COMMON_CATEGORY(Enum);
-		REGISTER_COMMON_CATEGORY(Struct);
-		REGISTER_COMMON_CATEGORY(Macro);
-
-#undef REGISTER_COMMON_CATEGORY
+		CommonCategoryKeys.Add(FCommonEditorCategory::AI, "AI");
+		CommonCategoryKeys.Add(FCommonEditorCategory::Animation, "Animation");
+		CommonCategoryKeys.Add(FCommonEditorCategory::Audio, "Audio");
+		CommonCategoryKeys.Add(FCommonEditorCategory::Development, "Development");
+		CommonCategoryKeys.Add(FCommonEditorCategory::Effects, "Effects");
+		CommonCategoryKeys.Add(FCommonEditorCategory::Gameplay, "Gameplay");
+		CommonCategoryKeys.Add(FCommonEditorCategory::Input, "Input");
+		CommonCategoryKeys.Add(FCommonEditorCategory::Math, "Math");
+		CommonCategoryKeys.Add(FCommonEditorCategory::Networking, "Networking");
+		CommonCategoryKeys.Add(FCommonEditorCategory::Pawn, "Pawn");
+		CommonCategoryKeys.Add(FCommonEditorCategory::Rendering, "Rendering");
+		CommonCategoryKeys.Add(FCommonEditorCategory::Utilities, "Utilities");
+		CommonCategoryKeys.Add(FCommonEditorCategory::Delegates, "Delegates");
+		CommonCategoryKeys.Add(FCommonEditorCategory::Variables, "Variables");
+
+		CommonCategoryKeys.Add(FCommonEditorCategory::FlowControl, "FlowControl");
+		CommonCategoryKeys.Add(FCommonEditorCategory::Transformation, "Transformation");
+		CommonCategoryKeys.Add(FCommonEditorCategory::String, "String");
+		CommonCategoryKeys.Add(FCommonEditorCategory::Text, "Text");
+		CommonCategoryKeys.Add(FCommonEditorCategory::Name, "Name");
+		CommonCategoryKeys.Add(FCommonEditorCategory::Enum, "Enum");
+		CommonCategoryKeys.Add(FCommonEditorCategory::Struct, "Struct");
+		CommonCategoryKeys.Add(FCommonEditorCategory::Macro, "Macro");
 	}
 
 	if (FString* CategoryKey = CommonCategoryKeys.Find(CategoryId))
@@ -202,7 +259,7 @@ FString FEditorCategoryUtils::GetCategoryDisplayString(FString const& Unsanitize
 				if (KeyLen > 0)
 				{
 					FString Key(KeyLen, *DisplayString + KeyIndex+1);
-					ReplacementStr = FEditorCategoryUtilsImpl::GetCategory(Key.Trim()).ToString();
+					ReplacementStr = FEditorCategoryUtilsImpl::GetCategory(*Key.Trim()).ToString();
 				}
 				DisplayString.ReplaceInline(*ToReplaceStr, *ReplacementStr);
 			}
@@ -316,5 +373,22 @@ bool FEditorCategoryUtils::IsCategoryHiddenFromClass(UClass const* Class, FStrin
 	return bIsHidden;
 }
 
-#undef ENUM_STRING
+//------------------------------------------------------------------------------
+void FEditorCategoryUtils::GetCategoryTooltipInfo(const FString& Category, FText& Tooltip, FString& DocLink, FString& DocExcerpt)
+{
+	if (FEditorCategoryUtilsImpl::FCategoryInfo const* FoundCategory = FEditorCategoryUtilsImpl::GetCategoryTable().Find(Category))
+	{
+		DocLink = FoundCategory->DocLink;
+		DocExcerpt = FoundCategory->DocExcerpt;
+		Tooltip = FoundCategory->Tooltip;
+	}
+	else
+	{
+		// Fall back to some defaults
+		DocLink = TEXT("Shared/GraphNodes/Blueprint/NodeCategories");
+		DocExcerpt = Category;
+		Tooltip = FEditorCategoryUtilsImpl::GetTooltipForCategory(GetCategoryDisplayString(Category), DocLink, DocExcerpt);
+	}
+}
+
 #undef LOCTEXT_NAMESPACE
diff --git a/Engine/Source/Editor/UnrealEd/Public/EditorCategoryUtils.h b/Engine/Source/Editor/UnrealEd/Public/EditorCategoryUtils.h
--- a/Engine/Source/Editor/UnrealEd/Public/EditorCategoryUtils.h
+++ b/Engine/Source/Editor/UnrealEd/Public/EditorCategoryUtils.h
@@ -18,7 +18,7 @@ struct FCommonEditorCategory
 	enum EValue
 	{
 		// Function categories:
-		Ai,			
+		AI,			
 		Animation,	
 		Audio,		
 		Development,
@@ -51,7 +51,7 @@ struct FCommonEditorCategory
  * Set of utilities for parsing, querying, and resolving class/function/field 
  * category data.
  */
-struct UNREALED_API FEditorCategoryUtils
+namespace FEditorCategoryUtils
 {
 	/**
 	 * To facilitate simple category renaming/reordering, we offer a key 
@@ -65,8 +65,17 @@ struct UNREALED_API FEditorCategoryUtils
 	 * 
 	 * @param  Key		A string key that people will use in metadata to reflect this category mapping.
 	 * @param  Category	The qualified category path that you want the key expanded to.
+	 * @param  Tooltip  An optional tooltip text to use for the category.  If not specified an attempt to find it from the NodeCategories UDN file will be made
 	 */
-	static void RegisterCategoryKey(FString const& Key, FText const& Category);
+	UNREALED_API void RegisterCategoryKey(FString const& Key, FText const& Category, FText const& Tooltip = FText::GetEmpty());
+
+	/**
+	 * @param  Key			A string key that people will use in metadata to reflect this category mapping.
+	 * @param  Category		The qualified category path that you want the key expanded to.
+	 * @param  DocLink		Path to the document page that contains the excerpt for this category
+	 * @param  DocExcerpt	Name of the excerpt within the document page for this category
+	 */
+	UNREALED_API void RegisterCategoryKey(FString const& Key, FText const& Category, FString const& DocLink, FString const& DocExcerpt);
 
 	/**
 	 * Retrieves a qualified category path for the desired common category.
@@ -74,7 +83,7 @@ struct UNREALED_API FEditorCategoryUtils
 	 * @param  CategoryId	The common category you want a path for.
 	 * @return A text string, (empty if the common category was not registered)
 	 */
-	static FText const& GetCommonCategory(FCommonEditorCategory::EValue CategoryId);
+	UNREALED_API FText const& GetCommonCategory(const FCommonEditorCategory::EValue CategoryId);
 
 	/**
 	 * Utility function that concatenates the supplied sub-category with one 
@@ -84,7 +93,7 @@ struct UNREALED_API FEditorCategoryUtils
 	 * @param  SubCategory	A sub-category that you want postfixed to the root category.
 	 * @return A concatenated text string, with the two categories separated by a pipe, '|', character.
 	 */
-	static FText BuildCategoryString(FCommonEditorCategory::EValue RootCategory, FText const& SubCategory);
+	UNREALED_API FText BuildCategoryString(const FCommonEditorCategory::EValue RootCategory, FText const& SubCategory);
 
 	/**
 	 * Expands any keys found in the category string (any terms found in square 
@@ -94,7 +103,7 @@ struct UNREALED_API FEditorCategoryUtils
 	 * @param  SubCategory	A sub-category that you want postfixing the result.
 	 * @return A concatenated text string, with the two categories separated by a pipe, '|', character.
 	 */
-	static FText GetCategoryDisplayString(FText const& UnsanitizedCategory);
+	UNREALED_API FText GetCategoryDisplayString(FText const& UnsanitizedCategory);
 
 	/**
 	 * Expands any keys found in the category string (any terms found in square 
@@ -104,7 +113,7 @@ struct UNREALED_API FEditorCategoryUtils
 	 * @param  SubCategory	A sub-category that you want postfixing the result.
 	 * @return A concatenated string, with the two categories separated by a pipe, '|', character.
 	 */
-	static FString GetCategoryDisplayString(FString const& UnsanitizedCategory);
+	UNREALED_API FString GetCategoryDisplayString(FString const& UnsanitizedCategory);
 
 	/**
 	 * Parses out the class's "HideCategories" metadata, and returns it 
@@ -113,7 +122,7 @@ struct UNREALED_API FEditorCategoryUtils
 	 * @param  Class			The class you want to pull data from.
 	 * @param  CategoriesOut	An array that will be filled with a list of hidden categories.
 	 */
-	static void GetClassHideCategories(UClass const* Class, TArray<FString>& CategoriesOut);
+	UNREALED_API void GetClassHideCategories(UClass const* Class, TArray<FString>& CategoriesOut);
 
 	/**
 	 * Parses out the class's "ShowCategories" metadata, and returns it 
@@ -122,7 +131,7 @@ struct UNREALED_API FEditorCategoryUtils
 	 * @param  Class			The class you want to pull data from.
 	 * @param  CategoriesOut	An array that will be filled with a list of shown categories.
 	 */
-	static void GetClassShowCategories(UClass const* Class, TArray<FString>& CategoriesOut);
+	UNREALED_API void GetClassShowCategories(UClass const* Class, TArray<FString>& CategoriesOut);
 
 	/**
 	 * Checks to see if the category associated with the supplied common 
@@ -132,7 +141,7 @@ struct UNREALED_API FEditorCategoryUtils
 	 * @param  CategoryId	An id associated with a category that you want to check.
 	 * @return True if the common category is hidden, false if not.
 	 */
-	static bool IsCategoryHiddenFromClass(UClass const* Class, FCommonEditorCategory::EValue CategoryId);
+	UNREALED_API bool IsCategoryHiddenFromClass(UClass const* Class, const FCommonEditorCategory::EValue CategoryId);
 
 	/**
 	 * Checks to see if the specified category is hidden from the supplied class.
@@ -141,7 +150,7 @@ struct UNREALED_API FEditorCategoryUtils
 	 * @param  Category A category path that you want to check.
 	 * @return True if the category is hidden, false if not.
 	 */
-	static bool IsCategoryHiddenFromClass(UClass const* Class, FText const& Category);
+	UNREALED_API bool IsCategoryHiddenFromClass(UClass const* Class, FText const& Category);
 
 	/**
 	 * Checks to see if the specified category is hidden from the supplied class.
@@ -150,5 +159,15 @@ struct UNREALED_API FEditorCategoryUtils
 	 * @param  Category A category path that you want to check.
 	 * @return True if the category is hidden, false if not.
 	 */
-	static bool IsCategoryHiddenFromClass(UClass const* Class, FString const& Category);
+	UNREALED_API bool IsCategoryHiddenFromClass(UClass const* Class, FString const& Category);
+
+	/**
+	 * Returns tooltip information for the specified category
+	 * 
+	 * @param  Category		The category that you want to tooltip details for.
+	 * @param  Tooltip		[OUT] The tooltip to display for this category.
+	 * @param  DocLink		[OUT] The link to the documentation page for this category.
+	 * @param  DocExcerpt	[OUT] A category path that you want to check.
+	 */
+	UNREALED_API void GetCategoryTooltipInfo(const FString& Category, FText& Tooltip, FString& DocLink, FString& DocExcerpt);
 };
\ No newline at end of file
