#ue4 - Optimize FBlueprintEditorUtils::PostEditChangeBlueprintActors() to use GetObjectsOfClass rather than TActorIterator and casting all the results.

diff --git a/Engine/Source/Editor/UnrealEd/Private/Kismet2/BlueprintEditorUtils.cpp b/Engine/Source/Editor/UnrealEd/Private/Kismet2/BlueprintEditorUtils.cpp
--- a/Engine/Source/Editor/UnrealEd/Private/Kismet2/BlueprintEditorUtils.cpp
+++ b/Engine/Source/Editor/UnrealEd/Private/Kismet2/BlueprintEditorUtils.cpp
@@ -5687,20 +5687,24 @@ void FBlueprintEditorUtils::UpdateOldPureFunctions( UBlueprint* Blueprint )
 /** Call PostEditChange() on any Actors that are based on this Blueprint */
 void FBlueprintEditorUtils::PostEditChangeBlueprintActors(UBlueprint* Blueprint)
 {
-	for(TObjectIterator<AActor> ActorIt; ActorIt; ++ActorIt)
+	if (Blueprint->GeneratedClass && Blueprint->GeneratedClass->IsChildOf(AActor::StaticClass()))
 	{
-		AActor* Actor = *ActorIt;
-		if(Actor->GetClass()->IsChildOf(Blueprint->GeneratedClass) && !Actor->IsPendingKill())
+		TArray<UObject*> MatchingBlueprintObjects;
+		GetObjectsOfClass(Blueprint->GeneratedClass, MatchingBlueprintObjects, true, (RF_ClassDefaultObject | RF_PendingKill));
+
+		for (auto ObjIt : MatchingBlueprintObjects)
 		{
+			// We know the class was derived from AActor because we checked the Blueprint->GeneratedClass.
+			AActor* Actor = static_cast<AActor*>(ObjIt);
 			Actor->PostEditChange();
 
-			// Let components that got re-regsitered by PostEditChange also get begun play
+			// Let components that got re-registered by PostEditChange also get begun play
 			if(Actor->bActorInitialized)
 			{
 				Actor->InitializeComponents();
 			}
 		}
-	}
+	}	
 
 	// Let the blueprint thumbnail renderer know that a blueprint has been modified so it knows to reinstance components for visualization
 	FThumbnailRenderingInfo* RenderInfo = GUnrealEd->GetThumbnailManager()->GetRenderingInfo( Blueprint );
