fixed typo

diff --git a/Engine/Source/Editor/UnrealEd/Private/Settings/ProjectPackagingSettingsCustomization.h b/Engine/Source/Editor/UnrealEd/Private/Settings/ProjectPackagingSettingsCustomization.h
--- a/Engine/Source/Editor/UnrealEd/Private/Settings/ProjectPackagingSettingsCustomization.h
+++ b/Engine/Source/Editor/UnrealEd/Private/Settings/ProjectPackagingSettingsCustomization.h
@@ -54,11 +54,11 @@ protected:
 		{
 			IDetailCategoryBuilder& ProjectCategory = LayoutBuilder.EditCategory("Project");
 			{
-				TSharedRef<FPropertyRestriction> BuildConfigurationRestruction = MakeShareable(new FPropertyRestriction(LOCTEXT("DebugGameRestrictionReason", "The DebugGame build configuration is not available in content-only projects.")));
-				BuildConfigurationRestruction->AddValue("DebugGame");
+				TSharedRef<FPropertyRestriction> BuildConfigurationRestriction = MakeShareable(new FPropertyRestriction(LOCTEXT("DebugGameRestrictionReason", "The DebugGame build configuration is not available in content-only projects.")));
+				BuildConfigurationRestriction->AddValue("DebugGame");
 
 				TSharedRef<IPropertyHandle> BuildConfigurationHandle = LayoutBuilder.GetProperty("BuildConfiguration");
-				BuildConfigurationHandle->AddRestriction(BuildConfigurationRestruction);
+				BuildConfigurationHandle->AddRestriction(BuildConfigurationRestriction);
 			}
 		}
 	}
