Disabling the convert to test as I regress the issue locally.

diff --git a/Engine/Source/Editor/UnrealEd/Private/Tests/EditorAutomationTests.cpp b/Engine/Source/Editor/UnrealEd/Private/Tests/EditorAutomationTests.cpp
--- a/Engine/Source/Editor/UnrealEd/Private/Tests/EditorAutomationTests.cpp
+++ b/Engine/Source/Editor/UnrealEd/Private/Tests/EditorAutomationTests.cpp
@@ -1128,7 +1128,7 @@ bool FStaticMeshValidation::RunTest(const FString& Parameters)
 /**
  * QA Convert Meshes Regression Testing
  */
-IMPLEMENT_SIMPLE_AUTOMATION_TEST(FConvertToValidation, "QA.Convert Meshes", EAutomationTestFlags::ATF_Editor);
+IMPLEMENT_SIMPLE_AUTOMATION_TEST(FConvertToValidation, "QA.Convert Meshes", EAutomationTestFlags::ATF_Editor | EAutomationTestFlags::ATF_RequiresUser);
 
 //gather all brushes that exist now
 void ConvertTestFindAllBrushes(TArray<ABrush*> &PreviousBrushes)
