FBlueprintCompileOnLoadTest::RunTest - fix

diff --git a/Engine/Source/Editor/UnrealEd/Private/Tests/BlueprintAutomationTests.cpp b/Engine/Source/Editor/UnrealEd/Private/Tests/BlueprintAutomationTests.cpp
--- a/Engine/Source/Editor/UnrealEd/Private/Tests/BlueprintAutomationTests.cpp
+++ b/Engine/Source/Editor/UnrealEd/Private/Tests/BlueprintAutomationTests.cpp
@@ -855,22 +855,30 @@ bool FBlueprintCompileOnLoadTest::RunTest(const FString& BlueprintAssetPath)
 		return false;
 	}
 
+	TArray<TWeakObjectPtr<UClass>> BlueprintDependenciesWP;
+	{
 #if WITH_EDITOR
-	TArray<UClass*> BlueprintDependencies = FScopedClassDependencyGather::GetCachedDependencies();
+		TArray<UClass*> BlueprintDependencies = FScopedClassDependencyGather::GetCachedDependencies();
 #else 
-	TArray<UClass*> BlueprintDependencies;
-	for (TObjectIterator<UBlueprint> BlueprintIt; BlueprintIt; ++BlueprintIt)
-	{
-		if (*BlueprintIt == InitialBlueprint)
+		TArray<UClass*> BlueprintDependencies;
+		for (TObjectIterator<UBlueprint> BlueprintIt; BlueprintIt; ++BlueprintIt)
 		{
-			continue;
+			if (*BlueprintIt == InitialBlueprint)
+			{
+				continue;
+			}
+			else if (BlueprintIt->GeneratedClass != NULL)
+			{
+				BlueprintDependencies.Add(BlueprintIt->GeneratedClass);
+			}
 		}
-		else if (BlueprintIt->GeneratedClass != NULL)
+#endif 
+		for (UClass* ClassDependency : BlueprintDependencies)
 		{
-			BlueprintDependencies.Add(BlueprintIt->GeneratedClass);
+			BlueprintDependenciesWP.Add(ClassDependency);
 		}
 	}
-#endif 
+	
 
 	// store off data for the initial blueprint so we can unload it (and reconstruct 
 	// later to compare it with a second one)
@@ -901,10 +909,14 @@ bool FBlueprintCompileOnLoadTest::RunTest(const FString& BlueprintAssetPath)
 		ClassRedirects.Add(UnloadedBlueprint->GeneratedClass, ReloadedBlueprint->GeneratedClass);
 		ClassRedirects.Add(UnloadedBlueprint->SkeletonGeneratedClass, ReloadedBlueprint->SkeletonGeneratedClass);
 
-		for (UClass* ClassDependency : BlueprintDependencies)
+		for (auto ClassDependencyWP : BlueprintDependenciesWP)
 		{
-			FArchiveReplaceObjectRef<UBlueprint>(ClassDependency, BlueprintRedirects, /*bNullPrivateRefs=*/false, /*bIgnoreOuterRef=*/true, /*bIgnoreArchetypeRef=*/false);
-			FArchiveReplaceObjectRef<UClass>(ClassDependency, ClassRedirects, /*bNullPrivateRefs=*/false, /*bIgnoreOuterRef=*/true, /*bIgnoreArchetypeRef=*/false);
+			auto ClassDependency = ClassDependencyWP.Get();
+			if (ClassDependency)
+			{
+				FArchiveReplaceObjectRef<UBlueprint>(ClassDependency, BlueprintRedirects, /*bNullPrivateRefs=*/false, /*bIgnoreOuterRef=*/true, /*bIgnoreArchetypeRef=*/false);
+				FArchiveReplaceObjectRef<UClass>(ClassDependency, ClassRedirects, /*bNullPrivateRefs=*/false, /*bIgnoreOuterRef=*/true, /*bIgnoreArchetypeRef=*/false);
+			}
 		}
 
 		UPackage* AssetPackage = ReloadedBlueprint->GetOutermost();
