Blueprint API Documentation generation first draft

diff --git a/Engine/Source/Editor/UnrealEd/Classes/Commandlets/GenerateBlueprintAPICommandlet.h b/Engine/Source/Editor/UnrealEd/Classes/Commandlets/GenerateBlueprintAPICommandlet.h
new file mode 100644
--- /dev/null
+++ b/Engine/Source/Editor/UnrealEd/Classes/Commandlets/GenerateBlueprintAPICommandlet.h
@@ -0,0 +1,16 @@
+// Copyright 1998-2014 Epic Games, Inc. All Rights Reserved.
+
+#pragma once
+#include "Commandlets/Commandlet.h"
+#include "GenerateBlueprintAPICommandlet.generated.h"
+
+UCLASS()
+class UGenerateBlueprintAPICommandlet : public UCommandlet
+{
+	GENERATED_UCLASS_BODY()
+
+public:		
+	// Begin UCommandlet Interface
+	virtual int32 Main(FString const& Params) override;
+	// End UCommandlet Interface
+};
diff --git a/Engine/Source/Editor/UnrealEd/Private/Commandlets/GenerateBlueprintAPICommandlet.cpp b/Engine/Source/Editor/UnrealEd/Private/Commandlets/GenerateBlueprintAPICommandlet.cpp
new file mode 100644
--- /dev/null
+++ b/Engine/Source/Editor/UnrealEd/Private/Commandlets/GenerateBlueprintAPICommandlet.cpp
@@ -0,0 +1,1031 @@
+// Copyright 1998-2014 Epic Games, Inc. All Rights Reserved.
+
+#include "UnrealEd.h"
+#include "ObjectEditorUtils.h"
+#include "Kismet2/KismetEditorUtilities.h"
+#include "K2ActionMenuBuilder.h"
+#include "EdGraphSchema_K2_Actions.h"
+#include "K2Node.h"
+#include "AnimationGraph.h"
+#include "AnimGraphNode_StateMachine.h"
+#include "BlueprintEditorUtils.h"
+#include "AnimStateConduitNode.h"
+#include "AnimStateNode.h"
+#include "Engine/LevelScriptBlueprint.h"
+#include "BehaviorTree/BehaviorTree.h"
+#include "AssetSelection.h"
+#include "K2Node_CustomEvent.h"
+#include "ScopedTimers.h"
+#include "BlueprintActionMenuBuilder.h"
+#include "BlueprintActionFilter.h"
+#include "BlueprintActionDatabase.h"
+#include "BlueprintActionMenuUtils.h"
+#include "EditorCategoryUtils.h"
+
+DEFINE_LOG_CATEGORY_STATIC(LogBlueprintInfoDump, Log, All);
+
+/*******************************************************************************
+ * Static Helpers
+ ******************************************************************************/
+
+namespace GenerateBlueprintAPIUtils
+{
+	static const FString HelpString = TEXT("\n\
+\n\
+GenerateBlueprintAPI commandlet params: \n\
+\n\
+    -class=<Class>      Used to specify the blueprint's parent class, if left  \n\
+                        unset then it will go through EVERY blueprint parent   \n\
+                        class available.                                       \n\
+\n\
+    -multifile          Used to keep file size down, will split each blueprint \n\
+                        into its own file (meaning only one file will be       \n\
+                        created when used with -class).                        \n\
+\n\
+    -palette            Defaults to on, unless other flags were specified.     \n\
+                        Dumps all actions from the blueprint's palette menu    \n\
+                        (constant across all graphs).                          \n\
+\n\
+    -palfilter=<Class>  Simulates picking a class from the blueprint palette's \n\
+                        drop down. Setting -palfilter=all will dump the palette\n\
+                        for every possible class.                              \n\
+\n\
+    -time               When enabled, will record timings during menu building \n\
+                        (as it has been time sync in the past). This is not    \n\
+                        ideal for diffs though (since times can easily vary).  \n\
+\n\
+    -experimental       Uses an new way of constructing Blueprint action menus \n\
+                        (that will replace the current system).                \n\
+\n\
+    -name=<Filename>    Overrides the default filename. Leave off the extention\n\
+                        (this will add .json to the end). When -multifile is   \n\
+                        supplied, the class name will be postfixed to the name.\n\
+\n\
+    -help, -h, -?       Display this message and then exit.                    \n\
+\n");
+
+	/** Flags that govern the verbosity of the dump. */
+	enum EDumpFlags
+	{
+		BPDUMP_FilteredPalette		= (1<<0), 
+
+		BPDUMP_LogHelp				= (1<<1),
+		BPDUMP_RecordTiming			= (1<<2),
+		BPDUMP_UseLegacyMenuBuilder = (1<<3),
+	};
+
+	/**
+	 * A collection of variables that represent the various command switches 
+	 * that users can specify when running the commandlet. See the HelpString 
+	 * variable for a listing of supported switches.
+	 */
+	struct CommandletOptions
+	{
+		CommandletOptions()
+			: BlueprintClass(nullptr)
+			, DumpFlags(0)
+			, PaletteFilter(nullptr)
+		{
+		}
+
+		/**
+		* Parses the string command switches into flags, class pointers, and 
+		* booleans that will govern what should be dumped. Logs errors if any 
+		* switch was misused.
+		*/
+		CommandletOptions(TArray<FString> const& Switches);
+
+		UClass*    BlueprintClass;
+		uint32     DumpFlags;
+		UClass*    PaletteFilter;
+		FString    SaveDir;
+		FString    Filename;
+	};
+
+	/** Static instance of the command switches (so we don't have to pass one along the call stack) */
+	static CommandletOptions CommandOptions;
+
+	/** Tracks instantiated blueprints (so we don't have to create more the we have to). */
+	static TMap<UClass*, UBlueprint*> ClassBlueprints;
+
+	/**
+	 * Certain blueprints (like level blueprints) require a level outer, and 
+	 * for certain actions we need a level actor selected. This utility function
+	 * provides an easy way to grab the world (which has a level that we can use 
+	 * for these purposes).
+	 */
+	static UWorld* GetWorld();
+
+	/**
+	 * Spawns a transient blueprint of the specified type. Adds all possible 
+	 * graph types (function, macro, etc.), and does some additional setup for
+	 * unique blueprint types (like level and anim blueprints).
+	 */
+	static UBlueprint* MakeTempBlueprint(UClass* BlueprintClass);
+
+	/**
+	 * Certain nodes add specific graph types that we want to dump info on (like 
+	 * state machine graphs for anim blueprints). The best way to add those 
+	 * graphs is through the natural process of adding those nodes (which this
+	 * method is intended for).
+	 */
+	template<class NodeType>
+	static NodeType* AddNodeToGraph(UEdGraph* Graph);
+
+	/**
+	 * Builds a fully qualified file path for a new dump file. If using the 
+	 * -multifile command switch, then this will create a sub-directory and name
+	 * the file after the class. Generally, dump files are placed in the 
+	 * project's ".../Saved/Commandlets/" directory.
+	 */
+	static FString BuildDumpFilePath(UClass* BlueprintClass);
+
+	/**
+	 * Utility function to convert a tab integer into a string of whitespace.
+	 * Defaults to tab characters, but if bUseSpaces is enabled, then single 
+	 * spaces are used.
+	 */
+	static FString BuildIndentString(uint32 IndentCount, bool bUseSpaces = false);
+
+	/**
+	 * Concatenates the action's category with its menu name (to help
+	 * distinguish similarly named actions). Can then be used to sort and 
+	 * uniquely identify actions.
+	 */
+	static FString GetActionKey(FGraphActionListBuilderBase::ActionGroup const& Action);
+
+	/**
+	* Goes through all of the blueprint skeleton's object properties and pulls
+	* out the ones that are associated with an UActorComponent (and are visbile
+	* to the blueprint).
+	*/
+	static void GetComponentProperties(UBlueprint* Blueprint, TArray<UObjectProperty*>& PropertiesOut);
+	
+	/**
+	 * Constructs a temporary blueprint (of the class type specified) and kicks 
+	 * off a dump of all its nested information (palette, graph, contextual 
+	 * actions, etc.). 
+	 */
+	static void DumpInfoForClass(uint32 Indent, UClass* BlueprintClass, FArchive* FileOutWriter);
+	
+	/**
+	 * Writes out all the category details that have been accumulated during the palette construction
+	 */
+	static void DumpCategoryInfo(uint32 Indent, FArchive* FileOutWriter);
+
+	/**
+	 * Assumes that the specified ActionListBuilder is configured with the 
+	 * proper blueprint. Starts by clearing any actions it contained and then
+	 * runs through and re-adds all actions.
+	 *
+	 * @return The amount of time (in seconds) that the menu building took.
+	 */
+	static double GetPaletteMenuActions(FBlueprintPaletteListBuilder& ActionListBuilder, UClass* PaletteFilter);
+
+	/**
+	 * Dumps all palette actions listed for the specified blueprint. Determines
+	 * if the user specified any filter class for the palette and adjusts 
+	 * accordingly (can dump multiple palettes if -palfilter=all was specified).
+	 */
+	static void DumpPalette(uint32 Indent, UBlueprint* Blueprint, FArchive* FileOutWriter);
+
+	/**
+	 * Dumps a single instance of the blueprint's palette (using the ClassFilter).
+	 * ClassFilter can be null and the full unfiltered palette will be dumped.
+	 */
+	static void DumpPalette(uint32 Indent, UBlueprint* Blueprint, UClass* ClassFilter, FArchive* FileOutWriter);
+
+	/**
+	 * Generic function utilized by both palette and context-menu dumps. Take a
+	 * GraphActionListBuilder and writes out every action that it has captured.
+	 */
+	static void DumpActionList(uint32 Indent, FGraphActionListBuilderBase& ActionList, FArchive* FileOutWriter);
+
+	/**
+	 * Generic function that dumps information on a single action (like it's 
+	 * name, category, an associated node if it has one, etc.).
+	 */
+	static void DumpActionMenuItem(uint32 Indent, FGraphActionListBuilderBase::ActionGroup const& Action, FGraphActionListBuilderBase& ActionList, FArchive* FileOutWriter);
+}
+
+//------------------------------------------------------------------------------
+FString MakeJsonString( const FString& String )
+{
+	FString OutString;
+	for (const TCHAR* Char = *String; *Char != TCHAR('\0'); ++Char)
+	{
+		switch (*Char)
+		{
+		case TCHAR('\\'): OutString += TEXT("\\\\"); break;
+		case TCHAR('\n'): OutString += TEXT("\\n"); break;
+		case TCHAR('\t'): OutString += TEXT("\\t"); break;
+		case TCHAR('\b'): OutString += TEXT("\\b"); break;
+		case TCHAR('\f'): OutString += TEXT("\\f"); break;
+		case TCHAR('\r'): OutString += TEXT("\\r"); break;
+		case TCHAR('\"'): OutString += TEXT("\\\""); break;
+		default: OutString += *Char;
+		}
+	}
+	return OutString;
+}
+
+//------------------------------------------------------------------------------
+GenerateBlueprintAPIUtils::CommandletOptions::CommandletOptions(TArray<FString> const& Switches)
+	: BlueprintClass(AActor::StaticClass())
+	, DumpFlags(BPDUMP_UseLegacyMenuBuilder)
+	, PaletteFilter(nullptr)
+{
+	uint32 NewDumpFlags = BPDUMP_UseLegacyMenuBuilder;
+	for (FString const& Switch : Switches)
+	{
+		if (Switch.StartsWith("class="))
+		{
+			FString ClassSwitch, ClassName;
+			Switch.Split(TEXT("="), &ClassSwitch, &ClassName);
+			BlueprintClass = FindObject<UClass>(ANY_PACKAGE, *ClassName);
+
+			if (BlueprintClass == nullptr)
+			{
+				UE_LOG(LogBlueprintInfoDump, Error, TEXT("Unrecognized blueprint class '%s', defaulting to 'Actor'"), *ClassName);
+				BlueprintClass = AActor::StaticClass();
+			}
+		}
+		else if (Switch.StartsWith("palfilter="))
+		{
+			FString ClassSwitch, ClassName;
+			Switch.Split(TEXT("="), &ClassSwitch, &ClassName);
+			PaletteFilter = FindObject<UClass>(ANY_PACKAGE, *ClassName);
+
+			NewDumpFlags |= BPDUMP_FilteredPalette;
+			if (PaletteFilter == nullptr)
+			{
+				if (ClassName.Compare("all", ESearchCase::IgnoreCase))
+				{
+					UE_LOG(LogBlueprintInfoDump, Error, TEXT("Unrecognized palette filter '%s', defaulting to unfiltered"), *ClassName);
+					NewDumpFlags &= ~(GenerateBlueprintAPIUtils::BPDUMP_FilteredPalette);
+				}
+			}
+		}
+		else if (!Switch.Compare("h", ESearchCase::IgnoreCase) || 
+		         !Switch.Compare("?", ESearchCase::IgnoreCase) || 
+				 !Switch.Compare("help", ESearchCase::IgnoreCase))
+		{
+			NewDumpFlags |= BPDUMP_LogHelp;
+		}
+		else if (!Switch.Compare("time", ESearchCase::IgnoreCase))
+		{
+			NewDumpFlags |= BPDUMP_RecordTiming;
+		}
+		else if (!Switch.Compare("experimental", ESearchCase::IgnoreCase))
+		{
+			NewDumpFlags &= ~BPDUMP_UseLegacyMenuBuilder;
+		}
+		else if (Switch.StartsWith("name="))
+		{
+			FString NameSwitch;
+			Switch.Split(TEXT("="), &NameSwitch, &Filename);
+		}
+		else if (Switch.StartsWith("path="))
+		{
+			FString PathSwitch;
+			Switch.Split(TEXT("="), &PathSwitch, &SaveDir);
+		}
+		else
+		{
+			UE_LOG(LogBlueprintInfoDump, Warning, TEXT("Unrecognized command switch '%s', use -help for a listing of all accepted params"), *Switch);
+		}
+	}
+
+	if (NewDumpFlags != 0)
+	{
+		DumpFlags = NewDumpFlags;
+	}
+}
+
+//------------------------------------------------------------------------------
+static UWorld* GenerateBlueprintAPIUtils::GetWorld()
+{
+	UWorld* World = nullptr;
+	for (FWorldContext const& WorldContext : GEngine->GetWorldContexts())
+	{
+		World = WorldContext.World();
+		if (World != nullptr)
+		{
+			break;
+		}
+	}
+
+	if (World == nullptr)
+	{
+		static UWorld* CommandletWorld = nullptr;
+		if (CommandletWorld == nullptr)
+		{
+			if (GUnrealEd == nullptr)
+			{
+				UE_LOG(LogBlueprintInfoDump, Error, TEXT("Cannot create a temp map to test within, without a valid editor world"));
+			}
+			else
+			{
+				CommandletWorld = GEditor->NewMap();
+			}
+		}
+
+		World = CommandletWorld;
+	}
+
+	return World;
+}
+
+
+//------------------------------------------------------------------------------
+static UBlueprint* GenerateBlueprintAPIUtils::MakeTempBlueprint(UClass* ParentClass)
+{
+	UBlueprint* MadeBlueprint = nullptr;
+	if (UBlueprint** FoundBlueprint = ClassBlueprints.Find(ParentClass))
+	{
+		MadeBlueprint = *FoundBlueprint;
+	}
+	else
+	{
+		UObject* BlueprintOuter = GetTransientPackage();
+
+		bool bIsAnimBlueprint = ParentClass->IsChildOf(UAnimInstance::StaticClass());
+		bool bIsLevelBlueprint = ParentClass->IsChildOf(ALevelScriptActor::StaticClass());
+
+		UClass* BlueprintClass = UBlueprint::StaticClass();
+		UClass* GeneratedClass = UBlueprintGeneratedClass::StaticClass();
+		EBlueprintType BlueprintType = BPTYPE_Normal;
+
+		if (bIsAnimBlueprint)
+		{
+			BlueprintClass = UAnimBlueprint::StaticClass();
+			GeneratedClass = UAnimBlueprintGeneratedClass::StaticClass();
+		}
+		else if (bIsLevelBlueprint)
+		{
+			UWorld* World = GetWorld();
+			if (World == nullptr)
+			{
+				UE_LOG(LogBlueprintInfoDump, Error, TEXT("Cannot make a proper level blueprint without a valid editor level for its outer."));
+			}
+			else
+			{
+				BlueprintClass = ULevelScriptBlueprint::StaticClass();
+				BlueprintType = BPTYPE_LevelScript;
+				BlueprintOuter = World->GetCurrentLevel();
+			}
+		}
+		// @TODO: UEditorUtilityBlueprint
+
+		FString const ClassName = ParentClass->GetName();
+		FString const DesiredName = FString::Printf(TEXT("COMMANDLET_TEMP_Blueprint_%s"), *ClassName);
+		FName   const TempBpName = MakeUniqueObjectName(BlueprintOuter, BlueprintClass, FName(*DesiredName));
+
+		check(FKismetEditorUtilities::CanCreateBlueprintOfClass(ParentClass));
+		MadeBlueprint = FKismetEditorUtilities::CreateBlueprint(ParentClass, BlueprintOuter, TempBpName, BlueprintType, BlueprintClass, GeneratedClass);
+		
+		// if this is an animation blueprint, then we want anim specific graphs to test as well (if it has an anim graph)...
+		if (bIsAnimBlueprint && (MadeBlueprint->FunctionGraphs.Num() > 0))
+		{
+			UAnimationGraph* AnimGraph = CastChecked<UAnimationGraph>(MadeBlueprint->FunctionGraphs[0]);
+			check(AnimGraph != nullptr);
+
+			// should add a state-machine graph
+			UAnimGraphNode_StateMachine* StateMachineNode = AddNodeToGraph<UAnimGraphNode_StateMachine>(AnimGraph);
+
+			UAnimationStateMachineGraph* StateMachineGraph = StateMachineNode->EditorStateMachineGraph;
+			// should add an conduit graph 
+			UAnimStateConduitNode* ConduitNode = AddNodeToGraph<UAnimStateConduitNode>((UEdGraph*)StateMachineGraph);
+
+			UAnimStateNode* StateNode = AddNodeToGraph<UAnimStateNode>((UEdGraph*)StateMachineGraph);
+			// should create a transition graph
+			StateNode->AutowireNewNode(ConduitNode->GetOutputPin());
+		}
+		else if (bIsLevelBlueprint)
+		{
+			ULevel* Level = CastChecked<ULevel>(BlueprintOuter);
+			Level->LevelScriptBlueprint = Cast<ULevelScriptBlueprint>(MadeBlueprint);
+		}
+
+		// may have been altered in CreateBlueprint()
+		BlueprintType = MadeBlueprint->BlueprintType;
+
+		bool bCanAddFunctions = (BlueprintType != BPTYPE_MacroLibrary); // taken from FBlueprintEditor::NewDocument_IsVisibleForType()
+		if (bCanAddFunctions)
+		{
+			// add a functions graph that isn't the construction script (or an animation graph)
+			FName FuncGraphName = MakeUniqueObjectName(MadeBlueprint, UEdGraph::StaticClass(), FName(TEXT("NewFunction")));
+			UEdGraph* FuncGraph = FBlueprintEditorUtils::CreateNewGraph(MadeBlueprint, FuncGraphName, UEdGraph::StaticClass(), UEdGraphSchema_K2::StaticClass());
+			FBlueprintEditorUtils::AddFunctionGraph<UClass>(MadeBlueprint, FuncGraph, /*bIsUserCreated =*/true, nullptr);
+		}
+
+		bool bCanAddMacros = ((BlueprintType == BPTYPE_MacroLibrary) || (BlueprintType == BPTYPE_Normal) || (BlueprintType == BPTYPE_LevelScript));
+		if (bCanAddMacros)
+		{
+			FName MacroGraphName = MakeUniqueObjectName(MadeBlueprint, UEdGraph::StaticClass(), FName(TEXT("NewMacro")));
+			UEdGraph* MacroGraph = FBlueprintEditorUtils::CreateNewGraph(MadeBlueprint, MacroGraphName, UEdGraph::StaticClass(), UEdGraphSchema_K2::StaticClass());
+			FBlueprintEditorUtils::AddMacroGraph(MadeBlueprint, MacroGraph, /*bIsUserCreated =*/true, nullptr);
+		}
+
+		// if you can add custom events to this blueprint, do so (so show that we 
+		// can call an event on ourselves)
+		if (MadeBlueprint->UbergraphPages.Num() > 0)
+		{
+			UK2Node_CustomEvent* CustomEventNode = AddNodeToGraph<UK2Node_CustomEvent>(MadeBlueprint->UbergraphPages[0]);
+			CustomEventNode->CustomFunctionName = FBlueprintEditorUtils::FindUniqueCustomEventName(MadeBlueprint);
+		}
+
+		FBlueprintEditorUtils::MarkBlueprintAsStructurallyModified(MadeBlueprint);
+		MadeBlueprint->SetFlags(RF_RootSet); // to keep the BP from being garbage collected
+		FKismetEditorUtilities::CompileBlueprint(MadeBlueprint);
+		ClassBlueprints.Add(ParentClass, MadeBlueprint);
+	}
+
+	check(MadeBlueprint != nullptr);
+	return MadeBlueprint;
+}
+
+//------------------------------------------------------------------------------
+template<class NodeType>
+static NodeType* GenerateBlueprintAPIUtils::AddNodeToGraph(UEdGraph* Graph)
+{
+	check(Graph != nullptr);
+	NodeType* NewNode = NewObject<NodeType>(Graph);
+	Graph->AddNode(NewNode, /*bFromUI =*/true, /*bSelectNewNode =*/false);
+
+	NewNode->CreateNewGuid();
+	NewNode->PostPlacedNewNode();
+	NewNode->AllocateDefaultPins();
+	return NewNode;
+}
+
+//------------------------------------------------------------------------------
+static FString GenerateBlueprintAPIUtils::BuildDumpFilePath(UClass* BlueprintClass)
+{
+	FString CommandletSaveDir, Pathname;
+
+	if (CommandOptions.SaveDir.IsEmpty())
+	{
+		CommandletSaveDir = FPaths::GameSavedDir() + TEXT("Commandlets/");
+		CommandletSaveDir = FPaths::ConvertRelativePathToFull(CommandletSaveDir);
+	}
+	else
+	{
+		CommandletSaveDir= CommandOptions.SaveDir;
+	}
+	IFileManager::Get().MakeDirectory(*CommandletSaveDir);
+	
+	if (CommandOptions.Filename.IsEmpty())
+	{
+		FString Pathname = FString::Printf(TEXT("GenerateBlueprintAPI_%s"), FPlatformTime::StrTimestamp());
+		Pathname = Pathname.Replace(TEXT(" "), TEXT("_"));
+		Pathname = Pathname.Replace(TEXT("/"), TEXT("-"));
+		Pathname = Pathname.Replace(TEXT(":"), TEXT("."));
+	}
+	else
+	{
+		Pathname = CommandOptions.Filename;
+	}
+
+	Pathname += ".json";
+
+	return CommandletSaveDir / *Pathname;
+}
+
+//------------------------------------------------------------------------------
+static FString GenerateBlueprintAPIUtils::BuildIndentString(uint32 IndentCount, bool bUseSpaces)
+{
+	char RepatingChar = '\t';
+	if (bUseSpaces)
+	{
+		RepatingChar = ' ';
+	}
+
+	FString IndentString;
+	while (IndentCount > 0)
+	{
+		IndentString += RepatingChar;
+		--IndentCount;
+	}
+	return IndentString;
+}
+
+//------------------------------------------------------------------------------
+static FString GenerateBlueprintAPIUtils::GetActionKey(FGraphActionListBuilderBase::ActionGroup const& Action)
+{
+	TArray<FString> MenuHierarchy;
+	Action.GetCategoryChain(MenuHierarchy);
+
+	FString ActionKey;
+	for (FString const& SubCategory : MenuHierarchy)
+	{
+		ActionKey += SubCategory + TEXT("|");
+	}
+	if (MenuHierarchy.Num() > 0)
+	{
+		ActionKey.RemoveAt(ActionKey.Len() - 1); // remove the trailing '|'
+	}
+
+	TSharedPtr<FEdGraphSchemaAction> MainAction = Action.Actions[0];
+	ActionKey += MainAction->MenuDescription.ToString();
+
+	return ActionKey;
+}
+
+//------------------------------------------------------------------------------
+static void GenerateBlueprintAPIUtils::GetComponentProperties(UBlueprint* Blueprint, TArray<UObjectProperty*>& PropertiesOut)
+{
+	UClass* BpClass = Blueprint->GeneratedClass;
+	if (BpClass->IsChildOf<AActor>())
+	{
+		for (TFieldIterator<UObjectProperty> PropertyIt(BpClass, EFieldIteratorFlags::IncludeSuper); PropertyIt; ++PropertyIt)
+		{
+			// SMyBlueprint filters out component variables in SMyBlueprint::CollectAllActions() using CPF_BlueprintVisible/CPF_Parm flags
+			if (PropertyIt->PropertyClass->IsChildOf(UActorComponent::StaticClass()) &&
+				PropertyIt->HasAnyPropertyFlags(CPF_BlueprintVisible) && !PropertyIt->HasAnyPropertyFlags(CPF_Parm))
+			{
+				PropertiesOut.Add(*PropertyIt);
+			}
+		}
+	}
+}
+
+//------------------------------------------------------------------------------
+static void GenerateBlueprintAPIUtils::DumpInfoForClass(uint32 Indent, UClass* BlueprintClass, FArchive* FileOutWriter)
+{
+	FString const ClassName = BlueprintClass->GetName();
+	UE_LOG(LogBlueprintInfoDump, Display, TEXT("%sDumping BP class: '%s'..."), *BuildIndentString(Indent, true), *ClassName);
+
+	FString const ClassEntryIndent = BuildIndentString(Indent);
+	FString BeginClassEntry = FString::Printf(TEXT("%s\"%s\" : {"), *ClassEntryIndent, *ClassName);
+
+	FString IndentedNewline = "\n" + BuildIndentString(Indent + 1);
+
+	BeginClassEntry += IndentedNewline + "\"ClassContext\" : \"" + ClassName + "\",\n";
+	FileOutWriter->Serialize(TCHAR_TO_ANSI(*BeginClassEntry), BeginClassEntry.Len());
+
+	UBlueprint* TempBlueprint = MakeTempBlueprint(BlueprintClass);
+
+	DumpPalette(Indent + 1, TempBlueprint, FileOutWriter);
+
+	FString EndClassEntry = "\n" + ClassEntryIndent + "}";
+	FileOutWriter->Serialize(TCHAR_TO_ANSI(*EndClassEntry), EndClassEntry.Len());
+}
+
+static TMap<FString, FText> CategoryTooltipsMap;
+
+//------------------------------------------------------------------------------
+static void GenerateBlueprintAPIUtils::DumpCategoryInfo(uint32 Indent, FArchive* FileOutWriter)
+{
+	const FString CategoryEntryIndent = BuildIndentString(Indent);
+	FString CategoriesEntry = FString::Printf(TEXT(",\n%s\"Categories\" : {"), *CategoryEntryIndent);
+
+	const FString IndentedNewline = "\n" + BuildIndentString(Indent + 1);
+	const FString CategoryDetailsIndent = BuildIndentString(Indent + 2);
+	bool bNeedComma = false;
+
+	for (auto Category : CategoryTooltipsMap)
+	{
+		if (!Category.Value.IsEmpty())
+		{
+			CategoriesEntry += FString::Printf(TEXT("%s%s\"%s\" : {"), (bNeedComma ? TEXT(",") : TEXT("")), *IndentedNewline, *Category.Key);
+			CategoriesEntry += FString::Printf(TEXT("\n%s\"Tooltip\"\t: \"%s\""), *CategoryDetailsIndent, *MakeJsonString(Category.Value.ToString()));
+			CategoriesEntry += FString::Printf(TEXT("\n%s}"), *IndentedNewline);
+
+			bNeedComma = true;
+		}
+	}
+
+	CategoriesEntry += FString::Printf(TEXT("\n%s}"), *CategoryEntryIndent);
+
+	FileOutWriter->Serialize(TCHAR_TO_ANSI(*CategoriesEntry), CategoriesEntry.Len());
+}
+
+//------------------------------------------------------------------------------
+static double GenerateBlueprintAPIUtils::GetPaletteMenuActions(FBlueprintPaletteListBuilder& PaletteBuilder, UClass* PaletteFilter)
+{
+	PaletteBuilder.Empty();
+	UEdGraphSchema_K2 const* K2Schema = GetDefault<UEdGraphSchema_K2>();
+
+	double MenuBuildDuration = 0.0;
+	
+	if (CommandOptions.DumpFlags & BPDUMP_UseLegacyMenuBuilder)
+	{
+		FScopedDurationTimer DurationTimer(MenuBuildDuration);
+		FK2ActionMenuBuilder(K2Schema).GetPaletteActions(PaletteBuilder, PaletteFilter);
+	}
+	else
+	{
+		FBlueprintActionContext FilterContext;
+		FilterContext.Blueprints.Add(const_cast<UBlueprint*>(PaletteBuilder.Blueprint));
+		
+		FBlueprintActionMenuBuilder MenuBuilder(nullptr);
+		{
+			// prime the database so it's not recorded in our timing capture
+			FBlueprintActionDatabase::Get();
+			
+			FScopedDurationTimer DurationTimer(MenuBuildDuration);
+			FBlueprintActionMenuUtils::MakePaletteMenu(FilterContext, PaletteFilter, MenuBuilder);
+		}
+		PaletteBuilder.Append(MenuBuilder);
+	}
+
+	return MenuBuildDuration;
+}
+
+//------------------------------------------------------------------------------
+static void GenerateBlueprintAPIUtils::DumpPalette(uint32 Indent, UBlueprint* Blueprint, FArchive* FileOutWriter)
+{
+	UClass* PaletteFilter = CommandOptions.PaletteFilter;
+	DumpPalette(Indent, Blueprint, PaletteFilter, FileOutWriter);
+	bool bNeedsEndline = true;
+
+	if ((CommandOptions.DumpFlags & BPDUMP_FilteredPalette) && (PaletteFilter == nullptr))
+	{
+		// anim blueprints don't have a palette, so it is ok to assume this
+		UEdGraphSchema_K2 const* K2Schema = GetDefault<UEdGraphSchema_K2>();
+
+		for (TObjectIterator<UClass> ClassIt; ClassIt; ++ClassIt)
+		{
+			UClass* Class = *ClassIt;
+			if (!K2Schema->ClassHasBlueprintAccessibleMembers(Class))
+			{
+				continue;
+			}
+
+			FileOutWriter->Serialize(TCHAR_TO_ANSI(TEXT(",\n")), 2);
+			DumpPalette(Indent, Blueprint, Class, FileOutWriter);
+		}
+	}
+}
+
+//------------------------------------------------------------------------------
+static void GenerateBlueprintAPIUtils::DumpPalette(uint32 Indent, UBlueprint* Blueprint, UClass* ClassFilter, FArchive* FileOutWriter)
+{
+	FString const PaletteEntryIndent = BuildIndentString(Indent);
+	FString BeginPaletteEntry = FString::Printf(TEXT("%s\"Palette"), *PaletteEntryIndent);
+
+	FString FilterClassName("<UNFILTERED>");
+	if (ClassFilter != nullptr)
+	{
+		FilterClassName = ClassFilter->GetName();
+		BeginPaletteEntry += "-" + FilterClassName;
+	}
+	BeginPaletteEntry += "\" : {\n";
+
+	FString const NestedIndent = BuildIndentString(Indent + 1);
+	UE_LOG(LogBlueprintInfoDump, Display, TEXT("%sDumping palette: %s"), *BuildIndentString(Indent, true), *FilterClassName);
+
+	bool bIsAnimBlueprint = (Cast<UAnimBlueprint>(Blueprint) != nullptr);
+	// animation blueprints don't have a palette
+	if (bIsAnimBlueprint)
+	{
+		BeginPaletteEntry += NestedIndent + "\"IsAnimBlueprint\" : true";
+		FileOutWriter->Serialize(TCHAR_TO_ANSI(*BeginPaletteEntry), BeginPaletteEntry.Len());
+	}
+	else
+	{
+		FBlueprintPaletteListBuilder PaletteBuilder(Blueprint, TEXT("Library"));
+		double MenuBuildDuration = GetPaletteMenuActions(PaletteBuilder, ClassFilter);
+
+		BeginPaletteEntry += NestedIndent + "\"FilterClass\" : \"" + FilterClassName + "\",\n";
+		if (CommandOptions.DumpFlags & BPDUMP_RecordTiming)
+		{
+			BeginPaletteEntry += FString::Printf(TEXT("%s\"MenuBuildTime_Seconds\" : %f,\n"), *NestedIndent, MenuBuildDuration);
+		}
+
+		FileOutWriter->Serialize(TCHAR_TO_ANSI(*BeginPaletteEntry), BeginPaletteEntry.Len());
+		DumpActionList(Indent + 1, PaletteBuilder, FileOutWriter);
+	}
+
+	FString EndPaletteEntry = "\n" + PaletteEntryIndent + "}";
+	FileOutWriter->Serialize(TCHAR_TO_ANSI(*EndPaletteEntry), EndPaletteEntry.Len());
+}
+
+//------------------------------------------------------------------------------
+static void GenerateBlueprintAPIUtils::DumpActionList(uint32 Indent, FGraphActionListBuilderBase& ActionList, FArchive* FileOutWriter)
+{
+	TArray< FGraphActionListBuilderBase::ActionGroup const* > SortedActions;
+	for (int32 ActionIndex = 0; ActionIndex < ActionList.GetNumActions(); ++ActionIndex)
+	{
+		FGraphActionListBuilderBase::ActionGroup& Action = ActionList.GetAction(ActionIndex);
+		if (Action.Actions.Num() <= 0)
+		{
+			continue;
+		}
+
+		SortedActions.Add(&Action);
+	}
+
+	FString const ActionListIndent = BuildIndentString(Indent);
+	FString const NestedIndent     = BuildIndentString(Indent + 1);
+
+	FString BeginActionListEntry = FString::Printf(TEXT("%s\"ActionSet\" : {\n%s\"ActionCount\" : %d"), *ActionListIndent, *NestedIndent, SortedActions.Num());
+	BeginActionListEntry += ",\n" + NestedIndent + "\"Actions\" : {";
+
+	FileOutWriter->Serialize(TCHAR_TO_ANSI(*BeginActionListEntry), BeginActionListEntry.Len());
+
+	struct ActionSortFunctor
+	{
+		bool operator()(FGraphActionListBuilderBase::ActionGroup const& LHS, FGraphActionListBuilderBase::ActionGroup const& RHS) const
+		{
+			TSharedPtr<FEdGraphSchemaAction> LHSAction = LHS.Actions[0];
+			TSharedPtr<FEdGraphSchemaAction> RHSAction = RHS.Actions[0];
+			
+			if (LHSAction->Grouping != RHSAction->Grouping)
+			{
+				return LHSAction->Grouping > RHSAction->Grouping;
+			}
+			
+			FString LhKey = GetActionKey(LHS);
+			FString RhKey = GetActionKey(RHS);
+			return (LhKey.Compare(RhKey) < 0);
+		}
+	};
+	// need to sort so we can easily compare from one generation to the next
+	SortedActions.Sort(ActionSortFunctor());
+
+	FString LineEnding("\n");
+	for (FGraphActionListBuilderBase::ActionGroup const* Action : SortedActions)
+	{
+		FileOutWriter->Serialize(TCHAR_TO_ANSI(*LineEnding), LineEnding.Len());
+		DumpActionMenuItem(Indent + 2, *Action, ActionList, FileOutWriter);
+		LineEnding = TEXT(",\n");
+	}
+
+	FString EndActionListEntry = "\n" + NestedIndent + "}";
+
+	EndActionListEntry += "\n" + ActionListIndent + "}";
+	FileOutWriter->Serialize(TCHAR_TO_ANSI(*EndActionListEntry), EndActionListEntry.Len());
+}
+
+//------------------------------------------------------------------------------
+static void GenerateBlueprintAPIUtils::DumpActionMenuItem(uint32 Indent, FGraphActionListBuilderBase::ActionGroup const& Action, FGraphActionListBuilderBase& ActionList, FArchive* FileOutWriter)
+{
+	check(Action.Actions.Num() > 0);
+
+	// Get action category info
+	TArray<FString> MenuHierarchy;
+	Action.GetCategoryChain(MenuHierarchy);
+
+	FString ActionCategory = TEXT("");
+
+	bool bHasCategory = (MenuHierarchy.Num() > 0);
+	if (bHasCategory)
+	{
+		for (FString const& SubCategory : MenuHierarchy)
+		{
+			ActionCategory += SubCategory + TEXT("|");
+		}
+	}
+
+	TArray<FString> Categories;
+	ActionCategory.ParseIntoArray(&Categories, TEXT("|"), true);
+
+	for (const FString& Category : Categories)
+	{
+		FString CategoryDisplayName = FEditorCategoryUtils::GetCategoryDisplayString(Category);
+		if (!CategoryTooltipsMap.Contains(CategoryDisplayName))
+		{
+			FText Tooltip;
+			FString DocLink, DocExcerpt;
+
+			FEditorCategoryUtils::GetCategoryTooltipInfo(CategoryDisplayName, Tooltip, DocLink, DocExcerpt);
+			CategoryTooltipsMap.Add(CategoryDisplayName, Tooltip);
+		}
+	}
+
+	TSharedPtr<FEdGraphSchemaAction> PrimeAction = Action.Actions[0];
+	const FString ActionName = PrimeAction->MenuDescription.ToString();
+
+	const FString ActionEntryIndent = BuildIndentString(Indent);
+	FString ActionEntry = ActionEntryIndent + "\"" + MakeJsonString(ActionCategory + ActionName) + "\"";
+
+	++Indent;
+	FString const IndentedNewline = "\n" + BuildIndentString(Indent);
+
+	ActionEntry += " : {";
+
+	const FString TooltipFieldLabel("\"Tooltip\"     : \"");
+	const FString TooltipStr = PrimeAction->TooltipDescription.Replace(TEXT("\n"), *(IndentedNewline + BuildIndentString(TooltipFieldLabel.Len(), /*bUseSpaces =*/true)));
+
+	ActionEntry += IndentedNewline + TooltipFieldLabel + MakeJsonString(TooltipStr) + "\"";
+		
+	// Get action node type info
+	UK2Node const* NodeTemplate = FK2SchemaActionUtils::ExtractNodeTemplateFromAction(PrimeAction);
+	if (NodeTemplate != nullptr)
+	{
+		UK2Node* Node = DuplicateObject<UK2Node>(NodeTemplate,ActionList.OwnerOfTemporaries);
+		ActionList.OwnerOfTemporaries->AddNode(Node);
+		Node->AllocateDefaultPins();
+
+		if (Node->Pins.Num() > 0)
+		{
+			ActionEntry += "," + IndentedNewline + "\"Pins\"        : {";
+
+			const FString PinEntryIndentedNewline = "\n" + BuildIndentString(Indent+1);
+			const FString PinDetailsIndentedNewline = "\n" + BuildIndentString(Indent+2);
+
+			static UEnum* PinDirectionEnum = FindObject<UEnum>(ANY_PACKAGE, TEXT("EEdGraphPinDirection"));
+
+			bool bFirst = true;
+			for (UEdGraphPin* Pin : Node->Pins)
+			{
+				if (!Pin->bHidden)
+				{
+					if (!bFirst)
+					{
+						ActionEntry += ",";
+					}
+					bFirst = false;
+
+					ActionEntry += PinEntryIndentedNewline + "\"" + MakeJsonString(Pin->GetName()) + "\" : {";
+
+					const FString DisplayName = Pin->GetDisplayName().ToString();
+
+					if (!DisplayName.IsEmpty())
+					{
+						ActionEntry += PinDetailsIndentedNewline + "\"Name\"                 : \"" + MakeJsonString(DisplayName) + "\",";
+					}
+					ActionEntry += PinDetailsIndentedNewline + "\"Direction\"            : \"" + (Pin->Direction == EGPD_Input ? "input" : "output") + "\"";
+
+					if (!Pin->PinToolTip.IsEmpty())
+					{
+						const FString PinTooltipFieldLabel("\"Tooltip\"              : \"");
+						const FString PinTooltipStr = Pin->PinToolTip.Replace(TEXT("\n"), *(PinDetailsIndentedNewline + BuildIndentString(PinTooltipFieldLabel.Len(), /*bUseSpaces =*/true)));
+
+						ActionEntry += "," + PinDetailsIndentedNewline + PinTooltipFieldLabel + MakeJsonString(PinTooltipStr) + "\"";
+					}
+
+					ActionEntry += "," + PinDetailsIndentedNewline + "\"PinCategory\"          : \"" + MakeJsonString(Pin->PinType.PinCategory) + "\"";
+
+					if (!Pin->PinType.PinSubCategory.IsEmpty())
+					{
+						ActionEntry += "," + PinDetailsIndentedNewline + "\"PinSubCategory\"       : \"" + MakeJsonString(Pin->PinType.PinSubCategory) + "\"";
+					}
+
+					if (Pin->PinType.PinSubCategoryObject.IsValid())
+					{
+						ActionEntry += "," + PinDetailsIndentedNewline + "\"PinSubCategoryObject\" : \"" + MakeJsonString(Pin->PinType.PinSubCategoryObject->GetName()) + "\"";
+					}
+
+					if (Pin->PinType.bIsArray)
+					{
+						ActionEntry += "," + PinDetailsIndentedNewline + "\"IsArray\"              : \"true\"";
+					}
+
+					if (Pin->PinType.bIsConst)
+					{
+						ActionEntry += "," + PinDetailsIndentedNewline + "\"IsConst\"              : \"true\"";
+					}
+
+					if (Pin->PinType.bIsReference)
+					{
+						ActionEntry += "," + PinDetailsIndentedNewline + "\"IsReference\"          : \"true\"";
+					}
+
+					ActionEntry += PinEntryIndentedNewline + "}";
+				}
+			}
+
+			ActionEntry += IndentedNewline + "}";
+		}
+
+		ActionList.OwnerOfTemporaries->RemoveNode(Node);
+	}
+	// Finish action entry
+	ActionEntry += "\n" + ActionEntryIndent + "}";
+
+	// Write entry to file
+	FileOutWriter->Serialize(TCHAR_TO_ANSI(*ActionEntry), ActionEntry.Len());
+}
+
+/*******************************************************************************
+ * UGenerateBlueprintAPICommandlet
+ ******************************************************************************/
+
+//------------------------------------------------------------------------------
+UGenerateBlueprintAPICommandlet::UGenerateBlueprintAPICommandlet(class FPostConstructInitializeProperties const& PCIP)
+	: Super(PCIP)
+{
+}
+
+//------------------------------------------------------------------------------
+int32 UGenerateBlueprintAPICommandlet::Main(FString const& Params)
+{
+	TArray<FString> Tokens, Switches;
+	ParseCommandLine(*Params, Tokens, Switches);
+
+	GenerateBlueprintAPIUtils::CommandOptions = GenerateBlueprintAPIUtils::CommandletOptions(Switches);
+	GenerateBlueprintAPIUtils::CommandletOptions const& CommandOptions = GenerateBlueprintAPIUtils::CommandOptions;
+
+	FString ActiveFilePath;
+	FArchive* FileOut = nullptr;
+
+	// this lambda is responsible for adding closing characters to the file, and 
+	// closing out the writer (and diffing the resultant file if the user deigns us to do so)
+	auto CloseFileStream = [](FArchive** FileOutPtr)
+	{
+		FArchive*& FileOut = (*FileOutPtr);
+		if (FileOut != nullptr)
+		{
+			FileOut->Serialize(TCHAR_TO_ANSI(TEXT("\n}")), 2);
+			FileOut->Close();
+
+			delete FileOut;
+			FileOut = nullptr;
+		}
+	};
+
+	// this lambda is responsible for opening a file for write, and adding 
+	// opening json characters to the file (contextually tracks the active filepath as well)
+	auto OpenFileStream = [&ActiveFilePath, &CloseFileStream, &FileOut](UClass* Class)->FArchive*
+	{
+		CloseFileStream(&FileOut);
+
+		ActiveFilePath = GenerateBlueprintAPIUtils::BuildDumpFilePath(Class);
+		FileOut = IFileManager::Get().CreateFileWriter(*ActiveFilePath);
+		check(FileOut != nullptr);
+		FileOut->Serialize(TCHAR_TO_ANSI(TEXT("{\n")), 2);
+
+		return FileOut;
+	};
+	
+	bool bNeedsJsonComma = false;
+	// this lambda will dump blueprint info for the specified class, if the user
+	// set -multifile, then this will also close the existing file and open a new one for this class
+	auto WriteClassInfo = [&bNeedsJsonComma, &OpenFileStream, &CloseFileStream, &FileOut](UClass* Class)
+	{
+		if (FileOut == nullptr)
+		{
+			FileOut = OpenFileStream(Class);
+		}
+		// if we're adding all the class entries into one file, then we need to 
+		// separate them by a comma (or invalid json)
+		else if (bNeedsJsonComma)
+		{
+			FileOut->Serialize(TCHAR_TO_ANSI(TEXT(",\n")), 2);
+		}
+
+		GenerateBlueprintAPIUtils::DumpInfoForClass(1, Class, FileOut);
+		bNeedsJsonComma = true;
+	};
+
+	// this lambda is used as a precursory check to ensure that the specified 
+	// class is a blueprintable type... broken into it's own lambda to save on reuse
+	auto IsInvalidBlueprintClass = [](UClass* Class)->bool
+	{
+		return !IsValid(Class) ||
+			Class->HasAnyClassFlags(CLASS_NewerVersionExists) ||
+			FKismetEditorUtilities::IsClassABlueprintSkeleton(Class)  ||
+			!FKismetEditorUtilities::CanCreateBlueprintOfClass(Class) ||
+			(Class->GetOuterUPackage() == GetTransientPackage());
+	};
+
+	UClass* const BlueprintClass = CommandOptions.BlueprintClass;
+	if (CommandOptions.DumpFlags & GenerateBlueprintAPIUtils::BPDUMP_LogHelp)
+	{
+		UE_LOG(LogBlueprintInfoDump, Display, TEXT("%s"), *GenerateBlueprintAPIUtils::HelpString);
+	}
+	else if (BlueprintClass != nullptr)
+	{
+		UE_LOG(LogBlueprintInfoDump, Display, TEXT("Dumping Blueprint info..."));
+		// make sure the class that the user specified is a blueprintable type
+		if (IsInvalidBlueprintClass(BlueprintClass))
+		{
+			UE_LOG(LogBlueprintInfoDump, Error, TEXT("Cannot create a blueprint from class: '%s'"), *BlueprintClass->GetName());
+			if (FileOut != nullptr)
+			{
+				FString const InvalidClassEntry = GenerateBlueprintAPIUtils::BuildIndentString(1) + "\"INVALID_BLUEPRINT_CLASS\" : \"" + BlueprintClass->GetName() + "\"";
+				FileOut->Serialize(TCHAR_TO_ANSI(*InvalidClassEntry), InvalidClassEntry.Len());
+			}
+		}
+		else
+		{
+			WriteClassInfo(BlueprintClass);
+		}
+	}
+	// if the user didn't specify a class, then we take that to mean dump ALL the classes!
+	else
+	{
+		UE_LOG(LogBlueprintInfoDump, Display, TEXT("Dumping Blueprint info..."));
+		for (TObjectIterator<UClass> ClassIt; ClassIt; ++ClassIt)
+		{
+			UClass* Class = *ClassIt;
+			if (IsInvalidBlueprintClass(Class))
+			{
+				continue;
+			}
+			WriteClassInfo(Class);
+		}
+	}
+
+	GenerateBlueprintAPIUtils::DumpCategoryInfo(1, FileOut);
+
+	CloseFileStream(&FileOut);
+	return 0;
+}
+
+
+
+
diff --git a/Engine/Source/Programs/AutomationTool/Scripts/BuildDocumentation.Automation.cs b/Engine/Source/Programs/AutomationTool/Scripts/BuildDocumentation.Automation.cs
--- a/Engine/Source/Programs/AutomationTool/Scripts/BuildDocumentation.Automation.cs
+++ b/Engine/Source/Programs/AutomationTool/Scripts/BuildDocumentation.Automation.cs
@@ -19,6 +19,7 @@ public class BuildDocumentation : BuildCommand
 		bool bCleanMeta = bClean || ParseParam("cleanmeta");
 		bool bCleanEnv = bClean || ParseParam("cleanenv");
 		bool bCleanXml = bClean || ParseParam("cleanxml");
+		bool bCleanJson = bClean || ParseParam("cleanjson");
 		bool bCleanUdn = bClean || ParseParam("cleanudn");
 		bool bCleanHtml = bClean || ParseParam("cleanhtml");
 		bool bCleanChm = bClean || ParseParam("cleanchm");
@@ -27,6 +28,7 @@ public class BuildDocumentation : BuildCommand
 		bool bBuildMeta = bBuild || ParseParam("buildmeta");
 		bool bBuildEnv = bBuild || ParseParam("buildenv");
 		bool bBuildXml = bBuild || ParseParam("buildxml");
+		bool bBuildJson = bBuild || ParseParam("buildjson");
 		bool bBuildUdn = bBuild || ParseParam("buildudn");
 		bool bBuildHtml = bBuild || ParseParam("buildhtml");
 		bool bBuildChm = bBuild || ParseParam("buildchm");
@@ -49,6 +51,8 @@ public class BuildDocumentation : BuildCommand
 		CommandUtils.CreateDirectory(BuildDir);
 		string DoxygenDir = Path.Combine(IntermediateDir, "doxygen");
 		CommandUtils.CreateDirectory(DoxygenDir);
+		string JsonDir = Path.Combine(IntermediateDir, "json");
+		CommandUtils.CreateDirectory(JsonDir);
 		string MetadataDir = Path.Combine(IntermediateDir, "metadata");
 		CommandUtils.CreateDirectory(MetadataDir);
 		string ArchiveDir = Path.Combine(CmdEnv.LocalRoot, "Engine\\Documentation\\Builds");
@@ -76,6 +80,7 @@ public class BuildDocumentation : BuildCommand
 		ApiToolCommandLine += " -enginedir=\"" + Path.Combine(CmdEnv.LocalRoot, "Engine") + "\"";
 		ApiToolCommandLine += " -targetinfo=\"" + TargetInfoPath + "\"";
 		ApiToolCommandLine += " -xmldir=\"" + DoxygenDir + "\"";
+		ApiToolCommandLine += " -jsondir=\"" + JsonDir + "\"";
 		ApiToolCommandLine += " -metadatadir=\"" + MetadataDir + "\"";
 		if (Filter != null) ApiToolCommandLine += " -filter=" + Filter;
 
@@ -92,7 +97,11 @@ public class BuildDocumentation : BuildCommand
 		{
 			RunAndLog(CmdEnv, ApiDocToolPath, "-cleanxml" + ApiToolCommandLine, "ApiDocTool-CleanXML");
 		}
-		if(bCleanUdn)
+		if (bCleanJson)
+		{
+			RunAndLog(CmdEnv, ApiDocToolPath, "-cleanjson" + ApiToolCommandLine, "ApiDocTool-CleanJSON");
+		}
+		if (bCleanUdn)
 		{
 			RunAndLog(CmdEnv, ApiDocToolPath, "-cleanudn" + ApiToolCommandLine, "ApiDocTool-CleanUDN");
 		}
@@ -119,7 +128,20 @@ public class BuildDocumentation : BuildCommand
 			RunAndLog(CmdEnv, ApiDocToolPath, "-buildxml" + ApiToolCommandLine, "ApiDocTool-BuildXML");
 			if (bMakeArchives)
 			{
-				CreateAndSubmitArchiveFromDir("Engine/Documentation/Builds/API-XML.tgz", "API documentation intermediates", IntermediateDir);
+				List<string> InputFiles = new List<string>();
+				FindFilesForTar(IntermediateDir, "build", InputFiles, true);
+				FindFilesForTar(IntermediateDir, "doxygen", InputFiles, true);
+				FindFilesForTar(IntermediateDir, "metadata", InputFiles, true);
+
+				CreateAndSubmitArchiveFromFiles("Engine/Documentation/Builds/API-XML.tgz", "API documentation intermediates", IntermediateDir, InputFiles);
+			}
+		}
+		if (bBuildJson)
+		{
+			RunAndLog(CmdEnv, ApiDocToolPath, "-buildjson" + ApiToolCommandLine, "ApiDocTool-BuildJSON");
+			if (bMakeArchives)
+			{
+				CreateAndSubmitArchiveFromDir("Engine/Documentation/Builds/API-JSON.tgz", "Blueprint API documentation intermediates", Path.Combine(IntermediateDir, "json"));
 			}
 		}
 		if (bBuildUdn)
@@ -127,8 +149,14 @@ public class BuildDocumentation : BuildCommand
 			RunAndLog(CmdEnv, ApiDocToolPath, "-buildudn" + (bStats? (" -stats=\"" + StatsPath + "\"") : "") + ApiToolCommandLine, "ApiDocTool-BuildUDN");
 			if (bMakeArchives)
 			{
-				CreateAndSubmitArchiveFromDir("Engine/Documentation/Builds/API-UDN.tgz", "API documentation UDN output", Path.Combine(CmdEnv.LocalRoot, "Engine\\Documentation\\Source\\API"));
-				CreateAndSubmitArchiveFromFiles("Engine/Documentation/Builds/API-Sitemap.tgz", "API documentation sitemap output", Path.Combine(CmdEnv.LocalRoot, "Engine\\Documentation\\CHM"), new string[] { "API.hhc", "API.hhk" });
+				string BaseDir = Path.Combine(CmdEnv.LocalRoot, "Engine\\Documentation\\Source");
+
+				List<string> InputFiles = new List<string>();
+				FindFilesForTar(BaseDir, "API", InputFiles, true);
+				FindFilesForTar(BaseDir, "BlueprintAPI", InputFiles, true);
+
+				CreateAndSubmitArchiveFromFiles("Engine/Documentation/Builds/API-UDN.tgz", "API documentation UDN output", BaseDir, InputFiles);
+				CreateAndSubmitArchiveFromFiles("Engine/Documentation/Builds/API-Sitemap.tgz", "API documentation sitemap output", Path.Combine(CmdEnv.LocalRoot, "Engine\\Documentation\\CHM"), new string[] { "API.hhc", "API.hhk", "BlueprintAPI.hhc", "BlueprintAPI.hhk" });
 			}
 		}
 		if (bBuildHtml)
@@ -153,6 +181,7 @@ public class BuildDocumentation : BuildCommand
 		{
 			RunAndLog(CmdEnv, ApiDocToolPath, "-buildchm" + ApiToolCommandLine, "ApiDocTool-BuildCHM");
 			SubmitFile("Engine/Documentation/CHM/API.chm", "API documentation CHM output");
+			SubmitFile("Engine/Documentation/CHM/BlueprintAPI.chm", "Blueprint API documentation CHM output");
 		}
 	}
 
diff --git a/Engine/Source/Programs/UnrealDocTool/APIDocTool/APIDocTool/APIAction.cs b/Engine/Source/Programs/UnrealDocTool/APIDocTool/APIDocTool/APIAction.cs
new file mode 100644
--- /dev/null
+++ b/Engine/Source/Programs/UnrealDocTool/APIDocTool/APIDocTool/APIAction.cs
@@ -0,0 +1,89 @@
+﻿using System;
+using System.Collections.Generic;
+using System.Linq;
+using System.Text;
+using System.Threading.Tasks;
+
+namespace APIDocTool
+{
+	class APIAction : APIPage
+	{
+		public readonly string Tooltip;
+		private List<APIActionPin> Pins;
+
+		public APIAction(APIPage InParent, string InName, Dictionary<string, object> ActionProperties)
+			: base(InParent, InName)
+		{
+			object Value;
+			
+			if (ActionProperties.TryGetValue("Tooltip", out Value))
+			{
+				Tooltip = String.Concat(((string)Value).Split('\n').Select(Line => Line.Trim() + '\n'));
+//				Tooltip = new string(Array.FindAll<char>(((string)Value).ToCharArray(), c => c != '\t'));
+			}
+			else
+			{
+				Tooltip = "";
+			}
+
+			if (ActionProperties.TryGetValue("Pins", out Value))
+			{
+				Pins = new List<APIActionPin>();
+
+				foreach (var Pin in (Dictionary<string, object>)Value)
+				{
+					Pins.Add(new APIActionPin(this, APIActionPin.GetPinName((Dictionary<string, object>)Pin.Value), (Dictionary<string, object>)Pin.Value));
+				}
+			}
+		}
+
+		public override void WritePage(UdnManifest Manifest, string OutputPath)
+		{
+			using (UdnWriter Writer = new UdnWriter(OutputPath))
+			{
+				Writer.WritePageHeader(Name, PageCrumbs, Name);
+
+				// Tooltip/description
+				Writer.WriteLine(Tooltip);
+
+				if (Pins != null && Pins.Count() > 0)
+				{
+					// Inputs
+					Writer.EnterSection("inputs", "Inputs");
+					Writer.WriteList(Pins.Where(x => x.bInputPin).Select(x => x.GetListItem()));
+					Writer.LeaveSection();
+
+					// Outputs
+					Writer.EnterSection("outputs", "Outputs");
+					Writer.WriteList(Pins.Where(x => !x.bInputPin).Select(x => x.GetListItem()));
+					Writer.LeaveSection();
+				}
+			}
+		}
+
+		public override void GatherReferencedPages(List<APIPage> Pages)
+		{
+			/*if (Pins != null)
+			{
+				Pages.AddRange(Pins);
+			}*/
+		}
+
+		public override void AddToManifest(UdnManifest Manifest)
+		{
+			Manifest.Add(Name, this);
+			/*if (Pins != null)
+			{
+				foreach (APIPage Pin in Pins)
+				{
+					Pin.AddToManifest(Manifest);
+				}
+			}*/
+		}
+
+		public UdnListItem GetListItem()
+		{
+			return new UdnListItem(Name, Tooltip, LinkPath);
+		}
+	}
+}
diff --git a/Engine/Source/Programs/UnrealDocTool/APIDocTool/APIDocTool/APIActionPin.cs b/Engine/Source/Programs/UnrealDocTool/APIDocTool/APIDocTool/APIActionPin.cs
new file mode 100644
--- /dev/null
+++ b/Engine/Source/Programs/UnrealDocTool/APIDocTool/APIDocTool/APIActionPin.cs
@@ -0,0 +1,151 @@
+﻿using System;
+using System.Collections.Generic;
+using System.Diagnostics;
+using System.Linq;
+using System.Text;
+using System.Threading.Tasks;
+
+namespace APIDocTool
+{
+	class APIActionPin : APIPage
+	{
+		public readonly string Tooltip;
+
+		public readonly string PinCategory;
+		public readonly string PinSubCategory;
+		public readonly string PinSubCategoryObject;
+
+		public readonly bool bInputPin;
+		public readonly bool bIsConst;
+		public readonly bool bIsReference;
+		public readonly bool bIsArray;
+
+		public static string GetPinName(Dictionary<string, object> PinProperties)
+		{
+			object Value;
+			string Name;
+
+			if (PinProperties.TryGetValue("Name", out Value))
+			{
+				Name = (string)Value;
+			}
+			else
+			{
+				Debug.Assert((string)PinProperties["PinCategory"] == "exec");
+				
+				Name = ((string)PinProperties["Direction"] == "input" ? "In" : "Out");
+			}
+
+			return Name;
+		}
+
+		public APIActionPin(APIPage InParent, string InName, Dictionary<string, object> PinProperties)
+			: base(InParent, InName)
+		{
+			object Value;
+
+			Debug.Assert(InName == GetPinName(PinProperties));
+
+			bInputPin = (string)PinProperties["Direction"] == "input";
+			PinCategory = (string)PinProperties["PinCategory"];
+
+			if (PinProperties.TryGetValue("Tooltip", out Value))
+			{
+				Tooltip = String.Concat(((string)Value).Split('\n').Select(Line => Line.Trim() + '\n'));
+//				Tooltip = new string(Array.FindAll<char>(((string)Value).ToCharArray(), c => c != '\t'));
+			}
+			else
+			{
+				Tooltip = "";
+			}
+
+			if (PinProperties.TryGetValue("PinSubCategory", out Value))
+			{
+				PinSubCategory = (string)Value;
+			}
+			else
+			{
+				PinSubCategory = "";
+			}
+
+			if (PinProperties.TryGetValue("PinSubCategoryObject", out Value))
+			{
+				PinSubCategoryObject = (string)Value;
+			}
+			else
+			{
+				PinSubCategoryObject = "";
+			}
+
+			if (PinProperties.TryGetValue("IsConst", out Value))
+			{
+				bIsConst = Convert.ToBoolean((string)Value);
+			}
+
+			if (PinProperties.TryGetValue("IsReference", out Value))
+			{
+				bIsReference = Convert.ToBoolean((string)Value);
+			}
+
+			if (PinProperties.TryGetValue("IsArray", out Value))
+			{
+				bIsArray = Convert.ToBoolean((string)Value);
+			}
+		}
+
+		private string GetPinCategory()
+		{
+			if (PinCategory == "struct")
+			{
+				if (PinSubCategoryObject == "Vector")
+				{
+					return "vector";
+				}
+				else if (PinSubCategoryObject == "Transform")
+				{
+					return "transform";
+				}
+				else if (PinSubCategoryObject == "Rotator")
+				{
+					return "rotator";
+				}
+			}
+
+			return PinCategory;
+		}
+
+		public override void WritePage(UdnManifest Manifest, string OutputPath)
+		{
+			using (UdnWriter Writer = new UdnWriter(OutputPath))
+			{
+				Writer.WritePageHeader(Name, PageCrumbs, Name);
+			}
+		}
+
+		public override void GatherReferencedPages(List<APIPage> Pages)
+		{
+
+		}
+
+		public override void AddToManifest(UdnManifest Manifest)
+		{
+			Manifest.Add(Name, this);
+		}
+
+		public UdnIconListItem GetListItem()
+		{
+			List<Icon> ItemIcons = new List<Icon>();
+
+			if (bIsArray)
+			{
+				ItemIcons.Add(Icons.ArrayVariablePinIcons[GetPinCategory()]);
+			}
+			else
+			{
+				ItemIcons.Add(Icons.VariablePinIcons[GetPinCategory()]);
+			}
+
+			return new UdnIconListItem(ItemIcons, Name, Tooltip, null);
+		}
+	}
+}
diff --git a/Engine/Source/Programs/UnrealDocTool/APIDocTool/APIDocTool/APICategory.cs b/Engine/Source/Programs/UnrealDocTool/APIDocTool/APIDocTool/APICategory.cs
new file mode 100644
--- /dev/null
+++ b/Engine/Source/Programs/UnrealDocTool/APIDocTool/APIDocTool/APICategory.cs
@@ -0,0 +1,183 @@
+﻿using System;
+using System.Collections.Generic;
+using System.Linq;
+using System.Text;
+using System.Threading.Tasks;
+
+namespace APIDocTool
+{
+	class APICategory : APIPage
+	{
+		static private Dictionary<string, string> CategoryTooltips = new Dictionary<string, string>();
+
+		static public void LoadTooltips(Dictionary<string,object> Categories)
+		{
+			foreach (var Category in Categories)
+			{
+				CategoryTooltips.Add(Category.Key,(string)((Dictionary<string,object>)Category.Value)["Tooltip"]);
+			}
+		}
+
+		public readonly Dictionary<string, APICategory> SubCategories;
+		public readonly List<APIAction> Actions;
+		public readonly bool bIsRootCategory;
+
+		public APICategory(APIPage InParent, string InName, bool bInIsRootCategory = false)
+			: base(InParent, InName)
+		{
+			SubCategories = new Dictionary<string, APICategory>();
+			Actions = new List<APIAction>();
+
+			bIsRootCategory = bInIsRootCategory;
+		}
+
+		public APICategory GetSubCategory(string CategoryName)
+		{
+			APICategory ReturnValue;
+			if (SubCategories.TryGetValue(CategoryName, out ReturnValue) == false)
+			{
+				ReturnValue = new APICategory(this, CategoryName);
+				SubCategories.Add(CategoryName, ReturnValue);
+			}
+			return ReturnValue;
+		}
+
+		public void AddAction(APIAction Action)
+		{
+			Actions.Add(Action);
+		}
+
+		public override SitemapNode CreateSitemapNode()
+		{
+			SitemapNode Node = new SitemapNode(Name, SitemapLinkPath);
+			Node.Children.AddRange(SubCategories.OrderBy(x => x.Value.Name).Select(x => x.Value.CreateSitemapNode()));
+			return Node;
+		}
+
+		public void WriteSitemapContents(string OutputPath)
+		{
+			List<SitemapNode> RootNodes = new List<SitemapNode>();
+			if (bIsRootCategory)
+			{
+				RootNodes.Add(new SitemapNode("Contents", SitemapLinkPath));
+			}
+			foreach (APIPage SubCategory in SubCategories.Values)
+			{
+				RootNodes.Add(SubCategory.CreateSitemapNode());
+			}
+			foreach (APIPage Action in Actions)
+			{
+				RootNodes.Add(Action.CreateSitemapNode());
+			}
+			Sitemap.WriteContentsFile(OutputPath, RootNodes);
+		}
+
+		public void WriteSitemapIndex(string OutputPath)
+		{
+			// Gather all the pages
+			List<APIPage> AllPages = new List<APIPage>(GatherPages());
+
+			// Find the shortest link path to each member
+			Dictionary<string, List<APIMember>> MemberIndexLinks = new Dictionary<string, List<APIMember>>();
+			foreach (APIMember Member in AllPages.OfType<APIMember>())
+			{
+				// Find or create the list of members for this link path
+				List<APIMember> Members;
+				if (!MemberIndexLinks.TryGetValue(Member.FullName, out Members))
+				{
+					Members = new List<APIMember>();
+					MemberIndexLinks.Add(Member.FullName, Members);
+				}
+
+				// Remove any members 
+				Members.RemoveAll(x => x.HasParent(Member));
+				if (!Members.Any(x => Member.HasParent(x))) Members.Add(Member);
+			}
+
+			// Create all the root pages
+			List<SitemapNode> RootNodes = new List<SitemapNode>();
+			foreach (APIPage Page in AllPages)
+			{
+				if (!(Page is APIMember))
+				{
+					RootNodes.Add(new SitemapNode(Page.Name, Page.SitemapLinkPath));
+				}
+			}
+			foreach (KeyValuePair<string, List<APIMember>> MemberIndexLink in MemberIndexLinks)
+			{
+				foreach (APIMember Member in MemberIndexLink.Value)
+				{
+					RootNodes.Add(new SitemapNode(MemberIndexLink.Key, Member.SitemapLinkPath));
+				}
+			}
+
+			// Write the index file
+			Sitemap.WriteIndexFile(OutputPath, RootNodes);
+		}
+
+		public override void WritePage(UdnManifest Manifest, string OutputPath)
+		{
+			using (UdnWriter Writer = new UdnWriter(OutputPath))
+			{
+
+				if (bIsRootCategory)
+				{
+					Writer.WritePageHeader("Unreal Engine Blueprint API Reference", PageCrumbs, "Unreal Engine Blueprint API Reference");
+
+					Writer.WriteHeading(2, "Disclaimer");
+
+					Writer.WriteLine("The Blueprint API reference is an early work in progress, and some information may be missing or out of date. It strives to reflect all available nodes, but it is not guaranteed to be an exhaustive list.");
+					Writer.WriteLine("For tutorials, walkthroughs and detailed guides to using Blueprints in UE4, please see the [Introduction to Blueprints](http://docs.unrealengine.com/latest/INT/Engine/Blueprints/GettingStarted/index.html) on the web.");
+				}
+				else
+				{
+					Writer.WritePageHeader(Name, PageCrumbs, Name);
+				}
+
+				if (Actions.Count > 0)
+				{
+					Writer.WriteHeading(2, "Actions");
+					Writer.WriteList(Actions.OrderBy(x => x.Name).Select(x => x.GetListItem()));
+				}
+
+				if (SubCategories.Count > 0)
+				{
+					Writer.WriteHeading(2, "Categories");
+					Writer.WriteList(SubCategories.OrderBy(x => x.Key).Select(x => x.Value.GetListItem()));
+				}
+			}
+
+		}
+
+		public override void GatherReferencedPages(List<APIPage> Pages)
+		{
+			Pages.AddRange(SubCategories.Values);
+			Pages.AddRange(Actions);
+		}
+
+		public override void AddToManifest(UdnManifest Manifest)
+		{
+			Manifest.Add(Name, this);
+			foreach (APIPage SubCategory in SubCategories.Values)
+			{
+				SubCategory.AddToManifest(Manifest);
+			}
+			foreach (APIPage Action in Actions)
+			{
+				Action.AddToManifest(Manifest);
+			}
+		}
+
+		public UdnListItem GetListItem()
+		{
+			string Tooltip;
+			if (!CategoryTooltips.TryGetValue(Name, out Tooltip))
+			{
+				Tooltip = "";
+			}
+
+			return new UdnListItem(Name, Tooltip, LinkPath);
+		}
+
+	}
+}
diff --git a/Engine/Source/Programs/UnrealDocTool/APIDocTool/APIDocTool/APIDocTool.csproj b/Engine/Source/Programs/UnrealDocTool/APIDocTool/APIDocTool/APIDocTool.csproj
--- a/Engine/Source/Programs/UnrealDocTool/APIDocTool/APIDocTool/APIDocTool.csproj
+++ b/Engine/Source/Programs/UnrealDocTool/APIDocTool/APIDocTool/APIDocTool.csproj
@@ -65,6 +65,9 @@
     <Compile Include="..\..\..\DotNETCommon\MetaData.cs">
       <Link>Properties\MetaData.cs</Link>
     </Compile>
+    <Compile Include="APIAction.cs" />
+    <Compile Include="APIActionPin.cs" />
+    <Compile Include="APICategory.cs" />
     <Compile Include="APIConstant.cs" />
     <Compile Include="APIConstantIndex.cs" />
     <Compile Include="APIEnumIndex.cs" />
@@ -79,6 +82,13 @@
     <Compile Include="APIRecord.cs" />
     <Compile Include="APIRecordIndex.cs" />
     <Compile Include="BuildTarget.cs" />
+    <Compile Include="fastJSON\Formatter.cs" />
+    <Compile Include="fastJSON\Getters.cs" />
+    <Compile Include="fastJSON\JSON.cs" />
+    <Compile Include="fastJSON\JsonParser.cs" />
+    <Compile Include="fastJSON\JsonSerializer.cs" />
+    <Compile Include="fastJSON\Reflection.cs" />
+    <Compile Include="fastJSON\SafeDictionary.cs" />
     <Compile Include="Icons.cs" />
     <Compile Include="IniFile.cs" />
     <Compile Include="Metadata.cs" />
diff --git a/Engine/Source/Programs/UnrealDocTool/APIDocTool/APIDocTool/Icons.cs b/Engine/Source/Programs/UnrealDocTool/APIDocTool/APIDocTool/Icons.cs
--- a/Engine/Source/Programs/UnrealDocTool/APIDocTool/APIDocTool/Icons.cs
+++ b/Engine/Source/Programs/UnrealDocTool/APIDocTool/APIDocTool/Icons.cs
@@ -146,6 +146,47 @@ namespace APIDocTool
 			new MetadataIcon(new Icon("%ROOT%/api_variable_meta_other.png", "Other"), new string[]{ }),
 		};
 
+		//// Blueprint Variable type icons 
+		public static readonly Dictionary<string, Icon> VariablePinIcons = new Dictionary<string, Icon>
+		{
+			{ "bool", new Icon("%ROOT%/bp_api_pin_bool.png", "Boolean") },
+			{ "byte", new Icon("%ROOT%/bp_api_pin_byte.png", "Byte") },
+			{ "class", new Icon("%ROOT%/bp_api_pin_class.png", "Class") },
+			{ "delegate", new Icon("%ROOT%/bp_api_pin_delegate.png", "Delegate") },
+			{ "exec", new Icon("%ROOT%/bp_api_pin_exec.png", "Exec") },
+			{ "float", new Icon("%ROOT%/bp_api_pin_float.png", "Float") },
+			{ "int", new Icon("%ROOT%/bp_api_pin_int.png", "Integer") },
+			{ "interface", new Icon("%ROOT%/bp_api_pin_interface.png", "Interface") },
+			{ "name", new Icon("%ROOT%/bp_api_pin_name.png", "Name") },
+			{ "object", new Icon("%ROOT%/bp_api_pin_object.png", "Object") },
+			{ "rotator", new Icon("%ROOT%/bp_api_pin_rotator.png", "Rotator") },
+			{ "string", new Icon("%ROOT%/bp_api_pin_string.png", "String") },
+			{ "struct", new Icon("%ROOT%/bp_api_pin_struct.png", "Struct") },
+			{ "text", new Icon("%ROOT%/bp_api_pin_text.png", "Text") },
+			{ "transform", new Icon("%ROOT%/bp_api_pin_transform.png", "Transform") },
+			{ "vector", new Icon("%ROOT%/bp_api_pin_vector.png", "Vector") },
+			{ "wildcard", new Icon("%ROOT%/bp_api_pin_wildcard.png", "Wildcard") },
+		};
+
+		//// Blueprint Array variable type icons 
+		public static readonly Dictionary<string, Icon> ArrayVariablePinIcons = new Dictionary<string, Icon>
+		{
+			{ "bool", new Icon("%ROOT%/bp_api_pin_array_bool.png", "Boolean Array") },
+			{ "byte", new Icon("%ROOT%/bp_api_pin_array_byte.png", "Byte Array") },
+			{ "class", new Icon("%ROOT%/bp_api_pin_array_class.png", "Class Array") },
+			{ "float", new Icon("%ROOT%/bp_api_pin_array_float.png", "Float Array") },
+			{ "int", new Icon("%ROOT%/bp_api_pin_array_int.png", "Integer Array") },
+			{ "name", new Icon("%ROOT%/bp_api_pin_array_name.png", "Name Array") },
+			{ "object", new Icon("%ROOT%/bp_api_pin_array_object.png", "Object Array") },
+			{ "rotator", new Icon("%ROOT%/bp_api_pin_array_rotator.png", "Rotator Array") },
+			{ "string", new Icon("%ROOT%/bp_api_pin_array_string.png", "String Array") },
+			{ "struct", new Icon("%ROOT%/bp_api_pin_array_struct.png", "Struct Array") },
+			{ "text", new Icon("%ROOT%/bp_api_pin_array_text.png", "Text Array") },
+			{ "transform", new Icon("%ROOT%/bp_api_pin_array_transform.png", "Transform Array") },
+			{ "vector", new Icon("%ROOT%/bp_api_pin_array_vector.png", "Vector Array") },
+			{ "wildcard", new Icon("%ROOT%/bp_api_pin_array_wildcard.png", "Wildcard Array") },
+		};
+
 		//// Meta icons ///////////////////////////////////////////////////////////////////////////////////////////
 
 		public static readonly MetadataIcon[] MetaMetadataIcons =
diff --git a/Engine/Source/Programs/UnrealDocTool/APIDocTool/APIDocTool/Program.cs b/Engine/Source/Programs/UnrealDocTool/APIDocTool/APIDocTool/Program.cs
--- a/Engine/Source/Programs/UnrealDocTool/APIDocTool/APIDocTool/Program.cs
+++ b/Engine/Source/Programs/UnrealDocTool/APIDocTool/APIDocTool/Program.cs
@@ -34,6 +34,7 @@ namespace APIDocTool
 
 		public const string TabSpaces = "&nbsp;&nbsp;&nbsp;&nbsp;";
         public const string APIFolder = "API";
+		public const string BlueprintAPIFolder = "BlueprintAPI";
 
 		public static bool bIndexOnly = false;
 		public static bool bOutputPublish = false;
@@ -159,9 +160,6 @@ namespace APIDocTool
 			"DECLARE_LOG_CATEGORY_EXTERN(CategoryName, DefaultVerbosity, CompileTimeVerbosity)= ",
 		};
 
-		const string SitemapContentsFileName = "API.hhc";
-		const string SitemapIndexFileName = "API.hhk";
-
 		static Program()
 		{
 			List<string> DelegateMacros = new List<string>();
@@ -200,6 +198,7 @@ namespace APIDocTool
 			bool bCleanTargetInfo = false;
 			bool bCleanMetadata = false;
 			bool bCleanXml = false;
+			bool bCleanJson = false;
 			bool bCleanUdn = false;
 			bool bCleanHtml = false;
 			bool bCleanChm = false;
@@ -207,14 +206,19 @@ namespace APIDocTool
 			bool bBuildTargetInfo = false;
 			bool bBuildMetadata = false;
 			bool bBuildXml = false;
+			bool bBuildJson = false;
 			bool bBuildUdn = false;
 			bool bBuildHtml = false;
 			bool bBuildChm = false;
 
+			bool bBuildCodeAPI = false;
+			bool bBuildBlueprintAPI = false;
+
 			string TargetInfoPath = Path.Combine(IntermediateRootDir, "build\\targetinfo.xml");
 			string EngineDir = EngineRootDir;
 			string MetadataDir = Path.Combine(IntermediateRootDir, "metadata");
 			string XmlDir = Path.Combine(IntermediateRootDir, "doxygen");
+			string JsonDir = Path.Combine(IntermediateRootDir, "json");
 			string StatsPath = null;
 
 			List<string> Filters = new List<string>();
@@ -226,7 +230,15 @@ namespace APIDocTool
 				string OptionName = (EqualsIdx == -1)? Argument : Argument.Substring(0, EqualsIdx + 1);
 				string OptionValue = (EqualsIdx == -1)? null : Argument.Substring(EqualsIdx + 1);
 
-				if (OptionName == "-clean")
+				if (OptionName == "-blueprint")
+				{
+					bBuildBlueprintAPI = true;
+				}
+				else if (OptionName == "-code")
+				{
+					bBuildCodeAPI = true;
+				}
+				else if (OptionName == "-clean")
 				{
 					bCleanTargetInfo = bCleanMetadata = bCleanXml = bCleanUdn = bCleanHtml = true;
 				}
@@ -271,6 +283,18 @@ namespace APIDocTool
 				{
 					bCleanXml = bBuildXml = true;
 				}
+				else if (OptionName == "-cleanjson")
+				{
+					bCleanJson = true;
+				}
+				else if (OptionName == "-buildjson")
+				{
+					bBuildJson = true;
+				}
+				else if (OptionName == "-rebuildjson")
+				{
+					bCleanJson = bBuildJson = true;
+				}
 				else if (OptionName == "-cleanudn")
 				{
 					bCleanUdn = true;
@@ -319,6 +343,10 @@ namespace APIDocTool
 				{
 					XmlDir = Path.GetFullPath(OptionValue);
 				}
+				else if (OptionName == "-jsondir=")
+				{
+					JsonDir = Path.GetFullPath(OptionValue);
+				}
 				else if (OptionName == "-metadatadir=")
 				{
 					MetadataDir = Path.GetFullPath(OptionValue);
@@ -342,6 +370,13 @@ namespace APIDocTool
 				}
 			}
 
+			// If neither -blueprint or -code was specified, build both
+			if (!bBuildCodeAPI && !bBuildBlueprintAPI)
+			{
+				bBuildCodeAPI = true;
+				bBuildBlueprintAPI = true;
+			}
+
 			// Check we have all the required parameters
 			if (bBuildXml && TargetInfoPath == null)
 			{
@@ -349,69 +384,122 @@ namespace APIDocTool
 			}
 			else if (bValidArgs && EngineDir != null)
 			{
-				// If we don't intermediate paths, make them up
-				if (XmlDir == null) XmlDir = Path.Combine(EngineDir, "Intermediate\\Documentation\\Default\\Xml");
-				if (MetadataDir == null) MetadataDir = Path.Combine(EngineDir, "Intermediate\\Documentation\\Default\\Metadata");
-
 				// Derive all the engine paths we need in one place
 				string DoxygenPath = Path.Combine(EngineDir, "Extras\\NotForLicensees\\Doxygen\\bin\\doxygen.exe");
 				string UdnDir = Path.Combine(EngineDir, "Documentation\\Source");
 				string HtmlDir = Path.Combine(EngineDir, "Documentation\\HTML");
 				string ChmDir = Path.Combine(EngineDir, "Documentation\\CHM");
+				string EditorPath = Path.Combine(EngineDir, "Binaries\\Win64\\UE4Editor.exe");
 				string DocToolPath = Path.Combine(EngineDir, "Binaries\\DotNET\\UnrealDocTool.exe");
 				string ChmCompilerPath = Path.Combine(EngineDir, "Extras\\NotForLicensees\\HTML Help Workshop\\hhc.exe");
 				string SettingsPath = Path.Combine(EngineDir, "Documentation\\Extras\\API\\API.ini");
 				string MetadataPath = Path.Combine(MetadataDir, "metadata.xml");
 
-				// Read the settings file
-				Settings = IniFile.Read(SettingsPath);
-				IgnoredFunctionMacros = new HashSet<string>(Settings.FindValueOrDefault("Input.IgnoredFunctionMacros", "").Split('\n'));
-				IncludedSourceFiles = new HashSet<string>(Settings.FindValueOrDefault("Output.IncludedSourceFiles", "").Split('\n'));
-
-				// Find all the metadata pages
-				AddMetadataKeyword(UdnDir, "UCLASS", "Programming/UnrealArchitecture/Reference/Classes#classdeclaration", "Programming/UnrealArchitecture/Reference/Classes/Specifiers");
-				AddMetadataKeyword(UdnDir, "UFUNCTION", "Programming/UnrealArchitecture/Reference/Functions", "Programming/UnrealArchitecture/Reference/Functions/Specifiers");
-				AddMetadataKeyword(UdnDir, "UPROPERTY", "Programming/UnrealArchitecture/Reference/Properties", "Programming/UnrealArchitecture/Reference/Properties/Specifiers");
-				AddMetadataKeyword(UdnDir, "USTRUCT", "Programming/UnrealArchitecture/Reference/Structs", "Programming/UnrealArchitecture/Reference/Structs/Specifiers");
-
-				// Clean the output folders
-				if (bCleanTargetInfo)
-				{
-					CleanTargetInfo(TargetInfoPath);
-				}
-				if (bCleanMetadata)
-				{
-					CleanMetadata(MetadataDir);
-				}
-				if (bCleanXml)
+				if (bBuildCodeAPI)
 				{
-					CleanXml(XmlDir);
-				}
-				if (bCleanUdn)
-				{
-					CleanUdn(UdnDir);
-				}
-				if (bCleanHtml)
-				{
-					CleanHtml(HtmlDir);
-				}
-				if (bCleanChm)
-				{
-					CleanChm(ChmDir);
+					// If we don't intermediate paths, make them up
+					if (XmlDir == null) XmlDir = Path.Combine(EngineDir, "Intermediate\\Documentation\\Default\\Xml");
+					if (MetadataDir == null) MetadataDir = Path.Combine(EngineDir, "Intermediate\\Documentation\\Default\\Metadata");
+
+					// Read the settings file
+					Settings = IniFile.Read(SettingsPath);
+					IgnoredFunctionMacros = new HashSet<string>(Settings.FindValueOrDefault("Input.IgnoredFunctionMacros", "").Split('\n'));
+					IncludedSourceFiles = new HashSet<string>(Settings.FindValueOrDefault("Output.IncludedSourceFiles", "").Split('\n'));
+
+					// Find all the metadata pages
+					AddMetadataKeyword(UdnDir, "UCLASS", "Programming/UnrealArchitecture/Reference/Classes#classdeclaration", "Programming/UnrealArchitecture/Reference/Classes/Specifiers");
+					AddMetadataKeyword(UdnDir, "UFUNCTION", "Programming/UnrealArchitecture/Reference/Functions", "Programming/UnrealArchitecture/Reference/Functions/Specifiers");
+					AddMetadataKeyword(UdnDir, "UPROPERTY", "Programming/UnrealArchitecture/Reference/Properties", "Programming/UnrealArchitecture/Reference/Properties/Specifiers");
+					AddMetadataKeyword(UdnDir, "USTRUCT", "Programming/UnrealArchitecture/Reference/Structs", "Programming/UnrealArchitecture/Reference/Structs/Specifiers");
+
+					// Clean the output folders
+					if (bCleanTargetInfo)
+					{
+						CleanTargetInfo(TargetInfoPath);
+					}
+					if (bCleanMetadata)
+					{
+						CleanMetadata(MetadataDir);
+					}
+					if (bCleanXml)
+					{
+						CleanXml(XmlDir);
+					}
+					if (bCleanUdn)
+					{
+						CleanUdn(Path.Combine(UdnDir, "API"));
+					}
+					if (bCleanHtml)
+					{
+						CleanHtml(Path.Combine(HtmlDir, "INT\\API"));
+					}
+					if (bCleanChm)
+					{
+						CleanChm(ChmDir);
+					}
+
+					// Build the data
+					if (!bBuildTargetInfo || BuildTargetInfo(TargetInfoPath, EngineDir))
+					{
+						if (!bBuildMetadata || BuildMetadata(DoxygenPath, EngineDir, MetadataDir, MetadataPath))
+						{
+							if (!bBuildXml || BuildXml(EngineDir, TargetInfoPath, DoxygenPath, XmlDir, Filters))
+							{
+								if (!bBuildUdn || BuildUdn(EngineDir, XmlDir, UdnDir, ChmDir, MetadataPath, StatsPath, Filters))
+								{
+									if (!bBuildHtml || BuildHtml(EngineDir, DocToolPath, UdnDir, HtmlDir, false))
+									{
+										if (!bBuildChm || BuildChm(ChmCompilerPath, HtmlDir, ChmDir, false))
+										{
+											Console.WriteLine("Complete.");
+										}
+									}
+								}
+							}
+						}
+					}
 				}
 
-				// Build the data
-				if (!bBuildTargetInfo || BuildTargetInfo(TargetInfoPath, EngineDir))
+				if (bBuildBlueprintAPI)
 				{
-					if (!bBuildMetadata || BuildMetadata(DoxygenPath, EngineDir, MetadataDir, MetadataPath))
+					// If we don't intermediate paths, make them up
+					if (JsonDir == null) JsonDir = Path.Combine(EngineDir, "Intermediate\\Documentation\\Default\\Json");
+
+					// Read the settings file
+					Settings = IniFile.Read(SettingsPath);
+
+					// Clean the output folders
+					if (bCleanTargetInfo)
+					{
+						CleanTargetInfo(TargetInfoPath);
+					}
+					if (bCleanJson)
 					{
-						if (!bBuildXml || BuildXml(EngineDir, TargetInfoPath, DoxygenPath, XmlDir, Filters))
+						CleanJson(JsonDir);
+					}
+					if (bCleanUdn)
+					{
+						CleanUdn(Path.Combine(UdnDir, "BlueprintAPI"));
+					}
+					if (bCleanHtml)
+					{
+						CleanHtml(Path.Combine(HtmlDir, "INT\\BlueprintAPI"));
+					}
+					if (bCleanChm)
+					{
+						CleanChm(ChmDir);
+					}
+
+					// Build the data
+					if (!bBuildTargetInfo || BuildTargetInfo(TargetInfoPath, EngineDir))
+					{
+						if (!bBuildJson || BuildJson(JsonDir, EngineDir, EditorPath))
 						{
-							if (!bBuildUdn || BuildUdn(EngineDir, XmlDir, UdnDir, ChmDir, MetadataPath, StatsPath, Filters))
+							if (!bBuildUdn || BuildBlueprintUdn(JsonDir, UdnDir, ChmDir, StatsPath))
 							{
-								if (!bBuildHtml || BuildHtml(EngineDir, DocToolPath, UdnDir, HtmlDir))
+								if (!bBuildHtml || BuildHtml(EngineDir, DocToolPath, UdnDir, HtmlDir, true))
 								{
-									if (!bBuildChm || BuildChm(ChmCompilerPath, HtmlDir, ChmDir))
+									if (!bBuildChm || BuildChm(ChmCompilerPath, HtmlDir, ChmDir, true))
 									{
 										Console.WriteLine("Complete.");
 									}
@@ -427,6 +515,8 @@ namespace APIDocTool
 				Console.WriteLine("APIDocTool.exe [options] -enginedir=<...>");                                    // <-- 80 character limit to start of comment
 				Console.WriteLine();
 				Console.WriteLine("Options:");
+				Console.WriteLine("    -blueprint:                      Build the blueprint API");
+				Console.WriteLine("    -code:                           Build the code API");
 				Console.WriteLine("    -rebuild:                        Clean and build everything");
 				Console.WriteLine("    -rebuild[meta|xml|udn|html|chm]: Clean and build specific files");
 				Console.WriteLine("    -clean:                          Clean all files");
@@ -436,6 +526,7 @@ namespace APIDocTool
 				Console.WriteLine("    -targetinfo=<...>:			    Specifies the build info, created by");
 				Console.WriteLine("                                     running UBT with -writetargetinfo=<...>");
 				Console.WriteLine("    -xmldir=<...>:				    Output directory for xml files");
+				Console.WriteLine("    -jsondir=<...>:				    Output directory for json files");
 				Console.WriteLine("    -metadatadir=<...>:			    Output directory for metadata files");
 				Console.WriteLine("    -filter=<...>,<...>:             Filter conversion, eg.");
 				Console.WriteLine("                                       Folders:  -filter=Core/Containers/...");
@@ -695,9 +786,47 @@ namespace APIDocTool
 			return true;
 		}
 
-		static void CleanUdn(string UdnDir)
+		static void CleanJson(string JsonDir)
+		{
+			Console.WriteLine("Cleaning '{0}'", JsonDir);
+			Utility.SafeDeleteDirectoryContents(JsonDir, true);
+		}
+
+		static bool BuildJson(string JsonDir, string EngineDir, string EditorPath)
+		{
+			// Create the output directory
+			Utility.SafeCreateDirectory(JsonDir);
+
+			using (Process JsonExportProcess = new Process())
+			{
+				JsonExportProcess.StartInfo.WorkingDirectory = EngineDir;
+				JsonExportProcess.StartInfo.FileName = EditorPath;
+				JsonExportProcess.StartInfo.Arguments = "-run=GenerateBlueprintAPI -path=" + JsonDir + " -name=BlueprintAPI";
+				JsonExportProcess.StartInfo.UseShellExecute = false;
+				JsonExportProcess.StartInfo.RedirectStandardOutput = true;
+				JsonExportProcess.StartInfo.RedirectStandardError = true;
+
+				JsonExportProcess.OutputDataReceived += new DataReceivedEventHandler(ProcessOutputReceived);
+				JsonExportProcess.ErrorDataReceived += new DataReceivedEventHandler(ProcessOutputReceived);
+
+				try
+				{
+					JsonExportProcess.Start();
+					JsonExportProcess.BeginOutputReadLine();
+					JsonExportProcess.BeginErrorReadLine();
+					JsonExportProcess.WaitForExit();
+					return JsonExportProcess.ExitCode == 0;
+				}
+				catch (Exception Ex)
+				{
+					Console.WriteLine(Ex.ToString() + "\n" + Ex.StackTrace);
+					return false;
+				}
+			}
+		}
+
+		static void CleanUdn(string CleanDir)
 		{
-			string CleanDir = Path.Combine(UdnDir, "API");
 			Console.WriteLine("Cleaning '{0}'", CleanDir);
 
 			// Delete all the files
@@ -802,11 +931,86 @@ namespace APIDocTool
 
 			// Write the sitemap contents
 			Console.WriteLine("Writing sitemap contents...");
-			Index.WriteSitemapContents(Path.Combine(SitemapDir, SitemapContentsFileName));
+			Index.WriteSitemapContents(Path.Combine(SitemapDir, "API.hhc"));
 
 			// Write the sitemap index
 			Console.WriteLine("Writing sitemap index...");
-			Index.WriteSitemapIndex(Path.Combine(SitemapDir, SitemapIndexFileName));
+			Index.WriteSitemapIndex(Path.Combine(SitemapDir, "API.hhk"));
+
+			return true;
+		}
+
+		static bool BuildBlueprintUdn(string JsonDir, string UdnDir, string SitemapDir, string StatsPath)
+		{
+			// Read the input json file
+			string JsonFilePath = Path.Combine(JsonDir, "BlueprintAPI.json");
+			var json = (Dictionary<string, object>)fastJSON.JSON.Instance.Parse(File.ReadAllText(JsonFilePath));
+
+			APICategory.LoadTooltips((Dictionary<string, object>)json["Categories"]);
+
+			// TODO: This path is clearly sketchy as hell, but we'll clean it up later maybe
+			var Actions = (Dictionary<string, object>)((Dictionary<string, object>)((Dictionary<string, object>)((Dictionary<string, object>)json["Actor"])["Palette"])["ActionSet"])["Actions"];
+
+			APICategory RootCategory = new APICategory(null, "BlueprintAPI", true);
+
+			foreach (var Action in Actions)
+			{
+				var CategoryList = Action.Key.Split('|');
+				var ActionCategory = RootCategory;
+
+				Debug.Assert(CategoryList.Length > 0);
+				Debug.Assert(CategoryList[0] == "Library");
+
+				for (int CategoryIndex = 1; CategoryIndex < CategoryList.Length - 1; ++CategoryIndex)
+				{
+					ActionCategory = ActionCategory.GetSubCategory(CategoryList[CategoryIndex]);
+				}
+
+				ActionCategory.AddAction(new APIAction(ActionCategory, CategoryList.Last(), (Dictionary<string, object>)Action.Value));
+			}
+
+			// Build a list of pages to output
+			List<APIPage> OutputPages = new List<APIPage>(RootCategory.GatherPages().OrderBy(x => x.LinkPath));
+
+			// Create the output directory
+			Utility.SafeCreateDirectory(UdnDir);
+			Utility.SafeCreateDirectory(Path.Combine(UdnDir, APIFolder));
+
+			// Build the manifest
+			Console.WriteLine("Writing manifest...");
+			UdnManifest Manifest = new UdnManifest(RootCategory);
+			Manifest.PrintConflicts();
+			Manifest.Write(Path.Combine(UdnDir, BlueprintAPIFolder + "\\API.manifest"));
+
+			Console.WriteLine("Categories: " + OutputPages.Count(page => page is APICategory));
+			Console.WriteLine("Actions: " + OutputPages.Count(page => page is APIAction));
+
+			// Write all the pages
+			using (Tracker UdnTracker = new Tracker("Writing UDN pages...", OutputPages.Count))
+			{
+				foreach (int Idx in UdnTracker.Indices)
+				{
+					APIPage Page = OutputPages[Idx];
+
+					// Create the output directory
+					string MemberDirectory = Path.Combine(UdnDir, Page.LinkPath);
+					if (!Directory.Exists(MemberDirectory))
+					{
+						Directory.CreateDirectory(MemberDirectory);
+					}
+
+					// Write the page
+					Page.WritePage(Manifest, Path.Combine(MemberDirectory, "index.INT.udn"));
+				}
+			}
+
+			// Write the sitemap contents
+			Console.WriteLine("Writing sitemap contents...");
+			RootCategory.WriteSitemapContents(Path.Combine(SitemapDir, "BlueprintAPI.hhc"));
+
+			// Write the sitemap index
+			Console.WriteLine("Writing sitemap index...");
+			RootCategory.WriteSitemapIndex(Path.Combine(SitemapDir, "BlueprintAPI.hhk"));
 
 			return true;
 		}
@@ -905,14 +1109,13 @@ namespace APIDocTool
 			}
 		}
 
-		public static void CleanHtml(string HtmlPath)
+		public static void CleanHtml(string CleanDir)
 		{
-			string CleanDir = Path.Combine(HtmlPath, "INT\\API");
 			Console.WriteLine("Cleaning '{0}'", CleanDir);
 			Utility.SafeDeleteDirectoryContents(CleanDir, true);
 		}
 
-		public static bool BuildHtml(string EngineDir, string DocToolPath, string UdnPath, string HtmlDir)
+		public static bool BuildHtml(string EngineDir, string DocToolPath, string UdnPath, string HtmlDir, bool bBlueprint)
 		{
 			Utility.SafeCreateDirectory(HtmlDir);
 
@@ -920,7 +1123,7 @@ namespace APIDocTool
 			{
 				DocToolProcess.StartInfo.WorkingDirectory = EngineDir;
 				DocToolProcess.StartInfo.FileName = DocToolPath;
-				DocToolProcess.StartInfo.Arguments = "API\\* -lang=INT -t=DefaultAPI.html -v=warn";
+				DocToolProcess.StartInfo.Arguments = (bBlueprint ? "BlueprintAPI" : "API") + "\\* -lang=INT -t=DefaultAPI.html -v=warn";
 				DocToolProcess.StartInfo.UseShellExecute = false;
 				DocToolProcess.StartInfo.RedirectStandardOutput = true;
 				DocToolProcess.StartInfo.RedirectStandardError = true;
@@ -956,20 +1159,34 @@ namespace APIDocTool
 			}
 		}
 
-		public static bool BuildChm(string ChmCompilerPath, string BaseHtmlDir, string ChmDir)
+		public static bool BuildChm(string ChmCompilerPath, string BaseHtmlDir, string ChmDir, bool bBlueprint)
 		{
-			const string ProjectFileName = "API.hhp";
+			string ProjectFileName = (bBlueprint ? "BlueprintAPI.hhp" : "API.hhp");
 			Console.WriteLine("Searching for CHM input files...");
 
 			// Build a list of all the files we want to copy
 			List<string> FilePaths = new List<string>();
 			List<string> DirectoryPaths = new List<string>();
-			Utility.FindRelativeContents(BaseHtmlDir, "Images\\api*", false, FilePaths, DirectoryPaths);
+			if (bBlueprint)
+			{
+				Utility.FindRelativeContents(BaseHtmlDir, "Images\\bp_api*", false, FilePaths, DirectoryPaths);
+			}
+			else
+			{
+				Utility.FindRelativeContents(BaseHtmlDir, "Images\\api*", false, FilePaths, DirectoryPaths);
+			}
 			Utility.FindRelativeContents(BaseHtmlDir, "Include\\*", true, FilePaths, DirectoryPaths);
 
 			// Find all the HTML files
 			List<string> HtmlFilePaths = new List<string>();
-			Utility.FindRelativeContents(BaseHtmlDir, "INT\\API\\*.html", true, HtmlFilePaths, DirectoryPaths);
+			if (bBlueprint)
+			{
+				Utility.FindRelativeContents(BaseHtmlDir, "INT\\BlueprintAPI\\*.html", true, HtmlFilePaths, DirectoryPaths);
+			}
+			else 
+			{
+				Utility.FindRelativeContents(BaseHtmlDir, "INT\\API\\*.html", true, HtmlFilePaths, DirectoryPaths);
+			}
 
 			// Create all the target directories
 			foreach (string DirectoryPath in DirectoryPaths)
@@ -1015,16 +1232,32 @@ namespace APIDocTool
 			using (StreamWriter Writer = new StreamWriter(Path.Combine(ChmDir, ProjectFileName)))
 			{
 				Writer.WriteLine("[OPTIONS]");
-				Writer.WriteLine("Title=UE4 API Documentation");
-				Writer.WriteLine("Binary TOC=Yes");
-				Writer.WriteLine("Compatibility=1.1 or later");
-				Writer.WriteLine("Compiled file=API.chm");
-				Writer.WriteLine("Contents file=" + SitemapContentsFileName);
-				Writer.WriteLine("Index file=" + SitemapIndexFileName);
-				Writer.WriteLine("Default topic=INT\\API\\index.html");
-				Writer.WriteLine("Full-text search=Yes");
-				Writer.WriteLine("Display compile progress=Yes");
-				Writer.WriteLine("Language=0x409 English (United States)");
+				if (bBlueprint)
+				{
+					Writer.WriteLine("Title=UE4 Blueprint API Documentation");
+					Writer.WriteLine("Binary TOC=Yes");
+					Writer.WriteLine("Compatibility=1.1 or later");
+					Writer.WriteLine("Compiled file=BlueprintAPI.chm");
+					Writer.WriteLine("Contents file=BlueprintAPI.hhc");
+					Writer.WriteLine("Index file=BlueprintAPI.hhk");
+					Writer.WriteLine("Default topic=INT\\BlueprintAPI\\index.html");
+					Writer.WriteLine("Full-text search=Yes");
+					Writer.WriteLine("Display compile progress=Yes");
+					Writer.WriteLine("Language=0x409 English (United States)");
+				}
+				else
+				{
+					Writer.WriteLine("Title=UE4 API Documentation");
+					Writer.WriteLine("Binary TOC=Yes");
+					Writer.WriteLine("Compatibility=1.1 or later");
+					Writer.WriteLine("Compiled file=API.chm");
+					Writer.WriteLine("Contents file=API.hhc");
+					Writer.WriteLine("Index file=API.hhk");
+					Writer.WriteLine("Default topic=INT\\API\\index.html");
+					Writer.WriteLine("Full-text search=Yes");
+					Writer.WriteLine("Display compile progress=Yes");
+					Writer.WriteLine("Language=0x409 English (United States)");
+				}
 				Writer.WriteLine();
 				Writer.WriteLine("[FILES]");
 				foreach (string FilePath in FilePaths)
diff --git a/Engine/Source/Programs/UnrealDocTool/APIDocTool/APIDocTool/UdnWriter.cs b/Engine/Source/Programs/UnrealDocTool/APIDocTool/APIDocTool/UdnWriter.cs
--- a/Engine/Source/Programs/UnrealDocTool/APIDocTool/APIDocTool/UdnWriter.cs
+++ b/Engine/Source/Programs/UnrealDocTool/APIDocTool/APIDocTool/UdnWriter.cs
@@ -456,6 +456,28 @@ namespace APIDocTool
 			WriteObject("MemberListTail");
 		}
 
+		public void WriteList(IEnumerable<UdnIconListItem> Items)
+		{
+			// Write the head
+			WriteObject("MemberIconListHeadBlank");
+
+			// Write the body
+			foreach (UdnIconListItem Item in Items)
+			{
+				if (Item.Link == null)
+				{
+					WriteObject("MemberIconListItem", "icon", Item.Icons, "name", Item.Name, "desc", Item.Description);
+				}
+				else
+				{
+					WriteObject("MemberIconListItemLinked", "icon", Item.Icons, "name", Item.Name, "link", "[RELATIVE:" + Item.Link + "]", "desc", Item.Description);
+				}
+			}
+
+			// Write the tail
+			WriteObject("MemberIconListTail");
+		}
+
 		public void WriteList(string NameColumn, string DescColumn, IEnumerable<UdnIconListItem> Items)
 		{
 			// Write the head
diff --git a/Engine/Source/Programs/UnrealDocTool/APIDocTool/APIDocTool/fastJSON/Formatter.cs b/Engine/Source/Programs/UnrealDocTool/APIDocTool/APIDocTool/fastJSON/Formatter.cs
new file mode 100644
--- /dev/null
+++ b/Engine/Source/Programs/UnrealDocTool/APIDocTool/APIDocTool/fastJSON/Formatter.cs
@@ -0,0 +1,97 @@
+﻿using System.Collections.Generic;
+using System.Text;
+
+namespace fastJSON
+{
+    internal static class Formatter
+    {
+        public static string Indent = "   ";
+
+        public static void AppendIndent(StringBuilder sb, int count)
+        {
+            for (; count > 0; --count) sb.Append(Indent);
+        }
+
+        public static bool IsEscaped(StringBuilder sb, int index)
+        {
+            bool escaped = false;
+            while (index > 0 && sb[--index] == '\\') escaped = !escaped;
+            return escaped;
+        }
+
+        public static string PrettyPrint(string input)
+        {
+            var output = new StringBuilder(input.Length * 2);
+            char? quote = null;
+            int depth = 0;
+
+            for (int i = 0; i < input.Length; ++i)
+            {
+                char ch = input[i];
+
+                if (ch == '\"') // found string span
+                {
+                    bool str = true;
+                    while (str)
+                    {
+                        output.Append(ch);
+                        ch = input[++i];
+                        if (ch == '\\')
+                        {
+                            if (input[i + 1] == '\"')
+                            {
+                                output.Append(ch);
+                                ch = input[++i];
+                            }
+                        }
+                        else if (ch == '\"')
+                            str = false;
+                    }
+                }
+
+                switch (ch)
+                {
+                    case '{':
+                    case '[':
+                        output.Append(ch);
+                        if (!quote.HasValue)
+                        {
+                            output.AppendLine();
+                            AppendIndent(output, ++depth);
+                        }
+                        break;
+                    case '}':
+                    case ']':
+                        if (quote.HasValue)
+                            output.Append(ch);
+                        else
+                        {
+                            output.AppendLine();
+                            AppendIndent(output, --depth);
+                            output.Append(ch);
+                        }
+                        break;
+                    case ',':
+                        output.Append(ch);
+                        if (!quote.HasValue)
+                        {
+                            output.AppendLine();
+                            AppendIndent(output, depth);
+                        }
+                        break;
+                    case ':':
+                        if (quote.HasValue) output.Append(ch);
+                        else output.Append(" : ");
+                        break;
+                    default:
+                        if (quote.HasValue || !char.IsWhiteSpace(ch))
+                            output.Append(ch);
+                        break;
+                }
+
+            }
+
+            return output.ToString();
+        }
+    }
+}
\ No newline at end of file
diff --git a/Engine/Source/Programs/UnrealDocTool/APIDocTool/APIDocTool/fastJSON/Getters.cs b/Engine/Source/Programs/UnrealDocTool/APIDocTool/APIDocTool/fastJSON/Getters.cs
new file mode 100644
--- /dev/null
+++ b/Engine/Source/Programs/UnrealDocTool/APIDocTool/APIDocTool/fastJSON/Getters.cs
@@ -0,0 +1,11 @@
+﻿using System;
+using System.Collections.Generic;
+
+namespace fastJSON
+{
+    public sealed class DatasetSchema
+    {
+        public List<string> Info { get; set; }
+        public string Name { get; set; }
+    }
+}
diff --git a/Engine/Source/Programs/UnrealDocTool/APIDocTool/APIDocTool/fastJSON/JSON.cs b/Engine/Source/Programs/UnrealDocTool/APIDocTool/APIDocTool/fastJSON/JSON.cs
new file mode 100644
--- /dev/null
+++ b/Engine/Source/Programs/UnrealDocTool/APIDocTool/APIDocTool/fastJSON/JSON.cs
@@ -0,0 +1,898 @@
+﻿using System;
+using System.Collections;
+using System.Collections.Generic;
+#if !SILVERLIGHT
+using System.Data;
+#endif
+using System.Globalization;
+using System.IO;
+using System.Reflection;
+using System.Reflection.Emit;
+
+namespace fastJSON
+{
+    public delegate string Serialize(object data);
+    public delegate object Deserialize(string data);
+
+    public sealed class JSONParameters
+    {
+        /// <summary>
+        /// Use the optimized fast Dataset Schema format (default = True)
+        /// </summary>
+        public bool UseOptimizedDatasetSchema = true;
+        /// <summary>
+        /// Use the fast GUID format (default = True)
+        /// </summary>
+        public bool UseFastGuid = true;
+        /// <summary>
+        /// Serialize null values to the output (default = True)
+        /// </summary>
+        public bool SerializeNullValues = true;
+        /// <summary>
+        /// Use the UTC date format (default = True)
+        /// </summary>
+        public bool UseUTCDateTime = true;
+        /// <summary>
+        /// Show the readonly properties of types in the output (default = False)
+        /// </summary>
+        public bool ShowReadOnlyProperties = false;
+        /// <summary>
+        /// Use the $types extension to optimise the output json (default = True)
+        /// </summary>
+        public bool UsingGlobalTypes = true;
+        /// <summary>
+        /// ** work in progress
+        /// </summary>
+        public bool IgnoreCaseOnDeserialize = false;
+        /// <summary>
+        /// Anonymous types have read only properties 
+        /// </summary>
+        public bool EnableAnonymousTypes = false;
+        /// <summary>
+        /// Enable fastJSON extensions $types, $type, $map (default = True)
+        /// </summary>
+        public bool UseExtensions = true;
+        /// <summary>
+        /// Use escaped unicode i.e. \uXXXX format for non ASCII characters (default = True)
+        /// </summary>
+        public bool UseEscapedUnicode = true;
+
+        public void FixValues()
+        {
+            if (UseExtensions == false) // disable conflicting params
+            {
+                UsingGlobalTypes = false;
+            }
+        }
+    }
+
+    public sealed class JSON
+    {
+        //public readonly static JSON Instance = new JSON();
+        [ThreadStatic]
+        private static JSON _instance;
+
+        public static JSON Instance
+        {
+            get { return _instance ?? (_instance = new JSON()); }
+        }
+
+        private JSON()
+        {
+        }
+        /// <summary>
+        /// You can set these paramters globally for all calls
+        /// </summary>
+        public JSONParameters Parameters = new JSONParameters();
+        private JSONParameters _params;
+
+        public string ToJSON(object obj)
+        {
+            _params = Parameters;
+            _params.FixValues();
+            Reflection.Instance.ShowReadOnlyProperties = _params.ShowReadOnlyProperties;
+            return ToJSON(obj, Parameters);
+        }
+
+        public string ToJSON(object obj, JSONParameters param)
+        {
+            _params = param;
+            _params.FixValues();
+            Reflection.Instance.ShowReadOnlyProperties = _params.ShowReadOnlyProperties;
+            Type t = null;
+
+            if (obj == null)
+                return "null";
+
+            if (obj.GetType().IsGenericType)
+                t = obj.GetType().GetGenericTypeDefinition();
+            if (t == typeof(Dictionary<,>) || t == typeof(List<>))
+                _params.UsingGlobalTypes = false;
+
+            // FEATURE : enable extensions when you can deserialize anon types
+            if (_params.EnableAnonymousTypes) { _params.UseExtensions = false; _params.UsingGlobalTypes = false; Reflection.Instance.ShowReadOnlyProperties = true; }
+            _usingglobals = _params.UsingGlobalTypes;
+            return new JSONSerializer(_params).ConvertToJSON(obj);
+        }
+
+        public object Parse(string json)
+        {
+            _params = Parameters;
+            Reflection.Instance.ShowReadOnlyProperties = _params.ShowReadOnlyProperties;
+            return new JsonParser(json, _params.IgnoreCaseOnDeserialize).Decode();
+        }
+
+        public T ToObject<T>(string json)
+        {
+            return (T)ToObject(json, typeof(T));
+        }
+
+        public object ToObject(string json)
+        {
+            return ToObject(json, null);
+        }
+
+        public object ToObject(string json, Type type)
+        {
+            _params = Parameters;
+            _params.FixValues();
+            Reflection.Instance.ShowReadOnlyProperties = _params.ShowReadOnlyProperties;
+            Type t = null;
+            if (type != null && type.IsGenericType)
+                t = type.GetGenericTypeDefinition();
+            if (t == typeof(Dictionary<,>) || t == typeof(List<>))
+                _params.UsingGlobalTypes = false;
+            _usingglobals = _params.UsingGlobalTypes;
+
+            object o = new JsonParser(json, Parameters.IgnoreCaseOnDeserialize).Decode();
+            if (o == null)
+                return null;
+
+#if !SILVERLIGHT
+            if (type != null && type == typeof(DataSet))
+                return CreateDataset(o as Dictionary<string, object>, null);
+
+            if (type != null && type == typeof(DataTable))
+                return CreateDataTable(o as Dictionary<string, object>, null);
+#endif
+            if (o is IDictionary)
+            {
+                if (type != null && type.IsGenericType && type.GetGenericTypeDefinition() == typeof(Dictionary<,>)) // deserialize a dictionary
+                    return RootDictionary(o, type);
+                else // deserialize an object
+                    return ParseDictionary(o as Dictionary<string, object>, null, type, null);
+            }
+
+            if (o is List<object>)
+            {
+                if (type != null && type.IsGenericType && type.GetGenericTypeDefinition() == typeof(Dictionary<,>)) // kv format
+                    return RootDictionary(o, type);
+
+                if (type != null && type.IsGenericType && type.GetGenericTypeDefinition() == typeof(List<>)) // deserialize to generic list
+                    return RootList(o, type);
+                else
+                    return (o as List<object>).ToArray();
+            }
+
+            if (type != null && o.GetType() != type)
+                return ChangeType(o, type);
+
+            return o;
+        }
+
+        public string Beautify(string input)
+        {
+            return Formatter.PrettyPrint(input);
+        }
+
+        public object FillObject(object input, string json)
+        {
+            _params = Parameters;
+            _params.FixValues();
+            Reflection.Instance.ShowReadOnlyProperties = _params.ShowReadOnlyProperties;
+            Dictionary<string, object> ht = new JsonParser(json, Parameters.IgnoreCaseOnDeserialize).Decode() as Dictionary<string, object>;
+            if (ht == null) return null;
+            return ParseDictionary(ht, null, input.GetType(), input);
+        }
+
+        public object DeepCopy(object obj)
+        {
+            return ToObject(ToJSON(obj));
+        }
+
+        public T DeepCopy<T>(T obj)
+        {
+            return ToObject<T>(ToJSON(obj));
+        }
+
+        internal SafeDictionary<Type, Serialize> _customSerializer = new SafeDictionary<Type, Serialize>();
+        internal SafeDictionary<Type, Deserialize> _customDeserializer = new SafeDictionary<Type, Deserialize>();
+
+        public void RegisterCustomType(Type type, Serialize serializer, Deserialize deserializer)
+        {
+            if (type != null && serializer != null && deserializer != null)
+            {
+                _customSerializer.Add(type, serializer);
+                _customDeserializer.Add(type, deserializer);
+                // reset property cache
+                _propertycache = new SafeDictionary<string, SafeDictionary<string, myPropInfo>>();
+            }
+        }
+
+        internal bool IsTypeRegistered(Type t)
+        {
+            if (_customSerializer.Count == 0)
+                return false;
+            Serialize s;
+            return _customSerializer.TryGetValue(t, out s);
+        }
+
+        #region [   JSON specific reflection   ]
+
+        private enum myPropInfoType
+        {
+            Int,
+            Long,
+            String,
+            Bool,
+            DateTime,
+            Enum,
+            Guid,
+
+            Array,
+            ByteArray,
+            Dictionary,
+            StringDictionary,
+#if !SILVERLIGHT
+            Hashtable,
+            DataSet,
+            DataTable,
+#endif
+			Custom,
+            Unknown,
+        }
+
+        [Flags]
+        private enum myPropInfoFlags
+        {
+            Filled = 1 << 0,
+            CanWrite = 1 << 1
+        }
+
+        private struct myPropInfo
+        {
+            public Type pt;
+            public Type bt;
+            public Type changeType;
+            public Reflection.GenericSetter setter;
+            public Reflection.GenericGetter getter;
+            public Type[] GenericTypes;
+            public string Name;
+            public myPropInfoType Type;
+            public myPropInfoFlags Flags;
+
+            public bool IsClass;
+            public bool IsValueType;
+            public bool IsGenericType;
+        }
+
+        SafeDictionary<string, SafeDictionary<string, myPropInfo>> _propertycache = new SafeDictionary<string, SafeDictionary<string, myPropInfo>>();
+        private SafeDictionary<string, myPropInfo> Getproperties(Type type, string typename)
+        {
+            SafeDictionary<string, myPropInfo> sd = null;
+            if (_propertycache.TryGetValue(typename, out sd))
+            {
+                return sd;
+            }
+            else
+            {
+                sd = new SafeDictionary<string, myPropInfo>();
+                PropertyInfo[] pr = type.GetProperties(BindingFlags.Public | BindingFlags.Instance);
+                foreach (PropertyInfo p in pr)
+                {
+                    myPropInfo d = CreateMyProp(p.PropertyType, p.Name);
+                    d.Flags |= myPropInfoFlags.CanWrite;
+                    d.setter = Reflection.CreateSetMethod(type, p);
+                    d.getter = Reflection.CreateGetMethod(type, p);
+                    sd.Add(p.Name, d);
+                }
+                FieldInfo[] fi = type.GetFields(BindingFlags.Public | BindingFlags.Instance);
+                foreach (FieldInfo f in fi)
+                {
+                    myPropInfo d = CreateMyProp(f.FieldType, f.Name);
+                    d.setter = Reflection.CreateSetField(type, f);
+                    d.getter = Reflection.CreateGetField(type, f);
+                    sd.Add(f.Name, d);
+                }
+
+                _propertycache.Add(typename, sd);
+                return sd;
+            }
+        }
+
+        private myPropInfo CreateMyProp(Type t, string name)
+        {
+            myPropInfo d = new myPropInfo();
+            myPropInfoType d_type = myPropInfoType.Unknown;
+            myPropInfoFlags d_flags = myPropInfoFlags.Filled | myPropInfoFlags.CanWrite;
+
+            if (t == typeof(int) || t == typeof(int?)) d_type = myPropInfoType.Int;
+            else if (t == typeof(long) || t == typeof(long?)) d_type = myPropInfoType.Long;
+            else if (t == typeof(string)) d_type = myPropInfoType.String;
+            else if (t == typeof(bool) || t == typeof(bool?)) d_type = myPropInfoType.Bool;
+            else if (t == typeof(DateTime) || t == typeof(DateTime?)) d_type = myPropInfoType.DateTime;
+            else if (t.IsEnum) d_type = myPropInfoType.Enum;
+            else if (t == typeof(Guid) || t == typeof(Guid?)) d_type = myPropInfoType.Guid;
+            else if (t.IsArray)
+            {
+                d.bt = t.GetElementType();
+                if (t == typeof(byte[]))
+                    d_type = myPropInfoType.ByteArray;
+                else
+                    d_type = myPropInfoType.Array;
+            }
+            else if (t.Name.Contains("Dictionary"))
+            {
+                d.GenericTypes = t.GetGenericArguments();
+                if (d.GenericTypes.Length > 0 && d.GenericTypes[0] == typeof(string))
+                    d_type = myPropInfoType.StringDictionary;
+                else
+                    d_type = myPropInfoType.Dictionary;
+            }
+#if !SILVERLIGHT
+            else if (t == typeof(Hashtable)) d_type = myPropInfoType.Hashtable;
+            else if (t == typeof(DataSet)) d_type = myPropInfoType.DataSet;
+            else if (t == typeof(DataTable)) d_type = myPropInfoType.DataTable;
+#endif
+
+            else if (IsTypeRegistered(t))								
+                d_type = myPropInfoType.Custom;
+
+            d.IsClass = t.IsClass;
+            d.IsValueType = t.IsValueType;
+            if (t.IsGenericType)
+            {
+                d.IsGenericType = true;
+                d.bt = t.GetGenericArguments()[0];
+            }
+
+            d.pt = t;
+            d.Name = name;
+            d.changeType = GetChangeType(t);
+            d.Type = d_type;
+            d.Flags = d_flags;
+
+            return d;
+        }
+
+        private object ChangeType(object value, Type conversionType)
+        {
+            if (conversionType == typeof(int))
+                return (int)((long)value);
+
+            else if (conversionType == typeof(long))
+                return (long)value;
+
+            else if (conversionType == typeof(string))
+                return (string)value;
+
+            else if (conversionType == typeof(Guid))
+                return CreateGuid((string)value);
+
+            else if (conversionType.IsEnum)
+                return CreateEnum(conversionType, (string)value);
+
+            else if (IsTypeRegistered(conversionType))
+                return CreateCustom((string)value, conversionType);
+
+            return Convert.ChangeType(value, conversionType, CultureInfo.InvariantCulture);
+        }
+        #endregion
+
+        #region [   p r i v a t e   m e t h o d s   ]
+
+        private object RootList(object parse, Type type)
+        {
+            Type[] gtypes = type.GetGenericArguments();
+            IList o = (IList)Reflection.Instance.FastCreateInstance(type);
+            foreach (var k in (IList)parse)
+            {
+                _usingglobals = false;
+                object v = k;
+                if (k is Dictionary<string, object>)
+                    v = ParseDictionary(k as Dictionary<string, object>, null, gtypes[0], null);
+                else
+                    v = ChangeType(k, gtypes[0]);
+
+                o.Add(v);
+            }
+            return o;
+        }
+
+        private object RootDictionary(object parse, Type type)
+        {
+            Type[] gtypes = type.GetGenericArguments();
+            if (parse is Dictionary<string, object>)
+            {
+                IDictionary o = (IDictionary)Reflection.Instance.FastCreateInstance(type);
+
+                foreach (var kv in (Dictionary<string, object>)parse)
+                {
+                    object v;
+                    object k = ChangeType(kv.Key, gtypes[0]);
+                    if (kv.Value is Dictionary<string, object>)
+                        v = ParseDictionary(kv.Value as Dictionary<string, object>, null, gtypes[1], null);
+                    else if (kv.Value is List<object>)
+                        v = CreateArray(kv.Value as List<object>, typeof(object), typeof(object), null);
+                    else
+                        v = ChangeType(kv.Value, gtypes[1]);
+                    o.Add(k, v);
+                }
+
+                return o;
+            }
+            if (parse is List<object>)
+                return CreateDictionary(parse as List<object>, type, gtypes, null);
+
+            return null;
+        }
+
+        bool _usingglobals = false;
+        private object ParseDictionary(Dictionary<string, object> d, Dictionary<string, object> globaltypes, Type type, object input)
+        {
+            object tn = "";
+
+            if (d.TryGetValue("$types", out tn))
+            {
+                _usingglobals = true;
+                globaltypes = new Dictionary<string, object>();
+                foreach (var kv in (Dictionary<string, object>)tn)
+                {
+                    globaltypes.Add((string)kv.Value, kv.Key);
+                }
+            }
+
+            bool found = d.TryGetValue("$type", out tn);
+#if !SILVERLIGHT
+            if (found == false && type == typeof(System.Object))
+            {
+                return CreateDataset(d, globaltypes);
+            }
+#endif
+            if (found)
+            {
+                if (_usingglobals)
+                {
+                    object tname = "";
+                    if (globaltypes.TryGetValue((string)tn, out tname))
+                        tn = tname;
+                }
+                type = Reflection.Instance.GetTypeFromCache((string)tn);
+            }
+
+            if (type == null)
+                throw new Exception("Cannot determine type");
+
+            string typename = type.FullName;
+            object o = input;
+            if (o == null)
+                o = Reflection.Instance.FastCreateInstance(type);
+
+            SafeDictionary<string, myPropInfo> props = Getproperties(type, typename);
+            foreach (string n in d.Keys)
+            {
+                string name = n;
+                if (_params.IgnoreCaseOnDeserialize) name = name.ToLower();
+                if (name == "$map")
+                {
+                    ProcessMap(o, props, (Dictionary<string, object>)d[name]);
+                    continue;
+                }
+                myPropInfo pi;
+                if (props.TryGetValue(name, out pi) == false)
+                    continue;
+                if ((pi.Flags & (myPropInfoFlags.Filled | myPropInfoFlags.CanWrite)) != 0)
+                {
+                    object v = d[name];
+
+                    if (v != null)
+                    {
+                        object oset = null;
+
+                        switch (pi.Type)
+                        {
+                            case myPropInfoType.Int: oset = (int)((long)v); break;
+                            case myPropInfoType.Long: oset = (long)v; break;
+                            case myPropInfoType.String: oset = (string)v; break;
+                            case myPropInfoType.Bool: oset = (bool)v; break;
+                            case myPropInfoType.DateTime: oset = CreateDateTime((string)v); break;
+                            case myPropInfoType.Enum: oset = CreateEnum(pi.pt, (string)v); break;
+                            case myPropInfoType.Guid: oset = CreateGuid((string)v); break;
+
+                            case myPropInfoType.Array:
+                                if (!pi.IsValueType)
+                                    oset = CreateArray((List<object>)v, pi.pt, pi.bt, globaltypes);
+                                // what about 'else'?
+                                break;
+                            case myPropInfoType.ByteArray: oset = Convert.FromBase64String((string)v); break;
+#if !SILVERLIGHT
+                            case myPropInfoType.DataSet: oset = CreateDataset((Dictionary<string, object>)v, globaltypes); break;
+                            case myPropInfoType.DataTable: oset = this.CreateDataTable((Dictionary<string, object>)v, globaltypes); break;
+                            case myPropInfoType.Hashtable: // same case as Dictionary
+#endif
+                            case myPropInfoType.Dictionary: oset = CreateDictionary((List<object>)v, pi.pt, pi.GenericTypes, globaltypes); break;
+                            case myPropInfoType.StringDictionary: oset = CreateStringKeyDictionary((Dictionary<string, object>)v, pi.pt, pi.GenericTypes, globaltypes); break;
+
+							case myPropInfoType.Custom: oset = CreateCustom((string)v, pi.pt); break;
+                            default:
+                                {
+                                    if (pi.IsGenericType && pi.IsValueType == false && v is List<object>)
+                                        oset = CreateGenericList((List<object>)v, pi.pt, pi.bt, globaltypes);
+
+                                    else if (pi.IsClass && v is Dictionary<string, object>)
+                                        oset = ParseDictionary((Dictionary<string, object>)v, globaltypes, pi.pt, pi.getter(o));
+
+                                    else if (v is List<object>)
+                                        oset = CreateArray((List<object>)v, pi.pt, typeof(object), globaltypes);
+
+                                    else if (pi.IsValueType)
+                                        oset = ChangeType(v, pi.changeType);
+
+                                    else
+                                        oset = v;
+                                }
+                                break;
+                        }
+
+                        o = pi.setter(o, oset);
+                    }
+                }
+            }
+            return o;
+        }
+
+        private object CreateCustom(string v, Type type)
+        {
+            Deserialize d;
+            _customDeserializer.TryGetValue(type, out d);
+            return d(v);
+        }
+
+        private void ProcessMap(object obj, SafeDictionary<string, JSON.myPropInfo> props, Dictionary<string, object> dic)
+        {
+            foreach (KeyValuePair<string, object> kv in dic)
+            {
+                myPropInfo p = props[kv.Key];
+                object o = p.getter(obj);
+                Type t = Type.GetType((string)kv.Value);
+                if (t == typeof(Guid))
+                    p.setter(obj, CreateGuid((string)o));
+            }
+        }
+
+        static int CreateInteger(out int num, string s, int index, int count)
+        {
+            num = 0;
+            bool neg = false;
+            for (int x = 0; x < count; x++, index++)
+            {
+                char cc = s[index];
+
+                if (cc == '-')
+                    neg = true;
+                else if (cc == '+')
+                    neg = false;
+                else
+                {
+                    num *= 10;
+                    num += (int)(cc - '0');
+                }
+            }
+            if (neg) num = -num;
+
+            return num;
+        }
+
+        internal static long CreateLong(out long num, char[] s, int index, int count)
+        {
+            num = 0;
+            bool neg = false;
+            for (int x = 0; x < count; x++, index++)
+            {
+                char cc = s[index];
+
+                if (cc == '-')
+                    neg = true;
+                else if (cc == '+')
+                    neg = false;
+                else
+                {
+                    num *= 10;
+                    num += (int)(cc - '0');
+                }
+            }
+            if (neg) num = -num;
+
+            return num;
+        }
+
+        private object CreateEnum(Type pt, string v)
+        {
+            // TODO : optimize create enum
+#if !SILVERLIGHT
+            return Enum.Parse(pt, v);
+#else
+            return Enum.Parse(pt, v, true);
+#endif
+        }
+
+        private Guid CreateGuid(string s)
+        {
+            if (s.Length > 30)
+                return new Guid(s);
+            else
+                return new Guid(Convert.FromBase64String(s));
+        }
+
+        private DateTime CreateDateTime(string value)
+        {
+            bool utc = false;
+            //                   0123456789012345678
+            // datetime format = yyyy-MM-dd HH:mm:ss
+            int year;// = (int)CreateLong(value.Substring(0, 4));
+            int month;// = (int)CreateLong(value.Substring(5, 2));
+            int day;// = (int)CreateLong(value.Substring(8, 2));
+            int hour;// = (int)CreateLong(value.Substring(11, 2));
+            int min;// = (int)CreateLong(value.Substring(14, 2));
+            int sec;// = (int)CreateLong(value.Substring(17, 2));
+            CreateInteger(out year, value, 0, 4);
+            CreateInteger(out month, value, 5, 2);
+            CreateInteger(out day, value, 8, 2);
+            CreateInteger(out hour, value, 11, 2);
+            CreateInteger(out min, value, 14, 2);
+            CreateInteger(out sec, value, 17, 2);
+
+            //if (value.EndsWith("Z"))
+            if (value[value.Length - 1] == 'Z')
+                utc = true;
+
+            if (_params.UseUTCDateTime == false && utc == false)
+                return new DateTime(year, month, day, hour, min, sec);
+            else
+                return new DateTime(year, month, day, hour, min, sec, DateTimeKind.Utc).ToLocalTime();
+        }
+
+        private object CreateArray(List<object> data, Type pt, Type bt, Dictionary<string, object> globalTypes)
+        {
+            Array col = Array.CreateInstance(bt, data.Count);
+            // create an array of objects
+            for (int i = 0; i < data.Count; i++)// each (object ob in data)
+            {
+                object ob = data[i];
+                if (ob is IDictionary)
+                    //col.SetValue(ParseDictionary((Dictionary<string, object>)ob, globalTypes, bt, null), i); // Original line
+                    col.SetValue(RootDictionary(ob, typeof(Dictionary<string,object>)), i);	// Change for UE4
+                else
+                    col.SetValue(ChangeType(ob, bt), i);
+            }
+
+            return col;
+        }
+
+
+        private object CreateGenericList(List<object> data, Type pt, Type bt, Dictionary<string, object> globalTypes)
+        {
+            IList col = (IList)Reflection.Instance.FastCreateInstance(pt);
+            // create an array of objects
+            foreach (object ob in data)
+            {
+                if (ob is IDictionary)
+                    col.Add(ParseDictionary((Dictionary<string, object>)ob, globalTypes, bt, null));
+
+                else if (ob is List<object>)
+                    col.Add(((List<object>)ob).ToArray());
+
+                else
+                    col.Add(ChangeType(ob, bt));
+            }
+            return col;
+        }
+
+        private object CreateStringKeyDictionary(Dictionary<string, object> reader, Type pt, Type[] types, Dictionary<string, object> globalTypes)
+        {
+            var col = (IDictionary)Reflection.Instance.FastCreateInstance(pt);
+            //Type t1 = null;
+            Type t2 = null;
+            if (types != null)
+            {
+                //t1 = types[0];
+                t2 = types[1];
+            }
+
+            foreach (KeyValuePair<string, object> values in reader)
+            {
+                var key = values.Key;//ChangeType(values.Key, t1);
+                object val = null;
+                if (values.Value is Dictionary<string, object>)
+                    val = ParseDictionary((Dictionary<string, object>)values.Value, globalTypes, t2, null);
+                else
+                    val = ChangeType(values.Value, t2);
+                col.Add(key, val);
+            }
+
+            return col;
+        }
+
+        private object CreateDictionary(List<object> reader, Type pt, Type[] types, Dictionary<string, object> globalTypes)
+        {
+            IDictionary col = (IDictionary)Reflection.Instance.FastCreateInstance(pt);
+            Type t1 = null;
+            Type t2 = null;
+            if (types != null)
+            {
+                t1 = types[0];
+                t2 = types[1];
+            }
+
+            foreach (Dictionary<string, object> values in reader)
+            {
+                object key = values["k"];
+                object val = values["v"];
+
+                if (key is Dictionary<string, object>)
+                    key = ParseDictionary((Dictionary<string, object>)key, globalTypes, t1, null);
+                else
+                    key = ChangeType(key, t1);
+
+                if (val is Dictionary<string, object>)
+                    val = ParseDictionary((Dictionary<string, object>)val, globalTypes, t2, null);
+                else
+                    val = ChangeType(val, t2);
+
+                col.Add(key, val);
+            }
+
+            return col;
+        }
+
+        private Type GetChangeType(Type conversionType)
+        {
+            if (conversionType.IsGenericType && conversionType.GetGenericTypeDefinition().Equals(typeof(Nullable<>)))
+                return conversionType.GetGenericArguments()[0];
+
+            return conversionType;
+        }
+
+#if !SILVERLIGHT
+        private DataSet CreateDataset(Dictionary<string, object> reader, Dictionary<string, object> globalTypes)
+        {
+            DataSet ds = new DataSet();
+            ds.EnforceConstraints = false;
+            ds.BeginInit();
+
+            // read dataset schema here
+            ReadSchema(reader, ds, globalTypes);
+
+            foreach (KeyValuePair<string, object> pair in reader)
+            {
+                if (pair.Key == "$type" || pair.Key == "$schema") continue;
+
+                List<object> rows = (List<object>)pair.Value;
+                if (rows == null) continue;
+
+                DataTable dt = ds.Tables[pair.Key];
+                ReadDataTable(rows, dt);
+            }
+
+            ds.EndInit();
+
+            return ds;
+        }
+
+        private void ReadSchema(Dictionary<string, object> reader, DataSet ds, Dictionary<string, object> globalTypes)
+        {
+            var schema = reader["$schema"];
+
+            if (schema is string)
+            {
+                TextReader tr = new StringReader((string)schema);
+                ds.ReadXmlSchema(tr);
+            }
+            else
+            {
+                DatasetSchema ms = (DatasetSchema)ParseDictionary((Dictionary<string, object>)schema, globalTypes, typeof(DatasetSchema), null);
+                ds.DataSetName = ms.Name;
+                for (int i = 0; i < ms.Info.Count; i += 3)
+                {
+                    if (ds.Tables.Contains(ms.Info[i]) == false)
+                        ds.Tables.Add(ms.Info[i]);
+                    ds.Tables[ms.Info[i]].Columns.Add(ms.Info[i + 1], Type.GetType(ms.Info[i + 2]));
+                }
+            }
+        }
+
+        private void ReadDataTable(List<object> rows, DataTable dt)
+        {
+            dt.BeginInit();
+            dt.BeginLoadData();
+            List<int> guidcols = new List<int>();
+            List<int> datecol = new List<int>();
+
+            foreach (DataColumn c in dt.Columns)
+            {
+                if (c.DataType == typeof(Guid) || c.DataType == typeof(Guid?))
+                    guidcols.Add(c.Ordinal);
+                if (_params.UseUTCDateTime && (c.DataType == typeof(DateTime) || c.DataType == typeof(DateTime?)))
+                    datecol.Add(c.Ordinal);
+            }
+
+            foreach (List<object> row in rows)
+            {
+                object[] v = new object[row.Count];
+                row.CopyTo(v, 0);
+                foreach (int i in guidcols)
+                {
+                    string s = (string)v[i];
+                    if (s != null && s.Length < 36)
+                        v[i] = new Guid(Convert.FromBase64String(s));
+                }
+                if (_params.UseUTCDateTime)
+                {
+                    foreach (int i in datecol)
+                    {
+                        string s = (string)v[i];
+                        if (s != null)
+                            v[i] = CreateDateTime(s);
+                    }
+                }
+                dt.Rows.Add(v);
+            }
+
+            dt.EndLoadData();
+            dt.EndInit();
+        }
+
+        DataTable CreateDataTable(Dictionary<string, object> reader, Dictionary<string, object> globalTypes)
+        {
+            var dt = new DataTable();
+
+            // read dataset schema here
+            var schema = reader["$schema"];
+
+            if (schema is string)
+            {
+                TextReader tr = new StringReader((string)schema);
+                dt.ReadXmlSchema(tr);
+            }
+            else
+            {
+                var ms = (DatasetSchema)this.ParseDictionary((Dictionary<string, object>)schema, globalTypes, typeof(DatasetSchema), null);
+                dt.TableName = ms.Info[0];
+                for (int i = 0; i < ms.Info.Count; i += 3)
+                {
+                    dt.Columns.Add(ms.Info[i + 1], Type.GetType(ms.Info[i + 2]));
+                }
+            }
+
+            foreach (var pair in reader)
+            {
+                if (pair.Key == "$type" || pair.Key == "$schema")
+                    continue;
+
+                var rows = (List<object>)pair.Value;
+                if (rows == null)
+                    continue;
+
+                if (!dt.TableName.Equals(pair.Key, StringComparison.InvariantCultureIgnoreCase))
+                    continue;
+
+                ReadDataTable(rows, dt);
+            }
+
+            return dt;
+        }
+#endif
+        #endregion
+    }
+
+}
\ No newline at end of file
diff --git a/Engine/Source/Programs/UnrealDocTool/APIDocTool/APIDocTool/fastJSON/JsonParser.cs b/Engine/Source/Programs/UnrealDocTool/APIDocTool/APIDocTool/fastJSON/JsonParser.cs
new file mode 100644
--- /dev/null
+++ b/Engine/Source/Programs/UnrealDocTool/APIDocTool/APIDocTool/fastJSON/JsonParser.cs
@@ -0,0 +1,441 @@
+using System;
+using System.Collections;
+using System.Collections.Generic;
+using System.Globalization;
+using System.Text;
+
+namespace fastJSON
+{
+    /// <summary>
+    /// This class encodes and decodes JSON strings.
+    /// Spec. details, see http://www.json.org/
+    /// </summary>
+    internal sealed class JsonParser
+    {
+        enum Token
+        {
+            None = -1,           // Used to denote no Lookahead available
+            Curly_Open,
+            Curly_Close,
+            Squared_Open,
+            Squared_Close,
+            Colon,
+            Comma,
+            String,
+            Number,
+            True,
+            False,
+            Null
+        }
+
+        readonly char[] json;
+        readonly StringBuilder s = new StringBuilder();
+        Token lookAheadToken = Token.None;
+        int index;
+        bool _ignorecase = false;
+
+
+        internal JsonParser(string json, bool ignorecase)
+        {
+            this.json = json.ToCharArray();
+            _ignorecase = ignorecase;
+        }
+
+        public object Decode()
+        {
+            return ParseValue();
+        }
+
+        private Dictionary<string, object> ParseObject()
+        {
+            Dictionary<string, object> table = new Dictionary<string, object>();
+
+            ConsumeToken(); // {
+
+            while (true)
+            {
+                switch (LookAhead())
+                {
+
+                    case Token.Comma:
+                        ConsumeToken();
+                        break;
+
+                    case Token.Curly_Close:
+                        ConsumeToken();
+                        return table;
+
+                    default:
+                        {
+
+                            // name
+                            string name = ParseString();
+                            if (_ignorecase)
+                                name = name.ToLower();
+
+                            // :
+                            if (NextToken() != Token.Colon)
+                            {
+                                throw new Exception("Expected colon at index " + index);
+                            }
+
+                            // value
+                            object value = ParseValue();
+
+                            table[name] = value;
+                        }
+                        break;
+                }
+            }
+        }
+
+        private List<object> ParseArray()
+        {
+            List<object> array = new List<object>();
+            ConsumeToken(); // [
+
+            while (true)
+            {
+                switch (LookAhead())
+                {
+
+                    case Token.Comma:
+                        ConsumeToken();
+                        break;
+
+                    case Token.Squared_Close:
+                        ConsumeToken();
+                        return array;
+
+                    default:
+                        array.Add(ParseValue());
+                        break;
+                }
+            }
+        }
+
+        private object ParseValue()
+        {
+            switch (LookAhead())
+            {
+                case Token.Number:
+                    return ParseNumber();
+
+                case Token.String:
+                    return ParseString();
+
+                case Token.Curly_Open:
+                    return ParseObject();
+
+                case Token.Squared_Open:
+                    return ParseArray();
+
+                case Token.True:
+                    ConsumeToken();
+                    return true;
+
+                case Token.False:
+                    ConsumeToken();
+                    return false;
+
+                case Token.Null:
+                    ConsumeToken();
+                    return null;
+            }
+
+            throw new Exception("Unrecognized token at index" + index);
+        }
+
+        private string ParseString()
+        {
+            ConsumeToken(); // "
+
+            s.Length = 0;
+
+            int runIndex = -1;
+
+            while (index < json.Length)
+            {
+                var c = json[index++];
+
+                if (c == '"')
+                {
+                    if (runIndex != -1)
+                    {
+                        if (s.Length == 0)
+                            return new string(json, runIndex, index - runIndex - 1);
+
+                        s.Append(json, runIndex, index - runIndex - 1);
+                    }
+                    return s.ToString();
+                }
+
+                if (c != '\\')
+                {
+                    if (runIndex == -1)
+                        runIndex = index - 1;
+
+                    continue;
+                }
+
+                if (index == json.Length) break;
+
+                if (runIndex != -1)
+                {
+                    s.Append(json, runIndex, index - runIndex - 1);
+                    runIndex = -1;
+                }
+
+                switch (json[index++])
+                {
+                    case '"':
+                        s.Append('"');
+                        break;
+
+                    case '\\':
+                        s.Append('\\');
+                        break;
+
+                    case '/':
+                        s.Append('/');
+                        break;
+
+                    case 'b':
+                        s.Append('\b');
+                        break;
+
+                    case 'f':
+                        s.Append('\f');
+                        break;
+
+                    case 'n':
+                        s.Append('\n');
+                        break;
+
+                    case 'r':
+                        s.Append('\r');
+                        break;
+
+                    case 't':
+                        s.Append('\t');
+                        break;
+
+                    case 'u':
+                        {
+                            int remainingLength = json.Length - index;
+                            if (remainingLength < 4) break;
+
+                            // parse the 32 bit hex into an integer codepoint
+                            uint codePoint = ParseUnicode(json[index], json[index + 1], json[index + 2], json[index + 3]);
+                            s.Append((char)codePoint);
+
+                            // skip 4 chars
+                            index += 4;
+                        }
+                        break;
+                }
+            }
+
+            throw new Exception("Unexpectedly reached end of string");
+        }
+
+        private uint ParseSingleChar(char c1, uint multipliyer)
+        {
+            uint p1 = 0;
+            if (c1 >= '0' && c1 <= '9')
+                p1 = (uint)(c1 - '0') * multipliyer;
+            else if (c1 >= 'A' && c1 <= 'F')
+                p1 = (uint)((c1 - 'A') + 10) * multipliyer;
+            else if (c1 >= 'a' && c1 <= 'f')
+                p1 = (uint)((c1 - 'a') + 10) * multipliyer;
+            return p1;
+        }
+
+        private uint ParseUnicode(char c1, char c2, char c3, char c4)
+        {
+            uint p1 = ParseSingleChar(c1, 0x1000);
+            uint p2 = ParseSingleChar(c2, 0x100);
+            uint p3 = ParseSingleChar(c3, 0x10);
+            uint p4 = ParseSingleChar(c4, 1);
+
+            return p1 + p2 + p3 + p4;
+        }
+
+        private long CreateLong(string s)
+        {
+            long num = 0;
+            bool neg = false;
+            foreach (char cc in s)
+            {
+                if (cc == '-')
+                    neg = true;
+                else if (cc == '+')
+                    neg = false;
+                else
+                {
+                    num *= 10;
+                    num += (int)(cc - '0');
+                }
+            }
+
+            return neg ? -num : num;
+        }
+
+        private object ParseNumber()
+        {
+            ConsumeToken();
+
+            // Need to start back one place because the first digit is also a token and would have been consumed
+            var startIndex = index - 1;
+            bool dec = false;
+            do
+            {
+                if (index == json.Length)
+                    break;
+                var c = json[index];
+
+                if ((c >= '0' && c <= '9') || c == '.' || c == '-' || c == '+' || c == 'e' || c == 'E')
+                {
+                    if (c == '.' || c == 'e' || c == 'E')
+                        dec = true;
+                    if (++index == json.Length)
+                        break;                        //throw new Exception("Unexpected end of string whilst parsing number");
+                    continue;
+                }
+                break;
+            } while (true);
+
+			if (dec)
+			{
+				string s = new string(json, startIndex, index - startIndex);
+				return double.Parse(s, NumberFormatInfo.InvariantInfo);
+			}
+			long num;
+			return JSON.CreateLong(out num, json, startIndex, index - startIndex);
+        }
+
+        private Token LookAhead()
+        {
+            if (lookAheadToken != Token.None) return lookAheadToken;
+
+            return lookAheadToken = NextTokenCore();
+        }
+
+        private void ConsumeToken()
+        {
+            lookAheadToken = Token.None;
+        }
+
+        private Token NextToken()
+        {
+            var result = lookAheadToken != Token.None ? lookAheadToken : NextTokenCore();
+
+            lookAheadToken = Token.None;
+
+            return result;
+        }
+
+        private Token NextTokenCore()
+        {
+            char c;
+
+            // Skip past whitespace
+            do
+            {
+                c = json[index];
+
+                if (c > ' ') break;
+                if (c != ' ' && c != '\t' && c != '\n' && c != '\r') break;
+
+            } while (++index < json.Length);
+
+            if (index == json.Length)
+            {
+                throw new Exception("Reached end of string unexpectedly");
+            }
+
+            c = json[index];
+
+            index++;
+
+            //if (c >= '0' && c <= '9')
+            //    return Token.Number;
+
+            switch (c)
+            {
+                case '{':
+                    return Token.Curly_Open;
+
+                case '}':
+                    return Token.Curly_Close;
+
+                case '[':
+                    return Token.Squared_Open;
+
+                case ']':
+                    return Token.Squared_Close;
+
+                case ',':
+                    return Token.Comma;
+
+                case '"':
+                    return Token.String;
+
+                case '0':
+                case '1':
+                case '2':
+                case '3':
+                case '4':
+                case '5':
+                case '6':
+                case '7':
+                case '8':
+                case '9':
+                case '-':
+                case '+':
+                case '.':
+                    return Token.Number;
+
+                case ':':
+                    return Token.Colon;
+
+                case 'f':
+                    if (json.Length - index >= 4 &&
+                        json[index + 0] == 'a' &&
+                        json[index + 1] == 'l' &&
+                        json[index + 2] == 's' &&
+                        json[index + 3] == 'e')
+                    {
+                        index += 4;
+                        return Token.False;
+                    }
+                    break;
+
+                case 't':
+                    if (json.Length - index >= 3 &&
+                        json[index + 0] == 'r' &&
+                        json[index + 1] == 'u' &&
+                        json[index + 2] == 'e')
+                    {
+                        index += 3;
+                        return Token.True;
+                    }
+                    break;
+
+                case 'n':
+                    if (json.Length - index >= 3 &&
+                        json[index + 0] == 'u' &&
+                        json[index + 1] == 'l' &&
+                        json[index + 2] == 'l')
+                    {
+                        index += 3;
+                        return Token.Null;
+                    }
+                    break;
+
+            }
+
+            throw new Exception("Could not find token at index " + --index);
+        }
+    }
+}
diff --git a/Engine/Source/Programs/UnrealDocTool/APIDocTool/APIDocTool/fastJSON/JsonSerializer.cs b/Engine/Source/Programs/UnrealDocTool/APIDocTool/APIDocTool/fastJSON/JsonSerializer.cs
new file mode 100644
--- /dev/null
+++ b/Engine/Source/Programs/UnrealDocTool/APIDocTool/APIDocTool/fastJSON/JsonSerializer.cs
@@ -0,0 +1,505 @@
+﻿using System;
+using System.Collections;
+using System.Collections.Generic;
+#if !SILVERLIGHT
+using System.Data;
+#endif
+using System.Globalization;
+using System.IO;
+using System.Text;
+
+namespace fastJSON
+{
+    internal sealed class JSONSerializer
+    {
+        private StringBuilder _output = new StringBuilder();
+        private StringBuilder _before = new StringBuilder();
+        readonly int _MAX_DEPTH = 10;
+        int _current_depth = 0;
+        private Dictionary<string, int> _globalTypes = new Dictionary<string, int>();
+        private JSONParameters _params;
+        private bool _useEscapedUnicode = false;
+
+        internal JSONSerializer(JSONParameters param)
+        {
+            _params = param;
+            _useEscapedUnicode = _params.UseEscapedUnicode;
+        }
+
+        internal string ConvertToJSON(object obj)
+        {
+            WriteValue(obj);
+
+            string str = "";
+            if (_params.UsingGlobalTypes && _globalTypes != null && _globalTypes.Count > 0)
+            {
+                StringBuilder sb = _before;
+                sb.Append("\"$types\":{");
+                bool pendingSeparator = false;
+                foreach (var kv in _globalTypes)
+                {
+                    if (pendingSeparator) sb.Append(',');
+                    pendingSeparator = true;
+                    sb.Append('\"');
+                    sb.Append(kv.Key);
+                    sb.Append("\":\"");
+                    sb.Append(kv.Value);
+                    sb.Append('\"');
+                }
+                sb.Append("},");
+                sb.Append(_output.ToString());
+                str = sb.ToString();
+            }
+            else
+                str = _output.ToString();
+
+            return str;
+        }
+
+        private void WriteValue(object obj)
+        {
+            if (obj == null || obj is DBNull)
+                _output.Append("null");
+
+            else if (obj is string || obj is char)
+                WriteString(obj.ToString());
+
+            else if (obj is Guid)
+                WriteGuid((Guid)obj);
+
+            else if (obj is bool)
+                _output.Append(((bool)obj) ? "true" : "false"); // conform to standard
+
+            else if (
+                obj is int || obj is long || obj is double ||
+                obj is decimal || obj is float ||
+                obj is byte || obj is short ||
+                obj is sbyte || obj is ushort ||
+                obj is uint || obj is ulong
+            )
+                _output.Append(((IConvertible)obj).ToString(NumberFormatInfo.InvariantInfo));
+
+            else if (obj is DateTime)
+                WriteDateTime((DateTime)obj);
+
+            else if (obj is IDictionary && obj.GetType().IsGenericType && obj.GetType().GetGenericArguments()[0] == typeof(string))
+                WriteStringDictionary((IDictionary)obj);
+
+            else if (obj is IDictionary)
+                WriteDictionary((IDictionary)obj);
+#if !SILVERLIGHT
+            else if (obj is DataSet)
+                WriteDataset((DataSet)obj);
+
+            else if (obj is DataTable)
+                this.WriteDataTable((DataTable)obj);
+#endif
+            else if (obj is byte[])
+                WriteBytes((byte[])obj);
+
+            else if (obj is IEnumerable)// Array || obj is IList || obj is ICollection)
+                WriteArray((IEnumerable)obj);
+
+            else if (obj is Enum)
+                WriteEnum((Enum)obj);
+
+            else if (JSON.Instance.IsTypeRegistered(obj.GetType()))
+                WriteCustom(obj);
+
+            else
+                WriteObject(obj);
+        }
+
+        private void WriteCustom(object obj)
+        {
+            Serialize s;
+            JSON.Instance._customSerializer.TryGetValue(obj.GetType(), out s);
+            WriteStringFast(s(obj));
+        }
+
+        private void WriteEnum(Enum e)
+        {
+            // TODO : optimize enum write
+            WriteStringFast(e.ToString());
+        }
+
+        private void WriteGuid(Guid g)
+        {
+            if (_params.UseFastGuid == false)
+                WriteStringFast(g.ToString());
+            else
+                WriteBytes(g.ToByteArray());
+        }
+
+        private void WriteBytes(byte[] bytes)
+        {
+#if !SILVERLIGHT
+            WriteStringFast(Convert.ToBase64String(bytes, 0, bytes.Length, Base64FormattingOptions.None));
+#else
+            WriteStringFast(Convert.ToBase64String(bytes, 0, bytes.Length));
+#endif
+        }
+
+        private void WriteDateTime(DateTime dateTime)
+        {
+            // datetime format standard : yyyy-MM-dd HH:mm:ss
+            DateTime dt = dateTime;
+            if (_params.UseUTCDateTime)
+                dt = dateTime.ToUniversalTime();
+
+            _output.Append('\"');
+            _output.Append(dt.Year.ToString("0000", NumberFormatInfo.InvariantInfo));
+            _output.Append('-');
+            _output.Append(dt.Month.ToString("00", NumberFormatInfo.InvariantInfo));
+            _output.Append('-');
+            _output.Append(dt.Day.ToString("00", NumberFormatInfo.InvariantInfo));
+            _output.Append(' ');
+            _output.Append(dt.Hour.ToString("00", NumberFormatInfo.InvariantInfo));
+            _output.Append(':');
+            _output.Append(dt.Minute.ToString("00", NumberFormatInfo.InvariantInfo));
+            _output.Append(':');
+            _output.Append(dt.Second.ToString("00", NumberFormatInfo.InvariantInfo));
+
+            if (_params.UseUTCDateTime)
+                _output.Append('Z');
+
+            _output.Append('\"');
+        }
+
+#if !SILVERLIGHT
+        private DatasetSchema GetSchema(DataTable ds)
+        {
+            if (ds == null) return null;
+
+            DatasetSchema m = new DatasetSchema();
+            m.Info = new List<string>();
+            m.Name = ds.TableName;
+
+            foreach (DataColumn c in ds.Columns)
+            {
+                m.Info.Add(ds.TableName);
+                m.Info.Add(c.ColumnName);
+                m.Info.Add(c.DataType.ToString());
+            }
+            // FEATURE : serialize relations and constraints here
+
+            return m;
+        }
+
+        private DatasetSchema GetSchema(DataSet ds)
+        {
+            if (ds == null) return null;
+
+            DatasetSchema m = new DatasetSchema();
+            m.Info = new List<string>();
+            m.Name = ds.DataSetName;
+
+            foreach (DataTable t in ds.Tables)
+            {
+                foreach (DataColumn c in t.Columns)
+                {
+                    m.Info.Add(t.TableName);
+                    m.Info.Add(c.ColumnName);
+                    m.Info.Add(c.DataType.ToString());
+                }
+            }
+            // FEATURE : serialize relations and constraints here
+
+            return m;
+        }
+
+        private string GetXmlSchema(DataTable dt)
+        {
+            using (var writer = new StringWriter())
+            {
+                dt.WriteXmlSchema(writer);
+                return dt.ToString();
+            }
+        }
+
+        private void WriteDataset(DataSet ds)
+        {
+            _output.Append('{');
+            if (_params.UseExtensions)
+            {
+                WritePair("$schema", _params.UseOptimizedDatasetSchema ? (object)GetSchema(ds) : ds.GetXmlSchema());
+                _output.Append(',');
+            }
+            bool tablesep = false;
+            foreach (DataTable table in ds.Tables)
+            {
+                if (tablesep) _output.Append(',');
+                tablesep = true;
+                WriteDataTableData(table);
+            }
+            // end dataset
+            _output.Append('}');
+        }
+
+        private void WriteDataTableData(DataTable table)
+        {
+            _output.Append('\"');
+            _output.Append(table.TableName);
+            _output.Append("\":[");
+            DataColumnCollection cols = table.Columns;
+            bool rowseparator = false;
+            foreach (DataRow row in table.Rows)
+            {
+                if (rowseparator) _output.Append(',');
+                rowseparator = true;
+                _output.Append('[');
+
+                bool pendingSeperator = false;
+                foreach (DataColumn column in cols)
+                {
+                    if (pendingSeperator) _output.Append(',');
+                    WriteValue(row[column]);
+                    pendingSeperator = true;
+                }
+                _output.Append(']');
+            }
+
+            _output.Append(']');
+        }
+
+        void WriteDataTable(DataTable dt)
+        {
+            this._output.Append('{');
+            if (_params.UseExtensions)
+            {
+                this.WritePair("$schema", _params.UseOptimizedDatasetSchema ? (object)this.GetSchema(dt) : this.GetXmlSchema(dt));
+                this._output.Append(',');
+            }
+
+            WriteDataTableData(dt);
+
+            // end datatable
+            this._output.Append('}');
+        }
+#endif
+
+        bool _TypesWritten = false;
+        private void WriteObject(object obj)
+        {
+            if (_params.UsingGlobalTypes == false)
+                _output.Append('{');
+            else
+            {
+                if (_TypesWritten == false)
+                {
+                    _output.Append('{');
+                    _before = _output;
+                    _output = new StringBuilder();
+                }
+                else
+                    _output.Append('{');
+            }
+            _TypesWritten = true;
+            _current_depth++;
+            if (_current_depth > _MAX_DEPTH)
+                throw new Exception("Serializer encountered maximum depth of " + _MAX_DEPTH);
+
+
+            Dictionary<string, string> map = new Dictionary<string, string>();
+            Type t = obj.GetType();
+            bool append = false;
+            if (_params.UseExtensions)
+            {
+                if (_params.UsingGlobalTypes == false)
+                    WritePairFast("$type", Reflection.Instance.GetTypeAssemblyName(t));
+                else
+                {
+                    int dt = 0;
+                    string ct = Reflection.Instance.GetTypeAssemblyName(t);
+                    if (_globalTypes.TryGetValue(ct, out dt) == false)
+                    {
+                        dt = _globalTypes.Count + 1;
+                        _globalTypes.Add(ct, dt);
+                    }
+                    WritePairFast("$type", dt.ToString());
+                }
+                append = true;
+            }
+
+            List<Getters> g = Reflection.Instance.GetGetters(t);
+
+            foreach (var p in g)
+            {
+                object o = p.Getter(obj);
+                if ((o == null || o is DBNull) && _params.SerializeNullValues == false)
+                {
+                    //append = false;
+                }
+                else
+                {
+                    if (append)
+                        _output.Append(',');
+
+                    WritePair(p.Name, o);
+                    if (o != null && _params.UseExtensions)
+                    {
+                        Type tt = o.GetType();
+                        if (tt == typeof(System.Object))
+                            map.Add(p.Name, tt.ToString());
+                    }
+                    append = true;
+                }
+            }
+            if (map.Count > 0 && _params.UseExtensions)
+            {
+                _output.Append(",\"$map\":");
+                WriteStringDictionary(map);
+            }
+            _current_depth--;
+            _output.Append('}');
+            _current_depth--;
+        }
+
+        private void WritePairFast(string name, string value)
+        {
+            if ((value == null) && _params.SerializeNullValues == false)
+                return;
+            WriteStringFast(name);
+
+            _output.Append(':');
+
+            WriteStringFast(value);
+        }
+
+        private void WritePair(string name, object value)
+        {
+            if ((value == null || value is DBNull) && _params.SerializeNullValues == false)
+                return;
+            WriteStringFast(name);
+
+            _output.Append(':');
+
+            WriteValue(value);
+        }
+
+        private void WriteArray(IEnumerable array)
+        {
+            _output.Append('[');
+
+            bool pendingSeperator = false;
+
+            foreach (object obj in array)
+            {
+                if (pendingSeperator) _output.Append(',');
+
+                WriteValue(obj);
+
+                pendingSeperator = true;
+            }
+            _output.Append(']');
+        }
+
+        private void WriteStringDictionary(IDictionary dic)
+        {
+            _output.Append('{');
+
+            bool pendingSeparator = false;
+
+            foreach (DictionaryEntry entry in dic)
+            {
+                if (pendingSeparator) _output.Append(',');
+
+                WritePair((string)entry.Key, entry.Value);
+
+                pendingSeparator = true;
+            }
+            _output.Append('}');
+        }
+
+        private void WriteDictionary(IDictionary dic)
+        {
+            _output.Append('[');
+
+            bool pendingSeparator = false;
+
+            foreach (DictionaryEntry entry in dic)
+            {
+                if (pendingSeparator) _output.Append(',');
+                _output.Append('{');
+                WritePair("k", entry.Key);
+                _output.Append(",");
+                WritePair("v", entry.Value);
+                _output.Append('}');
+
+                pendingSeparator = true;
+            }
+            _output.Append(']');
+        }
+
+        private void WriteStringFast(string s)
+        {
+            _output.Append('\"');
+            _output.Append(s);
+            _output.Append('\"');
+        }
+
+        private void WriteString(string s)
+        {
+            _output.Append('\"');
+
+            int runIndex = -1;
+
+            for (var index = 0; index < s.Length; ++index)
+            {
+                var c = s[index];
+
+                if (_useEscapedUnicode)
+                {
+                    if (c >= ' ' && c < 128 && c != '\"' && c != '\\')
+                    {
+                        if (runIndex == -1)
+                            runIndex = index;
+
+                        continue;
+                    }
+                }
+                else
+                {
+                    if (c != '\t' && c != '\n' && c != '\r' && c != '\"' && c != '\\')// && c != ':' && c!=',')
+                    {
+                        if (runIndex == -1)
+                            runIndex = index;
+
+                        continue;
+                    }
+                }
+
+                if (runIndex != -1)
+                {
+                    _output.Append(s, runIndex, index - runIndex);
+                    runIndex = -1;
+                }
+
+                switch (c)
+                {
+                    case '\t': _output.Append("\\t"); break;
+                    case '\r': _output.Append("\\r"); break;
+                    case '\n': _output.Append("\\n"); break;
+                    case '"':
+                    case '\\': _output.Append('\\'); _output.Append(c); break;
+                    default:
+                        if (_useEscapedUnicode)
+                        {
+                            _output.Append("\\u");
+                            _output.Append(((int)c).ToString("X4", NumberFormatInfo.InvariantInfo));
+                        }
+                        else
+                            _output.Append(c);
+
+                        break;
+                }
+            }
+
+            if (runIndex != -1)
+                _output.Append(s, runIndex, s.Length - runIndex);
+
+
+            _output.Append('\"');
+        }
+    }
+}
diff --git a/Engine/Source/Programs/UnrealDocTool/APIDocTool/APIDocTool/fastJSON/Reflection.cs b/Engine/Source/Programs/UnrealDocTool/APIDocTool/APIDocTool/fastJSON/Reflection.cs
new file mode 100644
--- /dev/null
+++ b/Engine/Source/Programs/UnrealDocTool/APIDocTool/APIDocTool/fastJSON/Reflection.cs
@@ -0,0 +1,305 @@
+﻿using System;
+using System.Collections.Generic;
+using System.Text;
+using System.Reflection.Emit;
+using System.Reflection;
+using System.Collections;
+
+namespace fastJSON
+{
+    internal struct Getters
+    {
+        public string Name;
+        public Reflection.GenericGetter Getter;
+        public Type propertyType;
+    }
+
+    internal sealed class Reflection
+    {
+        public readonly static Reflection Instance = new Reflection();
+        private Reflection()
+        {
+        }
+
+        public bool ShowReadOnlyProperties = false;
+        internal delegate object GenericSetter(object target, object value);
+        internal delegate object GenericGetter(object obj);
+        private delegate object CreateObject();
+
+        private SafeDictionary<Type, string> _tyname = new SafeDictionary<Type, string>();
+        private SafeDictionary<string, Type> _typecache = new SafeDictionary<string, Type>();
+        private SafeDictionary<Type, CreateObject> _constrcache = new SafeDictionary<Type, CreateObject>();
+        private SafeDictionary<Type, List<Getters>> _getterscache = new SafeDictionary<Type, List<Getters>>();
+
+        #region [   PROPERTY GET SET   ]
+        internal string GetTypeAssemblyName(Type t)
+        {
+            string val = "";
+            if (_tyname.TryGetValue(t, out val))
+                return val;
+            else
+            {
+                string s = t.AssemblyQualifiedName;
+                _tyname.Add(t, s);
+                return s;
+            }
+        }
+
+        internal Type GetTypeFromCache(string typename)
+        {
+            Type val = null;
+            if (_typecache.TryGetValue(typename, out val))
+                return val;
+            else
+            {
+                Type t = Type.GetType(typename);
+                _typecache.Add(typename, t);
+                return t;
+            }
+        }
+
+        internal object FastCreateInstance(Type objtype)
+        {
+            try
+            {
+                CreateObject c = null;
+                if (_constrcache.TryGetValue(objtype, out c))
+                {
+                    return c();
+                }
+                else
+                {
+                    if (objtype.IsClass)
+                    {
+                        DynamicMethod dynMethod = new DynamicMethod("_", objtype, null);
+                        ILGenerator ilGen = dynMethod.GetILGenerator();
+                        ilGen.Emit(OpCodes.Newobj, objtype.GetConstructor(Type.EmptyTypes));
+                        ilGen.Emit(OpCodes.Ret);
+                        c = (CreateObject)dynMethod.CreateDelegate(typeof(CreateObject));
+                        _constrcache.Add(objtype, c);
+                    }
+                    else // structs
+                    {
+                        DynamicMethod dynMethod = new DynamicMethod("_", typeof(object), null);
+                        ILGenerator ilGen = dynMethod.GetILGenerator();
+                        var lv = ilGen.DeclareLocal(objtype);
+                        ilGen.Emit(OpCodes.Ldloca_S, lv);
+                        ilGen.Emit(OpCodes.Initobj, objtype);
+                        ilGen.Emit(OpCodes.Ldloc_0);
+                        ilGen.Emit(OpCodes.Box, objtype);
+                        ilGen.Emit(OpCodes.Ret);
+                        c = (CreateObject)dynMethod.CreateDelegate(typeof(CreateObject));
+                        _constrcache.Add(objtype, c);
+                    }
+                    return c();
+                }
+            }
+            catch (Exception exc)
+            {
+                throw new Exception(string.Format("Failed to fast create instance for type '{0}' from assemebly '{1}'",
+                    objtype.FullName, objtype.AssemblyQualifiedName), exc);
+            }
+        }
+
+        internal static GenericSetter CreateSetField(Type type, FieldInfo fieldInfo)
+        {
+            Type[] arguments = new Type[2];
+            arguments[0] = arguments[1] = typeof(object);
+
+            DynamicMethod dynamicSet = new DynamicMethod("_", typeof(object), arguments, type);
+
+            ILGenerator il = dynamicSet.GetILGenerator();
+
+            if (!type.IsClass) // structs
+            {
+                var lv = il.DeclareLocal(type);
+                il.Emit(OpCodes.Ldarg_0);
+                il.Emit(OpCodes.Unbox_Any, type);
+                il.Emit(OpCodes.Stloc_0);
+                il.Emit(OpCodes.Ldloca_S, lv);
+                il.Emit(OpCodes.Ldarg_1);
+                if (fieldInfo.FieldType.IsClass)
+                    il.Emit(OpCodes.Castclass, fieldInfo.FieldType);
+                else
+                    il.Emit(OpCodes.Unbox_Any, fieldInfo.FieldType);
+                il.Emit(OpCodes.Stfld, fieldInfo);
+                il.Emit(OpCodes.Ldloc_0);
+                il.Emit(OpCodes.Box, type);
+                il.Emit(OpCodes.Ret);
+            }
+            else
+            {
+                il.Emit(OpCodes.Ldarg_0);
+                il.Emit(OpCodes.Ldarg_1);
+                if (fieldInfo.FieldType.IsValueType)
+                    il.Emit(OpCodes.Unbox_Any, fieldInfo.FieldType);
+                il.Emit(OpCodes.Stfld, fieldInfo);
+                il.Emit(OpCodes.Ldarg_0);
+                il.Emit(OpCodes.Ret);
+            }
+            return (GenericSetter)dynamicSet.CreateDelegate(typeof(GenericSetter));
+        }
+
+        internal static GenericSetter CreateSetMethod(Type type, PropertyInfo propertyInfo)
+        {
+            MethodInfo setMethod = propertyInfo.GetSetMethod();
+            if (setMethod == null)
+                return null;
+
+            Type[] arguments = new Type[2];
+            arguments[0] = arguments[1] = typeof(object);
+
+            DynamicMethod setter = new DynamicMethod("_", typeof(object), arguments);
+            ILGenerator il = setter.GetILGenerator();
+
+            if (!type.IsClass) // structs
+            {
+                var lv = il.DeclareLocal(type);
+                il.Emit(OpCodes.Ldarg_0);
+                il.Emit(OpCodes.Unbox_Any, type);
+                il.Emit(OpCodes.Stloc_0);
+                il.Emit(OpCodes.Ldloca_S, lv);
+                il.Emit(OpCodes.Ldarg_1);
+                if (propertyInfo.PropertyType.IsClass)
+                    il.Emit(OpCodes.Castclass, propertyInfo.PropertyType);
+                else
+                    il.Emit(OpCodes.Unbox_Any, propertyInfo.PropertyType);
+                il.EmitCall(OpCodes.Call, setMethod, null);
+                il.Emit(OpCodes.Ldloc_0);
+                il.Emit(OpCodes.Box, type);
+            }
+            else
+            {
+                il.Emit(OpCodes.Ldarg_0);
+                il.Emit(OpCodes.Castclass, propertyInfo.DeclaringType);
+                il.Emit(OpCodes.Ldarg_1);
+                if (propertyInfo.PropertyType.IsClass)
+                    il.Emit(OpCodes.Castclass, propertyInfo.PropertyType);
+                else
+                    il.Emit(OpCodes.Unbox_Any, propertyInfo.PropertyType);
+                il.EmitCall(OpCodes.Callvirt, setMethod, null);
+                il.Emit(OpCodes.Ldarg_0);
+            }
+
+            il.Emit(OpCodes.Ret);
+
+            return (GenericSetter)setter.CreateDelegate(typeof(GenericSetter));
+        }
+
+        internal static GenericGetter CreateGetField(Type type, FieldInfo fieldInfo)
+        {
+            DynamicMethod dynamicGet = new DynamicMethod("_", typeof(object), new Type[] { typeof(object) }, type);
+
+            ILGenerator il = dynamicGet.GetILGenerator();
+
+            if (!type.IsClass) // structs
+            {
+                var lv = il.DeclareLocal(type);
+                il.Emit(OpCodes.Ldarg_0);
+                il.Emit(OpCodes.Unbox_Any, type);
+                il.Emit(OpCodes.Stloc_0);
+                il.Emit(OpCodes.Ldloca_S, lv);
+                il.Emit(OpCodes.Ldfld, fieldInfo);
+                if (fieldInfo.FieldType.IsValueType)
+                    il.Emit(OpCodes.Box, fieldInfo.FieldType);
+            }
+            else
+            {
+                il.Emit(OpCodes.Ldarg_0);
+                il.Emit(OpCodes.Ldfld, fieldInfo);
+                if (fieldInfo.FieldType.IsValueType)
+                    il.Emit(OpCodes.Box, fieldInfo.FieldType);
+            }
+
+            il.Emit(OpCodes.Ret);
+
+            return (GenericGetter)dynamicGet.CreateDelegate(typeof(GenericGetter));
+        }
+
+        internal static GenericGetter CreateGetMethod(Type type, PropertyInfo propertyInfo)
+        {
+            MethodInfo getMethod = propertyInfo.GetGetMethod();
+            if (getMethod == null)
+                return null;
+
+            DynamicMethod getter = new DynamicMethod("_", typeof(object), new Type[] { typeof(object) }, type);
+
+            ILGenerator il = getter.GetILGenerator();
+
+            if (!type.IsClass) // structs
+            {
+                var lv = il.DeclareLocal(type);
+                il.Emit(OpCodes.Ldarg_0);
+                il.Emit(OpCodes.Unbox_Any, type);
+                il.Emit(OpCodes.Stloc_0);
+                il.Emit(OpCodes.Ldloca_S, lv);
+                il.EmitCall(OpCodes.Call, getMethod, null);
+                if (propertyInfo.PropertyType.IsValueType)
+                    il.Emit(OpCodes.Box, propertyInfo.PropertyType);
+            }
+            else
+            {
+                il.Emit(OpCodes.Ldarg_0);
+                il.Emit(OpCodes.Castclass, propertyInfo.DeclaringType);
+                il.EmitCall(OpCodes.Callvirt, getMethod, null);
+                if (propertyInfo.PropertyType.IsValueType)
+                    il.Emit(OpCodes.Box, propertyInfo.PropertyType);
+            }
+
+            il.Emit(OpCodes.Ret);
+
+            return (GenericGetter)getter.CreateDelegate(typeof(GenericGetter));
+        }
+
+        internal List<Getters> GetGetters(Type type)
+        {
+            List<Getters> val = null;
+            if (_getterscache.TryGetValue(type, out val))
+                return val;
+
+            PropertyInfo[] props = type.GetProperties(BindingFlags.Public | BindingFlags.Instance);
+            List<Getters> getters = new List<Getters>();
+            foreach (PropertyInfo p in props)
+            {
+                if (!p.CanWrite && ShowReadOnlyProperties == false) continue;
+
+                object[] att = p.GetCustomAttributes(typeof(System.Xml.Serialization.XmlIgnoreAttribute), false);
+                if (att != null && att.Length > 0)
+                    continue;
+
+                GenericGetter g = CreateGetMethod(type, p);
+                if (g != null)
+                {
+                    Getters gg = new Getters();
+                    gg.Name = p.Name;
+                    gg.Getter = g;
+                    gg.propertyType = p.PropertyType;
+                    getters.Add(gg);
+                }
+            }
+
+            FieldInfo[] fi = type.GetFields(BindingFlags.Instance | BindingFlags.Public);
+            foreach (var f in fi)
+            {
+                object[] att = f.GetCustomAttributes(typeof(System.Xml.Serialization.XmlIgnoreAttribute), false);
+                if (att != null && att.Length > 0)
+                    continue;
+
+                GenericGetter g = CreateGetField(type, f);
+                if (g != null)
+                {
+                    Getters gg = new Getters();
+                    gg.Name = f.Name;
+                    gg.Getter = g;
+                    gg.propertyType = f.FieldType;
+                    getters.Add(gg);
+                }
+            }
+
+            _getterscache.Add(type, getters);
+            return getters;
+        }
+
+        #endregion
+    }
+}
diff --git a/Engine/Source/Programs/UnrealDocTool/APIDocTool/APIDocTool/fastJSON/SafeDictionary.cs b/Engine/Source/Programs/UnrealDocTool/APIDocTool/APIDocTool/fastJSON/SafeDictionary.cs
new file mode 100644
--- /dev/null
+++ b/Engine/Source/Programs/UnrealDocTool/APIDocTool/APIDocTool/fastJSON/SafeDictionary.cs
@@ -0,0 +1,52 @@
+﻿using System;
+using System.Collections.Generic;
+
+namespace fastJSON
+{
+    internal sealed class SafeDictionary<TKey, TValue>
+    {
+        private readonly object _Padlock = new object();
+        private readonly Dictionary<TKey, TValue> _Dictionary;
+
+        public SafeDictionary(int capacity)
+        {
+            _Dictionary = new Dictionary<TKey, TValue>(capacity);
+        }
+
+        public SafeDictionary()
+        {
+            _Dictionary = new Dictionary<TKey, TValue>();
+        }
+
+        public bool TryGetValue(TKey key, out TValue value)
+        {
+            lock (_Padlock)
+                return _Dictionary.TryGetValue(key, out value);
+        }
+
+        public int Count { get { lock (_Padlock) return _Dictionary.Count; } }
+
+        public TValue this[TKey key]
+        {
+            get
+            {
+                lock (_Padlock)
+                    return _Dictionary[key];
+            }
+            set
+            {
+                lock (_Padlock)
+                    _Dictionary[key] = value;
+            }
+        }
+
+        public void Add(TKey key, TValue value)
+        {
+            lock (_Padlock)
+            {
+                if (_Dictionary.ContainsKey(key) == false)
+                    _Dictionary.Add(key, value);
+            }
+        }
+    }
+}
