Remove confusing warning after already committing to force delete

diff --git a/Engine/Source/Editor/UnrealEd/Private/ObjectTools.cpp b/Engine/Source/Editor/UnrealEd/Private/ObjectTools.cpp
--- a/Engine/Source/Editor/UnrealEd/Private/ObjectTools.cpp
+++ b/Engine/Source/Editor/UnrealEd/Private/ObjectTools.cpp
@@ -1745,7 +1745,6 @@ namespace ObjectTools
 	int32 ForceDeleteObjects( const TArray< UObject* >& InObjectsToDelete, bool ShowConfirmation )
 	{
 		int32 NumDeletedObjects = 0;
-		bool ForceDeleteAll = false;
 
 		// Confirm that the delete was intentional
 		if ( ShowConfirmation && !ShowDeleteConfirmationDialog(InObjectsToDelete) )
@@ -1774,70 +1773,7 @@ namespace ObjectTools
 
 			GEditor->GetSelectedObjects()->Deselect( CurrentObject );
 
-			if( !ForceDeleteAll )
-			{
-				FReferencerInformationList Refs;
-
-				// Check and see whether we are referenced by any objects that won't be garbage collected. 
-				bool bIsReferenced = IsReferenced( CurrentObject, GARBAGE_COLLECTION_KEEPFLAGS, true, &Refs );
-
-				if ( bIsReferenced )
-				{
-					// Create a string list of all referenced properties.
-					// Check if this object is referenced in default properties
-					FString RefObjNames;
-					FString DefaultPropertiesObjNames;
-					ComposeStringOfReferencingObjects( Refs.ExternalReferences, RefObjNames, DefaultPropertiesObjNames );
-					ComposeStringOfReferencingObjects( Refs.InternalReferences, RefObjNames, DefaultPropertiesObjNames );
-
-					FFormatNamedArguments Args;
-					Args.Add( TEXT("ObjectName"), FText::FromString( CurrentObject->GetName() ) );
-					Args.Add( TEXT("ReferencedObjectNames"), FText::FromString( RefObjNames ) );
-					const FText Message = FText::Format( NSLOCTEXT("Core", "Warning_ForceDelete", "Deleting {ObjectName}.\n\nForcing delete on a referenced object is potentially dangerous and could cause data corruption.  The following objects may have invalid references if you proceed:\n {ReferencedObjectNames}.\n\nDo you wish to delete this referenced object?"), Args );
-
-					int32 YesNoCancelReply = FMessageDialog::Open( EAppMsgType::YesNoYesAllNoAll, Message );
-					switch ( YesNoCancelReply )
-					{
-					case EAppReturnType::Yes: // Yes
-						{
-							ObjectsToDelete.Add( CurrentObject );
-							break;
-						}
-
-					case EAppReturnType::YesAll: // Yes to All
-						{
-							ForceDeleteAll = true;
-							ObjectsToDelete.Add( CurrentObject );
-							break;
-						}
-
-					case EAppReturnType::Cancel:
-					case EAppReturnType::No: // No
-						{
-							//Skip to the next object and proceed
-							continue;
-							break;
-						}
-
-					case EAppReturnType::NoAll: // No to All
-						{
-							GWarn->EndSlowTask();
-							return NumDeletedObjects;
-						}
-
-					default:
-						break;
-					}
-				}
-				else
-				{
-					ObjectsToDelete.Add( CurrentObject );
-				}
-			}
-			else
-			{
-				ObjectsToDelete.Add( CurrentObject );
-			}
+			ObjectsToDelete.Add( CurrentObject );
 		}
 
 		// If the current editor world is in this list, transition to a new map and reload the world to finish the delete
