Enable Split Pin functionality by default and remove from experimental settings

diff --git a/Engine/Source/Editor/BlueprintGraph/Private/K2Node.cpp b/Engine/Source/Editor/BlueprintGraph/Private/K2Node.cpp
--- a/Engine/Source/Editor/BlueprintGraph/Private/K2Node.cpp
+++ b/Engine/Source/Editor/BlueprintGraph/Private/K2Node.cpp
@@ -577,7 +577,7 @@ void UK2Node::ReconstructSinglePin(UEdGraphPin* NewPin, UEdGraphPin* OldPin, ERe
 
 bool UK2Node::AllowSplitPins() const
 {
-	return GetDefault<UEditorExperimentalSettings>()->bAllowSplitStructPins;
+	return true;
 }
 
 void UK2Node::ExpandSplitPin(FKismetCompilerContext* CompilerContext, UEdGraph* SourceGraph, UEdGraphPin* Pin)
diff --git a/Engine/Source/Editor/UnrealEd/Classes/Settings/EditorExperimentalSettings.h b/Engine/Source/Editor/UnrealEd/Classes/Settings/EditorExperimentalSettings.h
--- a/Engine/Source/Editor/UnrealEd/Classes/Settings/EditorExperimentalSettings.h
+++ b/Engine/Source/Editor/UnrealEd/Classes/Settings/EditorExperimentalSettings.h
@@ -56,9 +56,6 @@ public:
 	UPROPERTY(EditAnywhere, config, Category=Blueprints, meta=(DisplayName="Draw midpoint arrows in Blueprints"))
 	bool bDrawMidpointArrowsInBlueprints;
 
-	UPROPERTY(EditAnywhere, config, Category=Blueprints, meta=(DisplayName="Allow splitting of struct pins"))
-	bool bAllowSplitStructPins;
-
 	/** Whether to show Audio Streaming options for SoundWaves (disabling will not stop all audio streaming) */
 	UPROPERTY(EditAnywhere, config, Category=Audio)
 	bool bShowAudioStreamingOptions;
