CurveEditor - Refactor curve selection controls.
+ Remove curve selection check boxes.
+ Add controls to show/hide curves, and to lock/unlock curves.
+ Small refactor to paint method to ensure correct draw order.

diff --git a/Engine/Source/Editor/EditorStyle/Private/SlateEditorStyle.cpp b/Engine/Source/Editor/EditorStyle/Private/SlateEditorStyle.cpp
--- a/Engine/Source/Editor/EditorStyle/Private/SlateEditorStyle.cpp
+++ b/Engine/Source/Editor/EditorStyle/Private/SlateEditorStyle.cpp
@@ -778,14 +778,23 @@ void FSlateEditorStyle::FStyle::SetupGeneralStyles()
 	}
 
 	{
-		Set( "CurveEd.TimelineArea", new IMAGE_BRUSH( "Old/White", Icon16x16, FLinearColor(1,1,1,0.25f) ) );
-		Set( "CurveEd.FitHorizontal", new IMAGE_BRUSH( "Icons/FitHorz_16x", Icon16x16 ) );
-		Set( "CurveEd.FitVertical", new IMAGE_BRUSH( "Icons/FitVert_16x", Icon16x16 ) );
-		Set( "CurveEd.CurveKey", new IMAGE_BRUSH( "Common/Key", FVector2D(11.0f, 11.0f) ) );
-		Set( "CurveEd.CurveKeySelected", new IMAGE_BRUSH( "Common/Key", FVector2D(11.0f, 11.0f), SelectionColor ) );
-		Set( "CurveEd.InfoFont", TTF_CORE_FONT( "Fonts/Roboto-Regular", 8 ) );
-		Set( "CurveEd.Tangent", new IMAGE_BRUSH( "Common/Tangent", FVector2D(7.0f, 7.0f), FLinearColor(0.0f, 0.66f, 0.7f) ) );
-		Set( "CurveEd.TangentColor", FLinearColor(0.0f, 0.66f, 0.7f) );
+		Set("CurveEd.TimelineArea", new IMAGE_BRUSH("Old/White", Icon16x16, FLinearColor(1, 1, 1, 0.25f)));
+		Set("CurveEd.FitHorizontal", new IMAGE_BRUSH("Icons/FitHorz_16x", Icon16x16));
+		Set("CurveEd.FitVertical", new IMAGE_BRUSH("Icons/FitVert_16x", Icon16x16));
+		Set("CurveEd.CurveKey", new IMAGE_BRUSH("Common/Key", FVector2D(11.0f, 11.0f)));
+		Set("CurveEd.CurveKeySelected", new IMAGE_BRUSH("Common/Key", FVector2D(11.0f, 11.0f), SelectionColor));
+		Set("CurveEd.InfoFont", TTF_CORE_FONT("Fonts/Roboto-Regular", 8));
+		Set("CurveEd.LabelFont", TTF_CORE_FONT("Fonts/Roboto-Bold", 10));
+		Set("CurveEd.Tangent", new IMAGE_BRUSH("Common/Tangent", FVector2D(7.0f, 7.0f), FLinearColor(0.0f, 0.66f, 0.7f)));
+		Set("CurveEd.TangentColor", FLinearColor(0.0f, 0.66f, 0.7f));
+		Set("CurveEd.Visible", new IMAGE_BRUSH("Icons/icon_levels_visible_16px", Icon16x16));
+		Set("CurveEd.VisibleHighlight", new IMAGE_BRUSH("Icons/icon_levels_visible_hi_16px", Icon16x16));
+		Set("CurveEd.Invisible", new IMAGE_BRUSH("Icons/icon_levels_invisible_16px", Icon16x16));
+		Set("CurveEd.InvisibleHighlight", new IMAGE_BRUSH("Icons/icon_levels_invisible_hi_16px", Icon16x16));
+		Set("CurveEd.Locked", new IMAGE_BRUSH("Icons/icon_locked_16px", Icon16x16));
+		Set("CurveEd.LockedHighlight", new IMAGE_BRUSH("Icons/icon_locked_highlight_16px", Icon16x16));
+		Set("CurveEd.Unlocked", new IMAGE_BRUSH("Icons/icon_unlocked_16px", Icon16x16));
+		Set("CurveEd.UnlockedHighlight", new IMAGE_BRUSH("Icons/icon_unlocked_highlight_16px", Icon16x16));
 	}
 	
 	// Scrub control buttons
diff --git a/Engine/Source/Editor/UnrealEd/Private/SCurveEditor.cpp b/Engine/Source/Editor/UnrealEd/Private/SCurveEditor.cpp
--- a/Engine/Source/Editor/UnrealEd/Private/SCurveEditor.cpp
+++ b/Engine/Source/Editor/UnrealEd/Private/SCurveEditor.cpp
@@ -70,7 +70,7 @@ void SCurveEditor::Construct(const FArguments& InArgs)
 	//Simple r/g/b for now
 	CurveColors.Add(FLinearColor(1.0f, 0.0f, 0.0f));
 	CurveColors.Add(FLinearColor(0.0f, 1.0f, 0.0f));
-	CurveColors.Add(FLinearColor(0.0f, 0.0f, 1.0f));
+	CurveColors.Add(FLinearColor(0.05f, 0.05f, 1.0f));
 
 	TransactionIndex = -1;
 
@@ -125,6 +125,7 @@ void SCurveEditor::Construct(const FArguments& InArgs)
 	SAssignNew(WarningMessageText, SErrorText);
 
 	TSharedRef<SBox> CurveSelector = SNew(SBox)
+		.VAlign(VAlign_Top)
 		.Visibility(this, &SCurveEditor::GetControlVisibility)
 		[
 			CreateCurveSelectionWidget()
@@ -134,146 +135,153 @@ void SCurveEditor::Construct(const FArguments& InArgs)
 
 	this->ChildSlot
 	[
-		SNew( SVerticalBox )
+		SNew( SHorizontalBox )
 
-		+ SVerticalBox::Slot()
-		.FillHeight(1.0f)
+		+ SHorizontalBox::Slot()
+		.FillWidth(1.0f)
 		[
-			SNew(SHorizontalBox)
-			.Visibility( this, &SCurveEditor::GetCurveAreaVisibility )
+			SNew( SVerticalBox )
 
-			+ SHorizontalBox::Slot()
-			.AutoWidth()
-			.Padding(FMargin(30, 12, 0, 0))
+			+ SVerticalBox::Slot()
+			.FillHeight(1.0f)
 			[
-				CurveSelector
-			]
+				SNew(SHorizontalBox)
+				.Visibility( this, &SCurveEditor::GetCurveAreaVisibility )
 
-			+ SHorizontalBox::Slot()
-			.AutoWidth()
-			[
-				SNew(SBorder)
-				.VAlign(VAlign_Top)
-				.HAlign(HAlign_Left)
-				.BorderImage( FEditorStyle::GetBrush("NoBorder") )
-				.DesiredSizeScale(FVector2D(256.0f,32.0f))
-				.Visibility(this, &SCurveEditor::GetControlVisibility)
-				.Padding(FMargin(2, 12, 0, 0))
+				+ SHorizontalBox::Slot()
+				.AutoWidth()
+				.Padding(FMargin(30, 12, 0, 0))
 				[
-					SNew(SHorizontalBox)
+					CurveSelector
+				]
 
-					+ SHorizontalBox::Slot()
-					.AutoWidth()
+				+ SHorizontalBox::Slot()
+				.AutoWidth()
+				[
+					SNew(SBorder)
+					.VAlign(VAlign_Top)
+					.HAlign(HAlign_Left)
+					.BorderImage( FEditorStyle::GetBrush("NoBorder") )
+					.DesiredSizeScale(FVector2D(256.0f,32.0f))
+					.Visibility(this, &SCurveEditor::GetControlVisibility)
+					.Padding(FMargin(2, 12, 0, 0))
 					[
-						SNew(SButton)
-						.ToolTipText(LOCTEXT("ZoomToFitHorizontal", "Zoom To Fit Horizontal"))
-						.Visibility(this, &SCurveEditor::GetZoomButtonVisibility)
-						.OnClicked(this, &SCurveEditor::ZoomToFitHorizontalClicked)
-						.ContentPadding(1)
+						SNew(SHorizontalBox)
+
+						+ SHorizontalBox::Slot()
+						.AutoWidth()
 						[
-							SNew(SImage) 
-							.Image( FEditorStyle::GetBrush("CurveEd.FitHorizontal") ) 
-							.ColorAndOpacity( FSlateColor::UseForeground() ) 
+							SNew(SButton)
+							.ToolTipText(LOCTEXT("ZoomToFitHorizontal", "Zoom To Fit Horizontal"))
+							.Visibility(this, &SCurveEditor::GetZoomButtonVisibility)
+							.OnClicked(this, &SCurveEditor::ZoomToFitHorizontalClicked)
+							.ContentPadding(1)
+							[
+								SNew(SImage) 
+								.Image( FEditorStyle::GetBrush("CurveEd.FitHorizontal") ) 
+								.ColorAndOpacity( FSlateColor::UseForeground() ) 
+							]
 						]
-					]
 					
-					+ SHorizontalBox::Slot()
-					.AutoWidth()
-					[
-						SNew(SButton)
-						.ToolTipText(LOCTEXT("ZoomToFitVertical", "Zoom To Fit Vertical"))
-						.Visibility(this, &SCurveEditor::GetZoomButtonVisibility)
-						.OnClicked(this, &SCurveEditor::ZoomToFitVerticalClicked)
-						.ContentPadding(1)
+						+ SHorizontalBox::Slot()
+						.AutoWidth()
 						[
-							SNew(SImage) 
-							.Image( FEditorStyle::GetBrush("CurveEd.FitVertical") ) 
-							.ColorAndOpacity( FSlateColor::UseForeground() ) 
+							SNew(SButton)
+							.ToolTipText(LOCTEXT("ZoomToFitVertical", "Zoom To Fit Vertical"))
+							.Visibility(this, &SCurveEditor::GetZoomButtonVisibility)
+							.OnClicked(this, &SCurveEditor::ZoomToFitVerticalClicked)
+							.ContentPadding(1)
+							[
+								SNew(SImage) 
+								.Image( FEditorStyle::GetBrush("CurveEd.FitVertical") ) 
+								.ColorAndOpacity( FSlateColor::UseForeground() ) 
+							]
 						]
-					]
 
-					+ SHorizontalBox::Slot()
-					.AutoWidth()
-					[
-						SNew(SBorder)
-						.BorderImage(FEditorStyle::GetBrush("NoBorder"))
-						.Visibility(this, &SCurveEditor::GetEditVisibility)
-						.VAlign(VAlign_Center)
+						+ SHorizontalBox::Slot()
+						.AutoWidth()
 						[
-							SNew(SNumericEntryBox<float>)
-							.IsEnabled(this, &SCurveEditor::GetInputEditEnabled)
-							.Font(FEditorStyle::GetFontStyle("CurveEd.InfoFont"))
-							.Value(this, &SCurveEditor::OnGetTime)
-							.UndeterminedString(LOCTEXT("MultipleValues", "Multiple Values").ToString())
-							.OnValueCommitted(this, &SCurveEditor::OnTimeComitted)
-							.OnValueChanged(this, &SCurveEditor::OnTimeChanged)
-							.OnBeginSliderMovement(this, &SCurveEditor::OnBeginSliderMovement, LOCTEXT("SetValue", "Set New Time"))
-							.OnEndSliderMovement(this, &SCurveEditor::OnEndSliderMovement)
-							.LabelVAlign(VAlign_Center)
-							.AllowSpin(true)
-							.MinValue(TOptional<float>())
-							.MaxValue(TOptional<float>())
-							.MaxSliderValue(TOptional<float>())
-							.MinSliderValue(TOptional<float>())
-							.Delta(0.001f)
-							.MinDesiredValueWidth(60.0f)
-							.Label()
+							SNew(SBorder)
+							.BorderImage(FEditorStyle::GetBrush("NoBorder"))
+							.Visibility(this, &SCurveEditor::GetEditVisibility)
+							.VAlign(VAlign_Center)
 							[
-								SNew(STextBlock)
-								.Text(LOCTEXT("Time", "Time"))
+								SNew(SNumericEntryBox<float>)
+								.IsEnabled(this, &SCurveEditor::GetInputEditEnabled)
+								.Font(FEditorStyle::GetFontStyle("CurveEd.InfoFont"))
+								.Value(this, &SCurveEditor::OnGetTime)
+								.UndeterminedString(LOCTEXT("MultipleValues", "Multiple Values").ToString())
+								.OnValueCommitted(this, &SCurveEditor::OnTimeComitted)
+								.OnValueChanged(this, &SCurveEditor::OnTimeChanged)
+								.OnBeginSliderMovement(this, &SCurveEditor::OnBeginSliderMovement, LOCTEXT("SetValue", "Set New Time"))
+								.OnEndSliderMovement(this, &SCurveEditor::OnEndSliderMovement)
+								.LabelVAlign(VAlign_Center)
+								.AllowSpin(true)
+								.MinValue(TOptional<float>())
+								.MaxValue(TOptional<float>())
+								.MaxSliderValue(TOptional<float>())
+								.MinSliderValue(TOptional<float>())
+								.Delta(0.001f)
+								.MinDesiredValueWidth(60.0f)
+								.Label()
+								[
+									SNew(STextBlock)
+									.Text(LOCTEXT("Time", "Time"))
+								]
 							]
 						]
-					]
 					
-					+ SHorizontalBox::Slot()
-					.AutoWidth()
-					[
-						SNew(SBorder)
-						.BorderImage( FEditorStyle::GetBrush("NoBorder") )
-						.Visibility(this, &SCurveEditor::GetEditVisibility)
-						.VAlign(VAlign_Center)
+						+ SHorizontalBox::Slot()
+						.AutoWidth()
 						[
-							SNew(SNumericEntryBox<float>)
-							.Font(FEditorStyle::GetFontStyle("CurveEd.InfoFont"))
-							.Value(this, &SCurveEditor::OnGetValue)
-							.UndeterminedString(LOCTEXT("MultipleValues", "Multiple Values").ToString())
-							.OnValueCommitted(this, &SCurveEditor::OnValueComitted)
-							.OnValueChanged(this, &SCurveEditor::OnValueChanged)
-							.OnBeginSliderMovement(this, &SCurveEditor::OnBeginSliderMovement, LOCTEXT("SetValue", "Set New Value"))
-							.OnEndSliderMovement(this, &SCurveEditor::OnEndSliderMovement)
-							.LabelVAlign(VAlign_Center)
-							.AllowSpin(true)
-							.MinValue(TOptional<float>())
-							.MaxValue(TOptional<float>())
-							.MaxSliderValue(TOptional<float>())
-							.MinSliderValue(TOptional<float>())
-							.Delta(.001f)
-							.MinDesiredValueWidth(60.0f)
-							.Label()
+							SNew(SBorder)
+							.BorderImage( FEditorStyle::GetBrush("NoBorder") )
+							.Visibility(this, &SCurveEditor::GetEditVisibility)
+							.VAlign(VAlign_Center)
 							[
-								SNew(STextBlock)
-								.Text(LOCTEXT("Value", "Value"))
+								SNew(SNumericEntryBox<float>)
+								.Font(FEditorStyle::GetFontStyle("CurveEd.InfoFont"))
+								.Value(this, &SCurveEditor::OnGetValue)
+								.UndeterminedString(LOCTEXT("MultipleValues", "Multiple Values").ToString())
+								.OnValueCommitted(this, &SCurveEditor::OnValueComitted)
+								.OnValueChanged(this, &SCurveEditor::OnValueChanged)
+								.OnBeginSliderMovement(this, &SCurveEditor::OnBeginSliderMovement, LOCTEXT("SetValue", "Set New Value"))
+								.OnEndSliderMovement(this, &SCurveEditor::OnEndSliderMovement)
+								.LabelVAlign(VAlign_Center)
+								.AllowSpin(true)
+								.MinValue(TOptional<float>())
+								.MaxValue(TOptional<float>())
+								.MaxSliderValue(TOptional<float>())
+								.MinSliderValue(TOptional<float>())
+								.Delta(.001f)
+								.MinDesiredValueWidth(60.0f)
+								.Label()
+								[
+									SNew(STextBlock)
+									.Text(LOCTEXT("Value", "Value"))
+								]
 							]
 						]
 					]
 				]
 			]
-		]
+		
 
-		+ SVerticalBox::Slot()
-		.VAlign(VAlign_Bottom)
-		.FillHeight(.75f)
-		[
-			SNew( SBorder )
-			.Visibility( this, &SCurveEditor::GetColorGradientVisibility )
-			.BorderImage( FEditorStyle::GetBrush("ToolPanel.GroupBorder") )
-			.BorderBackgroundColor( FLinearColor( .8f, .8f, .8f, .60f ) )
-			.Padding(1.0f)
+			+ SVerticalBox::Slot()
+			.VAlign(VAlign_Bottom)
+			.FillHeight(.75f)
 			[
-				SAssignNew( GradientViewer, SColorGradientEditor )
-				.ViewMinInput( ViewMinInput )
-				.ViewMaxInput( ViewMaxInput )
-				.IsEditingEnabled( this, &SCurveEditor::IsEditingEnabled )
+				SNew( SBorder )
+				.Visibility( this, &SCurveEditor::GetColorGradientVisibility )
+				.BorderImage( FEditorStyle::GetBrush("ToolPanel.GroupBorder") )
+				.BorderBackgroundColor( FLinearColor( .8f, .8f, .8f, .60f ) )
+				.Padding(1.0f)
+				[
+					SAssignNew( GradientViewer, SColorGradientEditor )
+					.ViewMinInput( ViewMinInput )
+					.ViewMaxInput( ViewMaxInput )
+					.IsEditingEnabled( this, &SCurveEditor::IsEditingEnabled )
+				]
 			]
 		]
 	];
@@ -284,6 +292,72 @@ void SCurveEditor::Construct(const FArguments& InArgs)
 	}
 }
 
+FText SCurveEditor::GetIsCurveVisibleToolTip(TSharedPtr<FCurveViewModel> CurveViewModel) const
+{
+	return CurveViewModel->bIsVisible ?
+		FText::Format(LOCTEXT("HideFormat", "Hide {0} curve"), FText::FromName(CurveViewModel->CurveInfo.CurveName)) :
+		FText::Format(LOCTEXT("ShowFormat", "Show {0} curve"), FText::FromName(CurveViewModel->CurveInfo.CurveName));
+}
+
+ESlateCheckBoxState::Type SCurveEditor::IsCurveVisible(TSharedPtr<FCurveViewModel> CurveViewModel) const
+{
+	return CurveViewModel->bIsVisible ? ESlateCheckBoxState::Checked : ESlateCheckBoxState::Unchecked;
+}
+
+void SCurveEditor::OnCurveIsVisibleChanged(ESlateCheckBoxState::Type NewCheckboxState, TSharedPtr<FCurveViewModel> CurveViewModel)
+{
+	if (NewCheckboxState == ESlateCheckBoxState::Checked)
+	{
+		CurveViewModel->bIsVisible = true;
+	}
+	else
+	{
+		CurveViewModel->bIsVisible = false;
+		RemoveCurveKeysFromSelection(CurveViewModel);
+	}
+}
+
+FText SCurveEditor::GetIsCurveLockedToolTip(TSharedPtr<FCurveViewModel> CurveViewModel) const
+{
+	return CurveViewModel->bIsLocked ?
+		FText::Format(LOCTEXT("UnlockFormat", "Unlock {0} curve for editing"), FText::FromName(CurveViewModel->CurveInfo.CurveName)) :
+		FText::Format(LOCTEXT("LockFormat", "Lock {0} curve for editing"), FText::FromName(CurveViewModel->CurveInfo.CurveName));
+}
+
+ESlateCheckBoxState::Type SCurveEditor::IsCurveLocked(TSharedPtr<FCurveViewModel> CurveViewModel) const
+{
+	return CurveViewModel->bIsLocked ? ESlateCheckBoxState::Checked : ESlateCheckBoxState::Unchecked;
+}
+
+void SCurveEditor::OnCurveIsLockedChanged(ESlateCheckBoxState::Type NewCheckboxState, TSharedPtr<FCurveViewModel> CurveViewModel)
+{
+	if (NewCheckboxState == ESlateCheckBoxState::Checked)
+	{
+		CurveViewModel->bIsLocked = true;
+		RemoveCurveKeysFromSelection(CurveViewModel);
+	}
+	else
+	{
+		CurveViewModel->bIsLocked = false;
+	}
+}
+
+void SCurveEditor::RemoveCurveKeysFromSelection(TSharedPtr<FCurveViewModel> CurveViewModel)
+{
+	TArray<FSelectedCurveKey> SelectedKeysForLockedCurve;
+	for (auto SelectedKey : SelectedKeys)
+	{
+		if (SelectedKey.Curve == CurveViewModel->CurveInfo.CurveToEdit)
+		{
+			SelectedKeysForLockedCurve.Add(SelectedKey);
+		}
+	}
+	for (auto KeyToDeselect : SelectedKeysForLockedCurve)
+	{
+		RemoveFromSelection(KeyToDeselect);
+	}
+}
+
 SCurveEditor::~SCurveEditor()
 {
 	if (GEditor != NULL)
@@ -294,33 +368,73 @@ SCurveEditor::~SCurveEditor()
 
 TSharedRef<SWidget> SCurveEditor::CreateCurveSelectionWidget() const
 {
-	TSharedRef<SVerticalBox> Box = SNew(SVerticalBox);
-
-	//Curve selection buttons only useful if you have more then 1 curve
-	if(Curves.Num() > 1)
+	TSharedRef<SVerticalBox> CurveBox = SNew(SVerticalBox);
+	if (CurveViewModels.Num() > 1)
 	{
-		for ( const FRichCurveEditInfo& CurveInfo : Curves )
+		// Only create curve controls if there are more than one.
+		for (auto CurveViewModel : CurveViewModels)
 		{
-			FRichCurve* Curve = CurveInfo.CurveToEdit;
-
-			Box->AddSlot()
+			CurveBox->AddSlot()
 			.AutoHeight()
 			[
 				SNew(SHorizontalBox)
-				+SHorizontalBox::Slot()
+
+				+ SHorizontalBox::Slot()
+				.Padding(0, 0, 5, 0)
+				.FillWidth(1.0f)
+				[
+					SNew(STextBlock)
+					.Font(FEditorStyle::GetFontStyle("CurveEd.LabelFont"))
+					.ColorAndOpacity(CurveViewModel->Color)
+					.Text(CurveViewModel->CurveInfo.CurveName.ToString())
+				]
+
+				+ SHorizontalBox::Slot()
 				.AutoWidth()
+				.VAlign(VAlign_Center)
 				[
 					SNew(SCheckBox)
-					.ForegroundColor(GetCurveColor(Curve))
-					.IsChecked(this, &SCurveEditor::GetCheckState, Curve)
-					.OnCheckStateChanged(this, &SCurveEditor::OnUserSelectedCurve, Curve)
-					.ToolTipText(this, &SCurveEditor::GetCurveCheckBoxToolTip, Curve)
+					.IsChecked(this, &SCurveEditor::IsCurveVisible, CurveViewModel)
+					.OnCheckStateChanged(this, &SCurveEditor::OnCurveIsVisibleChanged, CurveViewModel)
+					.ToolTipText(this, &SCurveEditor::GetIsCurveVisibleToolTip, CurveViewModel)
+					.CheckedImage(FEditorStyle::GetBrush("CurveEd.Visible"))
+					.CheckedHoveredImage(FEditorStyle::GetBrush("CurveEd.VisibleHighlight"))
+					.CheckedPressedImage(FEditorStyle::GetBrush("CurveEd.Visible"))
+					.UncheckedImage(FEditorStyle::GetBrush("CurveEd.Invisible"))
+					.UncheckedHoveredImage(FEditorStyle::GetBrush("CurveEd.InvisibleHighlight"))
+					.UncheckedPressedImage(FEditorStyle::GetBrush("CurveEd.Invisible"))
+				]
+
+				+ SHorizontalBox::Slot()
+				.AutoWidth()
+				.VAlign(VAlign_Center)
+				.Padding(2, 0, 0, 0)
+				[
+					SNew(SCheckBox)
+					.IsChecked(this, &SCurveEditor::IsCurveLocked, CurveViewModel)
+					.OnCheckStateChanged(this, &SCurveEditor::OnCurveIsLockedChanged, CurveViewModel)
+					.ToolTipText(this, &SCurveEditor::GetIsCurveLockedToolTip, CurveViewModel)
+					.CheckedImage(FEditorStyle::GetBrush("CurveEd.Locked"))
+					.CheckedHoveredImage(FEditorStyle::GetBrush("CurveEd.LockedHighlight"))
+					.CheckedPressedImage(FEditorStyle::GetBrush("CurveEd.Locked"))
+					.UncheckedImage(FEditorStyle::GetBrush("CurveEd.Unlocked"))
+					.UncheckedHoveredImage(FEditorStyle::GetBrush("CurveEd.UnlockedHighlight"))
+					.UncheckedPressedImage(FEditorStyle::GetBrush("CurveEd.Unlocked"))
+					.Visibility(bCanEditTrack ? EVisibility::Visible : EVisibility::Collapsed)
 				]
 			];
 		}
 	}
 
-	return Box;
+	TSharedRef<SBorder> Border = SNew(SBorder)
+		.Padding(FMargin(3, 2, 2, 2))
+		.BorderImage(FEditorStyle::GetBrush("ToolPanel.GroupBorder"))
+		.BorderBackgroundColor(FLinearColor(0.0f, 0.0f, 0.0f, 0.3f))
+		[
+			CurveBox
+		];
+
+	return Border;
 }
 
 void SCurveEditor::PushWarningMenu( FVector2D Position, const FText& Message )
@@ -335,10 +449,8 @@ void SCurveEditor::PushWarningMenu( FVector2D Position, const FText& Message )
 }
 
 
-void SCurveEditor::PushCurveMenu(const FGeometry& InMyGeometry, FVector2D Position)
+void SCurveEditor::PushKeyMenu(const FGeometry& InMyGeometry, FVector2D Position)
 {
-	FText Heading = (SelectedKeys.Num() > 0) ? LOCTEXT("KeyInterpolationMode", "Key Interpolation") : LOCTEXT("CurveInterpolationMode", "Curve Interpolation");
-
 	FMenuBuilder MenuBuilder(true, Commands.ToSharedRef());
 
 	if (SelectedKeys.Num() == 0)
@@ -355,7 +467,7 @@ void SCurveEditor::PushCurveMenu(const FGeometry& InMyGeometry, FVector2D Positi
 		MenuBuilder.EndSection();
 	}
 
-	MenuBuilder.BeginSection("CurveEditorInterpolation", Heading);
+	MenuBuilder.BeginSection("CurveEditorInterpolation", LOCTEXT("KeyInterpolationMode", "Key Interpolation"));
 	{
 		MenuBuilder.AddMenuEntry(FRichCurveEditorCommands::Get().InterpolationCubicAuto);
 		MenuBuilder.AddMenuEntry(FRichCurveEditorCommands::Get().InterpolationCubicUser);
@@ -425,29 +537,44 @@ int32 SCurveEditor::OnPaint(const FPaintArgs& Args, const FGeometry& AllottedGeo
 	float ZeroInputX = ScaleInfo.InputToLocalX(0.f);
 	float ZeroOutputY = ScaleInfo.OutputToLocalY(0.f);
 
-
-	// time=0 line
-	{
-		FSlateDrawElement::MakeBox
+	// timeline background
+	int32 BackgroundLayerId = LayerId;
+	float TimelineMaxX = ScaleInfo.InputToLocalX(TimelineLength.Get());
+	FSlateDrawElement::MakeBox
 		(
-			OutDrawElements,
-			LayerId,
-			CurveAreaGeometry.ToPaintGeometry( FVector2D(ZeroInputX, 0), FVector2D(1, CurveAreaGeometry.Size.Y) ),
-			WhiteBrush,
-			MyClippingRect,
-			DrawEffects,
-			WhiteBrush->GetTint( InWidgetStyle ) * InWidgetStyle.GetColorAndOpacityTint()
+		OutDrawElements,
+		BackgroundLayerId,
+		CurveAreaGeometry.ToPaintGeometry(FVector2D(ZeroInputX, 0), FVector2D(TimelineMaxX - ZeroInputX, CurveAreaGeometry.Size.Y)),
+		TimelineAreaBrush,
+		MyClippingRect,
+		DrawEffects,
+		TimelineAreaBrush->GetTint(InWidgetStyle) * InWidgetStyle.GetColorAndOpacityTint()
 		);
-	}
-	
 
+	// grid lines.
+	int32 GridLineLayerId = BackgroundLayerId + 1;
+	PaintGridLines(CurveAreaGeometry, ScaleInfo, OutDrawElements, GridLineLayerId, MyClippingRect, DrawEffects);
+
+	// time=0 line
+	int32 ZeroLineLayerId = GridLineLayerId + 1;
+	FSlateDrawElement::MakeBox
+	(
+		OutDrawElements,
+		ZeroLineLayerId,
+		CurveAreaGeometry.ToPaintGeometry(FVector2D(ZeroInputX, 0), FVector2D(1, CurveAreaGeometry.Size.Y)),
+		WhiteBrush,
+		MyClippingRect,
+		DrawEffects,
+		WhiteBrush->GetTint(InWidgetStyle) * InWidgetStyle.GetColorAndOpacityTint()
+	);
+
+	// value=0 line
 	if( AreCurvesVisible() )
 	{
-		// value=0 line
 		FSlateDrawElement::MakeBox
 		(
 			OutDrawElements,
-			LayerId,
+			ZeroLineLayerId,
 			CurveAreaGeometry.ToPaintGeometry( FVector2D(0, ZeroOutputY), FVector2D(CurveAreaGeometry.Size.X, 1) ),
 			WhiteBrush,
 			MyClippingRect,
@@ -456,73 +583,57 @@ int32 SCurveEditor::OnPaint(const FPaintArgs& Args, const FGeometry& AllottedGeo
 		);
 	}
 
-	// timeline background
-	float TimelineMaxX = ScaleInfo.InputToLocalX(TimelineLength.Get());
-
-	FSlateDrawElement::MakeBox
-	(
-		OutDrawElements,
-		LayerId,
-		CurveAreaGeometry.ToPaintGeometry( FVector2D(ZeroInputX, 0), FVector2D(TimelineMaxX-ZeroInputX, CurveAreaGeometry.Size.Y) ),
-		TimelineAreaBrush,
-		MyClippingRect,
-		DrawEffects,
-		TimelineAreaBrush->GetTint( InWidgetStyle ) * InWidgetStyle.GetColorAndOpacityTint()
-	);
-
-
-	LayerId++; // Ensure we draw line on top of background
+	int32 LockedCurveLayerID = ZeroLineLayerId + 1;
+	int32 CurveLayerId = LockedCurveLayerID + 1;
 
+	int32 KeyLayerId = CurveLayerId + 1;
+	int32 SelectedKeyLayerId = KeyLayerId + 1;
+	
 	if( AreCurvesVisible() )
 	{
+		//Paint the curves, unlocked curves will be on top
+		for ( auto CurveViewModel : CurveViewModels )
 		{
-			const int32 SelectedCurveLayer = LayerId;
-			const int32 UnSelectedCurveLayer = LayerId+1;
-
-			//Paint the curves, any selected curves will be on top
-			for ( const FRichCurveEditInfo& CurveInfo : Curves )
+			if (CurveViewModel->bIsVisible)
 			{
-				FRichCurve* Curve = CurveInfo.CurveToEdit;
-				int32 LayerToUse = (SelectedCurves.Find(Curve) != INDEX_NONE) ? SelectedCurveLayer : UnSelectedCurveLayer;
-				{
-					PaintCurve(Curve, CurveAreaGeometry, ScaleInfo, OutDrawElements, LayerToUse, MyClippingRect, DrawEffects, InWidgetStyle);
-				}
+				PaintCurve(CurveViewModel, CurveAreaGeometry, ScaleInfo, OutDrawElements, CurveViewModel->bIsLocked ? LockedCurveLayerID : CurveLayerId, MyClippingRect, DrawEffects, InWidgetStyle);
 			}
-			LayerId = LayerId + 2;
 		}
 
 		//Paint the keys on top of the curve
-		for ( const FRichCurveEditInfo& CurveInfo : Curves )
+		for (auto CurveViewModel : CurveViewModels)
 		{
-			FRichCurve* Curve = CurveInfo.CurveToEdit;
-			LayerId = PaintKeys(Curve, ScaleInfo, OutDrawElements, LayerId, CurveAreaGeometry, MyClippingRect, DrawEffects, InWidgetStyle);
+			if (CurveViewModel->bIsVisible)
+			{
+				PaintKeys(CurveViewModel, ScaleInfo, OutDrawElements, KeyLayerId, SelectedKeyLayerId, CurveAreaGeometry, MyClippingRect, DrawEffects, InWidgetStyle);
+			}
 		}
 	}
 
+	// Paint children
+	int32 ChildrenLayerId = SelectedKeyLayerId + 1;
+	int32 MarqueeLayerId = SCompoundWidget::OnPaint(Args, CurveAreaGeometry, MyClippingRect, OutDrawElements, ChildrenLayerId, InWidgetStyle, bParentEnabled);
+
+	// Paint marquee
 	if (DragState == EDragState::MarqueeSelect)
 	{
-		PaintMarquee(AllottedGeometry, MyClippingRect, OutDrawElements, LayerId);
+		PaintMarquee(AllottedGeometry, MyClippingRect, OutDrawElements, MarqueeLayerId);
 	}
 
-	// Paint children
-	SCompoundWidget::OnPaint(Args, CurveAreaGeometry, MyClippingRect, OutDrawElements, LayerId, InWidgetStyle, bParentEnabled);
-
-	PaintGridLines(CurveAreaGeometry, ScaleInfo, OutDrawElements, LayerId++, MyClippingRect, DrawEffects);
-
-	return LayerId;
+	return MarqueeLayerId + 1;
 }
 
-void SCurveEditor::PaintCurve( FRichCurve* Curve, const FGeometry &AllottedGeometry, FTrackScaleInfo &ScaleInfo, FSlateWindowElementList &OutDrawElements, 
+void SCurveEditor::PaintCurve(TSharedPtr<FCurveViewModel> CurveViewModel, const FGeometry &AllottedGeometry, FTrackScaleInfo &ScaleInfo, FSlateWindowElementList &OutDrawElements,
 	int32 LayerId, const FSlateRect& MyClippingRect, ESlateDrawEffect::Type DrawEffects, const FWidgetStyle &InWidgetStyle )const
 {
-	if (Curve != NULL)
+	if (CurveViewModel.IsValid())
 	{
 		if (bDrawCurve)
 		{
-			FLinearColor Color = InWidgetStyle.GetColorAndOpacityTint() * GetCurveColor(Curve);
+			FLinearColor Color = InWidgetStyle.GetColorAndOpacityTint() * CurveViewModel->Color;
 
-			//If there are multiple curves, alpha fade out non-selected curves
-			if(Curves.Num() > 1 && SelectedCurves.Num() > 0 && INDEX_NONE == SelectedCurves.Find(Curve))
+			// Fade out curves which are locked.
+			if(CurveViewModel->bIsLocked)
 			{
 				Color *= FLinearColor(1.0f,1.0f,1.0f,0.35f); 
 			}
@@ -530,6 +641,7 @@ void SCurveEditor::PaintCurve( FRichCurve* Curve, const FGeometry &AllottedGeome
 			TArray<FVector2D> LinePoints;
 			int32 CurveDrawInterval = 1;
 
+			FRichCurve* Curve = CurveViewModel->CurveInfo.CurveToEdit;
 			if (Curve->GetNumKeys() < 2) 
 			{
 				//Not enough point, just draw flat line
@@ -626,14 +738,13 @@ void SCurveEditor::CreateLinesForSegment( FRichCurve* Curve, const FRichCurveKey
 	}
 }
 
-int32 SCurveEditor::PaintKeys( FRichCurve* Curve, FTrackScaleInfo &ScaleInfo, FSlateWindowElementList &OutDrawElements, int32 LayerId, const FGeometry &AllottedGeometry, const FSlateRect& MyClippingRect, ESlateDrawEffect::Type DrawEffects, const FWidgetStyle &InWidgetStyle ) const
+void SCurveEditor::PaintKeys(TSharedPtr<FCurveViewModel> CurveViewModel, FTrackScaleInfo &ScaleInfo, FSlateWindowElementList &OutDrawElements, int32 LayerId, int32 SelectedLayerId, const FGeometry &AllottedGeometry, const FSlateRect& MyClippingRect, ESlateDrawEffect::Type DrawEffects, const FWidgetStyle &InWidgetStyle ) const
 {
-	const int32 UnSelectedLayerID = LayerId;
-	const int32 SelectedLayerID = LayerId+1; //selected keys should be on top
+	FLinearColor KeyColor = CurveViewModel->bIsLocked ? FLinearColor(0.1f,0.1f,0.1f,1.f) : InWidgetStyle.GetColorAndOpacityTint();
 
-	FLinearColor KeyColor = IsEditingEnabled() ? InWidgetStyle.GetColorAndOpacityTint() : FLinearColor(0.1f,0.1f,0.1f,1.f);
 	// Iterate over each key
 	ERichCurveInterpMode LastInterpMode = RCIM_Linear;
+	FRichCurve* Curve = CurveViewModel->CurveInfo.CurveToEdit;
 	for (auto It(Curve->GetKeyHandleIterator()); It; ++It)
 	{
 		FKeyHandle KeyHandle = It.Key();
@@ -647,7 +758,7 @@ int32 SCurveEditor::PaintKeys( FRichCurve* Curve, FTrackScaleInfo &ScaleInfo, FS
 		// Get brush
 		bool IsSelected = IsKeySelected(FSelectedCurveKey(Curve,KeyHandle));
 		const FSlateBrush* KeyBrush = IsSelected ? FEditorStyle::GetBrush("CurveEd.CurveKeySelected") : FEditorStyle::GetBrush("CurveEd.CurveKey");
-		int32 LayerToUse	=  IsSelected ? SelectedLayerID: UnSelectedLayerID;
+		int32 LayerToUse = IsSelected ? SelectedLayerId: LayerId;
 
 		FSlateDrawElement::MakeBox(
 			OutDrawElements,
@@ -667,7 +778,6 @@ int32 SCurveEditor::PaintKeys( FRichCurve* Curve, FTrackScaleInfo &ScaleInfo, FS
 
 		LastInterpMode = Curve->GetKeyInterpMode(KeyHandle);
 	}
-	return LayerId+2;
 }
 
 
@@ -755,7 +865,7 @@ void SCurveEditor::PaintGridLines(const FGeometry &AllottedGeometry, FTrackScale
 
 	const float GridPixelSpacing = FMath::Min(ScaleInfo.WidgetSize.GetMin()/1.5f, MaxGridPixelSpacing);
 
-	const FLinearColor GridColor = FLinearColor(0.0f,0.0f,0.0f, 0.25f) ;
+	const FLinearColor GridColor = FLinearColor(0.0f,0.0f,0.0f, 0.30f) ;
 	const FLinearColor GridTextColor = FLinearColor(1.0f,1.0f,1.0f, 0.75f) ;
 
 	//Vertical grid(time)
@@ -867,7 +977,6 @@ void SCurveEditor::SetCurveOwner(FCurveOwnerInterface* InCurveOwner, bool bCanEd
 	if(InCurveOwner != CurveOwner)
 	{
 		EmptySelection();
-		SelectedCurves.Empty();
 	}
 
 	GradientViewer->SetCurveOwner(InCurveOwner);
@@ -879,15 +988,14 @@ void SCurveEditor::SetCurveOwner(FCurveOwnerInterface* InCurveOwner, bool bCanEd
 	bAreCurvesVisible = !IsLinearColorCurve();
 	bIsGradientEditorVisible = IsLinearColorCurve();
 
-	Curves.Empty();
+	CurveViewModels.Empty();
 	if(CurveOwner != NULL)
 	{
-		Curves = InCurveOwner->GetCurves();
-
-		//Auto select first curve
-		if(Curves.Num() > 0)
+		int curveIndex = 0;
+		for (auto CurveInfo : CurveOwner->GetCurves())
 		{
-			SelectCurve(Curves[0].CurveToEdit, false);
+			CurveViewModels.Add(TSharedPtr<FCurveViewModel>(new FCurveViewModel(CurveInfo, CurveColors[curveIndex % CurveColors.Num()], !bCanEdit)));
+			curveIndex++;
 		}
 		CurveOwner->MakeTransactional();
 	}
@@ -920,9 +1028,9 @@ FCurveOwnerInterface* SCurveEditor::GetCurveOwner() const
 
 FRichCurve* SCurveEditor::GetCurve(int32 CurveIndex) const
 {
-	if(CurveIndex < Curves.Num())
+	if(CurveIndex < CurveViewModels.Num())
 	{
-		return Curves[CurveIndex].CurveToEdit;
+		return CurveViewModels[CurveIndex]->CurveInfo.CurveToEdit;
 	}
 	return NULL;
 }
@@ -972,46 +1080,45 @@ FReply SCurveEditor::OnMouseButtonDown( const FGeometry& InMyGeometry, const FPo
 
 void SCurveEditor::AddNewKey(FGeometry InMyGeometry, FVector2D ScreenPosition)
 {
-	if ( SelectedCurves.Num() == 0 )
+	bool bKeyAdded = false;
+	for (auto CurveViewModel : CurveViewModels)
 	{
-		if ( Curves.Num() == 1 )
+		if (!CurveViewModel->bIsLocked)
 		{
-			SelectCurve(Curves[0].CurveToEdit, false);
-		}
-		else
-		{
-			PushWarningMenu(ScreenPosition, LOCTEXT("NoCurveSelected", "Please select a curve"));
-		}
-	}
-
-	for ( auto It = SelectedCurves.CreateIterator(); It; ++It )
-	{
-		FRichCurve* SelectedCurve = *It;
-		if ( IsValidCurve(SelectedCurve) )
-		{
-			const FScopedTransaction Transaction(LOCTEXT("CurveEditor_AddKey", "Add Key(s)"));
-			CurveOwner->ModifyOwner();
+			FRichCurve* SelectedCurve = CurveViewModel->CurveInfo.CurveToEdit;
+			if ( IsValidCurve(SelectedCurve) )
+			{
+				const FScopedTransaction Transaction(LOCTEXT("CurveEditor_AddKey", "Add Key(s)"));
+				CurveOwner->ModifyOwner();
 
-			FTrackScaleInfo ScaleInfo(ViewMinInput.Get(), ViewMaxInput.Get(), ViewMinOutput, ViewMaxOutput, InMyGeometry.Size);
+				FTrackScaleInfo ScaleInfo(ViewMinInput.Get(), ViewMaxInput.Get(), ViewMinOutput, ViewMaxOutput, InMyGeometry.Size);
 
-			FVector2D LocalClickPos = InMyGeometry.AbsoluteToLocal(ScreenPosition);
+				FVector2D LocalClickPos = InMyGeometry.AbsoluteToLocal(ScreenPosition);
 
-			float NewInput = ScaleInfo.LocalXToInput(LocalClickPos.X);
-			float NewOutput = ScaleInfo.LocalYToOutput(LocalClickPos.Y);
-			FKeyHandle NewKeyHandle = SelectedCurve->AddKey(NewInput, NewOutput);
+				FVector2D NewKeyLocation = SnapLocation(FVector2D
+				(
+					ScaleInfo.LocalXToInput(LocalClickPos.X),
+					ScaleInfo.LocalYToOutput(LocalClickPos.Y))
+				);
+				FKeyHandle NewKeyHandle = SelectedCurve->AddKey(NewKeyLocation.X, NewKeyLocation.Y);
 
-			EmptySelection();
-			AddToSelection(FSelectedCurveKey(SelectedCurve, NewKeyHandle));
+				EmptySelection();
+				AddToSelection(FSelectedCurveKey(SelectedCurve, NewKeyHandle));
+				bKeyAdded = true;
+			}
 		}
 	}
 
-	CurveOwner->OnCurveChanged();
+	if (bKeyAdded)
+	{
+		CurveOwner->OnCurveChanged();
+	}
 }
 
 void SCurveEditor::OnMouseCaptureLost()
 {
 	// if we began a drag transaction we need to finish it to make sure undo doesn't get out of sync
-	if ( IsEditingEnabled() && ( DragState == EDragState::DragKey || DragState == EDragState::DragTangent ) )
+	if (DragState == EDragState::DragKey || DragState == EDragState::DragTangent)
 	{
 		EndDragTransaction();
 	}
@@ -1102,7 +1209,7 @@ FReply SCurveEditor::OnMouseWheel( const FGeometry& MyGeometry, const FPointerEv
 
 FReply SCurveEditor::OnKeyDown( const FGeometry& MyGeometry, const FKeyboardEvent& InKeyboardEvent )
 {
-	if ( IsEditingEnabled() && InKeyboardEvent.GetKey() == EKeys::Platform_Delete )
+	if (InKeyboardEvent.GetKey() == EKeys::Platform_Delete && SelectedKeys.Num() != 0)
 	{
 		DeleteSelectedKeys();
 		return FReply::Handled();
@@ -1126,7 +1233,7 @@ void SCurveEditor::TryStartDrag(const FGeometry& InMyGeometry, const FPointerEve
 	FVector2D DragVector = InMyGeometry.AbsoluteToLocal(InMouseEvent.GetScreenSpacePosition()) - MouseDownLocation;
 	if (DragVector.Size() >= DragThreshold)
 	{
-		if (IsEditingEnabled() && bLeftMouseButton)
+		if (bLeftMouseButton)
 		{
 			// Check if we should start dragging keys.
 			FSelectedCurveKey HitKey = HitTestKeys(InMyGeometry, InMyGeometry.LocalToAbsolute(MouseDownLocation));
@@ -1224,24 +1331,23 @@ void SCurveEditor::EndDrag(const FGeometry& InMyGeometry, const FPointerEvent& I
 
 	if (DragState == EDragState::DragKey || DragState == EDragState::DragTangent)
 	{
-		if (IsEditingEnabled())
-		{
-			EndDragTransaction();
-		}
+		EndDragTransaction();
 	}
 	else if (DragState == EDragState::MarqueeSelect)
 	{
-		FVector2D MarqueTopLeft(
+		FVector2D MarqueTopLeft
+		(
 			FMath::Min(MouseDownLocation.X, MouseMoveLocation.X),
 			FMath::Min(MouseDownLocation.Y, MouseMoveLocation.Y)
-			);
+		);
 
-		FVector2D MarqueBottomRight(
+		FVector2D MarqueBottomRight
+		(
 			FMath::Max(MouseDownLocation.X, MouseMoveLocation.X),
 			FMath::Max(MouseDownLocation.Y, MouseMoveLocation.Y)
-			);
+		);
 
-		TArray<FSelectedCurveKey> SelectedCurveKeys = GetKeysWithinMarquee(InMyGeometry, MarqueTopLeft, MarqueBottomRight);
+		TArray<FSelectedCurveKey> SelectedCurveKeys = GetEditableKeysWithinMarquee(InMyGeometry, MarqueTopLeft, MarqueBottomRight);
 
 		if (!bControlDown)
 		{
@@ -1250,6 +1356,7 @@ void SCurveEditor::EndDrag(const FGeometry& InMyGeometry, const FPointerEvent& I
 
 		for (auto SelectedCurveKey : SelectedCurveKeys)
 		{
+			
 			if (IsKeySelected(SelectedCurveKey))
 			{
 				RemoveFromSelection(SelectedCurveKey);
@@ -1271,72 +1378,61 @@ void SCurveEditor::ProcessClick(const FGeometry& InMyGeometry, const FPointerEve
 	const bool bShiftDown = InMouseEvent.IsShiftDown();
 
 	FSelectedCurveKey HitKey = HitTestKeys(InMyGeometry, InMouseEvent.GetScreenSpacePosition());
-	if (IsEditingEnabled())
+	if (bLeftMouseButton)
 	{
-		if (bLeftMouseButton)
+		// If the user left clicked a key, update selection based on modifier key state.
+		if (HitKey.IsValid())
 		{
-			// If the user left clicked a key, update selection based on modifier key state.
-			if (HitKey.IsValid())
+			if (!IsKeySelected(HitKey))
 			{
-				if (!IsKeySelected(HitKey))
+				if (!bControlDown)
 				{
-					if (!bControlDown)
-					{
-						EmptySelection();
-					}
-					AddToSelection(HitKey);
-				}
-				else if (bControlDown)
-				{
-					RemoveFromSelection(HitKey);
+					EmptySelection();
 				}
+				AddToSelection(HitKey);
 			}
-			else
+			else if (bControlDown)
 			{
-				// If the user didn't click a key, add a new one if shift is held down, or try to select a curve.
-				if (bShiftDown)
-				{
-					AddNewKey(InMyGeometry, InMouseEvent.GetScreenSpacePosition());
-				}
-				else
-				{
-					// clicking on background clears key selection
-					EmptySelection();
-
-					//select any curve that is nearby
-					if (Curves.Num() > 1)
-					{
-						HitTestCurves(InMyGeometry, InMouseEvent);
-					}
-				}
+				RemoveFromSelection(HitKey);
 			}
 		}
-		else if (bRightMouseButton)
+		else
 		{
-			// If the user right clicked, handle opening context menus.
-			if (HitKey.IsValid())
+			// If the user didn't click a key, add a new one if shift is held down, or try to select a curve.
+			if (bShiftDown)
 			{
-				// Make sure key is selected in readiness for context menu
-				if (!IsKeySelected(HitKey))
-				{
-					EmptySelection();
-					AddToSelection(HitKey);
-				}
-				PushCurveMenu(InMyGeometry, InMouseEvent.GetScreenSpacePosition());
+				AddNewKey(InMyGeometry, InMouseEvent.GetScreenSpacePosition());
 			}
-			else if (HitTestCurves(InMyGeometry, InMouseEvent))
+			else
 			{
-				// Empty the selection if a key wasn't clicked so that menu has the correct context.
+				// clicking on background clears key selection
 				EmptySelection();
-				PushCurveMenu(InMyGeometry, InMouseEvent.GetScreenSpacePosition());
+			}
+		}
+	}
+	else if (bRightMouseButton)
+	{
+		// If the user right clicked, handle opening context menus.
+		if (HitKey.IsValid())
+		{
+			// Make sure key is selected in readiness for context menu
+			if (!IsKeySelected(HitKey))
+			{
+				EmptySelection();
+				AddToSelection(HitKey);
+			}
+			PushKeyMenu(InMyGeometry, InMouseEvent.GetScreenSpacePosition());
+		}
+		else
+		{
+			FVector2D ScreenPos = InMouseEvent.GetScreenSpacePosition();
+			if (!HitTestCubicTangents(InMyGeometry, ScreenPos).IsValid())
+			{
+				CreateContextMenu(InMyGeometry, ScreenPos);
 			}
 			else
 			{
-				FVector2D ScreenPos = InMouseEvent.GetScreenSpacePosition();
-				if (!HitTestCubicTangents(InMyGeometry, ScreenPos).IsValid())
-				{
-					CreateContextMenu(InMyGeometry, ScreenPos);
-				}
+				EmptySelection();
 			}
 		}
 	}
@@ -1487,27 +1583,24 @@ SCurveEditor::FSelectedCurveKey SCurveEditor::HitTestKeys(const FGeometry& InMyG
 		const FVector2D HitPosition = InMyGeometry.AbsoluteToLocal( HitScreenPosition );
 
 
-		for(auto CurveIt = Curves.CreateConstIterator();CurveIt;++CurveIt)
+		for(auto CurveViewModel : CurveViewModels)
 		{
-			FRichCurve* Curve = CurveIt->CurveToEdit;
-			if(Curve != NULL)
+			if (!CurveViewModel->bIsLocked && CurveViewModel->bIsVisible)
 			{
-				for (auto It(Curve->GetKeyHandleIterator()); It; ++It)
+				FRichCurve* Curve = CurveViewModel->CurveInfo.CurveToEdit;
+				if(Curve != NULL)
 				{
-					float KeyScreenX = ScaleInfo.InputToLocalX(Curve->GetKeyTime(It.Key()));
-					float KeyScreenY = ScaleInfo.OutputToLocalY(Curve->GetKeyValue(It.Key()));
-
-					if(	HitPosition.X > (KeyScreenX - (0.5f * CONST_KeySize.X)) && 
-						HitPosition.X < (KeyScreenX + (0.5f * CONST_KeySize.X)) &&
-						HitPosition.Y > (KeyScreenY - (0.5f * CONST_KeySize.Y)) &&
-						HitPosition.Y < (KeyScreenY + (0.5f * CONST_KeySize.Y)) )
+					for (auto It(Curve->GetKeyHandleIterator()); It; ++It)
 					{
-						SelectedKey =  FSelectedCurveKey(Curve, It.Key());
+						float KeyScreenX = ScaleInfo.InputToLocalX(Curve->GetKeyTime(It.Key()));
+						float KeyScreenY = ScaleInfo.OutputToLocalY(Curve->GetKeyValue(It.Key()));
 
-						//prioritize selected curves
-						if(SelectedCurves.Contains(Curve))
+						if(	HitPosition.X > (KeyScreenX - (0.5f * CONST_KeySize.X)) && 
+							HitPosition.X < (KeyScreenX + (0.5f * CONST_KeySize.X)) &&
+							HitPosition.Y > (KeyScreenY - (0.5f * CONST_KeySize.Y)) &&
+							HitPosition.Y < (KeyScreenY + (0.5f * CONST_KeySize.Y)) )
 						{
-							return SelectedKey;
+							return  FSelectedCurveKey(Curve, It.Key());
 						}
 					}
 				}
@@ -1616,15 +1709,11 @@ TArray<FRichCurve*> SCurveEditor::GetCurvesToFit() const
 {
 	TArray<FRichCurve*> FitCurves;
 
-	if(SelectedCurves.Num() > 0)
-	{
-		FitCurves = SelectedCurves;
-	}
-	else
+	for(auto CurveViewModel : CurveViewModels)
 	{
-		for(auto It = Curves.CreateConstIterator();It;++It)
+		if (CurveViewModel->bIsVisible)
 		{
-			FitCurves.Add(It->CurveToEdit);
+			FitCurves.Add(CurveViewModel->CurveInfo.CurveToEdit);
 		}
 	}
 
@@ -1636,7 +1725,7 @@ void SCurveEditor::ZoomToFitHorizontal(bool OnlySelected)
 {
 	TArray<FRichCurve*> CurvesToFit = GetCurvesToFit();
 
-	if(Curves.Num() > 0)
+	if(CurveViewModels.Num() > 0)
 	{
 		float InMin = FLT_MAX;
 		float InMax = -FLT_MAX;
@@ -1886,9 +1975,9 @@ bool SCurveEditor::IsValidCurve( FRichCurve* Curve ) const
 	bool bIsValid = false;
 	if(Curve)
 	{
-		for(auto It = Curves.CreateConstIterator();It;++It)
+		for(auto CurveViewModel : CurveViewModels)
 		{
-			if(It->CurveToEdit == Curve)
+			if(CurveViewModel->CurveInfo.CurveToEdit == Curve)
 			{
 				bIsValid = true;
 				break;
@@ -1898,48 +1987,6 @@ bool SCurveEditor::IsValidCurve( FRichCurve* Curve ) const
 	return bIsValid;
 }
 
-FLinearColor SCurveEditor::GetCurveColor( FRichCurve* TestCurve ) const
-{
-	check(TestCurve);
-	for(int32 I = 0;I != Curves.Num() && I < CurveColors.Num();I++)
-	{
-		FRichCurve* Curve = Curves[I].CurveToEdit;
-		if(Curve == TestCurve)
-		{
-			return CurveColors[I];
-		}
-	}
-
-	return FLinearColor::Black;
-}
-
-void SCurveEditor::OnUserSelectedCurve( ESlateCheckBoxState::Type State, FRichCurve* Curve )
-{
-	if(IsValidCurve(Curve))
-	{
-		if(State == ESlateCheckBoxState::Checked)
-		{
-			SelectCurve(Curve, false);
-		}
-		else
-		{
-			SelectedCurves.Remove(Curve);
-		}	
-	}
-}
-
-FText SCurveEditor::GetCurveName( FRichCurve* Curve ) const
-{
-	for(auto It = Curves.CreateConstIterator();It;++It)
-	{
-		if(It->CurveToEdit == Curve)
-		{
-			return FText::FromName( It->CurveName );
-		}
-	}
-	return FText::GetEmpty();
-}
-
 void SCurveEditor::SetInputMinMax(float NewMin, float NewMax)
 {
 	if (SetInputViewRangeHandler.IsBound())
@@ -1961,23 +2008,20 @@ void SCurveEditor::SetInputMinMax(float NewMin, float NewMax)
 	}
 }
 
-bool SCurveEditor::HitTestCurves(  const FGeometry& InMyGeometry, const FPointerEvent& InMouseEvent )
+TSharedPtr<FCurveViewModel> SCurveEditor::HitTestCurves(  const FGeometry& InMyGeometry, const FPointerEvent& InMouseEvent )
 {
-	bool bHit = false;
-
 	if( AreCurvesVisible() )
 	{
-		const bool bControlDown = InMouseEvent.IsControlDown();
-
 		FTrackScaleInfo ScaleInfo(ViewMinInput.Get(),  ViewMaxInput.Get(), ViewMinOutput, ViewMaxOutput, InMyGeometry.Size);
 
 		const FVector2D HitPosition = InMyGeometry.AbsoluteToLocal( InMouseEvent.GetScreenSpacePosition()  );
 
 		TArray<FRichCurve*> CurvesHit;
 
-		for(auto It = Curves.CreateConstIterator();It;++It)
+		for(auto CurveViewModel : CurveViewModels)
 		{
-			FRichCurve* Curve = It->CurveToEdit;
+
+			FRichCurve* Curve = CurveViewModel->CurveInfo.CurveToEdit;
 			if(Curve != NULL)
 			{
 				float Time		 = ScaleInfo.LocalXToInput(HitPosition.X);
@@ -1986,55 +2030,13 @@ bool SCurveEditor::HitTestCurves(  const FGeometry& InMyGeometry, const FPointer
 				if( HitPosition.Y > (KeyScreenY - (0.5f * CONST_CurveSize.Y)) &&
 					HitPosition.Y < (KeyScreenY + (0.5f * CONST_CurveSize.Y)))
 				{
-					CurvesHit.Add(Curve);
-					bHit = true;
-				}
-			}
-		}
-
-		// Select the curves that were hit
-		if ( bHit )
-		{
-			if ( bControlDown )
-			{
-				// If control is down and we have to do multi-selection then don't worry about order, select everything that was hit.
-				for ( FRichCurve* Curve : CurvesHit )
-				{
-					SelectCurve(Curve, bControlDown);
-				}
-			}
-			else
-			{
-				FRichCurve* AlreadySelectedCurve = NULL;
-
-				// Without control down we want to bias the decision towards the curve that is already selected.
-				for ( FRichCurve* Curve : CurvesHit )
-				{
-					if ( SelectedCurves.Contains(Curve) )
-					{
-						AlreadySelectedCurve = Curve;
-						break;
-					}
-				}
-
-				if ( AlreadySelectedCurve != NULL )
-				{
-					// Select the already selected curve
-					SelectCurve(AlreadySelectedCurve, bControlDown);
-				}
-				else
-				{
-					// Select which ever curve was hit last
-					for ( FRichCurve* Curve : CurvesHit )
-					{
-						SelectCurve(Curve, bControlDown);
-					}
+					return CurveViewModel;
 				}
 			}
 		}
 	}
 
-	return bHit;
+	return TSharedPtr<FCurveViewModel>();
 }
 
 SCurveEditor::FSelectedTangent SCurveEditor::HitTestCubicTangents( const FGeometry& InMyGeometry, const FVector2D& HitScreenPosition )
@@ -2083,50 +2085,12 @@ SCurveEditor::FSelectedTangent SCurveEditor::HitTestCubicTangents( const FGeomet
 	return Tangent;
 }
 
-
-
-void SCurveEditor::SelectCurve( FRichCurve* Curve, bool bMultiple )
-{
-	if(bMultiple)
-	{
-		SelectedCurves.AddUnique(Curve);
-	}
-	else
-	{
-		SelectedCurves.Empty();
-		SelectedCurves.Add(Curve);
-	}	
-}
-
-ESlateCheckBoxState::Type SCurveEditor::GetCheckState( FRichCurve* Curve ) const
-{
-	return (SelectedCurves.Find(Curve) == INDEX_NONE) ?  ESlateCheckBoxState::Unchecked : ESlateCheckBoxState::Checked;
-}
-
-FText SCurveEditor::GetCurveCheckBoxToolTip( FRichCurve* Curve ) const
-{
-	FText Result;
-	if ( IsValidCurve(Curve) )
-	{
-		if(SelectedCurves.Find(Curve) != INDEX_NONE)
-		{
-			Result = FText::Format( LOCTEXT("CurveIndexSelectorToolTipChecked", "Toggle to Disable Adding Keys to Curve '{0}': shift + left mouse to Add Keys"), GetCurveName(Curve) );
-		}
-		else
-		{
-			Result = FText::Format( LOCTEXT("CurveIndexSelectorToolTipUnChecked", "Toggle to Enable Adding Keys to Curve '{0}'"), GetCurveName(Curve) );
-		}
-	}
-
-	return Result;
-}
-
 void SCurveEditor::OnSelectInterpolationMode(ERichCurveInterpMode InterpMode, ERichCurveTangentMode TangentMode)
 {
-	const FScopedTransaction Transaction( LOCTEXT("CurveEditor_SetInterpolationMode", "Select Interpolation Mode") );
-	CurveOwner->ModifyOwner();
 	if(SelectedKeys.Num() > 0)
 	{
+		const FScopedTransaction Transaction(LOCTEXT("CurveEditor_SetInterpolationMode", "Select Interpolation Mode"));
+		CurveOwner->ModifyOwner();
 		for(auto It = SelectedKeys.CreateIterator();It;++It)
 		{
 			FSelectedCurveKey& Key = *It;
@@ -2134,22 +2098,8 @@ void SCurveEditor::OnSelectInterpolationMode(ERichCurveInterpMode InterpMode, ER
 			Key.Curve->SetKeyInterpMode(Key.KeyHandle,InterpMode );
 			Key.Curve->SetKeyTangentMode(Key.KeyHandle,TangentMode );
 		}
+		CurveOwner->OnCurveChanged();
 	}
-	else
-	{
-		for(auto CurveIt = SelectedCurves.CreateIterator();CurveIt;++CurveIt)
-		{
-			FRichCurve* Curve = *CurveIt;
-			check(IsValidCurve(Curve));
-			for (auto It(Curve->GetKeyHandleIterator()); It; ++It)
-			{
-				Curve->SetKeyInterpMode(It.Key(),InterpMode );
-				Curve->SetKeyTangentMode(It.Key(),TangentMode );
-			}
-		}
-	}
-
-	CurveOwner->OnCurveChanged();
 }
 
 bool SCurveEditor::IsInterpolationModeSelected(ERichCurveInterpMode InterpMode, ERichCurveTangentMode TangentMode)
@@ -2163,22 +2113,12 @@ bool SCurveEditor::IsInterpolationModeSelected(ERichCurveInterpMode InterpMode,
 				return false;
 			}
 		}
+		return true;
 	}
 	else
 	{
-		for (auto SelectedCurve : SelectedCurves)
-		{
-			check(IsValidCurve(SelectedCurve));
-			for (auto KeyHandleIterator(SelectedCurve->GetKeyHandleIterator()); KeyHandleIterator; ++KeyHandleIterator)
-			{
-				if (SelectedCurve->GetKeyInterpMode(KeyHandleIterator->Key) != InterpMode || SelectedCurve->GetKeyTangentMode(KeyHandleIterator->Key) != TangentMode)
-				{
-					return false;
-				}
-			}
-		}
+		return false;
 	}
-	return true;
 }
 
 /* Given a tangent value for a key, calculates the 2D delta vector from that key in curve space */
@@ -2260,28 +2200,31 @@ void SCurveEditor::GetTangentPoints(  FTrackScaleInfo &ScaleInfo, const FSelecte
 	Leave = KeyScreenPosition + ToLeave*CONST_KeyTangentOffset;
 }
 
-TArray<SCurveEditor::FSelectedCurveKey> SCurveEditor::GetKeysWithinMarquee(const FGeometry& InMyGeometry, FVector2D MarqueeTopLeft, FVector2D MarqueeBottomRight) const
+TArray<SCurveEditor::FSelectedCurveKey> SCurveEditor::GetEditableKeysWithinMarquee(const FGeometry& InMyGeometry, FVector2D MarqueeTopLeft, FVector2D MarqueeBottomRight) const
 {
 	TArray<FSelectedCurveKey> KeysWithinMarquee;
 	if (AreCurvesVisible())
 	{
 		FTrackScaleInfo ScaleInfo(ViewMinInput.Get(), ViewMaxInput.Get(), ViewMinOutput, ViewMaxOutput, InMyGeometry.Size);
-		for (auto CurveIt = Curves.CreateConstIterator(); CurveIt; ++CurveIt)
+		for (auto CurveViewModel : CurveViewModels)
 		{
-			FRichCurve* Curve = CurveIt->CurveToEdit;
-			if (Curve != NULL)
+			if (!CurveViewModel->bIsLocked && CurveViewModel->bIsVisible)
 			{
-				for (auto It(Curve->GetKeyHandleIterator()); It; ++It)
+				FRichCurve* Curve = CurveViewModel->CurveInfo.CurveToEdit;
+				if (Curve != NULL)
 				{
-					float KeyScreenX = ScaleInfo.InputToLocalX(Curve->GetKeyTime(It.Key()));
-					float KeyScreenY = ScaleInfo.OutputToLocalY(Curve->GetKeyValue(It.Key()));
-
-					if (KeyScreenX >= (MarqueeTopLeft.X - (0.5f * CONST_KeySize.X)) &&
-						KeyScreenX <= (MarqueeBottomRight.X + (0.5f * CONST_KeySize.X)) &&
-						KeyScreenY >= (MarqueeTopLeft.Y - (0.5f * CONST_KeySize.Y)) &&
-						KeyScreenY <= (MarqueeBottomRight.Y + (0.5f * CONST_KeySize.Y)))
+					for (auto It(Curve->GetKeyHandleIterator()); It; ++It)
 					{
-						KeysWithinMarquee.Add(FSelectedCurveKey(Curve, It.Key()));
+						float KeyScreenX = ScaleInfo.InputToLocalX(Curve->GetKeyTime(It.Key()));
+						float KeyScreenY = ScaleInfo.OutputToLocalY(Curve->GetKeyValue(It.Key()));
+
+						if (KeyScreenX >= (MarqueeTopLeft.X - (0.5f * CONST_KeySize.X)) &&
+							KeyScreenX <= (MarqueeBottomRight.X + (0.5f * CONST_KeySize.X)) &&
+							KeyScreenY >= (MarqueeTopLeft.Y - (0.5f * CONST_KeySize.Y)) &&
+							KeyScreenY <= (MarqueeBottomRight.Y + (0.5f * CONST_KeySize.Y)))
+						{
+							KeysWithinMarquee.Add(FSelectedCurveKey(Curve, It.Key()));
+						}
 					}
 				}
 			}
diff --git a/Engine/Source/Editor/UnrealEd/Public/SCurveEditor.h b/Engine/Source/Editor/UnrealEd/Public/SCurveEditor.h
--- a/Engine/Source/Editor/UnrealEd/Public/SCurveEditor.h
+++ b/Engine/Source/Editor/UnrealEd/Public/SCurveEditor.h
@@ -71,6 +71,28 @@ struct FTrackScaleInfo
 
 };
 
+/** Represents UI state for a curve displayed in the curve editor. */
+class FCurveViewModel
+{
+public:
+	/** The curve info for the curve being edited. */
+	FRichCurveEditInfo CurveInfo;
+	/** The color which should be used to draw the curve and it's label in the UI. */
+	FLinearColor Color;
+	/** Whether or not the curve should be displayed in the UI. */
+	bool bIsVisible;
+	/** Whether or not the curve is locked from editing. */
+	bool bIsLocked;
+
+	FCurveViewModel(FRichCurveEditInfo InCurveInfo, FLinearColor InColor, bool bInIsLocked)
+	{
+		CurveInfo = InCurveInfo;
+		Color = InColor;
+		bIsLocked = bInIsLocked;
+		bIsLocked = false;
+	}
+};
+
 //////////////////////////////////////////////////////////////////////////
 // SCurveEditor
 
@@ -230,12 +252,6 @@ private:
 	/** Adds a new key to the curve. */
 	void AddNewKey(FGeometry InMyGeometry, FVector2D ScreenPosition);
 
-	/** Get the color associated with a given curve */
-	FLinearColor	GetCurveColor(FRichCurve* Curve) const;
-
-	/** Get the name for the given curve */
-	FText			GetCurveName(FRichCurve* Curve) const;
-
 	/** Test if the curve is exists, and if it being displayed on this widget */
 	bool		IsValidCurve(FRichCurve* Curve) const;
 
@@ -261,7 +277,7 @@ private:
 	void GetTangentPoints( FTrackScaleInfo &ScaleInfo,const FSelectedCurveKey &Key, FVector2D& Arrive, FVector2D& Leave) const;
 
 	/** Get the set of keys within a rectangle in local space */
-	TArray<FSelectedCurveKey> GetKeysWithinMarquee(const FGeometry& InMyGeometry, FVector2D MarqueeTopLeft, FVector2D MarqueeBottomRight) const;
+	TArray<FSelectedCurveKey> GetEditableKeysWithinMarquee(const FGeometry& InMyGeometry, FVector2D MarqueeTopLeft, FVector2D MarqueeBottomRight) const;
 
 	/** Empy key selecttion set */
 	void EmptySelection();
@@ -336,11 +352,11 @@ private:
 	UNREALED_API virtual FVector2D ComputeDesiredSize() const override;
 
 	/** Paint a curve */
-	void PaintCurve(	FRichCurve* Curve, const FGeometry &AllottedGeometry, FTrackScaleInfo &ScaleInfo, FSlateWindowElementList &OutDrawElements, 
+	void PaintCurve(TSharedPtr<FCurveViewModel> CurveViewModel, const FGeometry &AllottedGeometry, FTrackScaleInfo &ScaleInfo, FSlateWindowElementList &OutDrawElements, 
 					int32 LayerId, const FSlateRect& MyClippingRect, ESlateDrawEffect::Type DrawEffects, const FWidgetStyle &InWidgetStyle) const;
 
 	/** Paint the keys that make up a curve */
-	int32 PaintKeys( FRichCurve* Curve, FTrackScaleInfo &ScaleInfo, FSlateWindowElementList &OutDrawElements, int32 LayerId, const FGeometry &AllottedGeometry, const FSlateRect& MyClippingRect, ESlateDrawEffect::Type DrawEffects, const FWidgetStyle &InWidgetStyle ) const;
+	void PaintKeys(TSharedPtr<FCurveViewModel> CurveViewModel, FTrackScaleInfo &ScaleInfo, FSlateWindowElementList &OutDrawElements, int32 LayerId, int32 SelectedLayerId, const FGeometry &AllottedGeometry, const FSlateRect& MyClippingRect, ESlateDrawEffect::Type DrawEffects, const FWidgetStyle &InWidgetStyle ) const;
 
 	/** Paint the tangent for a key with cubic curves */
 	void PaintTangent( FTrackScaleInfo &ScaleInfo, FRichCurve* Curve, FKeyHandle KeyHandle, FVector2D KeyLocation, FSlateWindowElementList &OutDrawElements, int32 LayerId, 
@@ -381,10 +397,7 @@ private:
 	void CreateLinesForSegment( FRichCurve* Curve, const FRichCurveKey& Key1, const FRichCurveKey& Key2, TArray<FVector2D>& Points, FTrackScaleInfo &ScaleInfo) const;
 	
 	/** Detect if user is clicking on a curve */
-	bool HitTestCurves(  const FGeometry& InMyGeometry, const FPointerEvent& InMouseEvent );
-
-	/** Add curve to selection */
-	void SelectCurve(FRichCurve* Curve, bool bMultiple);
+	TSharedPtr<FCurveViewModel> HitTestCurves(const FGeometry& InMyGeometry, const FPointerEvent& InMouseEvent);
 
 	/* user is moving the tangent for a key */
 	void OnMoveTangent(FVector2D MouseCurvePosition);
@@ -394,20 +407,11 @@ private:
 	/** Construct widget that allows user to select which curve to edit if there are multiple */
 	TSharedRef<SWidget> CreateCurveSelectionWidget() const;
 
-	/** Get the state of a given curve's checkbox */
-	ESlateCheckBoxState::Type	GetCheckState(FRichCurve* Curve) const;
-
-	/** Called when the user clicks on the checkbox for a curve */
-	void			OnUserSelectedCurve( ESlateCheckBoxState::Type State, FRichCurve* Curve );
-
-	/** Get the tooltip for a curve selection checkbox */
-	FText			GetCurveCheckBoxToolTip(FRichCurve* Curve) const;
-
 	/** Create ontext Menu for waring menu*/
 	void	PushWarningMenu(FVector2D Position, const FText& Message);
 
 	/** Create context Menu for key interpolation settings*/
-	void	PushCurveMenu(const FGeometry& InMyGeometry, FVector2D Position);
+	void	PushKeyMenu(const FGeometry& InMyGeometry, FVector2D Position);
 
 	/** Called when the user selects the interpolation mode */
 	void	OnSelectInterpolationMode(ERichCurveInterpMode InterpMode, ERichCurveTangentMode TangentMode);
@@ -439,6 +443,16 @@ private:
 
 	FVector2D SnapLocation(FVector2D InLocation);
 
+	FText GetIsCurveVisibleToolTip(TSharedPtr<FCurveViewModel> CurveViewModel) const;
+	ESlateCheckBoxState::Type IsCurveVisible(TSharedPtr<FCurveViewModel> CurveViewModel) const;
+	void OnCurveIsVisibleChanged(ESlateCheckBoxState::Type NewCheckboxState, TSharedPtr<FCurveViewModel> CurveViewModel);
+
+	FText GetIsCurveLockedToolTip(TSharedPtr<FCurveViewModel> CurveViewModel) const;
+	ESlateCheckBoxState::Type IsCurveLocked(TSharedPtr<FCurveViewModel> CurveViewModel) const;
+	void OnCurveIsLockedChanged(ESlateCheckBoxState::Type NewCheckboxState, TSharedPtr<FCurveViewModel> CurveViewModel);
+
+	void RemoveCurveKeysFromSelection(TSharedPtr<FCurveViewModel> CurveViewModel);
+
 protected:
 
 	/** Set Default output values when range is too small **/
@@ -463,8 +477,6 @@ private:
 
 	/** Interface for curve supplier */
 	FCurveOwnerInterface*		CurveOwner;
-	/** Curves that we are editing */
-	TArray<FRichCurveEditInfo>	Curves;
 	/** Color for each curve */
 	TArray<FLinearColor>		CurveColors;
 
@@ -476,8 +488,6 @@ private:
 	bool				bAllowZoomOutput;
 	/** If we always show the color curves or allow the user to toggle this */
 	bool				bAlwaysDisplayColorCurves;
-	/** Which curves are currently selected for editing */
-	TArray<FRichCurve*>	SelectedCurves;
 
 	/** Array of selected keys*/
 	TArray<FSelectedCurveKey> SelectedKeys;
@@ -569,6 +579,8 @@ protected:
 
 	/** A map of selected key handle to their starting locations at the beginning of a drag operation. */
 	TMap<FKeyHandle, FVector2D> PreDragKeyLocations;
+
+	TArray<TSharedPtr<FCurveViewModel>> CurveViewModels;
 };
 
 #endif // __SCurveEditor_h__
diff --git a/Engine/Source/Runtime/Engine/Classes/Curves/CurveBase.h b/Engine/Source/Runtime/Engine/Classes/Curves/CurveBase.h
--- a/Engine/Source/Runtime/Engine/Classes/Curves/CurveBase.h
+++ b/Engine/Source/Runtime/Engine/Classes/Curves/CurveBase.h
@@ -387,6 +387,11 @@ template<class T> struct FRichCurveEditInfoTemplate
 	/** Pointer to curves to be edited */
 	T				CurveToEdit;
 
+	FRichCurveEditInfoTemplate()
+		: CurveName(NAME_None)
+		, CurveToEdit(nullptr)
+	{}
+
 	FRichCurveEditInfoTemplate(T InCurveToEdit)
 	:	CurveName(NAME_None)
 	,	CurveToEdit(InCurveToEdit)
