FBlueprintEditorUtils::EnsureCachedDependenciesUpToDate function was added

diff --git a/Engine/Source/Editor/UnrealEd/Private/Kismet2/BlueprintEditorUtils.cpp b/Engine/Source/Editor/UnrealEd/Private/Kismet2/BlueprintEditorUtils.cpp
--- a/Engine/Source/Editor/UnrealEd/Private/Kismet2/BlueprintEditorUtils.cpp
+++ b/Engine/Source/Editor/UnrealEd/Private/Kismet2/BlueprintEditorUtils.cpp
@@ -2167,6 +2167,15 @@ void FBlueprintEditorUtils::GatherDependencies(const UBlueprint* Blueprint, TSet
 	}
 }
 
+void FBlueprintEditorUtils::EnsureCachedDependenciesUpToDate(UBlueprint* Blueprint)
+{
+	if (Blueprint && !Blueprint->bCachedDependenciesUpToDate)
+	{
+		GatherDependencies(Blueprint, Blueprint->CachedDependencies);
+		Blueprint->bCachedDependenciesUpToDate = true;
+	}
+}
+
 void FBlueprintEditorUtils::GetDependentBlueprints(UBlueprint* Blueprint, TArray<UBlueprint*>& DependentBlueprints)
 {
 	DependentBlueprints.Empty();
@@ -2181,11 +2190,7 @@ void FBlueprintEditorUtils::GetDependentBlueprints(UBlueprint* Blueprint, TArray
 		UBlueprint* TestBP = (UBlueprint*) *ObjIt;
 		if (TestBP && !TestBP->HasAnyFlags(RF_PendingKill))
 		{
-			if (!TestBP->bCachedDependenciesUpToDate)
-			{
-				GatherDependencies(TestBP, TestBP->CachedDependencies);
-				TestBP->bCachedDependenciesUpToDate = true;
-			}
+			EnsureCachedDependenciesUpToDate(TestBP);
 
 			if (TestBP->CachedDependencies.Contains(Blueprint))
 			{
diff --git a/Engine/Source/Editor/UnrealEd/Public/Kismet2/BlueprintEditorUtils.h b/Engine/Source/Editor/UnrealEd/Public/Kismet2/BlueprintEditorUtils.h
--- a/Engine/Source/Editor/UnrealEd/Public/Kismet2/BlueprintEditorUtils.h
+++ b/Engine/Source/Editor/UnrealEd/Public/Kismet2/BlueprintEditorUtils.h
@@ -319,7 +319,10 @@ public:
 
 	/** Returns a list of loaded Blueprints that are dependent on the given Blueprint. */
 	static void GetDependentBlueprints(UBlueprint* Blueprint, TArray<UBlueprint*>& DependentBlueprints);
-	
+
+	/** Ensures, that CachedDependencies in BP are up to date */
+	static void EnsureCachedDependenciesUpToDate(UBlueprint* Blueprint);
+
 	/** returns if a graph is an intermediate build product */
 	static bool IsGraphIntermediate(const UEdGraph* Graph);
 
