Level editor viewport border should be hit test invisibile

diff --git a/Engine/Source/Editor/UnrealEd/Private/SEditorViewport.cpp b/Engine/Source/Editor/UnrealEd/Private/SEditorViewport.cpp
--- a/Engine/Source/Editor/UnrealEd/Private/SEditorViewport.cpp
+++ b/Engine/Source/Editor/UnrealEd/Private/SEditorViewport.cpp
@@ -309,7 +309,7 @@ void SEditorViewport::BindCommands()
 
 EVisibility SEditorViewport::OnGetViewportContentVisibility() const
 {
-	return GLevelEditorModeTools().IsViewportUIHidden() ? EVisibility::Collapsed : EVisibility::Visible;
+	return GLevelEditorModeTools().IsViewportUIHidden() ? EVisibility::Collapsed : EVisibility::SelfHitTestInvisible;
 }
 
 void SEditorViewport::OnToggleRealtime()
