UE4: The convert BSP to static mesh automation test now temporarily mounts a content root to the transient automation folder to save the result instead of relying on the read-only /Temp root. TTP#345210

diff --git a/Engine/Source/Editor/UnrealEd/Private/Tests/EditorAutomationTests.cpp b/Engine/Source/Editor/UnrealEd/Private/Tests/EditorAutomationTests.cpp
--- a/Engine/Source/Editor/UnrealEd/Private/Tests/EditorAutomationTests.cpp
+++ b/Engine/Source/Editor/UnrealEd/Private/Tests/EditorAutomationTests.cpp
@@ -1237,8 +1237,9 @@ bool FConvertToValidation::RunTest(const FString& Parameters)
 		TArray<AActor*> ToStaticMeshActors;
 		ToStaticMeshActors.Add(NewBrush);
 
-		//generate static mesh package name
-		FString PackageName = FPackageName::FilenameToLongPackageName(FPaths::AutomationTransientDir() + TEXT("ConvertToBSPToStaticMesh"));
+		//generate static mesh package name. Temporarily mount /Automation.
+		FPackageName::RegisterMountPoint(TEXT("/Automation/"), FPaths::AutomationTransientDir());
+		FString PackageName = TEXT("/Automation/ConvertToBSPToStaticMesh");
 		//Convert brush to specific package name
 		GEditor->DoConvertActors(ToStaticMeshActors, AStaticMeshActor::StaticClass(), TSet<FString>(), true, PackageName);
 
@@ -1258,6 +1259,9 @@ bool FConvertToValidation::RunTest(const FString& Parameters)
 		{
 			UE_LOG(LogEditorAutomationTests, Error, TEXT("Failed to save ConvertToBSPToStaticMesh."));
 		}
+
+		// Unmount /Automation now that the asset was successfully saved
+		FPackageName::UnRegisterMountPoint(TEXT("/Automation/"), FPaths::AutomationTransientDir());
 	}
 
 	TakeLatentAutomationScreenshot(ConvertMeshParameters, BaseFileName, FString::Printf(TEXT("FinalConvertMesh")), FString::Printf(TEXT("04_Final")), false);
