Fix PIE crash when simulating with the "use dedicated server" option checked

diff --git a/Engine/Source/Editor/UnrealEd/Private/PlayLevel.cpp b/Engine/Source/Editor/UnrealEd/Private/PlayLevel.cpp
--- a/Engine/Source/Editor/UnrealEd/Private/PlayLevel.cpp
+++ b/Engine/Source/Editor/UnrealEd/Private/PlayLevel.cpp
@@ -2568,10 +2568,12 @@ FViewport* UEditorEngine::GetPIEViewport()
 			FWorldContext &WorldContext = *It;
 			if (WorldContext.WorldType == EWorldType::PIE)
 			{
-				FSlatePlayInEditorInfo &SlatePlayInEditorSession = SlatePlayInEditorMap.FindChecked(WorldContext.ContextHandle);
-				if (SlatePlayInEditorSession.SlatePlayInEditorWindowViewport.IsValid() )
+				// We can't use FindChecked here because when using the dedicated server option we don't initialize this map 
+				//	(we don't use a viewport for the PIE context in this case)
+				FSlatePlayInEditorInfo * SlatePlayInEditorSessionPtr = SlatePlayInEditorMap.Find(WorldContext.ContextHandle);
+				if (SlatePlayInEditorSessionPtr != NULL && SlatePlayInEditorSessionPtr->SlatePlayInEditorWindowViewport.IsValid() )
 				{
-					return SlatePlayInEditorSession.SlatePlayInEditorWindowViewport.Get();
+					return SlatePlayInEditorSessionPtr->SlatePlayInEditorWindowViewport.Get();
 				}
 			}
 		}
