Generic SMultipleOptionTable widget. Used for project packaging settings.

diff --git a/Engine/Source/Developer/InternationalizationSettings/Private/InternationalizationSettingsModelDetails.cpp b/Engine/Source/Developer/InternationalizationSettings/Private/InternationalizationSettingsModelDetails.cpp
--- a/Engine/Source/Developer/InternationalizationSettings/Private/InternationalizationSettingsModelDetails.cpp
+++ b/Engine/Source/Developer/InternationalizationSettings/Private/InternationalizationSettingsModelDetails.cpp
@@ -1,4 +1,4 @@
-// Copyright 1998-2014 Epic Games, Inc. All Rights Reserved.
+﻿// Copyright 1998-2014 Epic Games, Inc. All Rights Reserved.
 
 #include "InternationalizationSettingsModulePrivatePCH.h"
 
@@ -248,7 +248,12 @@ void FInternationalizationSettingsModelDetails::CustomizeDetails( IDetailLayoutB
 void FInternationalizationSettingsModelDetails::RefreshAvailableCultures()
 {
 	AvailableCultures.Empty();
-	FInternationalization::Get().GetCulturesWithAvailableLocalization(FPaths::GetEditorLocalizationPaths(), AvailableCultures, true);
+	TArray<FCultureRef> LocalizedCultures;
+	FInternationalization::Get().GetCulturesWithAvailableLocalization(FPaths::GetEditorLocalizationPaths(), LocalizedCultures, true);
+	for (const FCultureRef& Culture : LocalizedCultures)
+	{
+		AvailableCultures.Add(Culture);
+	}
 
 	// Update our selected culture based on the available choices
 	if ( !AvailableCultures.Contains( SelectedCulture ) )
diff --git a/Engine/Source/Editor/UnrealEd/Private/Settings/ProjectPackagingSettingsCustomization.h b/Engine/Source/Editor/UnrealEd/Private/Settings/ProjectPackagingSettingsCustomization.h
--- a/Engine/Source/Editor/UnrealEd/Private/Settings/ProjectPackagingSettingsCustomization.h
+++ b/Engine/Source/Editor/UnrealEd/Private/Settings/ProjectPackagingSettingsCustomization.h
@@ -5,10 +5,66 @@
 #include "PropertyEditing.h"
 #include "PropertyCustomizationHelpers.h"
 #include "PropertyRestriction.h"
-
+#include "SMultipleOptionTable.h"
 
 #define LOCTEXT_NAMESPACE "FProjectPackagingSettingsCustomization"
 
+class SCulturePickerRowWidget : public SCompoundWidget
+{
+	SLATE_BEGIN_ARGS(SCulturePickerRowWidget){}
+	SLATE_END_ARGS()
+
+public:
+	void Construct(const FArguments& InArgs, FCulturePtr InCulture, TAttribute<bool> InIsFilteringCultures)
+	{
+		Culture = InCulture;
+		IsFilteringCultures = InIsFilteringCultures;
+
+		// Identify if this culture has localization data.
+		TArray< FCultureRef > LocalizedCultures;
+		TArray<FString> LocalizationPaths = FPaths::GetGameLocalizationPaths();
+		FInternationalization::Get().GetCulturesWithAvailableLocalization(LocalizationPaths, LocalizedCultures, true);
+		HasLocalizationData = LocalizedCultures.Contains(Culture.ToSharedRef());
+
+		ChildSlot
+			[
+				SNew(SHorizontalBox)
+					+SHorizontalBox::Slot()
+					.AutoWidth()
+					.Padding(FMargin(3.0, 2.0))
+					.VAlign(VAlign_Center)
+					[
+						// Warning Icon for whether or not this culture has localization data.
+						SNew(SImage)
+						.Image( FCoreStyle::Get().GetBrush("Icons.Warning") )
+						.Visibility(this, &SCulturePickerRowWidget::HandleWarningImageVisibility)
+						.ToolTipText(LOCTEXT("NotLocalizedWarning", "This project does not have localization data (translations) for this culture."))
+					]
+					+SHorizontalBox::Slot()
+					.FillWidth(1.0f)
+					.VAlign(VAlign_Center)
+					[
+						// Display name of culture.
+						SNew(STextBlock)
+						.Text(FText::FromString(Culture->GetDisplayName()))
+						.ToolTipText(FText::FromString(Culture->GetName()))
+					]
+			];
+	}
+
+	EVisibility HandleWarningImageVisibility() const
+	{
+		// Don't show the warning image if this culture has localization data.
+		// Collapse the widget entirely if we are filtering to only show cultures that have it - gets rid of awkward empty column of space.
+		bool bIsFilteringCultures = IsFilteringCultures.IsBound() ? IsFilteringCultures.Get() : false;
+		return bIsFilteringCultures ? EVisibility::Collapsed : (HasLocalizationData ? EVisibility::Hidden : EVisibility::Visible);
+	}
+
+private:
+	FCulturePtr Culture;
+	TAttribute<bool> IsFilteringCultures;
+	bool HasLocalizationData;
+};
 
 /**
  * Implements a details view customization for UProjectPackagingSettingsCustomization objects.
@@ -19,10 +75,10 @@ class FProjectPackagingSettingsCustomization
 public:
 
 	// IDetailCustomization interface
-
 	virtual void CustomizeDetails( IDetailLayoutBuilder& LayoutBuilder ) override
 	{
 		CustomizeProjectCategory(LayoutBuilder);
+		CustomizePackagingCategory(LayoutBuilder);
 	}
 
 public:
@@ -39,6 +95,25 @@ public:
 
 protected:
 
+	enum class EFilterCulturesChoices
+	{
+		/**
+		 * Only show cultures that have localization data.
+		 */
+		OnlyLocalizedCultures,
+
+		/**
+		 * Show all available cultures.
+		 */
+		AllAvailableCultures
+	};
+
+	FProjectPackagingSettingsCustomization()
+		: FilterCulturesChoice(EFilterCulturesChoices::AllAvailableCultures)
+	{
+
+	}
+
 	/**
 	 * Customizes the Project property category.
 	 *
@@ -62,6 +137,249 @@ protected:
 			}
 		}
 	}
+
+	/**
+	 * Customizes the Packaging property category.
+	 *
+	 * @param LayoutBuilder The layout builder.
+	 */
+	void CustomizePackagingCategory( IDetailLayoutBuilder& LayoutBuilder )
+	{
+		IDetailCategoryBuilder& PackagingCategory = LayoutBuilder.EditCategory("Packaging");
+		{
+			CulturesPropertyHandle = LayoutBuilder.GetProperty("CulturesToStage", UProjectPackagingSettings::StaticClass());
+			CulturesPropertyHandle->MarkHiddenByCustomization();
+			CulturesPropertyArrayHandle = CulturesPropertyHandle->AsArray();
+
+			PopulateCultureList();
+
+			PackagingCategory.AddCustomRow("CulturesToStage", true)
+				.NameContent()
+				.HAlign(HAlign_Fill)
+				.VAlign(VAlign_Top)
+				[
+					SNew(SHorizontalBox)
+
+					+ SHorizontalBox::Slot()
+					.AutoWidth()
+					[
+						CulturesPropertyHandle->CreatePropertyNameWidget()
+					]
+
+					+ SHorizontalBox::Slot()
+					.AutoWidth()
+					[
+						SNew(SImage)
+						.Image(FEditorStyle::GetBrush(TEXT("Icons.Error")))
+						.ToolTipText(LOCTEXT("NoCulturesToStageSelectedError", "At least one culture must be selected or fatal errors may occur when launching games."))
+						.Visibility(this, &FProjectPackagingSettingsCustomization::HandleNoCulturesErrorIconVisibility)
+					]
+				]
+				.ValueContent()
+				.HAlign(HAlign_Fill)
+				.VAlign(VAlign_Fill)
+				[
+					SNew(SVerticalBox)
+
+					+ SVerticalBox::Slot()
+					.AutoHeight()
+					.Padding(0.0f, 4.0f)
+					.VAlign(VAlign_Center)
+					[
+						SNew(SHorizontalBox)
+
+						+ SHorizontalBox::Slot()
+						.AutoWidth()
+						[
+							// all cultures radio button
+							SNew(SCheckBox)
+							.IsChecked(this, &FProjectPackagingSettingsCustomization::HandleShowCulturesCheckBoxIsChecked, EFilterCulturesChoices::AllAvailableCultures)
+							.OnCheckStateChanged(this, &FProjectPackagingSettingsCustomization::HandleShowCulturesCheckBoxCheckStateChanged, EFilterCulturesChoices::AllAvailableCultures)
+							.Style(FEditorStyle::Get(), "RadioButton")
+							[
+								SNew(STextBlock)
+								.Text(LOCTEXT("AllCulturesCheckBoxText", "Show all"))
+							]
+						]
+
+						+ SHorizontalBox::Slot()
+						.FillWidth(1.0f)
+						.Padding(8.0f, 0.0f, 0.0f, 0.0f)
+						[
+							// localized cultures radio button
+							SNew(SCheckBox)
+							.IsChecked(this, &FProjectPackagingSettingsCustomization::HandleShowCulturesCheckBoxIsChecked, EFilterCulturesChoices::OnlyLocalizedCultures)
+							.OnCheckStateChanged(this, &FProjectPackagingSettingsCustomization::HandleShowCulturesCheckBoxCheckStateChanged, EFilterCulturesChoices::OnlyLocalizedCultures)
+							.Style(FEditorStyle::Get(), "RadioButton")
+							[
+								SNew(STextBlock)
+								.Text(LOCTEXT("CookedCulturesCheckBoxText", "Show localized"))
+							]
+						]
+					]
+
+					+ SVerticalBox::Slot()
+					.AutoHeight()
+					[
+						SAssignNew(Table, SMultipleOptionTable<FCulturePtr>, &CultureList)
+						.OnPreBatchSelect(this, &FProjectPackagingSettingsCustomization::OnPreBatchSelect)
+						.OnPostBatchSelect(this, &FProjectPackagingSettingsCustomization::OnPostBatchSelect)
+						.OnGenerateOptionWidget(this, &FProjectPackagingSettingsCustomization::GenerateWidgetForCulture)
+						.OnOptionSelectionChanged(this, &FProjectPackagingSettingsCustomization::OnCultureSelectionChanged)
+						.IsOptionSelected(this, &FProjectPackagingSettingsCustomization::IsCultureSelected)
+						.ListHeight(100.0f)
+					]
+				];
+		}
+	}
+
+	void PopulateCultureList()
+	{
+		switch(FilterCulturesChoice)
+		{
+		case EFilterCulturesChoices::AllAvailableCultures:
+			CultureList.Empty();
+			for(const auto& Culture : FInternationalization::Get().GetAllCultures())
+			{
+				CultureList.Add(Culture);
+			}
+			break;
+		case EFilterCulturesChoices::OnlyLocalizedCultures:
+			TArray<FCultureRef> LocalizedCultureList;
+			TArray<FString> LocalizationPaths = FPaths::GetGameLocalizationPaths();
+			FInternationalization::Get().GetCulturesWithAvailableLocalization(LocalizationPaths, LocalizedCultureList, true);
+			CultureList.Empty();
+			for(const auto& Culture : LocalizedCultureList)
+			{
+				CultureList.Add(Culture);
+			}
+			break;
+		}
+	}
+
+	EVisibility HandleNoCulturesErrorIconVisibility() const
+	{
+		TArray<void*> RawData;
+		CulturesPropertyHandle->AccessRawData(RawData);
+		TArray<FString>* RawCultureStringArray = reinterpret_cast<TArray<FString>*>(RawData[0]);
+
+		return RawCultureStringArray->Num() ? EVisibility::Hidden : EVisibility::Visible;
+	}
+
+	ESlateCheckBoxState::Type HandleShowCulturesCheckBoxIsChecked( EFilterCulturesChoices Choice ) const
+	{
+		if (FilterCulturesChoice == Choice)
+		{
+			return ESlateCheckBoxState::Checked;
+		}
+
+		return ESlateCheckBoxState::Unchecked;
+	}
+
+	void HandleShowCulturesCheckBoxCheckStateChanged( ESlateCheckBoxState::Type NewState, EFilterCulturesChoices Choice )
+	{
+		if (NewState == ESlateCheckBoxState::Checked)
+		{
+			FilterCulturesChoice = Choice;
+		}
+
+		PopulateCultureList();
+		Table->RequestTableRefresh();
+	}
+
+	void AddCulture(FString CultureName)
+	{
+		if(!IsInBatchSelectOperation)
+		{
+			CulturesPropertyHandle->NotifyPreChange();
+		}
+		TArray<void*> RawData;
+		CulturesPropertyHandle->AccessRawData(RawData);
+		TArray<FString>* RawCultureStringArray = reinterpret_cast<TArray<FString>*>(RawData[0]);
+		RawCultureStringArray->Add(CultureName);
+		if(!IsInBatchSelectOperation)
+		{
+			CulturesPropertyHandle->NotifyPostChange();
+		}
+	}
+
+	void RemoveCulture(FString CultureName)
+	{
+		if(!IsInBatchSelectOperation)
+		{
+			CulturesPropertyHandle->NotifyPreChange();
+		}
+		TArray<void*> RawData;
+		CulturesPropertyHandle->AccessRawData(RawData);
+		TArray<FString>* RawCultureStringArray = reinterpret_cast<TArray<FString>*>(RawData[0]);
+		RawCultureStringArray->Remove(CultureName);
+		if(!IsInBatchSelectOperation)
+		{
+			CulturesPropertyHandle->NotifyPostChange();
+		}
+	}
+
+	bool IsFilteringCultures() const
+	{
+		return FilterCulturesChoice == EFilterCulturesChoices::OnlyLocalizedCultures;
+	}
+
+	void OnPreBatchSelect()
+	{
+		IsInBatchSelectOperation = true;
+		CulturesPropertyHandle->NotifyPreChange();
+	}
+
+	void OnPostBatchSelect()
+	{
+		CulturesPropertyHandle->NotifyPostChange();
+		IsInBatchSelectOperation = false;
+	}
+
+	TSharedRef<SWidget> GenerateWidgetForCulture(FCulturePtr Culture)
+	{
+		return SNew(SCulturePickerRowWidget, Culture, TAttribute<bool>(this, &FProjectPackagingSettingsCustomization::IsFilteringCultures));
+	}
+
+	void OnCultureSelectionChanged(bool IsSelected, FCulturePtr Culture)
+	{
+		if(IsSelected)
+		{
+			AddCulture(Culture->GetName());
+		}
+		else
+		{
+			RemoveCulture(Culture->GetName());
+		}
+	}
+
+	bool IsCultureSelected(FCulturePtr Culture)
+	{
+		FString CultureName = Culture->GetName();
+
+		uint32 ElementCount;
+		CulturesPropertyArrayHandle->GetNumElements(ElementCount);
+		for(uint32 Index = 0; Index < ElementCount; ++Index)
+		{
+			const TSharedRef<IPropertyHandle> PropertyHandle = CulturesPropertyArrayHandle->GetElement(Index);
+			FString CultureNameAtIndex;
+			PropertyHandle->GetValue(CultureNameAtIndex);
+			if(CultureNameAtIndex == CultureName)
+			{
+				return true;
+			}
+		}
+
+		return false;
+	}
+
+private:
+	TArray<FCulturePtr> CultureList;
+	TSharedPtr<IPropertyHandle> CulturesPropertyHandle;
+	TSharedPtr<IPropertyHandleArray> CulturesPropertyArrayHandle;
+	EFilterCulturesChoices FilterCulturesChoice;
+	TSharedPtr< SMultipleOptionTable<FCulturePtr> > Table;
+	bool IsInBatchSelectOperation;
 };
 
 
diff --git a/Engine/Source/Programs/AutomationTool/Scripts/CopyBuildToStagingDirectory.Automation.cs b/Engine/Source/Programs/AutomationTool/Scripts/CopyBuildToStagingDirectory.Automation.cs
--- a/Engine/Source/Programs/AutomationTool/Scripts/CopyBuildToStagingDirectory.Automation.cs
+++ b/Engine/Source/Programs/AutomationTool/Scripts/CopyBuildToStagingDirectory.Automation.cs
@@ -180,10 +180,6 @@ public partial class Project : CommandUtils
 
             // Initialize cultures to stage.
             List<string> CulturesToStage = null;
-            if (PlatformGameConfig != null)
-            {
-                PlatformGameConfig.GetArray("/Script/UnrealEd.ProjectPackagingSettings", "CulturesToStage", out CulturesToStage);
-            }
             if (CulturesToStage != null)
             {
                 CulturesToStage.AddRange(Params.CulturesToCook);
@@ -195,7 +191,15 @@ public partial class Project : CommandUtils
 
             if (CulturesToStage == null || CulturesToStage.Count == 0)
             {
-                throw new AutomationException("No cultures were specified for cooking and packaging. This will lead to fatal errors when launching. Specify culture codes via commandline (-CookCultures=) or using project packaging settings.");
+                if (PlatformGameConfig != null)
+                {
+                    PlatformGameConfig.GetArray("/Script/UnrealEd.ProjectPackagingSettings", "CulturesToStage", out CulturesToStage);
+                }
+            }
+
+            if (CulturesToStage == null || CulturesToStage.Count == 0)
+            {
+                throw new AutomationException("No cultures were specified for cooking and packaging. This will lead to fatal errors when launching. Specify culture codes via commandline (-CookCultures=) or using project packaging settings (+CulturesToStage).");
             }
 
 			// Engine ufs (content)
diff --git a/Engine/Source/Runtime/AppFramework/Public/Widgets/SMultipleOptionTable.h b/Engine/Source/Runtime/AppFramework/Public/Widgets/SMultipleOptionTable.h
new file mode 100644
--- /dev/null
+++ b/Engine/Source/Runtime/AppFramework/Public/Widgets/SMultipleOptionTable.h
@@ -0,0 +1,308 @@
+// Copyright 1998-2014 Epic Games, Inc. All Rights Reserved.
+
+#pragma once
+
+#include "SCheckBox.h"
+
+#define LOCTEXT_NAMESPACE "SMultipleOptionTable"
+
+/**
+ * Implements a row widget for an option list.
+ */
+template<typename OptionType>
+class SOptionTableRow
+	: public SMultiColumnTableRow< TSharedPtr<OptionType> >
+{
+public:
+	typedef typename TSlateDelegates< OptionType >::FOnGenerateWidget FOnGenerateWidget;
+
+	SLATE_BEGIN_ARGS(SOptionTableRow) { }
+		SLATE_ARGUMENT(TSharedPtr<STableViewBase>, OwnerTableView)
+		SLATE_ARGUMENT(OptionType, Option)
+		SLATE_EVENT(FOnCheckStateChanged, OnCheckStateChanged)
+		SLATE_ATTRIBUTE(ESlateCheckBoxState::Type, IsChecked)
+		SLATE_EVENT(FOnGenerateWidget, OnGenerateWidget)
+	SLATE_END_ARGS()
+
+public:
+
+	/**
+	 * Constructs the widget.
+	 *
+	 * @param InArgs The construction arguments.
+	 */
+	void Construct( const FArguments& InArgs)
+	{
+		Option = InArgs._Option;
+		OnCheckStateChanged = InArgs._OnCheckStateChanged;
+		IsChecked = InArgs._IsChecked;
+		OnGenerateWidget= InArgs._OnGenerateWidget;
+
+		SMultiColumnTableRow< TSharedPtr<OptionType> >::Construct(FSuperRowType::FArguments(), InArgs._OwnerTableView.ToSharedRef());
+	}
+
+	/**
+	 * Generates the widget for the specified column.
+	 *
+	 * @param ColumnName The name of the column to generate the widget for.
+	 *
+	 * @return The widget.
+	 */
+	virtual TSharedRef<SWidget> GenerateWidgetForColumn( const FName& ColumnName ) override
+	{
+		if (ColumnName == "IsSelected")
+		{
+			return SAssignNew(CheckBox, SCheckBox)
+				.IsChecked(this, &SOptionTableRow::HandleCheckBoxIsChecked)
+				.OnCheckStateChanged(this, &SOptionTableRow::HandleCheckBoxCheckStateChanged)
+				[
+					OnGenerateWidget.Execute(Option)
+				];
+		}
+
+		return SNullWidget::NullWidget;
+	}
+
+private:
+	// Callback for changing the checked state of the check box.
+	void HandleCheckBoxCheckStateChanged( ESlateCheckBoxState::Type NewState )
+	{
+		OnCheckStateChanged.Execute(NewState);
+	}
+
+	// Callback for determining the checked state of the check box.
+	ESlateCheckBoxState::Type HandleCheckBoxIsChecked( ) const
+	{
+		return IsChecked.Get();
+	}
+
+public:
+	TSharedPtr<SCheckBox> CheckBox;
+
+private:
+	OptionType Option;
+	FOnCheckStateChanged OnCheckStateChanged;
+	TAttribute<ESlateCheckBoxState::Type> IsChecked;
+	FOnGenerateWidget OnGenerateWidget;
+};
+
+template<typename OptionType>
+class SMultipleOptionTable : public SCompoundWidget
+{
+public:
+	typedef typename TSlateDelegates< OptionType >::FOnGenerateWidget FOnGenerateOptionWidget;
+
+	// This callback will be used before selecting or deselecting all of the options. May be of use for optimization.
+	DECLARE_DELEGATE(FOnPreBatchSelect);
+
+	// This callback will be used after selecting or deselecting all of the options. May be of use for optimization.
+	DECLARE_DELEGATE(FOnPostBatchSelect);
+
+	DECLARE_DELEGATE_TwoParams(FOnOptionSelectionChanged, bool, OptionType);
+	DECLARE_DELEGATE_RetVal_OneParam(bool, FIsOptionSelected, OptionType);
+
+	SLATE_BEGIN_ARGS( SMultipleOptionTable<OptionType> ) {}
+		SLATE_EVENT(FOnPreBatchSelect, OnPreBatchSelect)
+		SLATE_EVENT(FOnPostBatchSelect, OnPostBatchSelect)
+		SLATE_EVENT(FOnGenerateOptionWidget, OnGenerateOptionWidget)
+		SLATE_EVENT(FOnOptionSelectionChanged, OnOptionSelectionChanged)
+		SLATE_EVENT(FIsOptionSelected, IsOptionSelected)
+		SLATE_ATTRIBUTE(FOptionalSize, ListHeight)
+	SLATE_END_ARGS()
+
+public:
+	void Construct(const FArguments& InArgs, const TArray< OptionType >* const InOptions);
+
+	// Refreshes the list view. Useful if the options array is modified.
+	void RequestTableRefresh();
+
+private:
+	// Generates the special row widget for this list view widget.
+	TSharedRef<class ITableRow> HandleOptionListViewGenerateRow(OptionType Option, const TSharedRef< class STableViewBase >& OwnerTableView);
+
+	// Calls the provided user callback to generate a widget to display the specified option.
+	TSharedRef<SWidget> GenerateWidgetForOption(OptionType Option);
+
+	// Calls the provided user callback to notify them that an option has been checked or unchecked.
+	void HandleCheckBoxCheckStateChanged(ESlateCheckBoxState::Type NewState, OptionType Option) const;
+
+	// Calls the provided user callback to be notified that an option is selected or unselected in some way.
+	ESlateCheckBoxState::Type HandleCheckBoxIsChecked(OptionType Option) const;
+
+	// Selects all options.
+	void HandleAllHyperlinkNavigate();
+
+	// Deselects all options.
+	void HandleNoneHyperlinkNavigate();
+
+	// Collapses hyperlinks for selecting/deselecting all options based on whether there are multiple options.
+	EVisibility HandleHyperlinkVisibility() const;
+
+private:
+	const TArray< OptionType >* Options;
+	FOnPreBatchSelect OnPreBatchSelect;
+	FOnPostBatchSelect OnPostBatchSelect;
+	FOnGenerateOptionWidget OnGenerateOptionWidget;
+	FOnOptionSelectionChanged OnOptionSelectionChanged;
+	FIsOptionSelected IsOptionSelected;
+	TAttribute<FOptionalSize> ListHeight;
+
+	TSharedPtr< SListView< OptionType > > OptionListView;
+};
+
+template<typename OptionType>
+void SMultipleOptionTable<OptionType>::Construct(const FArguments& InArgs, const TArray< OptionType >* const InOptions)
+{
+	Options = InOptions;
+	OnPreBatchSelect = InArgs._OnPreBatchSelect;
+	OnPostBatchSelect = InArgs._OnPostBatchSelect;
+	OnGenerateOptionWidget = InArgs._OnGenerateOptionWidget;
+	OnOptionSelectionChanged = InArgs._OnOptionSelectionChanged;
+	IsOptionSelected = InArgs._IsOptionSelected;
+	ListHeight = InArgs._ListHeight;
+
+	ChildSlot
+		[
+			SNew(SVerticalBox)
+
+			+ SVerticalBox::Slot()
+			.FillHeight(1.0)
+			.Padding(0.0f, 2.0f, 0.0f, 0.0f)
+			[
+				SNew(SBox)
+				.HeightOverride(ListHeight)
+				[
+					// Options menu
+					SAssignNew(OptionListView, SListView<OptionType>)
+					.HeaderRow
+					(
+					SNew(SHeaderRow)
+					.Visibility(EVisibility::Collapsed)
+
+					+ SHeaderRow::Column("IsSelected")
+					.DefaultLabel(LOCTEXT("OptionListIsSelectedColumnHeader", "IsSelected"))
+					.FillWidth(1.0f)
+					)
+					.ItemHeight(16.0f)
+					.ListItemsSource(Options)
+					.OnGenerateRow(this, &SMultipleOptionTable::HandleOptionListViewGenerateRow)
+					.SelectionMode(ESelectionMode::None)
+				]
+			]
+
+			+ SVerticalBox::Slot()
+				.AutoHeight()
+				.Padding(0.0f, 6.0f, 0.0f, 4.0f)
+				[
+					SNew(SSeparator)
+					.Orientation(Orient_Horizontal)
+				]
+
+			+ SVerticalBox::Slot()
+				.AutoHeight()
+				[
+					SNew(SHorizontalBox)
+
+					+ SHorizontalBox::Slot()
+					.FillWidth(1.0f)
+					.HAlign(HAlign_Right)
+					[
+						SNew(STextBlock)
+						.Text(LOCTEXT("SelectLabel", "Select:"))
+					]
+
+					+ SHorizontalBox::Slot()
+						.AutoWidth()
+						.Padding(8.0f, 0.0f)
+						[
+							// Select all options hyper link
+							SNew(SHyperlink)
+							.OnNavigate(this, &SMultipleOptionTable::HandleAllHyperlinkNavigate)
+							.Text(LOCTEXT("AllHyperlinkLabel", "All"))
+							.ToolTipText(LOCTEXT("AllHyperlinkToolTip", "Select all options."))
+							.Visibility(this, &SMultipleOptionTable::HandleHyperlinkVisibility)									
+						]
+
+					+ SHorizontalBox::Slot()
+						.AutoWidth()
+						[
+							// Deselect all options hyper link
+							SNew(SHyperlink)
+							.OnNavigate(this, &SMultipleOptionTable::HandleNoneHyperlinkNavigate)
+							.Text(LOCTEXT("NoneHyperlinkLabel", "None"))
+							.ToolTipText(LOCTEXT("NoneHyperlinkToolTip", "Deselect all."))
+							.Visibility(this, &SMultipleOptionTable::HandleHyperlinkVisibility)									
+						]
+				]
+		];
+}
+
+template<typename OptionType>
+void SMultipleOptionTable<OptionType>::RequestTableRefresh()
+{
+	OptionListView->RequestListRefresh();
+}
+
+template<typename OptionType>
+TSharedRef<class ITableRow> SMultipleOptionTable<OptionType>::HandleOptionListViewGenerateRow(OptionType Option, const TSharedRef< class STableViewBase >& OwnerTableView)
+{
+	return SNew(SOptionTableRow<OptionType>)
+		.Option(Option)
+		.OwnerTableView(OwnerTableView)
+		.OnCheckStateChanged(this, &SMultipleOptionTable::HandleCheckBoxCheckStateChanged, Option)
+		.IsChecked(this, &SMultipleOptionTable::HandleCheckBoxIsChecked, Option)
+		.OnGenerateWidget(this, &SMultipleOptionTable::GenerateWidgetForOption);
+}
+
+template<typename OptionType>
+TSharedRef<SWidget> SMultipleOptionTable<OptionType>::GenerateWidgetForOption(OptionType Option)
+{
+	return OnGenerateOptionWidget.Execute(Option);
+}
+
+template<typename OptionType>
+void SMultipleOptionTable<OptionType>::HandleCheckBoxCheckStateChanged(ESlateCheckBoxState::Type NewState, OptionType Option) const
+{
+	OnOptionSelectionChanged.Execute(NewState == ESlateCheckBoxState::Checked, Option);
+}
+
+template<typename OptionType>
+ESlateCheckBoxState::Type SMultipleOptionTable<OptionType>::HandleCheckBoxIsChecked(OptionType Option) const
+{
+	return IsOptionSelected.Execute(Option) ? ESlateCheckBoxState::Checked : ESlateCheckBoxState::Unchecked;
+}
+
+template<typename OptionType>
+void SMultipleOptionTable<OptionType>::HandleAllHyperlinkNavigate()
+{
+	OnPreBatchSelect.ExecuteIfBound();
+	for (const OptionType& Option : *Options)
+	{
+		HandleCheckBoxCheckStateChanged(ESlateCheckBoxState::Checked, Option);
+	}
+	OnPostBatchSelect.ExecuteIfBound();
+}
+
+template<typename OptionType>
+void SMultipleOptionTable<OptionType>::HandleNoneHyperlinkNavigate()
+{
+	OnPreBatchSelect.ExecuteIfBound();
+	for (const OptionType& Option : *Options)
+	{
+		HandleCheckBoxCheckStateChanged(ESlateCheckBoxState::Unchecked, Option);
+	}
+	OnPostBatchSelect.ExecuteIfBound();
+}
+
+template<typename OptionType>
+EVisibility SMultipleOptionTable<OptionType>::HandleHyperlinkVisibility() const
+{
+	if (Options->Num() > 1)
+	{
+		return EVisibility::Visible;
+	}
+
+	return EVisibility::Collapsed;
+}
+
+#undef LOCTEXT_NAMESPACE
\ No newline at end of file
diff --git a/Engine/Source/Runtime/Core/Private/Internationalization/Internationalization.cpp b/Engine/Source/Runtime/Core/Private/Internationalization/Internationalization.cpp
--- a/Engine/Source/Runtime/Core/Private/Internationalization/Internationalization.cpp
+++ b/Engine/Source/Runtime/Core/Private/Internationalization/Internationalization.cpp
@@ -566,7 +566,7 @@ void FInternationalization::GetCultureNames(TArray<FString>& CultureNames) const
 	}
 }
 
-void FInternationalization::GetCulturesWithAvailableLocalization(const TArray<FString>& InLocalizationPaths, TArray< FCulturePtr >& OutAvailableCultures, const bool bIncludeDerivedCultures) const
+void FInternationalization::GetCulturesWithAvailableLocalization(const TArray<FString>& InLocalizationPaths, TArray< FCultureRef >& OutAvailableCultures, const bool bIncludeDerivedCultures) const
 {
 	OutAvailableCultures.Reset();
 
@@ -627,7 +627,7 @@ void FInternationalization::GetCulturesWithAvailableLocalization(const TArray<FS
 			FCulturePtr Culture = GetCulture(LocalizationFolder);
 			if(Culture.IsValid() && AllCultures.Contains(Culture.ToSharedRef()))
 			{
-				OutAvailableCultures.AddUnique(Culture);
+				OutAvailableCultures.AddUnique(Culture.ToSharedRef());
 			}
 		}
 	}
diff --git a/Engine/Source/Runtime/Core/Private/Internationalization/TextLocalizationManager.cpp b/Engine/Source/Runtime/Core/Private/Internationalization/TextLocalizationManager.cpp
--- a/Engine/Source/Runtime/Core/Private/Internationalization/TextLocalizationManager.cpp
+++ b/Engine/Source/Runtime/Core/Private/Internationalization/TextLocalizationManager.cpp
@@ -85,14 +85,14 @@ void EndInitTextLocalization()
 			}
 			LocalizationPaths += FPaths::GetEngineLocalizationPaths();
 
-			TArray< FCulturePtr > AvailableCultures;
+			TArray< FCultureRef > AvailableCultures;
 			I18N.GetCulturesWithAvailableLocalization(LocalizationPaths, AvailableCultures, false);
 
 			// Query for parent culture name based on request culture name.
 			FString ParentCultureName = RequestedCultureName;
 				
 			// If we do not have localization data, try the parent culture.
-			for(FCulturePtr ParentCulture = I18N.GetCulture(ParentCultureName); !AvailableCultures.Contains(ParentCulture); ParentCulture = I18N.GetCulture(ParentCultureName))
+			for(FCulturePtr ParentCulture = I18N.GetCulture(ParentCultureName); !ParentCulture.IsValid() || !AvailableCultures.Contains(ParentCulture.ToSharedRef()); ParentCulture = I18N.GetCulture(ParentCultureName))
 			{
 				ParentCultureName = FCulture::GetParentName(ParentCultureName);
 
diff --git a/Engine/Source/Runtime/Core/Public/Internationalization/Internationalization.h b/Engine/Source/Runtime/Core/Public/Internationalization/Internationalization.h
--- a/Engine/Source/Runtime/Core/Public/Internationalization/Internationalization.h
+++ b/Engine/Source/Runtime/Core/Public/Internationalization/Internationalization.h
@@ -68,8 +68,13 @@ public:
 
 	CORE_API void GetCultureNames(TArray<FString>& CultureNames) const;
 
+	CORE_API const TArray< FCultureRef >& GetAllCultures() const
+	{
+		return AllCultures;
+	}
+
 	// Given some paths to look at, populate a list of cultures that we have available localization information for. If bIncludeDerivedCultures, include cultures that are derived from those we have localization data for.
-	CORE_API void GetCulturesWithAvailableLocalization(const TArray<FString>& InLocalizationPaths, TArray< FCulturePtr >& OutAvailableCultures, const bool bIncludeDerivedCultures) const;
+	CORE_API void GetCulturesWithAvailableLocalization(const TArray<FString>& InLocalizationPaths, TArray< FCultureRef >& OutAvailableCultures, const bool bIncludeDerivedCultures) const;
 
 private:
 	FInternationalization();
