fixed navigation rebuild on BSP changes
ttp# 345444

diff --git a/Engine/Source/Editor/UnrealEd/Private/EditorCsg.cpp b/Engine/Source/Editor/UnrealEd/Private/EditorCsg.cpp
--- a/Engine/Source/Editor/UnrealEd/Private/EditorCsg.cpp
+++ b/Engine/Source/Editor/UnrealEd/Private/EditorCsg.cpp
@@ -293,15 +293,8 @@ void UEditorEngine::csgRebuild( UWorld* InWorld )
 
 	GWarn->UpdateProgress( 4, 4 );
 
-#if WITH_NAVIGATION_GENERATOR
-	// update static navigable geometry in every level
-	for ( int32 LevelIndex = 0; LevelIndex < InWorld->GetNumLevels(); ++LevelIndex ) 
-	{
-		ULevel* Level =  InWorld->GetLevel(LevelIndex);	
-
-		GEditor->RebuildStaticNavigableGeometry(Level);
-	}
-#endif
+	// update static navigable geometry in current level
+	RebuildStaticNavigableGeometry(InWorld->GetCurrentLevel());
 
 	// Empty EdPolys.
 	InWorld->GetModel()->Polys->Element.Empty();
diff --git a/Engine/Source/Editor/UnrealEd/Private/EditorExporters.cpp b/Engine/Source/Editor/UnrealEd/Private/EditorExporters.cpp
--- a/Engine/Source/Editor/UnrealEd/Private/EditorExporters.cpp
+++ b/Engine/Source/Editor/UnrealEd/Private/EditorExporters.cpp
@@ -2566,6 +2566,13 @@ void UEditorEngine::RebuildStaticNavigableGeometry(ULevel* Level)
 				}
 			}
 		}
+
+		UWorld* World = GetEditorWorldContext().World();
+		UNavigationSystem* NavSys = UNavigationSystem::GetCurrent(World);
+		if (NavSys)
+		{
+			NavSys->UpdateLevelCollision(Level);
+		}
 	}
 }
 
diff --git a/Engine/Source/Editor/UnrealEd/Private/EditorServer.cpp b/Engine/Source/Editor/UnrealEd/Private/EditorServer.cpp
--- a/Engine/Source/Editor/UnrealEd/Private/EditorServer.cpp
+++ b/Engine/Source/Editor/UnrealEd/Private/EditorServer.cpp
@@ -688,23 +688,9 @@ bool UEditorEngine::Exec_Brush( UWorld* InWorld, const TCHAR* Str, FOutputDevice
 		if ( NewBrush )
 		{
 			ULevel::LevelDirtiedEvent.Broadcast();
+			RebuildStaticNavigableGeometry(InWorld->GetCurrentLevel());
 		}
 
-#if WITH_NAVIGATION_GENERATOR
-		// update static navigable geometry in every level
-		for ( int32 LevelIndex = 0; LevelIndex < InWorld->GetNumLevels(); ++LevelIndex ) 
-		{
-			ULevel* Level =  InWorld->GetLevel(LevelIndex);	
-
-			GEditor->RebuildStaticNavigableGeometry(Level);
-		}
-
-		if (IsLoading() == false && InWorld->GetNavigationSystem()) 
-		{
-			InWorld->GetNavigationSystem()->Build();
-		}
-#endif
-
 		if(FParse::Command(&Str,TEXT("SELECTNEWBRUSH")))
 		{
 			GEditor->SelectNone(false, true);
@@ -785,6 +771,7 @@ bool UEditorEngine::Exec_Brush( UWorld* InWorld, const TCHAR* Str, FOutputDevice
 		if ( NewBrush )
 		{
 			ULevel::LevelDirtiedEvent.Broadcast();
+			RebuildStaticNavigableGeometry(InWorld->GetCurrentLevel());
 		}
 
 		if(FParse::Command(&Str,TEXT("SELECTNEWBRUSH")))
@@ -1451,6 +1438,8 @@ void UEditorEngine::RebuildLevel(ULevel& Level)
 	FBSPOps::GFastRebuild = 1;
 
 	Level.UpdateModelComponents();
+
+	RebuildStaticNavigableGeometry(&Level);
 }
 
 void UEditorEngine::RebuildModelFromBrushes(UModel* Model, bool bSelectedBrushesOnly)
@@ -1573,7 +1562,6 @@ void UEditorEngine::RebuildAlteredBSP()
 		if ( LevelToRebuild.IsValid() )
 		{
 			RebuildLevel(*LevelToRebuild.Get());
-			GEditor->RebuildStaticNavigableGeometry(LevelToRebuild.Get());
 			NumLevelsTouched++;
 		}
 	}
@@ -1584,21 +1572,6 @@ void UEditorEngine::RebuildAlteredBSP()
 	{
 		FEditorDelegates::MapChange.Broadcast( MapChangeEventFlags::MapRebuild );
 	}
-
-#if WITH_NAVIGATION_GENERATOR	
-	UWorld* World = GEditor->GetEditorWorldContext().World();
-	if (World != NULL && World->GetNavigationSystem() != NULL)
-	{
-		for (int32 LevelIdx = 0; LevelIdx < LevelsToRebuild.Num(); ++LevelIdx)
-		{
-			TWeakObjectPtr< ULevel > LevelToRebuild = LevelsToRebuild[LevelIdx];
-			if ( LevelToRebuild.IsValid() )
-			{
-				World->GetNavigationSystem()->UpdateLevelCollision(LevelToRebuild.Get());
-			}
-		}
-	}
-#endif // WITH_NAVIGATION_GENERATOR
 	
 	RedrawLevelEditingViewports();
 
