BP Automated test: better handling for cyclic dependency
ttp345055 AUTOMATION: QAGAME: QABP_Circular_ArrayProperty_A fails to compile

diff --git a/Engine/Source/Editor/UnrealEd/Private/Tests/BlueprintAutomationTests.cpp b/Engine/Source/Editor/UnrealEd/Private/Tests/BlueprintAutomationTests.cpp
--- a/Engine/Source/Editor/UnrealEd/Private/Tests/BlueprintAutomationTests.cpp
+++ b/Engine/Source/Editor/UnrealEd/Private/Tests/BlueprintAutomationTests.cpp
@@ -825,13 +825,13 @@ bool FBlueprintCompileOnLoadTest::RunTest(const FString& BlueprintAssetPath)
 			AddError(FString::Printf(TEXT("You have unsaved changes made to '%s', please save them before running this test."), *BlueprintAssetPath));
 			return false;
 		}
-		else 
+		else
 		{
 			AddWarning(FString::Printf(TEXT("Test may be invalid (the blueprint is already loaded): '%s'"), *BlueprintAssetPath));
 			FBlueprintAutomationTestUtilities::UnloadBlueprint(ExistingBP);
-		}		
+		}
 	}
-	
+
 	// tracks blueprints that were already loaded (and cleans up any that were 
 	// loaded in its lifetime, once it is destroyed)
 	FScopedBlueprintUnloader NewBlueprintUnloader(/*bAutoOpenScope =*/true, /*bRunGCOnCloseIn =*/true);
@@ -855,30 +855,48 @@ bool FBlueprintCompileOnLoadTest::RunTest(const FString& BlueprintAssetPath)
 		return false;
 	}
 
-	TArray<TWeakObjectPtr<UClass>> BlueprintDependenciesWP;
+	// GATHER SUBOBJECTS
+	TArray<TWeakObjectPtr<UObject>> InitialBlueprintSubobjects;
 	{
-#if WITH_EDITOR
-		TArray<UClass*> BlueprintDependencies = FScopedClassDependencyGather::GetCachedDependencies();
-#else 
-		TArray<UClass*> BlueprintDependencies;
-		for (TObjectIterator<UBlueprint> BlueprintIt; BlueprintIt; ++BlueprintIt)
+		TArray<UObject*> InitialBlueprintSubobjectsPtr;
+		GetObjectsWithOuter(InitialBlueprint, InitialBlueprintSubobjectsPtr);
+		for (auto Obj : InitialBlueprintSubobjectsPtr)
 		{
-			if (*BlueprintIt == InitialBlueprint)
-			{
-				continue;
-			}
-			else if (BlueprintIt->GeneratedClass != NULL)
-			{
-				BlueprintDependencies.Add(BlueprintIt->GeneratedClass);
-			}
+			InitialBlueprintSubobjects.Add(Obj);
 		}
-#endif 
-		for (UClass* ClassDependency : BlueprintDependencies)
+	}
+
+	// GATHER DEPENDENCIES
+	TSet<TWeakObjectPtr<UBlueprint>> BlueprintDependencies;
+	{
+		TArray<UBlueprint*> DependentBlueprints;
+		FBlueprintEditorUtils::GetDependentBlueprints(InitialBlueprint, DependentBlueprints);
+		for (auto BP : DependentBlueprints)
 		{
-			BlueprintDependenciesWP.Add(ClassDependency);
+			BlueprintDependencies.Add(BP);
+		}
+	}
+	BlueprintDependencies.Add(InitialBlueprint);
+
+	// GATHER DEPENDENCIES PERSISTENT DATA
+	struct FReplaceInnerData
+	{
+		TWeakObjectPtr<UClass> Class;
+		FStringAssetReference BlueprintAsset;
+	};
+	TArray<FReplaceInnerData> ReplaceInnerData;
+	for (auto BPToUnloadWP : BlueprintDependencies)
+	{
+		auto BPToUnload = BPToUnloadWP.Get();
+		auto OldClass = BPToUnload ? *BPToUnload->GeneratedClass : NULL;
+		if (OldClass)
+		{
+			FReplaceInnerData Data;
+			Data.Class = OldClass;
+			Data.BlueprintAsset = FStringAssetReference(BPToUnload);
+			ReplaceInnerData.Add(Data);
 		}
 	}
-	
 
 	// store off data for the initial blueprint so we can unload it (and reconstruct 
 	// later to compare it with a second one)
@@ -889,7 +907,18 @@ bool FBlueprintCompileOnLoadTest::RunTest(const FString& BlueprintAssetPath)
 	FName const BlueprintName = InitialBlueprint->GetFName();
 	// unload the blueprint so we can reload it (to catch any differences, now  
 	// that all its dependencies should be loaded as well)
-	FBlueprintAutomationTestUtilities::UnloadBlueprint(InitialBlueprint);
+
+	//UNLOAD DEPENDENCIES, all circular dependencies will be loaded again 
+	// unload the blueprint so we can reload it (to catch any differences, now  
+	// that all its dependencies should be loaded as well)
+	for (auto BPToUnloadWP : BlueprintDependencies)
+	{
+		if (auto BPToUnload = BPToUnloadWP.Get())
+		{
+			FBlueprintAutomationTestUtilities::UnloadBlueprint(BPToUnload);
+		}
+	}
+
 	// this blueprint is now dead (will be destroyed next garbage-collection pass)
 	UBlueprint* UnloadedBlueprint = InitialBlueprint;
 	InitialBlueprint = NULL;
@@ -897,25 +926,34 @@ bool FBlueprintCompileOnLoadTest::RunTest(const FString& BlueprintAssetPath)
 	// load the blueprint a second time; if the two separately loaded blueprints 
 	// are different, then this one is most likely the choice one (it has all its 
 	// dependencies loaded)
+
 	UBlueprint* ReloadedBlueprint = Cast<UBlueprint>(StaticLoadObject(UBlueprint::StaticClass(), NULL, *BlueprintAssetPath));
 
-	// fixup any circular dependencies that may now be referencing the unloaded 
-	// blueprint (replace them with the reloaded blueprint, class, etc.)...
+	UPackage* TransientPackage = GetTransientPackage();
+	FName ReconstructedName = MakeUniqueObjectName(TransientPackage, UBlueprint::StaticClass(), BlueprintName);
+	// reconstruct the initial blueprint (using the serialized data from its initial load)
+	EObjectFlags const StandardBlueprintFlags = RF_Public | RF_Standalone | RF_Transactional;
+	InitialBlueprint = ConstructObject<UBlueprint>(UBlueprint::StaticClass(), TransientPackage, ReconstructedName, StandardBlueprintFlags | RF_Transient);
+	FObjectReader(InitialBlueprint, InitialLoadData);
 	{
-		TMap<UBlueprint*, UBlueprint*> BlueprintRedirects;
-		BlueprintRedirects.Add(UnloadedBlueprint, ReloadedBlueprint);
-
-		TMap<UClass*, UClass*> ClassRedirects;
-		ClassRedirects.Add(UnloadedBlueprint->GeneratedClass, ReloadedBlueprint->GeneratedClass);
-		ClassRedirects.Add(UnloadedBlueprint->SkeletonGeneratedClass, ReloadedBlueprint->SkeletonGeneratedClass);
-
-		for (auto ClassDependencyWP : BlueprintDependenciesWP)
+		TMap<UObject*, UObject*> ClassRedirects;
+		for (auto& Data : ReplaceInnerData)
+		{
+			UClass* OriginalClass = Data.Class.Get();
+			UBlueprint* NewBlueprint = Cast<UBlueprint>(Data.BlueprintAsset.ResolveObject());
+			UClass* NewClass = NewBlueprint ? *NewBlueprint->GeneratedClass : NULL;
+			if (OriginalClass && NewClass)
+			{
+				ClassRedirects.Add(OriginalClass, NewClass);
+			}
+		}
+		// REPLACE OLD DATA
+		FArchiveReplaceObjectRef<UObject>(InitialBlueprint, ClassRedirects, /*bNullPrivateRefs=*/false, /*bIgnoreOuterRef=*/true, /*bIgnoreArchetypeRef=*/false);
+		for (auto SubobjWP : InitialBlueprintSubobjects)
 		{
-			auto ClassDependency = ClassDependencyWP.Get();
-			if (ClassDependency)
+			if (auto Subobj = SubobjWP.Get())
 			{
-				FArchiveReplaceObjectRef<UBlueprint>(ClassDependency, BlueprintRedirects, /*bNullPrivateRefs=*/false, /*bIgnoreOuterRef=*/true, /*bIgnoreArchetypeRef=*/false);
-				FArchiveReplaceObjectRef<UClass>(ClassDependency, ClassRedirects, /*bNullPrivateRefs=*/false, /*bIgnoreOuterRef=*/true, /*bIgnoreArchetypeRef=*/false);
+				FArchiveReplaceObjectRef<UObject>(Subobj, ClassRedirects, /*bNullPrivateRefs=*/false, /*bIgnoreOuterRef=*/true, /*bIgnoreArchetypeRef=*/false);
 			}
 		}
 
@@ -924,13 +962,6 @@ bool FBlueprintCompileOnLoadTest::RunTest(const FString& BlueprintAssetPath)
 		FBlueprintEditorUtils::RefreshAllNodes(ReloadedBlueprint);
 		AssetPackage->SetDirtyFlag(bHasUnsavedChanges);
 	}
-	
-	UPackage* TransientPackage = GetTransientPackage();
-	FName ReconstructedName = MakeUniqueObjectName(TransientPackage, UBlueprint::StaticClass(), BlueprintName);
-	// reconstruct the initial blueprint (using the serialized data from its initial load)
-	EObjectFlags const StandardBlueprintFlags = RF_Public | RF_Standalone | RF_Transactional;
-	InitialBlueprint = ConstructObject<UBlueprint>(UBlueprint::StaticClass(), TransientPackage, ReconstructedName, StandardBlueprintFlags | RF_Transient);
-	FObjectReader(InitialBlueprint, InitialLoadData);
 
 	// look for diffs between subsequent loads and log them as errors
 	TArray<FDiffSingleResult> BlueprintDiffs;
