Fix spelling error

diff --git a/Engine/Source/Editor/UnrealEd/Private/Kismet2/StructureEditorUtils.cpp b/Engine/Source/Editor/UnrealEd/Private/Kismet2/StructureEditorUtils.cpp
--- a/Engine/Source/Editor/UnrealEd/Private/Kismet2/StructureEditorUtils.cpp
+++ b/Engine/Source/Editor/UnrealEd/Private/Kismet2/StructureEditorUtils.cpp
@@ -157,7 +157,7 @@ bool FStructureEditorUtils::CanHaveAMemberVariableOfType(const UUserDefinedStruc
 		{
 			if (OutMsg)
 			{
-				*OutMsg = LOCTEXT("StructureIncorrectStructType", "Icorrect struct type in a structure member variable.").ToString();
+				*OutMsg = LOCTEXT("StructureIncorrectStructType", "Incorrect struct type in a structure member variable.").ToString();
 			}
 			return false;
 		}
@@ -169,7 +169,7 @@ bool FStructureEditorUtils::CanHaveAMemberVariableOfType(const UUserDefinedStruc
 	{
 		if (OutMsg)
 		{
-			*OutMsg = LOCTEXT("StructureIncorrectTypeCategory", "Icorrect type for a structure member variable.").ToString();
+			*OutMsg = LOCTEXT("StructureIncorrectTypeCategory", "Incorrect type for a structure member variable.").ToString();
 		}
 		return false;
 	}
