[344555] - Blueprints - Fix VM crash on compile when a property is accessed externally by a variable node in a dependent BP's construction script, and an instance of the BP being compiled exists and is linked to the target pin of the variable node.

- Deferred dependent BP change event notifications to occur after reinstancing has finished when performing a full compile on a non-interface BP. This ensures that instances of the BP being compiled will not be dereferenced by the dependent BP's construction script until after the object has been fully reinstanced.

diff --git a/Engine/Source/Editor/KismetCompiler/Private/KismetCompiler.cpp b/Engine/Source/Editor/KismetCompiler/Private/KismetCompiler.cpp
--- a/Engine/Source/Editor/KismetCompiler/Private/KismetCompiler.cpp
+++ b/Engine/Source/Editor/KismetCompiler/Private/KismetCompiler.cpp
@@ -3445,7 +3445,10 @@ void FKismetCompilerContext::Compile()
 
 			CurrentBP->Status = BS_Dirty;
 			FBlueprintEditorUtils::RefreshExternalBlueprintDependencyNodes(CurrentBP, NewClass);
-			CurrentBP->BroadcastChanged();
+			
+			// Note: We do not send a change notification event to the dependent BP here because
+			// we have not yet reinstanced any of the instances of the BP being compiled, which may
+			// be referenced by instances of the dependent BP that may be reconstructed as a result.
 
 			// Clear the package dirty state if it did not initially have any unsaved changes to begin with
 			if(Package != nullptr && Package->IsDirty() && !bStartedWithUnsavedChanges)
diff --git a/Engine/Source/Editor/UnrealEd/Private/Kismet2/Kismet2.cpp b/Engine/Source/Editor/UnrealEd/Private/Kismet2/Kismet2.cpp
--- a/Engine/Source/Editor/UnrealEd/Private/Kismet2/Kismet2.cpp
+++ b/Engine/Source/Editor/UnrealEd/Private/Kismet2/Kismet2.cpp
@@ -429,20 +429,26 @@ void FKismetEditorUtilities::CompileBlueprint(UBlueprint* BlueprintObj, bool bIs
 		BlueprintObj->NewVariables[VarIndex].DefaultValue.Empty();
 	}
 
-	// for interface changes, auto-refresh nodes on any dependent blueprints
-	if (FBlueprintEditorUtils::IsInterfaceBlueprint(BlueprintObj))
-	{
-		TArray<UBlueprint*> DependentBPs;
-		FBlueprintEditorUtils::GetDependentBlueprints(BlueprintObj, DependentBPs);
+	TArray<UBlueprint*> DependentBPs;
+	FBlueprintEditorUtils::GetDependentBlueprints(BlueprintObj, DependentBPs);
 
-		// refresh each dependent blueprint
-		for (UBlueprint* Dependent : DependentBPs)
+	// refresh each dependent blueprint
+	for (UBlueprint* Dependent : DependentBPs)
+	{
+		// for interface changes, auto-refresh nodes on any dependent blueprints
+		// note: RefreshAllNodes() will internally send a change notification event to the dependent blueprint
+		if (FBlueprintEditorUtils::IsInterfaceBlueprint(BlueprintObj))
 		{
 			bool bPreviousRegenValue = Dependent->bIsRegeneratingOnLoad;
 			Dependent->bIsRegeneratingOnLoad = Dependent->bIsRegeneratingOnLoad || BlueprintObj->bIsRegeneratingOnLoad;
 			FBlueprintEditorUtils::RefreshAllNodes(Dependent);
 			Dependent->bIsRegeneratingOnLoad = bPreviousRegenValue;
 		}
+		else if(!BlueprintObj->bIsRegeneratingOnLoad)
+		{
+			// for non-interface changes, nodes with an external dependency have already been refreshed, and it is now safe to send a change notification event
+			Dependent->BroadcastChanged();
+		}
 	}
 
 	if(!bIsRegeneratingOnLoad && BlueprintObj->GeneratedClass)
