Picks generated lightmap UV channel as first open instead of 1.

diff --git a/Engine/Source/Editor/UnrealEd/Private/Fbx/FbxStaticMeshImport.cpp b/Engine/Source/Editor/UnrealEd/Private/Fbx/FbxStaticMeshImport.cpp
--- a/Engine/Source/Editor/UnrealEd/Private/Fbx/FbxStaticMeshImport.cpp
+++ b/Engine/Source/Editor/UnrealEd/Private/Fbx/FbxStaticMeshImport.cpp
@@ -1067,6 +1067,7 @@ UStaticMesh* UnFbx::FFbxImporter::ImportStaticMeshAsSingle(UObject* InParent, TA
 		
 		// Remap material indices.
 		int32 MaxMaterialIndex = 0;
+		int32 FirstOpenUVChannel = 1;
 		{
 			FRawMesh RawMesh;
 			SrcModel.RawMeshBulkData->LoadRawMesh(RawMesh);
@@ -1106,6 +1107,15 @@ UStaticMesh* UnFbx::FFbxImporter::ImportStaticMeshAsSingle(UObject* InParent, TA
 				MaxMaterialIndex = FMath::Max<int32>(MaxMaterialIndex,RawMesh.FaceMaterialIndices[TriIndex]);
 			}
 
+			for( int32 i = 0; i < MAX_MESH_TEXTURE_COORDS; i++ )
+			{
+				if( RawMesh.WedgeTexCoords[i].Num() == 0 )
+				{
+					FirstOpenUVChannel = i;
+					break;
+				}
+			}
+
 			SrcModel.RawMeshBulkData->SaveRawMesh(RawMesh);
 		}
 
@@ -1158,6 +1168,7 @@ UStaticMesh* UnFbx::FFbxImporter::ImportStaticMeshAsSingle(UObject* InParent, TA
 		SrcModel.BuildSettings.bRecomputeNormals = ImportOptions->NormalImportMethod == FBXNIM_ComputeNormals;
 		SrcModel.BuildSettings.bRecomputeTangents = ImportOptions->NormalImportMethod != FBXNIM_ImportNormalsAndTangents;
 		SrcModel.BuildSettings.bGenerateLightmapUVs = ImportOptions->bGenerateLightmapUVs;
+		SrcModel.BuildSettings.DstLightmapIndex = FirstOpenUVChannel;
 
 		StaticMesh->LODGroup = ImportOptions->StaticMeshLODGroup;
 		StaticMesh->Build(false);
