Allow editor modes to override if they allow the viewport to use a drag tool or not

diff --git a/Engine/Source/Editor/UnrealEd/Private/EditorModes.cpp b/Engine/Source/Editor/UnrealEd/Private/EditorModes.cpp
--- a/Engine/Source/Editor/UnrealEd/Private/EditorModes.cpp
+++ b/Engine/Source/Editor/UnrealEd/Private/EditorModes.cpp
@@ -1567,6 +1567,16 @@ bool FEditorModeTools::EndTracking(FEditorViewportClient* InViewportClient, FVie
 	return bTransactionHandled;
 }
 
+bool FEditorModeTools::AllowsViewportDragTool() const
+{
+	bool bCanUseDragTool = false;
+	for (const TSharedPtr<FEdMode>& Mode : Modes)
+	{
+		bCanUseDragTool |= Mode->AllowsViewportDragTool();
+	}
+	return bCanUseDragTool;
+}
+
 /** Notifies all active modes that a map change has occured */
 void FEditorModeTools::MapChangeNotify()
 {
diff --git a/Engine/Source/Editor/UnrealEd/Private/EditorViewportClient.cpp b/Engine/Source/Editor/UnrealEd/Private/EditorViewportClient.cpp
--- a/Engine/Source/Editor/UnrealEd/Private/EditorViewportClient.cpp
+++ b/Engine/Source/Editor/UnrealEd/Private/EditorViewportClient.cpp
@@ -3101,6 +3101,11 @@ TSharedPtr<FDragTool> FEditorViewportClient::MakeDragTool(EDragTool::Type)
 	return MakeShareable( new FDragTool );
 }
 
+bool FEditorViewportClient::CanUseDragTool() const
+{
+	return !ShouldOrbitCamera() && GetCurrentWidgetAxis() == EAxisList::None && (!ModeTools || ModeTools->AllowsViewportDragTool());
+}
+
 bool FEditorViewportClient::ShouldOrbitCamera() const
 {
 	if( bCameraLock )
diff --git a/Engine/Source/Editor/UnrealEd/Private/MouseDeltaTracker.cpp b/Engine/Source/Editor/UnrealEd/Private/MouseDeltaTracker.cpp
--- a/Engine/Source/Editor/UnrealEd/Private/MouseDeltaTracker.cpp
+++ b/Engine/Source/Editor/UnrealEd/Private/MouseDeltaTracker.cpp
@@ -316,7 +316,7 @@ void FMouseDeltaTracker::ConditionalBeginUsingDragTool( FEditorViewportClient* I
 
 	if( bEnoughMouseMovement )
 	{
-		const bool bCanDrag = !InViewportClient->ShouldOrbitCamera() && InViewportClient->GetCurrentWidgetAxis() == EAxisList::None && !DragTool.IsValid() && !RightMouseButtonDown;
+		const bool bCanDrag = !DragTool.IsValid() && !RightMouseButtonDown && InViewportClient->CanUseDragTool();
 
 		if (bCanDrag && !bHasAttemptedDragTool)
 		{
diff --git a/Engine/Source/Editor/UnrealEd/Public/EditorModes.h b/Engine/Source/Editor/UnrealEd/Public/EditorModes.h
--- a/Engine/Source/Editor/UnrealEd/Public/EditorModes.h
+++ b/Engine/Source/Editor/UnrealEd/Public/EditorModes.h
@@ -249,6 +249,9 @@ public:
 	/** Check to see if an actor can be selected in this mode - no side effects */
 	virtual bool IsSelectionAllowed( AActor* InActor, bool bInSelection ) const { return true; }
 
+	/** @return True if this mode allows the viewport to use a drag tool */
+	virtual bool AllowsViewportDragTool() const { return true; }
+
 	/** Returns the editor mode identifier. */
 	FEditorModeID GetID() const { return Info.ID; }
 
@@ -547,6 +550,8 @@ public:
 	bool EndTracking(FEditorViewportClient* InViewportClient, FViewport* InViewport);
 	bool IsTracking() const { return bIsTracking; }
 
+	bool AllowsViewportDragTool() const;
+
 	/** Notifies all active modes that a map change has occured */
 	void MapChangeNotify();
 
diff --git a/Engine/Source/Editor/UnrealEd/Public/EditorViewportClient.h b/Engine/Source/Editor/UnrealEd/Public/EditorViewportClient.h
--- a/Engine/Source/Editor/UnrealEd/Public/EditorViewportClient.h
+++ b/Engine/Source/Editor/UnrealEd/Public/EditorViewportClient.h
@@ -611,6 +611,9 @@ public:
 	 */
 	virtual TSharedPtr<class FDragTool> MakeDragTool( EDragTool::Type DragToolType );
 
+	/** @return true if a drag tool can be used */
+	bool CanUseDragTool() const;
+
 	/** @return Whether or not to orbit the camera */
 	virtual bool ShouldOrbitCamera() const ;
 
