Disable steam on when PIE multiple processes

#ttp 332155 - +NEEDS REVIEW: Regression: EDITOR: SHOOTERGAME: DEDICATED SERVER: Clients get kicked to main menu

#branch UE4

#change "-MultiprocessOSS" is now specified on the commandline when launching multiple instances from PIE. If Steam detects this it auto-disables as it can't run the server and client on the same machine.

#change GetBool test could return false, indicating that it's not present in the config - yet the program may still want steam disabled for other reasons.

#remove Removed NoSteam test in IsEnabled - the reason being, if NoSteam was specified, the Steam module wouldn't have been loaded and thus the code for testing it not reached.

#change Also spits out a warning now when it fails to create the online subsystem, and if it just fails to create a server indicating to the user a workaround

diff --git a/Engine/Source/Editor/UnrealEd/Private/PlayLevel.cpp b/Engine/Source/Editor/UnrealEd/Private/PlayLevel.cpp
--- a/Engine/Source/Editor/UnrealEd/Private/PlayLevel.cpp
+++ b/Engine/Source/Editor/UnrealEd/Private/PlayLevel.cpp
@@ -685,7 +685,7 @@ void GetWindowSizeForInstanceType(FIntPoint &WindowSize, const ULevelEditorPlayS
  * @param	WinX			Window X position. This will contain the X position to use for the next window. (Not changed for dedicated server window).
  * @param	WinY			Window Y position. This will contain the X position to use for the next window. (Not changed for dedicated server window).
  * @param	InstanceNum		PIE instance index.
- * @param	IsServer		Is this instance a dedicated server. true if so else false.
+ * @param	IsServer		Is this instance a server. true if so else false.
  */
 FString GenerateCmdLineForNextPieInstance(int32 &WinX, int32 &WinY, int32 &InstanceNum, bool IsServer)
 {
@@ -697,7 +697,7 @@ FString GenerateCmdLineForNextPieInstance(int32 &WinX, int32 &WinY, int32 &Insta
 	//	-Override GameUserSettings.ini
 	//	-Force no steam
 	//	-Allow saving of config files (since we are giving them an override INI)
-	FString CmdLine = FString::Printf(TEXT("GameUserSettingsINI=%s -MultiprocessSaveConfig %s "), *GameUserSettingsOverride, *PlayInSettings->AdditionalLaunchOptions);
+	FString CmdLine = FString::Printf(TEXT("GameUserSettingsINI=%s -MultiprocessSaveConfig %s -MultiprocessOSS "), *GameUserSettingsOverride, *PlayInSettings->AdditionalLaunchOptions);
 
 	if (IsServer && PlayInSettings->PlayNetDedicated)
 	{
@@ -980,17 +980,14 @@ void UEditorEngine::PlayStandaloneLocalPc(FString MapNameOverride, FIntPoint* Wi
 {
 	const ULevelEditorPlaySettings* PlayInSettings = GetDefault<ULevelEditorPlaySettings>();
 	//const ULevelEditorPlaySettings* PlayInSettings = InPlaySettings != NULL ? InPlaySettings : GetDefault<ULevelEditorPlaySettings>();
+
 	FString CmdLine;
-	if (WindowPos != NULL)
+	if (WindowPos != NULL)	// If WindowPos == NULL, we're just launching one instance
 	{
 		CmdLine = GenerateCmdLineForNextPieInstance(WindowPos->X, WindowPos->Y, PIENum, bIsServer);
 	}
 	
-	FString URLParms;		
-	if (bIsServer == true)
-	{
-		URLParms = bIsServer == true ? TEXT("?Listen") : FString();
-	}
+	const FString URLParms = bIsServer == true ? TEXT("?Listen") : FString();
 	
 	// select map to play
 	TArray<FString> SavedMapNames;
diff --git a/Engine/Source/Runtime/Online/OnlineSubsystem/Private/OnlineSubsystemModule.cpp b/Engine/Source/Runtime/Online/OnlineSubsystem/Private/OnlineSubsystemModule.cpp
--- a/Engine/Source/Runtime/Online/OnlineSubsystem/Private/OnlineSubsystemModule.cpp
+++ b/Engine/Source/Runtime/Online/OnlineSubsystem/Private/OnlineSubsystemModule.cpp
@@ -199,6 +199,10 @@ IOnlineSubsystem* FOnlineSubsystemModule::GetOnlineSubsystem(const FName InSubsy
 					OnlineSubsystems.Add(KeyName, NewSubsystemInstance);
 					OnlineSubsystem = OnlineSubsystems.Find(KeyName);
 				}
+				else
+				{
+					UE_LOG(LogOnline, Warning, TEXT("Unable to create OnlineSubsystem module %s"), *SubsystemName.ToString());
+				}
 			}
 			else
 			{
diff --git a/Engine/Source/Runtime/Online/OnlineSubsystemSteam/Private/OnlineSubsystemSteam.cpp b/Engine/Source/Runtime/Online/OnlineSubsystemSteam/Private/OnlineSubsystemSteam.cpp
--- a/Engine/Source/Runtime/Online/OnlineSubsystemSteam/Private/OnlineSubsystemSteam.cpp
+++ b/Engine/Source/Runtime/Online/OnlineSubsystemSteam/Private/OnlineSubsystemSteam.cpp
@@ -341,6 +341,11 @@ bool FOnlineSubsystemSteam::Init()
 	}
 	else
 	{
+		// If the client succeeded, but the server didn't, this could be because there's a server and client running on the same machine - inform the user
+		if (bClientInitSuccess)
+		{
+			UE_LOG_ONLINE(Warning, TEXT("Failed to initialize Steam, this could be due to a Steam server and client running on the same machine. Try running with -NOSTEAM on the cmdline to disable."));
+		}
 		Shutdown();
 	}
 
@@ -406,10 +411,11 @@ bool FOnlineSubsystemSteam::IsEnabled()
 
 	// Check the ini for disabling Steam
 	bool bEnableSteam = true;
-	if (GConfig->GetBool(TEXT("OnlineSubsystemSteam"), TEXT("bEnabled"), bEnableSteam, GEngineIni) && bEnableSteam)
+	GConfig->GetBool(TEXT("OnlineSubsystemSteam"), TEXT("bEnabled"), bEnableSteam, GEngineIni);
+	if (bEnableSteam)
 	{
-		// Check the commandline for disabling Steam
-		bEnableSteam = !FParse::Param(FCommandLine::Get(),TEXT("NOSTEAM"));
+		// Steam doesn't support running both the server and client on the same machine
+		bEnableSteam = !FParse::Param(FCommandLine::Get(),TEXT("MultiprocessOSS"));
 #if UE_EDITOR
 		if (bEnableSteam)
 		{
