Fix for multiplayer PIE (viewport) not working for streaming levels. Sending server bad package names as they haven't had the PIE prefix removed

#ttp 225628 - +CRITICAL: Regression: SHOOTERGAME: PIE as client only opens one player window

#branch UE4

Note: Also requires fix 2284410, so that the server doesn't use Steam

diff --git a/Engine/Source/Editor/UnrealEd/Private/Editor.cpp b/Engine/Source/Editor/UnrealEd/Private/Editor.cpp
--- a/Engine/Source/Editor/UnrealEd/Private/Editor.cpp
+++ b/Engine/Source/Editor/UnrealEd/Private/Editor.cpp
@@ -6776,7 +6776,7 @@ FORCEINLINE bool NetworkRemapPath_local(FWorldContext &Context, FString &Str, bo
 bool UEditorEngine::NetworkRemapPath( UWorld *InWorld, FString &Str, bool reading)
 {
 	FWorldContext &Context = GetWorldContextFromWorldChecked(InWorld);
-	if (Context.PIEPrefix.IsEmpty() || Context.PIERemapPrefix.IsEmpty())
+	if ((Context.PIEPrefix.IsEmpty() && !reading) || (Context.PIERemapPrefix.IsEmpty() && reading))
 	{
 		return false;
 	}
@@ -6787,7 +6787,7 @@ bool UEditorEngine::NetworkRemapPath( UWorld *InWorld, FString &Str, bool readin
 bool UEditorEngine::NetworkRemapPath( UPendingNetGame *PendingNetGame, FString &Str, bool reading)
 {
 	FWorldContext &Context = GetWorldContextFromPendingNetGameChecked(PendingNetGame);
-	if (Context.PIEPrefix.IsEmpty() || Context.PIERemapPrefix.IsEmpty())
+	if ((Context.PIEPrefix.IsEmpty() && !reading) || (Context.PIERemapPrefix.IsEmpty() && reading))
 	{
 		return false;
 	}
