Turning off Surface Snapping in the editor viewport now prevents surface alignment during placement.

Previously this would only have been the case if the "Rotate to Surface Normal" option was disabled, but this was hidden and a little confusing.

I've also changed the actor positioning code to take parameter structs that have construction helpers so as to make the code more flexible to the client, whilst keeping the logic in one place.

I've fixed one or two bugs here as well:
	* Fixed vertex snapping helpers not showing up when dragging (they are still hidden when pasting)
	* Fixed point lights not having the same spawn offset as spot and directional lights
	* Fixed newly placed actors not getting their default rotations when not aligning to a surface

This addresses TTP#342860 - Static meshes that are drag and dropped into a world will have their rotation affected by the object the are dragged over
And TTP#341787 - CRITICAL: Regression: EDITOR: Direction light is at a flat angle by default

diff --git a/Engine/Source/Editor/UnrealEd/Private/AssetSelection.cpp b/Engine/Source/Editor/UnrealEd/Private/AssetSelection.cpp
--- a/Engine/Source/Editor/UnrealEd/Private/AssetSelection.cpp
+++ b/Engine/Source/Editor/UnrealEd/Private/AssetSelection.cpp
@@ -361,7 +361,12 @@ static AActor* PrivateAddActor( UObject* Asset, UActorFactory* Factory, bool Sel
 		return NULL;
 	}
 
-	const FTransform ActorTransform = FActorPositioning::GetSnappedSurfaceAlignedTransform(GCurrentLevelEditingViewportClient, Factory, GEditor->ClickLocation, GEditor->ClickPlane, NewActorTemplate->GetPlacementExtent());
+	const FSnappedPositioningData PositioningData = FSnappedPositioningData(GCurrentLevelEditingViewportClient, GEditor->ClickLocation, GEditor->ClickPlane)
+		.UseFactory(Factory)
+		.UseStartTransform(NewActorTemplate->GetTransform())
+		.UsePlacementExtent(NewActorTemplate->GetPlacementExtent());
+
+	const FTransform ActorTransform = FActorPositioning::GetSnappedSurfaceAlignedTransform(PositioningData);
 
 	// Do not fade snapping indicators over time if the viewport is not realtime
 	bool bClearImmediately = !GCurrentLevelEditingViewportClient || !GCurrentLevelEditingViewportClient->IsRealtime();
diff --git a/Engine/Source/Editor/UnrealEd/Private/Editor/ActorPositioning.cpp b/Engine/Source/Editor/UnrealEd/Private/Editor/ActorPositioning.cpp
--- a/Engine/Source/Editor/UnrealEd/Private/Editor/ActorPositioning.cpp
+++ b/Engine/Source/Editor/UnrealEd/Private/Editor/ActorPositioning.cpp
@@ -139,15 +139,17 @@ FTransform FActorPositioning::GetCurrentViewportPlacementTransform(const AActor&
 
 	if (CursorLocation.GetViewportType() == LVT_Perspective && !GCurrentLevelEditingViewportClient->Viewport->GetHitProxy( CursorPos.X, CursorPos.Y ))
 	{
-		ActorTransform.SetTranslation(FActorPositioning::GetActorPositionInFrontOfCamera(Actor, CursorLocation.GetOrigin(), CursorLocation.GetDirection()));
-	}
-	else if (bSnap)
-	{
-		ActorTransform = FActorPositioning::GetSnappedSurfaceAlignedTransform(GCurrentLevelEditingViewportClient, Factory, GEditor->ClickLocation, GEditor->ClickPlane, Actor.GetPlacementExtent());
+		ActorTransform.SetTranslation(GetActorPositionInFrontOfCamera(Actor, CursorLocation.GetOrigin(), CursorLocation.GetDirection()));
 	}
 	else
 	{
-		ActorTransform = FActorPositioning::GetSurfaceAlignedTransform(Factory, GEditor->ClickLocation, GEditor->ClickPlane, Actor.GetPlacementExtent());
+		const FSnappedPositioningData PositioningData = FSnappedPositioningData(GCurrentLevelEditingViewportClient, GEditor->ClickLocation, GEditor->ClickPlane)
+			.DrawSnapHelpers(true)
+			.UseFactory(Factory)
+			.UseStartTransform(Actor.GetTransform())
+			.UsePlacementExtent(Actor.GetPlacementExtent());
+
+		ActorTransform = bSnap ? GetSnappedSurfaceAlignedTransform(PositioningData) : GetSurfaceAlignedTransform(PositioningData);
 	}
 
 	return ActorTransform;
@@ -176,40 +178,41 @@ FVector FActorPositioning::GetActorPositionInFrontOfCamera(const AActor& InActor
 	return NewLocation;
 }
 
-FTransform FActorPositioning::GetSurfaceAlignedTransform(const UActorFactory* InFactory, const FVector& InSurfaceLocation, const FVector& SurfaceNormal, const FVector& InPlacementExtent, const FTransform& InStartTransform)
+FTransform FActorPositioning::GetSurfaceAlignedTransform(const FPositioningData& Data)
 {
 	// Sort out the rotation first, then do the location
-	FQuat RotatorQuat = InStartTransform.GetRotation();
-	if (InFactory && GetDefault<ULevelEditorViewportSettings>()->SnapToSurface.bSnapRotation)
+	FQuat RotatorQuat = Data.StartTransform.GetRotation();
+
+	if (Data.ActorFactory)
 	{
-		RotatorQuat = InFactory->AlignObjectToSurfaceNormal(SurfaceNormal, RotatorQuat);
+		RotatorQuat = Data.ActorFactory->AlignObjectToSurfaceNormal(Data.SurfaceNormal, RotatorQuat);
 	}
 
 	// Choose the largest location offset of the various options (global viewport settings, collision, factory offset)
 	const float SnapOffsetExtent = GetDefault<ULevelEditorViewportSettings>()->SnapToSurface.SnapOffsetExtent;
-	const float CollisionOffsetExtent = FVector::BoxPushOut(SurfaceNormal, InPlacementExtent);
+	const float CollisionOffsetExtent = FVector::BoxPushOut(Data.SurfaceNormal, Data.PlacementExtent);
 
-	FVector LocationOffset = SurfaceNormal * FMath::Max(SnapOffsetExtent, CollisionOffsetExtent);
-	if (InFactory && LocationOffset.SizeSquared() < InFactory->SpawnPositionOffset.SizeSquared())
+	FVector LocationOffset = Data.SurfaceNormal * FMath::Max(SnapOffsetExtent, CollisionOffsetExtent);
+	if (Data.ActorFactory && LocationOffset.SizeSquared() < Data.ActorFactory->SpawnPositionOffset.SizeSquared())
 	{
 		// Rotate the Spawn Position Offset to match our rotation
-		LocationOffset = RotatorQuat.RotateVector(-InFactory->SpawnPositionOffset);
+		LocationOffset = RotatorQuat.RotateVector(-Data.ActorFactory->SpawnPositionOffset);
 	}
 
-	return FTransform(RotatorQuat, InSurfaceLocation + LocationOffset);
+	return FTransform(Data.bAlignRotation ? RotatorQuat : Data.StartTransform.GetRotation(), Data.SurfaceLocation + LocationOffset);
 }
 
-FTransform FActorPositioning::GetSnappedSurfaceAlignedTransform(FLevelEditorViewportClient* InViewportClient, const UActorFactory* InFactory, FVector InSurfaceLocation, const FVector& SurfaceNormal, const FVector& InPlacementExtent, const FTransform& InStartTransform)
+FTransform FActorPositioning::GetSnappedSurfaceAlignedTransform(const FSnappedPositioningData& Data)
 {
-	FSnappingUtils::SnapPointToGrid(InSurfaceLocation, FVector(0.f));
+	FVector SnappedLocation = Data.SurfaceLocation;
+	FSnappingUtils::SnapPointToGrid(SnappedLocation, FVector(0.f));
 
-	bool bDrawVertexHelpers = false;
 	// Secondly, attempt vertex snapping
 	FVector AlignToNormal;
-	if (!InViewportClient || !FSnappingUtils::SnapLocationToNearestVertex( InSurfaceLocation, InViewportClient->GetDropPreviewLocation(), InViewportClient, AlignToNormal, bDrawVertexHelpers ))
+	if (!Data.LevelViewportClient || !FSnappingUtils::SnapLocationToNearestVertex( SnappedLocation, Data.LevelViewportClient->GetDropPreviewLocation(), Data.LevelViewportClient, AlignToNormal, Data.bDrawSnapHelpers ))
 	{
-		AlignToNormal = SurfaceNormal;
+		AlignToNormal = Data.SurfaceNormal;
 	}
 
-	return GetSurfaceAlignedTransform(InFactory, InSurfaceLocation, AlignToNormal, InPlacementExtent, InStartTransform);
+	return GetSurfaceAlignedTransform(Data);
 }
\ No newline at end of file
diff --git a/Engine/Source/Editor/UnrealEd/Private/Editor/ActorPositioning.h b/Engine/Source/Editor/UnrealEd/Private/Editor/ActorPositioning.h
--- a/Engine/Source/Editor/UnrealEd/Private/Editor/ActorPositioning.h
+++ b/Engine/Source/Editor/UnrealEd/Private/Editor/ActorPositioning.h
@@ -2,6 +2,65 @@
 
 #pragma once
 
+/** Positioning data struct */
+struct FPositioningData
+{
+	FPositioningData(const FVector& InSurfaceLocation, const FVector& InSurfaceNormal)
+		: SurfaceLocation(InSurfaceLocation)
+		, SurfaceNormal(InSurfaceNormal)
+		, ActorFactory(nullptr)
+		, StartTransform(FTransform::Identity)
+		, PlacementExtent(0.f)
+	{
+		auto& Settings = GetDefault<ULevelEditorViewportSettings>()->SnapToSurface;
+		bAlignRotation = Settings.bEnabled && Settings.bSnapRotation;
+	}
+
+	/** The surface location we want to position to */
+	FVector SurfaceLocation;
+
+	/** The surface normal we want to potentially align to */
+	FVector SurfaceNormal;
+
+	/** Placement extent offset to use (default = 0,0,0) */
+	FPositioningData& UsePlacementExtent(const FVector& InPlacementExtent) { PlacementExtent = InPlacementExtent; return *this; }
+	FVector PlacementExtent;
+
+	/** The start transform we are using for positioning. Ensures a natural alignment to the surface when using rotation. (default = Identity) */
+	FPositioningData& UseStartTransform(const FTransform& InStartTransform) { StartTransform = InStartTransform; return *this; }
+	FTransform StartTransform;
+
+	/** A factory to use for the alignment. Factories define the alignment routine and spawn offset amounts if necessary. */
+	FPositioningData& UseFactory(const UActorFactory* InActorFactory) { ActorFactory = InActorFactory; return *this; }
+	const UActorFactory* ActorFactory;
+
+	/** Whether to align to the surface normal, or just snap to its position */
+	FPositioningData& AlignToSurfaceRotation(bool bInAlignRotation) { bAlignRotation = bInAlignRotation; return *this; }
+	bool bAlignRotation;
+};
+
+/** Structure used for positioning actors with snapping */
+struct FSnappedPositioningData : FPositioningData
+{
+	FSnappedPositioningData(FLevelEditorViewportClient* InLevelViewportClient, const FVector& InSurfaceLocation, const FVector& InSurfaceNormal)
+		: FPositioningData(InSurfaceLocation, InSurfaceNormal)
+		, LevelViewportClient(InLevelViewportClient)
+		, bDrawSnapHelpers(false)
+	{}
+
+	/** The level viewport - required for vertex snapping routines */
+	FLevelEditorViewportClient* LevelViewportClient;
+
+	/** Whether to draw vertex snapping helpers or not when snapping */
+	FSnappedPositioningData& DrawSnapHelpers(bool bInDrawSnapHelpers) { bDrawSnapHelpers = bInDrawSnapHelpers; return *this; }
+	bool bDrawSnapHelpers;
+
+	/** Mask these construction helpers to return the correct type */
+	FSnappedPositioningData& UsePlacementExtent(const FVector& InPlacementExtent) { PlacementExtent = InPlacementExtent; return *this; }
+	FSnappedPositioningData& UseStartTransform(const FTransform& InStartTransform) { StartTransform = InStartTransform; return *this; }
+	FSnappedPositioningData& UseFactory(const UActorFactory* InActorFactory) { ActorFactory = InActorFactory; return *this; }
+	FSnappedPositioningData& AlignToSurfaceRotation(bool bInAlignRotation) { bAlignRotation = bInAlignRotation; return *this; }
+};
 
 struct FActorPositionTraceResult
 {
@@ -88,26 +147,18 @@ struct FActorPositioning
 	/**
 	 *	Get the snapped position and rotation transform for an actor to snap it to the specified surface. This will potentially align the actor to the specified surface normal, and offset it based on factory settings
 	 *
-	 *	@param	InFactory				The factory used to spawn the actor. Factories define how actors are aligned and offset.
-	 *	@param	InSurfaceLocation		The location on the surface to transform to
-	 *	@param	InSurfaceNormal			The surface normal of the spawn position
-	 *	@param	InPlacementExtent		(optional) Optional actor placement extent
-	 *	@param	InStartTransform		(optional) Optional initial transform. This will result in a transformation from InStartTransform, to the surface transform
+	 *	@param	PositioningData			Positioning data used to calculate the transformation
 	 *
 	 *	@return	The expected transform for the specified actor
 	 */
-	static UNREALED_API FTransform GetSnappedSurfaceAlignedTransform(FLevelEditorViewportClient* InViewportClient, const UActorFactory* InFactory, FVector InSurfaceLocation, const FVector& InSurfaceNormal, const FVector& InPlacementExtent = FVector(0.f), const FTransform& InStartTransform = FTransform::Identity);
+	static UNREALED_API FTransform GetSnappedSurfaceAlignedTransform(const FSnappedPositioningData& PositioningData);
 
 	/**
 	 *	Get the position and rotation transform for an actor to snap it to the specified surface. This will potentially align the actor to the specified surface normal, and offset it.
 	 *
-	 *	@param	InFactory				The factory used to spawn the actor. Factories define how actors are aligned and offset.
-	 *	@param	InSurfaceLocation		The location on the surface to transform to
-	 *	@param	InSurfaceNormal			The surface normal of the spawn position
-	 *	@param	InPlacementExtent		(optional) Optional actor placement extent
-	 *	@param	InStartTransform		(optional) Optional initial transform. This will result in a transformation from InStartTransform, to the surface transform
+	 *	@param	PositioningData			Positioning data used to calculate the transformation
 	 *
 	 *	@return	The expected transform for the specified actor
 	 */
-	static UNREALED_API FTransform GetSurfaceAlignedTransform(const UActorFactory* InFactory, const FVector& InSurfaceLocation, const FVector& InSurfaceNormal, const FVector& InPlacementExtent = FVector(0.f), const FTransform& InStartTransform = FTransform::Identity);
+	static UNREALED_API FTransform GetSurfaceAlignedTransform(const FPositioningData& PositioningData);
 };
\ No newline at end of file
diff --git a/Engine/Source/Editor/UnrealEd/Private/EditorServer.cpp b/Engine/Source/Editor/UnrealEd/Private/EditorServer.cpp
--- a/Engine/Source/Editor/UnrealEd/Private/EditorServer.cpp
+++ b/Engine/Source/Editor/UnrealEd/Private/EditorServer.cpp
@@ -3030,7 +3030,9 @@ void UEditorEngine::PasteSelectedActorsFromClipboard( UWorld* InWorld, const FTe
 		return;
 	}
 
-	FVector SaveClickLocation = FActorPositioning::GetSnappedSurfaceAlignedTransform(GCurrentLevelEditingViewportClient, nullptr, GEditor->ClickLocation, GEditor->ClickPlane).GetLocation();
+	const FSnappedPositioningData PositioningData = FSnappedPositioningData(GCurrentLevelEditingViewportClient, GEditor->ClickLocation, GEditor->ClickPlane)
+		.AlignToSurfaceRotation(false);
+	FVector SaveClickLocation = FActorPositioning::GetSnappedSurfaceAlignedTransform(PositioningData).GetLocation();
 	
 	ULevel* DesiredLevel = InWorld->GetCurrentLevel();
 
diff --git a/Engine/Source/Editor/UnrealEd/Private/Factories/ActorFactory.cpp b/Engine/Source/Editor/UnrealEd/Private/Factories/ActorFactory.cpp
--- a/Engine/Source/Editor/UnrealEd/Private/Factories/ActorFactory.cpp
+++ b/Engine/Source/Editor/UnrealEd/Private/Factories/ActorFactory.cpp
@@ -1319,6 +1319,8 @@ UActorFactoryPointLight::UActorFactoryPointLight(const class FPostConstructIniti
 {
 	DisplayName = LOCTEXT("PointLightDisplayName", "Point Light");
 	NewActorClass = APointLight::StaticClass();
+	SpawnPositionOffset = FVector(200, 0, 0);
+	bUseSurfaceOrientation = true;
 }
 
 /*-----------------------------------------------------------------------------
diff --git a/Engine/Source/Editor/UnrealEd/Private/LevelEditorViewport.cpp b/Engine/Source/Editor/UnrealEd/Private/LevelEditorViewport.cpp
--- a/Engine/Source/Editor/UnrealEd/Private/LevelEditorViewport.cpp
+++ b/Engine/Source/Editor/UnrealEd/Private/LevelEditorViewport.cpp
@@ -1011,7 +1011,14 @@ bool FLevelEditorViewportClient::UpdateDropPreviewActors(int32 MouseX, int32 Mou
 	for (AActor* Actor : DraggingActors)
 	{
 		const UActorFactory* ActorFactory = FactoryToUse ? FactoryToUse : GEditor->FindActorFactoryForActorClass(Actor->GetClass());
-		const FTransform ActorTransform = FActorPositioning::GetSnappedSurfaceAlignedTransform(this, ActorFactory, TraceResult.Location, TraceResult.SurfaceNormal, Actor->GetPlacementExtent());
+
+		const FSnappedPositioningData PositioningData = FSnappedPositioningData(this, TraceResult.Location, TraceResult.SurfaceNormal)
+			.DrawSnapHelpers(true)
+			.UseFactory(ActorFactory)
+			.UseStartTransform(PreDragActorTransforms.FindRef(Actor))
+			.UsePlacementExtent(Actor->GetPlacementExtent());
+
+		const FTransform ActorTransform = FActorPositioning::GetSnappedSurfaceAlignedTransform(PositioningData);
 
 		Actor->SetActorTransform(ActorTransform);
 		Actor->SetIsTemporarilyHiddenInEditor(false);
@@ -1216,6 +1223,9 @@ bool FLevelEditorViewportClient::DropObjectsAtCoordinates(int32 MouseX, int32 Mo
 				{
 					AActor* NewActor = *ActorIt;
 					DropPreviewActors.Add(NewActor);
+					
+					PreDragActorTransforms.Add(NewActor, NewActor->GetTransform());
+
 					NewActor->SetActorEnableCollision(false);
 
 					// Deselect if selected
@@ -1995,8 +2005,13 @@ void FLevelEditorViewportClient::ProjectActorsIntoWorld(const TArray<AActor*>& A
 			check(PreDragActorTransform);
 
 			// Compute the surface aligned transform. Note we do not use the snapped version here as our DragDelta is already snapped
-			FTransform ActorTransform = FActorPositioning::GetSurfaceAlignedTransform(Factory,
-				TraceResult.Location, TraceResult.SurfaceNormal, Actor->GetPlacementExtent(), *PreDragActorTransform);
+
+			const FPositioningData PositioningData = FPositioningData(TraceResult.Location, TraceResult.SurfaceNormal)
+				.UseStartTransform(*PreDragActorTransform)
+				.UsePlacementExtent(Actor->GetPlacementExtent())
+				.UseFactory(Factory);
+
+			FTransform ActorTransform = FActorPositioning::GetSurfaceAlignedTransform(PositioningData);
 			
 			ActorTransform.SetScale3D(Actor->GetActorScale3D());
 			if (auto* RootComponent = Actor->GetRootComponent())
@@ -2577,6 +2592,8 @@ void FLevelEditorViewportClient::TrackingStopped()
 
 		GEditor->RedrawLevelEditingViewports();
 	}
+
+	PreDragActorTransforms.Empty();
 }
 
 void FLevelEditorViewportClient::HandleViewportSettingChanged(FName PropertyName)
