Fix output spam from AnimBlueprint thumbnails

diff --git a/Engine/Source/Editor/UnrealEd/Private/ThumbnailHelpers.cpp b/Engine/Source/Editor/UnrealEd/Private/ThumbnailHelpers.cpp
--- a/Engine/Source/Editor/UnrealEd/Private/ThumbnailHelpers.cpp
+++ b/Engine/Source/Editor/UnrealEd/Private/ThumbnailHelpers.cpp
@@ -803,6 +803,7 @@ bool FAnimBlueprintThumbnailScene::SetAnimBlueprint(class UAnimBlueprint* InBlue
 	{
 		CleanupComponentChildren(PreviewActor->SkeletalMeshComponent);
 		PreviewActor->SkeletalMeshComponent->SetSkeletalMesh(nullptr);
+		PreviewActor->SkeletalMeshComponent->SetAnimInstanceClass(nullptr);
 	}
 
 	return bSetSucessfully;
