Fixing out of order initialization on clang platforms

diff --git a/Engine/Source/Editor/UnrealEd/Private/Editor/ActorPositioning.h b/Engine/Source/Editor/UnrealEd/Private/Editor/ActorPositioning.h
--- a/Engine/Source/Editor/UnrealEd/Private/Editor/ActorPositioning.h
+++ b/Engine/Source/Editor/UnrealEd/Private/Editor/ActorPositioning.h
@@ -8,9 +8,10 @@ struct FPositioningData
 	FPositioningData(const FVector& InSurfaceLocation, const FVector& InSurfaceNormal)
 		: SurfaceLocation(InSurfaceLocation)
 		, SurfaceNormal(InSurfaceNormal)
-		, ActorFactory(nullptr)
-		, StartTransform(FTransform::Identity)
 		, PlacementExtent(0.f)
+		, StartTransform(FTransform::Identity)
+		, ActorFactory(nullptr)
+		, bAlignRotation(false)
 	{
 		auto& Settings = GetDefault<ULevelEditorViewportSettings>()->SnapToSurface;
 		bAlignRotation = Settings.bEnabled && Settings.bSnapRotation;
