Can now specify existing blueprints to load for the blueprints-info commandlet.

diff --git a/Engine/Source/Editor/UnrealEd/Private/Commandlets/DumpBlueprintsInfoCommandlet.cpp b/Engine/Source/Editor/UnrealEd/Private/Commandlets/DumpBlueprintsInfoCommandlet.cpp
--- a/Engine/Source/Editor/UnrealEd/Private/Commandlets/DumpBlueprintsInfoCommandlet.cpp
+++ b/Engine/Source/Editor/UnrealEd/Private/Commandlets/DumpBlueprintsInfoCommandlet.cpp
@@ -25,6 +25,8 @@
 #include "AssetToolsModule.h"
 #include "BlueprintNodeSpawner.h"
 #include "AnimationGraphSchema.h"
+#include "IAssetRegistry.h"
+#include "PackageName.h"
 
 DEFINE_LOG_CATEGORY_STATIC(LogBlueprintInfoDump, Log, All);
 
@@ -106,6 +108,11 @@ DumpBlueprintsInfo commandlet params: \n\
     -interface=<Class>  Appends the desired interface to blueprints that are   \n\
                         being dumped. The <Class> param has to match a known   \n\
                         interface class.                                       \n\
+\n\
+    -loadBP=<Blueprint> Before recording any info, this will attempt to load   \n\
+                        the specified Blueprint. The blueprint name can contain\n\
+                        wildcards (to match multiple blueprints). Can also be  \n\
+                        set to \"all\", to load every non-developer blueprint. \n\
 \n\
     -help, -h, -?       Display this message and then exit.                    \n\
 \n");
@@ -248,7 +255,10 @@ DumpBlueprintsInfo commandlet params: \n\
 	static void GetComponentProperties(UBlueprint* Blueprint, TArray<UObjectProperty*>& PropertiesOut);
 	
 	/**
-	 * 
+	 * Dumps stats on the new blueprint menu system (database size, number of 
+	 * enteries, etc.).
+	 *
+	 * @return True if any data was written to the file, otherwise false (if nothing was dumped).
 	 */
 	static bool DumpActionDatabaseInfo(uint32 Indent, FArchive* FileOutWriter);
 
@@ -369,6 +379,26 @@ DumpBlueprintsInfo commandlet params: \n\
 	 * specified by the user).
 	 */
 	static void DiffDumpFiles(FString const& NewFilePath, FString const& OldFilePath, FString const& UserDiffCmd);
+
+	/**
+	 * Takes the user specified class name, and attempts to translate it into 
+	 * a class pointer. ClassName can be a blueprint name; if so, this will make
+	 * sure that blueprint is loaded and return that blueprint's generated class.
+	 * 
+	 * @param  ClassName	A name that the user provided, in attempt to identify a class.
+	 * @return Null if no matching class could be found, otherwise a valid class pointer.
+	 */
+	static UClass* GetUserNamedClass(FString ClassName);
+
+	/**
+	 * Attempts to load a subset of blueprint assets. The AssetName can be left
+	 * blank to force load all blueprints.
+	 * 
+	 * @param  AssetName			The name of the blueprint(s) that you want to load (can have wildcards). If left empty, will load all blueprints.
+	 * @param  bAllowDevBlueprints	If true, blueprint in the developer folder will also be loaded.
+	 * @return The number of blueprints that were loaded.
+	 */
+	static int32 LoadBlueprints(FString AssetName, bool bAllowDevBlueprints);
 }
 
 //------------------------------------------------------------------------------
@@ -387,7 +417,7 @@ DumpBlueprintInfoUtils::CommandletOptions::CommandletOptions(TArray<FString> con
 		{
 			FString ClassSwitch, ClassName;
 			Switch.Split(TEXT("="), &ClassSwitch, &ClassName);
-			BlueprintClass = FindObject<UClass>(ANY_PACKAGE, *ClassName);
+			BlueprintClass = GetUserNamedClass(ClassName);
 
 			if (BlueprintClass == nullptr)
 			{
@@ -399,7 +429,7 @@ DumpBlueprintInfoUtils::CommandletOptions::CommandletOptions(TArray<FString> con
 		{
 			FString ClassSwitch, ClassName;
 			Switch.Split(TEXT("="), &ClassSwitch, &ClassName);
-			PaletteFilter = FindObject<UClass>(ANY_PACKAGE, *ClassName);
+			PaletteFilter = GetUserNamedClass(ClassName);
 
 			NewDumpFlags |= BPDUMP_FilteredPalette;
 			if (PaletteFilter == nullptr)
@@ -464,7 +494,7 @@ DumpBlueprintInfoUtils::CommandletOptions::CommandletOptions(TArray<FString> con
 		{
 			FString ClassSwitch, ClassName;
 			Switch.Split(TEXT("="), &ClassSwitch, &ClassName);
-			SelectedObjectType = FindObject<UClass>(ANY_PACKAGE, *ClassName);
+			SelectedObjectType = GetUserNamedClass(ClassName);
 
 			if (!ClassName.Compare("all", ESearchCase::IgnoreCase))
 			{
@@ -554,7 +584,7 @@ DumpBlueprintInfoUtils::CommandletOptions::CommandletOptions(TArray<FString> con
 			FString InterfaceSwitch, InterfaceName;
 			Switch.Split(TEXT("="), &InterfaceSwitch, &InterfaceName);
 
-			if (UClass* Class = FindObject<UClass>(ANY_PACKAGE, *InterfaceName))
+			if (UClass* Class = GetUserNamedClass(InterfaceName))
 			{
 				if (Class->IsChildOf<UInterface>())
 				{
@@ -567,6 +597,19 @@ DumpBlueprintInfoUtils::CommandletOptions::CommandletOptions(TArray<FString> con
 				UE_LOG(LogBlueprintInfoDump, Warning, TEXT("Could not find a matching interface class matching this name: '%s'"), *InterfaceName);
 			}
 		}
+		else if (Switch.StartsWith("loadBP="))
+		{
+			FString LoadSwitch, BlueprintName;
+			Switch.Split(TEXT("="), &LoadSwitch, &BlueprintName);
+
+			if (!BlueprintName.Compare("all", ESearchCase::IgnoreCase))
+			{
+				BlueprintName = TEXT("");
+			}
+
+			bool const bAllowDevBlueprints = BlueprintName.IsEmpty() ? false : true;
+			LoadBlueprints(BlueprintName, bAllowDevBlueprints);
+		}
 		else if (!Switch.StartsWith("run=")) // account for the first switch, which invoked this commandlet
 		{
 			UE_LOG(LogBlueprintInfoDump, Warning, TEXT("Unrecognized command switch '%s', use -help for a listing of all accepted params"), *Switch);
@@ -672,6 +715,11 @@ static UBlueprint* DumpBlueprintInfoUtils::MakeTempBlueprint(UClass* ParentClass
 	{
 		MadeBlueprint = *FoundBlueprint;
 	}
+	else if (UBlueprint* ClassBlueprint = Cast<UBlueprint>(ParentClass->ClassGeneratedBy))
+	{
+		ClassBlueprints.Add(ParentClass, ClassBlueprint);
+		MadeBlueprint = ClassBlueprint;
+	}
 	else
 	{
 		UObject* BlueprintOuter = GetTransientPackage();
@@ -1461,8 +1509,14 @@ static void DumpBlueprintInfoUtils::DumpGraphContextActions(uint32 Indent, UEdGr
 	TArray<UObjectProperty*> ComponentProperties;
 	GetComponentProperties(Blueprint, ComponentProperties);
 
+	bool const bOnlyDumpSpecificComponents = (CommandOptions.SelectedObjectType != nullptr) &&
+		CommandOptions.SelectedObjectType->IsChildOf<UActorComponent>();
 	for (UObjectProperty* Component : ComponentProperties)
 	{
+		if (bOnlyDumpSpecificComponents && !Component->PropertyClass->IsChildOf(CommandOptions.SelectedObjectType))
+		{
+			continue;
+		}
 		UE_LOG(LogBlueprintInfoDump, Display, TEXT("%sDumping actions with component selection: '%s'..."), *BuildIndentString(Indent, true), *Component->GetName());
 
 		FString SelectionContextEntry = "," + IndentedNewline + "\"ComponentContextMenu-" + Component->GetName() + "\" : \n";
@@ -1908,6 +1962,119 @@ static void DumpBlueprintInfoUtils::DiffDumpFiles(FString const& NewFilePath, FS
 	}
 }
 
+//------------------------------------------------------------------------------
+static UClass* DumpBlueprintInfoUtils::GetUserNamedClass(FString ClassName)
+{
+	UClass* FoundClass = FindObject<UClass>(ANY_PACKAGE, *ClassName);
+
+	if (FoundClass == nullptr)
+	{
+		// maybe they meant a blueprint?
+		UBlueprint* Blueprint = FindObject<UBlueprint>(ANY_PACKAGE, *ClassName);
+		// maybe we have to load that blueprint?
+		if (Blueprint == nullptr)
+		{
+			// if this loaded something...
+			if (LoadBlueprints(ClassName, /*bAlloDevBlueprints =*/true) > 0)
+			{
+				Blueprint = FindObject<UBlueprint>(ANY_PACKAGE, *ClassName);
+				if (Blueprint != nullptr)
+				{
+					FoundClass = Blueprint->GeneratedClass;
+				}
+				else
+				{
+					FoundClass = FindObject<UClass>(ANY_PACKAGE, *ClassName);
+				}
+			}
+		}
+	}
+
+	return FoundClass;
+}
+
+//------------------------------------------------------------------------------
+static int32 DumpBlueprintInfoUtils::LoadBlueprints(FString AssetName, bool bAllowDevBlueprints)
+{
+	int32 LoadedCount = 0;
+
+	TArray<FString> PackagesToLoad;
+	if (!AssetName.IsEmpty())
+	{
+		FString PackageName = FString::Printf(TEXT("*%s*%s"), *AssetName, *FPackageName::GetAssetPackageExtension());
+
+		TArray<FString> UnusedPackageNames;
+		NormalizePackageNames(UnusedPackageNames, PackagesToLoad, PackageName, NORMALIZE_ExcludeMapPackages);
+
+		for (FString& FilePath : PackagesToLoad)
+		{
+			FString PackagePath;
+			if (FPackageName::TryConvertFilenameToLongPackageName(FilePath, PackagePath))
+			{
+				FilePath = PackagePath;
+			}
+		}
+	}
+
+	FAssetRegistryModule& AssetRegistryModule = FModuleManager::LoadModuleChecked<FAssetRegistryModule>(TEXT("AssetRegistry"));
+
+	static bool bAssetRegistryLoaded = false;
+	if (!bAssetRegistryLoaded)
+	{
+		UE_LOG(LogBlueprintInfoDump, Display, TEXT("Loading the asset registry..."));
+		AssetRegistryModule.Get().SearchAllAssets(/*bSynchronousSearch =*/true);
+		UE_LOG(LogBlueprintInfoDump, Display, TEXT("Finished loading the asset registry."));
+		bAssetRegistryLoaded = true;
+	}
+	
+	TArray<FAssetData> BlueprintAssetList;
+	AssetRegistryModule.Get().GetAssetsByClass(UBlueprint::StaticClass()->GetFName(), BlueprintAssetList);
+	AssetRegistryModule.Get().GetAssetsByClass(UAnimBlueprint::StaticClass()->GetFName(), BlueprintAssetList);
+
+ 	FString DevelopersRoot = FPaths::GameDevelopersDir().LeftChop(1);
+	FPackageName::TryConvertFilenameToLongPackageName(DevelopersRoot, DevelopersRoot);
+
+	bool const bLoadPackageSubset = (PackagesToLoad.Num() > 0);
+	for (FAssetData const& Asset : BlueprintAssetList)
+	{
+		if (Asset.IsAssetLoaded())
+		{
+			continue;
+		}
+
+		FString const PackageName = Asset.PackageName.ToString();
+		if (!bAllowDevBlueprints && PackageName.StartsWith(DevelopersRoot))
+		{
+			continue;
+		}
+
+		if (bLoadPackageSubset && !PackagesToLoad.Contains(PackageName))
+		{
+			continue;
+		}
+
+		FString const AssetPath = Asset.ObjectPath.ToString();
+		UE_LOG(LogBlueprintInfoDump, Display, TEXT("Loading '%s'..."), *AssetPath);
+
+		UBlueprint* LoadedBlueprint = Cast<UBlueprint>(StaticLoadObject(Asset.GetClass(), /*Outer =*/nullptr, *AssetPath));
+		if (LoadedBlueprint != nullptr)
+		{
+			++LoadedCount;
+		}
+		else
+		{
+			UE_LOG(LogBlueprintInfoDump, Warning, TEXT("Failed to load: '%s'."), *AssetPath);
+		}
+
+		if (bLoadPackageSubset && (LoadedCount >= PackagesToLoad.Num()))
+		{
+			break;
+		}
+ 	}
+
+	return LoadedCount;
+}
+
 /*******************************************************************************
  * UDumpBlueprintsInfoCommandlet
  ******************************************************************************/
