Fixed uses of the LOCTEXT_NAMESPACE macro as a parameter. It causes parse errors when gathering text from source, as it is an unexpected use of the macro.

diff --git a/Engine/Source/Editor/LandscapeEditor/Private/LandscapeEditorDetails.cpp b/Engine/Source/Editor/LandscapeEditor/Private/LandscapeEditorDetails.cpp
--- a/Engine/Source/Editor/LandscapeEditor/Private/LandscapeEditorDetails.cpp
+++ b/Engine/Source/Editor/LandscapeEditor/Private/LandscapeEditorDetails.cpp
@@ -190,7 +190,7 @@ FText FLandscapeEditorDetails::GetLocalizedName(FString Name)
 	}
 
 	FText Result;
-	ensure(FText::FindText(LOCTEXT_NAMESPACE, Name, Result));
+	ensure(FText::FindText(TEXT("LandscapeEditor"), Name, Result));
 	return Result;
 }
 
diff --git a/Engine/Source/Editor/UnrealEd/Private/UnrealEdEngine.cpp b/Engine/Source/Editor/UnrealEd/Private/UnrealEdEngine.cpp
--- a/Engine/Source/Editor/UnrealEd/Private/UnrealEdEngine.cpp
+++ b/Engine/Source/Editor/UnrealEd/Private/UnrealEdEngine.cpp
@@ -580,7 +580,7 @@ FString FClassPickerDefaults::GetName() const
 	}
 
 	FText OutName;
-	if ( FText::FindText(TEXT(LOCTEXT_NAMESPACE), LocTextNameID, OutName) )
+	if ( FText::FindText(TEXT("UnrealEd"), LocTextNameID, OutName) )
 	{
 		return OutName.ToString();
 	}
@@ -614,7 +614,7 @@ FString FClassPickerDefaults::GetDescription() const
 	}
 
 	FText OutDesc;
-	if ( FText::FindText(TEXT(LOCTEXT_NAMESPACE), LocTextDescriptionID, OutDesc) )
+	if ( FText::FindText(TEXT("UnrealEd"), LocTextDescriptionID, OutDesc) )
 	{
 		return OutDesc.ToString();
 	}
