Some missed includes for non-unity builds.

diff --git a/Engine/Source/Editor/UnrealEd/Private/Commandlets/DumpBlueprintsInfoCommandlet.cpp b/Engine/Source/Editor/UnrealEd/Private/Commandlets/DumpBlueprintsInfoCommandlet.cpp
--- a/Engine/Source/Editor/UnrealEd/Private/Commandlets/DumpBlueprintsInfoCommandlet.cpp
+++ b/Engine/Source/Editor/UnrealEd/Private/Commandlets/DumpBlueprintsInfoCommandlet.cpp
@@ -25,8 +25,9 @@
 #include "AssetToolsModule.h"
 #include "BlueprintNodeSpawner.h"
 #include "AnimationGraphSchema.h"
-#include "IAssetRegistry.h"
+#include "AssetRegistryModule.h"
 #include "PackageName.h"
+#include "PackageHelperFunctions.h" // for NORMALIZE_ExcludeMapPackages
 
 DEFINE_LOG_CATEGORY_STATIC(LogBlueprintInfoDump, Log, All);
 
@@ -254,6 +255,12 @@ DumpBlueprintsInfo commandlet params: \n\
 	 */
 	static void GetComponentProperties(UBlueprint* Blueprint, TArray<UObjectProperty*>& PropertiesOut);
 	
+
+	/**
+	 * 
+	 */
+	static FString BuildByteSizeString(int32 ByteSize);
+
 	/**
 	 * Dumps stats on the new blueprint menu system (database size, number of 
 	 * enteries, etc.).
@@ -1039,6 +1046,38 @@ static void DumpBlueprintInfoUtils::GetComponentProperties(UBlueprint* Blueprint
 	}
 }
 
+//------------------------------------------------------------------------------
+static FString DumpBlueprintInfoUtils::BuildByteSizeString(int32 ByteSize)
+{
+	enum 
+	{
+		Bytes,
+		KiloBytes,
+		MegaBytes,
+		GigaBytes,
+
+		ByteUnits_MAX
+	};
+
+	TCHAR* ByteUnits[ByteUnits_MAX];
+	ByteUnits[Bytes] = TEXT("Bytes");
+	ByteUnits[KiloBytes] = TEXT("KB");
+	ByteUnits[MegaBytes] = TEXT("MB");
+	ByteUnits[GigaBytes] = TEXT("GB");
+
+	int32 UnitsIndex = 0;
+	float ConvertedSize = ByteSize;
+
+	float const MetricStepSize = 1024.0f;
+	while ((ConvertedSize > MetricStepSize) && (UnitsIndex < ByteUnits_MAX))
+	{
+		ConvertedSize /= MetricStepSize;
+		++UnitsIndex;
+	}
+
+	return FString::Printf(TEXT("%.2f %s"), ConvertedSize, ByteUnits[UnitsIndex]);
+}
+
 //------------------------------------------------------------------------------
 static bool DumpBlueprintInfoUtils::DumpActionDatabaseInfo(uint32 Indent, FArchive* FileOutWriter)
 {
@@ -1089,8 +1128,7 @@ static bool DumpBlueprintInfoUtils::DumpActionDatabaseInfo(uint32 Indent, FArchi
 			}
 		}
 
-		float const DatabaseKbSize = Database.EstimatedSize() / 1024.f;
-		float EstimatedCacheKbSize = 0.0;
+		int32 EstimatedCacheKbSize = 0;
 		for (UBlueprint* CacheBlueprint : TemplateOuters)
 		{
 			TArray<UObject*> ChildObjs;
@@ -1102,7 +1140,13 @@ static bool DumpBlueprintInfoUtils::DumpActionDatabaseInfo(uint32 Indent, FArchi
 				EstimatedCacheKbSize += sizeof(*ChildObj);
 			}
 		}
-		EstimatedCacheKbSize /= 1024.f;
+
+		int32 const DatabaseSize = Database.EstimatedSize();
+		FString const DatabaseSizeStr = BuildByteSizeString(DatabaseSize);
+		FString const NodeCacheSizeStr = BuildByteSizeString(EstimatedCacheKbSize);
+		FString const TotalSystemSizeStr = BuildByteSizeString(DatabaseSize + EstimatedCacheKbSize);
+		FString const AvgDatabaseEntrySizeStr = BuildByteSizeString(DatabaseSize / DatabaseCount);
+		FString const AvgCachedNodeSizeStr = BuildByteSizeString(EstimatedCacheKbSize / TemplateCount);
 
 		FString const OriginalIndent  = BuildIndentString(Indent);
 		FString const IndentedNewline = "\n" + BuildIndentString(Indent + 1);
@@ -1112,11 +1156,11 @@ static bool DumpBlueprintInfoUtils::DumpActionDatabaseInfo(uint32 Indent, FArchi
 		DatabaseInfoEntry += FString::Printf(TEXT("%s\"CachingNodesDuration\"   : \"%.3f seconds\","), *IndentedNewline, NodeCachingDuration);
 		DatabaseInfoEntry += FString::Printf(TEXT("%s\"DatabaseEntryCount\"     : %d,"), *IndentedNewline, DatabaseCount);
 		DatabaseInfoEntry += FString::Printf(TEXT("%s\"NodeCacheCount\"         : %d,"), *IndentedNewline, TemplateCount);
-		DatabaseInfoEntry += FString::Printf(TEXT("%s\"DatabaseSize\"           : \"%.2f KB\","), *IndentedNewline, DatabaseKbSize);
-		DatabaseInfoEntry += FString::Printf(TEXT("%s\"AvgSizePerDatabaseEntry\": \"%.2f bytes\","), *IndentedNewline, (DatabaseKbSize * 1024.f) / DatabaseCount);
-		DatabaseInfoEntry += FString::Printf(TEXT("%s\"EstimatedNodeCacheSize\" : \"%.2f KB\","), *IndentedNewline, EstimatedCacheKbSize);
-		DatabaseInfoEntry += FString::Printf(TEXT("%s\"AvgSizePerCachedNode\"   : \"%.2f bytes\","), *IndentedNewline, (EstimatedCacheKbSize * 1024.f)/ TemplateCount);
-		DatabaseInfoEntry += FString::Printf(TEXT("%s\"EstimatedSystemSize\"    : \"%.2f KB\"\n%s}"), *IndentedNewline, DatabaseKbSize + EstimatedCacheKbSize, *OriginalIndent);
+		DatabaseInfoEntry += FString::Printf(TEXT("%s\"DatabaseSize\"           : \"%s\","), *IndentedNewline, *DatabaseSizeStr);
+		DatabaseInfoEntry += FString::Printf(TEXT("%s\"AvgSizePerDatabaseEntry\": \"%s\","), *IndentedNewline, *AvgDatabaseEntrySizeStr);
+		DatabaseInfoEntry += FString::Printf(TEXT("%s\"EstimatedNodeCacheSize\" : \"%s\","), *IndentedNewline, *NodeCacheSizeStr);
+		DatabaseInfoEntry += FString::Printf(TEXT("%s\"AvgSizePerCachedNode\"   : \"%s\","), *IndentedNewline, *AvgCachedNodeSizeStr);
+		DatabaseInfoEntry += FString::Printf(TEXT("%s\"EstimatedSystemSize\"    : \"%s\"\n%s}"), *IndentedNewline, *TotalSystemSizeStr, *OriginalIndent);
 		FileOutWriter->Serialize(TCHAR_TO_ANSI(*DatabaseInfoEntry), DatabaseInfoEntry.Len());
 		bWroteToFile = true;
 	}
