EditorPerformanceTests
--Updates to the Editor Performance test.
--Obtains the FPS, Frame Time, and Working Memory set and saves them to a csv file.

diff --git a/Engine/Source/Editor/UnrealEd/Private/Tests/AutomationEditorCommon.cpp b/Engine/Source/Editor/UnrealEd/Private/Tests/AutomationEditorCommon.cpp
--- a/Engine/Source/Editor/UnrealEd/Private/Tests/AutomationEditorCommon.cpp
+++ b/Engine/Source/Editor/UnrealEd/Private/Tests/AutomationEditorCommon.cpp
@@ -6,7 +6,7 @@
 #include "ModuleManager.h"
 #include "LevelEditor.h"
 #include "Editor/MainFrame/Public/MainFrame.h"
-
+#include "EngineVersion.h"
 
 
 DEFINE_LOG_CATEGORY_STATIC(LogAutomationEditorCommon, Log, All);
@@ -279,6 +279,79 @@ namespace AutomationEditorCommonUtils
 	}
 }
 
+/**
+* Writes a number to a text file.
+* @param FolderNameForTest is the folder that has the same name as the test. (For Example: "Performance").
+* @param FolderNameForBeingTested is the name for the thing that is being tested. (For Example: "MapName").
+* @param FileName is the name of the file with an extension (meaning no need to add the .txt)
+* @param NumberToWriteToFile is the float number that is expected to be written to the file.
+* @param Delimiter is the delimiter to be used. TEXT(",")
+*/
+void WriteToTextFile(FString FolderNameForTest ,FString FolderNameForBeingTested, FString FileName, float NumberToWriteToFile, FString Delimiter)
+{
+	//Performance file locations and setups.
+	FString FileSaveLocation = FPaths::AutomationLogDir() + FolderNameForTest + TEXT("/") + FolderNameForBeingTested + TEXT("/") + FileName + TEXT(".txt");
+
+	//Variables that hold the content from the text files
+	FString ReadFromTextFile;
+
+	//Log out to a text file the Duration.
+	FString CurrentNumberToWrite = FString::Printf(TEXT("%f"), NumberToWriteToFile);
+	FFileHelper::LoadFileToString(ReadFromTextFile, *FileSaveLocation);
+	FString FileSetup = ReadFromTextFile + CurrentNumberToWrite + Delimiter;
+	FFileHelper::SaveStringToFile(FileSetup, *FileSaveLocation);
+
+}
+
+/**
+* Returns the sum of the numbers available in a FString array.
+* @param NumberArray is the name of the array intended to be used.
+* @param bAverageOnly will return the average of the available numbers instead of the sum.
+*/
+float ValueFromStringArrayofNumbers(const TArray<FString>& NumberArray, bool bAverageInstead)
+{
+	//Total Value holds the sum of all the numbers available in the array.
+	float TotalValue = 0;
+	int32 NumberCounter = 0;
+
+	//Loop through the array.
+	for (int32 Index = 0; Index < NumberArray.Num(); ++Index)
+	{
+		//If the character in the array is a number then we'll want to add it to our total.
+		if (NumberArray[Index].IsNumeric())
+		{
+			NumberCounter++;
+			TotalValue += FCString::Atof(NumberArray[Index].GetCharArray().GetData());
+		}
+	}	
+
+	//If bAverageInstead equals true then only the average is returned.
+	//Otherwise return the total value.
+	if (bAverageInstead)
+	{
+		UE_LOG(LogEditorAutomationTests, Log, TEXT("Average value of the Array is %f"), (TotalValue / NumberCounter));
+		return (TotalValue / NumberCounter);
+	}
+	
+	UE_LOG(LogEditorAutomationTests, Log, TEXT("Total Value of the Array is %f"), TotalValue);
+	return TotalValue;
+
+}
+
+TArray<FString> CreateArrayFromFile(FString FullFileLocation)
+{
+
+	FString RawData;
+	TArray<FString> DataArray;
+
+	FFileHelper::LoadFileToString(RawData, *FullFileLocation);
+
+	RawData.ParseIntoArray(&DataArray, TEXT(","), false);
+		//ParseIntoArray(&DataArray, Delimiter, true);
+	
+	return DataArray;
+}
+
 ///////////////////////////////////////////////////////////////////////
 // Common Latent commands
 
@@ -369,83 +442,195 @@ bool FEndPlayMapCommand::Update()
 */
 bool FEditorPerformanceCommand::Update()
 {
-	//Get the base filename for the map that will be used.
-	FString ShortMapName = FPaths::GetBaseFilename(BaseMapName);	
-	
-	//Map Load time variables
-	double MapLoadStartTime = 0.0f;
-	double MapLoadTime = 0.0f;
-
-	//Gets the main frame module
+	//Gets the main frame module to get the name of our current level.
 	const IMainFrameModule& MainFrameModule = FModuleManager::GetModuleChecked< IMainFrameModule >("MainFrame");
-
+	FString ShortMapName = MainFrameModule.GetLoadedLevelName();
+	
 	//Variables that hold the content from the text files
-	FString SavedMapLoadTime;
 	FString SavedAverageFPS;
 	FString SavedMemory;
+	FString SavedDuration;
 	static SIZE_T StaticLastTotalAllocated = 0;
 
-	//Performance file locations and setups.
-	FString FileSaveLocation = FPaths::AutomationLogDir() + TEXT("Performance/") + ShortMapName + TEXT("/");
-	FString FPSFileSaveLocation = FileSaveLocation + TEXT("RawFPS.txt");
-	FString MemoryFileSaveLocation = FileSaveLocation + TEXT("RawMemory.txt");
-	FString MapLoadTimeFileSaveLocation = FileSaveLocation + TEXT("RAWMapLoadTime.txt");
+	//Variables needed for the WriteToTextFile function.
+	FString TestName = TEXT("Performance");
+	FString Delimiter = TEXT(",");
+	FString FileName;
 	
-	//Open the map if it hasn't been loaded already.
-	if (MainFrameModule.GetLoadedLevelName() != ShortMapName)
+	//Grab the FPS and working memory number.
+	float CurrentTime = FPlatformTime::Seconds();
+	if ((CurrentTime - StartTime) <= Duration)
 	{
-		UE_LOG(LogEditorAutomationTests, Log, TEXT("Intended map is not loaded in the editor.  Loading '%s' now."), *ShortMapName);
+		//Find the Average FPS
+		//Clamp to avoid huge averages at startup or after hitches
+		const float CurrentFPS = 1.0f / FSlateApplication::Get().GetAverageDeltaTime();
+		const float ClampedFPS = (CurrentFPS < 0.0f || CurrentFPS > 4000.0f) ? 0.0f : CurrentFPS;
 		
-		//Get the current number of seconds.  This will be used to track how long it took to load the map.
-		MapLoadStartTime = FPlatformTime::Seconds();
-		//Load the map
-		FEditorAutomationTestUtilities::LoadMap(BaseMapName);
-		//This is the time it took to load the map in the editor.
-		MapLoadTime = FPlatformTime::Seconds() - MapLoadStartTime;
-
-		UE_LOG(LogEditorAutomationTests, Log, TEXT("%s has been loaded."), *ShortMapName);
-
-		//Log out to a text file the time it takes to load the map.
-		FString LoadTime = FString::Printf(TEXT("%.3f"), MapLoadTime);
-		FFileHelper::LoadFileToString(SavedMapLoadTime, *MapLoadTimeFileSaveLocation);
-		FString MapLoadTimeFileSetup = SavedMapLoadTime + LoadTime + TEXT(",");
-		FFileHelper::SaveStringToFile(MapLoadTimeFileSetup, *MapLoadTimeFileSaveLocation);
-
-		//Display the time it took to launch the map.
-		UE_LOG(LogEditorAutomationTests, Display, TEXT("Map '%s' took '%.3f' to load"), *BaseMapName, MapLoadTime);
-		return true;
+		//Find the Frame Time in ms.
+		//Clamp to avoid huge averages at startup or after hitches
+		const float AverageMS = FSlateApplication::Get().GetAverageDeltaTime() * 1000.0f;
+		const float ClampedMS = (AverageMS < 0.0f || AverageMS > 4000.0f) ? 0.0f : AverageMS;
+
+		//Query OS for process memory used.
+		FPlatformMemoryStats MemoryStats = FPlatformMemory::GetStats();
+		StaticLastTotalAllocated = MemoryStats.UsedPhysical;
+
+		//Log out to a text file the Duration.
+		WriteToTextFile(TestName, ShortMapName, FileName = TEXT("RAWDuration"), CurrentTime, Delimiter);
+		//Log out to a text file the AverageFPS
+		WriteToTextFile(TestName, ShortMapName, FileName = TEXT("RAWAverageFPS"), ClampedFPS, Delimiter);
+		//Log out to a text file the AverageFrameTime
+		WriteToTextFile(TestName, ShortMapName, FileName = TEXT("RAWFrameTime"), ClampedMS, Delimiter);
+		//Log out to a text file the Memory in kilobytes.
+		WriteToTextFile(TestName, ShortMapName, FileName = TEXT("RAWMemoryUsedPhysical"), (float)StaticLastTotalAllocated / 1024, Delimiter);
+
+		return false;
 	}
 	
-	//FPS capture
-	const float CurrentFPS = 1.0f / FSlateApplication::Get().GetAverageDeltaTime();
-	const float ClampedFPS = (CurrentFPS < 0.3f || CurrentFPS > 4000.0f) ? 0.0f : CurrentFPS;
-	//Log out to a text file the AverageFPS
-	FString AverageFPS = FString::Printf(TEXT("%3.1f"), ClampedFPS);
-	FFileHelper::LoadFileToString(SavedAverageFPS, *FPSFileSaveLocation);
-	FString FPSFileSetup = SavedAverageFPS + AverageFPS + TEXT(",");
-	FFileHelper::SaveStringToFile(FPSFileSetup, *FPSFileSaveLocation);
-
-	//Working Set Memory capture
-	// Query OS for process memory used
-	FPlatformMemoryStats MemoryStats = FPlatformMemory::GetStats();
-	StaticLastTotalAllocated = MemoryStats.UsedPhysical;
-	//Log out to a text file the Memory in kilobytes.
-	FString Memory = FString::Printf(TEXT("%f"), ((float)StaticLastTotalAllocated / 1024));
-	FFileHelper::LoadFileToString(SavedMemory, *MemoryFileSaveLocation);
-	FString MemoryFileSetup = SavedMemory + Memory + TEXT(",");
-	FFileHelper::SaveStringToFile(MemoryFileSetup, *MemoryFileSaveLocation);
+	UE_LOG(LogEditorAutomationTests, Display, TEXT("Performance has been captured and saved to the '//Game/Saved/Automation/Log/Performance/%s' folder location."), *ShortMapName);
 	
 	return true;
 }
 
+bool FGenerateEditorPerformanceCharts::Update()
+{
+	//Get the current changelist number
+	FString Changelist = GEngineVersion.ToString(EVersionComponent::Changelist);
+
+	//The locations of the raw data.
+	FString FileSaveLocation = FPaths::AutomationLogDir() + TEXT("Performance/") + MapName + TEXT("/");
+	FString FPSFileSaveLocation = FileSaveLocation + TEXT("RAWAverageFPS.txt");
+	FString FrameTimeSaveLocation = FileSaveLocation + TEXT("RAWFrameTime.txt");
+	FString MemoryFileSaveLocation = FileSaveLocation + TEXT("RAWMemoryUsedPhysical.txt");
+	FString MapLoadTimeFileSaveLocation = FileSaveLocation + TEXT("RAWMapLoadTime.txt");
+	FString DurationFileSaveLocation = FileSaveLocation + TEXT("RAWDuration.txt");
+	//Get the data from the raw files and put them into an array.
+	TArray<FString> FPSRawArray = CreateArrayFromFile(FPSFileSaveLocation);
+	TArray<FString> FrameTimeRawArray = CreateArrayFromFile(FrameTimeSaveLocation);
+	TArray<FString> MemoryRawArray = CreateArrayFromFile(MemoryFileSaveLocation);
+	TArray<FString> MapLoadTimeRawArray = CreateArrayFromFile(MapLoadTimeFileSaveLocation);
+	TArray<FString> TestRunTimeRawArray = CreateArrayFromFile(DurationFileSaveLocation);
+
+	//The CSV files.
+	//RAW csv holds all of the information from the text files
+	//Final csv only holds the averaged information.
+	FString FinalCSVFilename = FString::Printf(TEXT("%s%s_Performance.csv"), *FileSaveLocation, *MapName);
+	FString RAWCSVFilename = FString::Printf(TEXT("%sRAW_%s_%s.csv"), *FileSaveLocation, *MapName, *FDateTime::Now().ToString());
+
+	//Create the .csv file that will hold the raw data from the text files.
+	FArchive* RAWCSVFile = IFileManager::Get().CreateFileWriter(*RAWCSVFilename);
+	if (!RAWCSVFile)
+	{
+		UE_LOG(LogEditorAutomationTests, Error, TEXT("Failed to create the csv file: %s"), *RAWCSVFilename);
+		return false;
+	}
+
+	//Add the top title row to the raw csv file
+	FString RAWCSVLine = (TEXT("Map Name, Changelist, Test Run Time, Map Load Time, Average FPS, Frame Time, Used Physical Memory\n"));
+	RAWCSVFile->Serialize(TCHAR_TO_ANSI(*RAWCSVLine), RAWCSVLine.Len());
+	//Loop through the text files and add them to the raw csv file.
+	for (int32 IterLoop = 0; IterLoop < FPSRawArray.Num(); IterLoop++)
+	{
+		RAWCSVLine = MapName + TEXT(",") + Changelist + TEXT(",") + TestRunTimeRawArray[IterLoop] + TEXT(",") + MapLoadTimeRawArray[0] + TEXT(",") + FPSRawArray[IterLoop] + TEXT(",") + FrameTimeRawArray[IterLoop] + TEXT(",") + MemoryRawArray[IterLoop] + LINE_TERMINATOR;
+		RAWCSVFile->Serialize(TCHAR_TO_ANSI(*RAWCSVLine), RAWCSVLine.Len());
+	}
+	
+	////Create the Performance .csv file	
+	//FArchive* FinalCSVFile = IFileManager::Get().CreateFileReader(*FinalCSVFilename);
+	//if (!FinalCSVFile)
+	//{
+	//	UE_LOG(LogEditorAutomationTests, Error, TEXT("Failed to create the csv file: %s"), *FinalCSVFilename);
+	//	return false;
+	//}
+
+
+
+	//Get the final data for the Performance csv file.
+	//This gets the averaged numbers using the raw data arrays.
+	float FPSAverage = ValueFromStringArrayofNumbers(FPSRawArray, true);
+	float FrameTimeAverage = ValueFromStringArrayofNumbers(FrameTimeRawArray, true);
+	float MemoryAverage = ValueFromStringArrayofNumbers(MemoryRawArray, true);
+	//This gets the actual test run time by subtracting the last test run entry against the first test run entry.
+	FString TestRunStartTime = TestRunTimeRawArray[0];
+	FString TestRunEndTime;
+	for (int32 i = (TestRunTimeRawArray.Num() - 1); i >= 0 ; --i)
+	{
+		if (!TestRunTimeRawArray[i].IsEmpty())
+		{
+			TestRunEndTime = TestRunTimeRawArray[i];
+			break;
+		}
+	}
+	float TestRunTime = FCString::Atof(TestRunEndTime.GetCharArray().GetData()) - FCString::Atof(TestRunStartTime.GetCharArray().GetData());
+	//Get the map load time.
+	FString MapLoadTime = MapLoadTimeRawArray[0];
+	
+	FString OldPerformanceCSVFile;
+
+	//If the Performance csv file does not have a size then it must be new.
+	//This adds the title row to the performance csv file.
+	if (IFileManager::Get().FileSize(FinalCSVFilename.GetCharArray().GetData()) <= 0)
+	{
+		FArchive* FinalCSVFile = IFileManager::Get().CreateFileWriter(*FinalCSVFilename);
+		if (!FinalCSVFile)
+		{
+			UE_LOG(LogEditorAutomationTests, Error, TEXT("Failed to create the csv file: %s"), *FinalCSVFilename);
+			return false;
+		}
+		FString FinalCSVLine = (TEXT("Date, Map Name, Changelist, Test Run Time, Map Load Time, Average FPS, Average MS, Used Physical KB\n"));
+		FinalCSVFile->Serialize(TCHAR_TO_ANSI(*FinalCSVLine), FinalCSVLine.Len());
+	}
+
+	FFileHelper::LoadFileToString(OldPerformanceCSVFile, *FinalCSVFilename);
+	FArchive* FinalCSVFile = IFileManager::Get().CreateFileWriter(*FinalCSVFilename);
+	if (!FinalCSVFile)
+	{
+		UE_LOG(LogEditorAutomationTests, Error, TEXT("Failed to create the csv file: %s"), *FinalCSVFilename);
+		return false;
+	}
+	FinalCSVFile->Serialize(TCHAR_TO_ANSI(*OldPerformanceCSVFile), OldPerformanceCSVFile.Len());
+
+
+
+	//Write the averaged numbers to the Performance CSV file.
+	FString FinalCSVLine = FString::Printf(TEXT("%s,%s,%s,%.0f,%s,%.3f,%.3f,%.0f%s"), *FDateTime::Today().ToString(), *MapName, *Changelist, TestRunTime, *MapLoadTime, FPSAverage, FrameTimeAverage, MemoryAverage, LINE_TERMINATOR);
+	FinalCSVFile->Serialize(TCHAR_TO_ANSI(*FinalCSVLine), FinalCSVLine.Len());
+
+	//Display the averaged numbers.
+	UE_LOG(LogEditorAutomationTests, Display, TEXT("AVG FPS: '%.1f'"), FPSAverage);
+	UE_LOG(LogEditorAutomationTests, Display, TEXT("AVG Frame Time: '%.1f' ms"), FrameTimeAverage);
+	UE_LOG(LogEditorAutomationTests, Display, TEXT("AVG Used Physical Memory: '%.0f' kb"), MemoryAverage);
+
+	//Close the csv files so we can use them while the editor is still open.
+	RAWCSVFile->Close();
+	FinalCSVFile->Close();
+
+	//	 the text files.
+	IFileManager::Get().Delete(*FPSFileSaveLocation);
+	IFileManager::Get().Delete(*FrameTimeSaveLocation);
+	IFileManager::Get().Delete(*MemoryFileSaveLocation);
+	IFileManager::Get().Delete(*MapLoadTimeFileSaveLocation);
+	IFileManager::Get().Delete(*DurationFileSaveLocation);
+
+	return true;
+}
+
 /**
 * This this command loads a map into the editor.
 */
 bool FEditorLoadMap::Update()
 {
+	//Get the base filename for the map that will be used.
+	FString ShortMapName = FPaths::GetBaseFilename(MapName);
+
 	double MapLoadStartTime = 0.0f;
 	double MapLoadTime = 0.0f;
-	
+
+	//Variables needed for the WriteToTextFile function.
+	FString TestName = TEXT("Performance");
+	FString Delimiter = TEXT(",");
+	FString FileName;
+
 	//Get the current number of seconds.  This will be used to track how long it took to load the map.
 	MapLoadStartTime = FPlatformTime::Seconds();
 
@@ -455,9 +640,13 @@ bool FEditorLoadMap::Update()
 	//This is the time it took to load the map in the editor.
 	MapLoadTime = FPlatformTime::Seconds() - MapLoadStartTime;
 
-	//Log how long it took to launch the map.
-	UE_LOG(LogEditorAutomationTests, Display, TEXT("Map '%s' took '%.3f' to load"), *MapName, MapLoadTime);
+	UE_LOG(LogEditorAutomationTests, Log, TEXT("%s has been loaded."), *ShortMapName);
 
+	//Log out to a text file the time it takes to load the map.
+	WriteToTextFile(TestName, ShortMapName, FileName = TEXT("RAWMapLoadTime"), MapLoadTime, Delimiter);
+	
+	UE_LOG(LogEditorAutomationTests, Display, TEXT("%s took %.3f to load."), *ShortMapName, MapLoadTime);
+	
 	return true;
 }
 
@@ -638,4 +827,5 @@ void FEditorAutomationTestUtilities::CollectMiscGameContentTestsByClass(TArray<F
 			}
 		}
 	//}
-}
\ No newline at end of file
+}
+
diff --git a/Engine/Source/Editor/UnrealEd/Private/Tests/EditorPerformanceTests.cpp b/Engine/Source/Editor/UnrealEd/Private/Tests/EditorPerformanceTests.cpp
--- a/Engine/Source/Editor/UnrealEd/Private/Tests/EditorPerformanceTests.cpp
+++ b/Engine/Source/Editor/UnrealEd/Private/Tests/EditorPerformanceTests.cpp
@@ -65,13 +65,14 @@ bool FMapPerformanceInEditor::RunTest(const FString& Parameters)
 	FString MapName = Parameters;
 	//Get the base filename for the map that will be used.
 	FString ShortMapName = FPaths::GetBaseFilename(MapName);
-	//Duration variable will be used to indicate how long the test will run for.
-	float Duration = 0.0f;
 	
 	//This gets the info we need from the automation test settings in the engine.ini.
 	UAutomationTestSettings const* AutomationTestSettings = GetDefault<UAutomationTestSettings>();
 	check(AutomationTestSettings);
 
+	//Duration variable will be used to indicate how long the test will run for.
+	float Duration = 0.0f;
+
 	//Now we find the test timer(aka duration) for our test.
 	for (auto Iter = AutomationTestSettings->EditorPerformanceTestMaps.CreateConstIterator(); Iter; ++Iter)
 	{
@@ -85,21 +86,26 @@ bool FMapPerformanceInEditor::RunTest(const FString& Parameters)
 			{
 				UE_LOG(LogEditorAutomationTests, Warning, TEXT("Please set the test timer for '%s' in the automation preferences or engine.ini."), *ShortMapName);
 			}
-			else 
+			else
 			if (Duration < 0)
 			{
 				UE_LOG(LogEditorAutomationTests, Error, TEXT("Test timer preference option '%s' is less than 0."), *ShortMapName);
-				return false; 
+				return false;
 			}
 		}
 	}
-	//Loop through the latent performance command equal to the duration time.
-	UE_LOG(LogEditorAutomationTests, Log, TEXT("Starting loop to run the latent FEditorPerformanceCommand on '%s'."), *ShortMapName);
-	for (int32 Counter = 0; Counter < Duration; Counter++)
-	{
-		ADD_LATENT_AUTOMATION_COMMAND(FWaitLatentCommand(1.0f));
-		ADD_LATENT_AUTOMATION_COMMAND(FEditorPerformanceCommand(MapName));
-	}
+
+	//Load Map and get the time it took to take to load the map.
+	ADD_LATENT_AUTOMATION_COMMAND(FEditorLoadMap(MapName));
+	ADD_LATENT_AUTOMATION_COMMAND(FWaitLatentCommand(1.0f));
+
+	//Grab the performance numbers based on the duration.
+	ADD_LATENT_AUTOMATION_COMMAND(FEditorPerformanceCommand(Duration));
+	ADD_LATENT_AUTOMATION_COMMAND(FWaitLatentCommand(0.5f));
+
+	//Combine performance data into one chart.
+	ADD_LATENT_AUTOMATION_COMMAND(FGenerateEditorPerformanceCharts(ShortMapName));
+
 	return true;
 }
 
diff --git a/Engine/Source/Editor/UnrealEd/Public/Tests/AutomationEditorCommon.h b/Engine/Source/Editor/UnrealEd/Public/Tests/AutomationEditorCommon.h
--- a/Engine/Source/Editor/UnrealEd/Public/Tests/AutomationEditorCommon.h
+++ b/Engine/Source/Editor/UnrealEd/Public/Tests/AutomationEditorCommon.h
@@ -130,13 +130,19 @@ DEFINE_LATENT_AUTOMATION_COMMAND(FEndPlayMapCommand);
 /**
 * This will display the average FPS and Memory numbers over a duration of time.
 */
-DEFINE_LATENT_AUTOMATION_COMMAND_ONE_PARAMETER(FEditorPerformanceCommand, FString, BaseMapName);
+DEFINE_LATENT_AUTOMATION_COMMAND_ONE_PARAMETER(FEditorPerformanceCommand, float, Duration);
 
 /**
 * Loads a map
 */
 DEFINE_LATENT_AUTOMATION_COMMAND_ONE_PARAMETER(FEditorLoadMap, FString, MapName);
 
+/**
+* Generates the editor performance capture chart.
+* @param Input the name of the map.
+*/
+DEFINE_LATENT_AUTOMATION_COMMAND_ONE_PARAMETER(FGenerateEditorPerformanceCharts, FString, MapName);
+
 //////////////////////////////////////////////////////////////////////////
 // FEditorAutomationTestUtilities
 
@@ -193,6 +199,7 @@ public:
 	*/
 	static void CollectMiscGameContentTestsByClass(TArray<FString>& OutBeautifiedNames, TArray <FString>& OutTestCommands);
 	
+
 };
 
 
