Clang compile fix

diff --git a/Engine/Source/Editor/UnrealEd/Private/Commandlets/DumpBlueprintsInfoCommandlet.cpp b/Engine/Source/Editor/UnrealEd/Private/Commandlets/DumpBlueprintsInfoCommandlet.cpp
--- a/Engine/Source/Editor/UnrealEd/Private/Commandlets/DumpBlueprintsInfoCommandlet.cpp
+++ b/Engine/Source/Editor/UnrealEd/Private/Commandlets/DumpBlueprintsInfoCommandlet.cpp
@@ -1059,7 +1059,7 @@ static FString DumpBlueprintInfoUtils::BuildByteSizeString(int32 ByteSize)
 		ByteUnits_MAX
 	};
 
-	TCHAR* ByteUnits[ByteUnits_MAX];
+	const TCHAR* ByteUnits[ByteUnits_MAX];
 	ByteUnits[Bytes] = TEXT("Bytes");
 	ByteUnits[KiloBytes] = TEXT("KB");
 	ByteUnits[MegaBytes] = TEXT("MB");
