Target Platform API added to provide more information about variants and variant support.

diff --git a/Engine/Source/Developer/TargetDeviceServices/Private/TargetDeviceServiceManager.cpp b/Engine/Source/Developer/TargetDeviceServices/Private/TargetDeviceServiceManager.cpp
--- a/Engine/Source/Developer/TargetDeviceServices/Private/TargetDeviceServiceManager.cpp
+++ b/Engine/Source/Developer/TargetDeviceServices/Private/TargetDeviceServiceManager.cpp
@@ -46,13 +46,13 @@ FTargetDeviceServiceManager::~FTargetDeviceServiceManager( )
 /* ITargetDeviceServiceManager interface
  *****************************************************************************/
 
-bool FTargetDeviceServiceManager::AddStartupService( const FTargetDeviceId& DeviceId, const FString& PreliminaryDeviceName )
+bool FTargetDeviceServiceManager::AddStartupService(const FString& DeviceName)
 {
 	FScopeLock Lock(&CriticalSection);
 	{
-		StartupServices.Add(DeviceId, false);
+		StartupServices.Add(DeviceName, false);
 
-		return AddService(DeviceId, PreliminaryDeviceName);
+		return AddService(DeviceName).IsValid();
 	}
 }
 
@@ -68,13 +68,13 @@ int32 FTargetDeviceServiceManager::GetServices( TArray<ITargetDeviceServicePtr>&
 }
 
 
-void FTargetDeviceServiceManager::RemoveStartupService( const FTargetDeviceId& DeviceId )
+void FTargetDeviceServiceManager::RemoveStartupService(const FString& DeviceName)
 {
 	FScopeLock Lock(&CriticalSection);
 	{
-		if (StartupServices.Remove(DeviceId) > 0)
+		if (StartupServices.Remove(DeviceName) > 0)
 		{
-			RemoveService(DeviceId);
+			RemoveService(DeviceName);
 		}
 	}
 }
@@ -83,40 +83,57 @@ void FTargetDeviceServiceManager::RemoveStartupService( const FTargetDeviceId& D
 /* FTargetDeviceServiceManager implementation
  *****************************************************************************/
 
-bool FTargetDeviceServiceManager::AddService( const FTargetDeviceId& DeviceId, const FString& PreliminaryDeviceName )
+ITargetDeviceServicePtr FTargetDeviceServiceManager::AddService(const FString& DeviceName)
 {
 	IMessageBusPtr MessageBus = MessageBusPtr.Pin();
 
 	if (!MessageBus.IsValid())
 	{
-		return false;
+		return nullptr;
 	}
 
-	ITargetDeviceServicePtr DeviceService = DeviceServices.FindRef(DeviceId);
+	ITargetDeviceServicePtr DeviceService = DeviceServices.FindRef(DeviceName);
 
 	// create service if needed
 	if (!DeviceService.IsValid())
 	{
-		DeviceService = MakeShareable(new FTargetDeviceService(DeviceId, PreliminaryDeviceName, MessageBus.ToSharedRef()));
-		DeviceServices.Add(DeviceId, DeviceService);
+		DeviceService = MakeShareable(new FTargetDeviceService(DeviceName, MessageBus.ToSharedRef()));
+		DeviceServices.Add(DeviceName, DeviceService);
 
 		ServiceAddedDelegate.Broadcast(DeviceService.ToSharedRef());
 	}
 
 	// share service if desired
-	const bool* Shared = StartupServices.Find(DeviceId);
-
+	const bool* Shared = StartupServices.Find(DeviceName);
 	if (Shared != nullptr)
 	{
 		DeviceService->SetShared(*Shared);
+		DeviceService->Start();
 	}
 
-	// start service if desired
-	ITargetDevicePtr TargetDevice = DeviceService->GetDevice();
+	return DeviceService;
+}
 
-	if ((Shared != nullptr) || (TargetDevice.IsValid() && TargetDevice->IsDefault()))
+bool FTargetDeviceServiceManager::AddTargetDevice(ITargetDevicePtr InDevice)
+{
+	IMessageBusPtr MessageBus = MessageBusPtr.Pin();
+	if (!MessageBus.IsValid())
 	{
-		DeviceService->Start();
+		return false;
+	}
+
+	const FString& DeviceName = InDevice->GetName();
+
+	ITargetDeviceServicePtr DeviceService = AddService(DeviceName);
+
+	if (DeviceService.IsValid())
+	{
+		DeviceService->AddTargetDevice(InDevice);
+
+		if (InDevice.IsValid() && InDevice->IsDefault())
+		{
+			DeviceService->Start();
+		}
 	}
 
 	return true;
@@ -144,7 +161,7 @@ void FTargetDeviceServiceManager::InitializeTargetPlatforms( )
 
 			if (Device.IsValid())
 			{
-				AddService(Device->GetId(), Device->GetName());
+				AddTargetDevice(Device);
 			}
 		}
 	}
@@ -176,31 +193,22 @@ void FTargetDeviceServiceManager::LoadSettings( )
 		const FString& ServiceString = *It.Value();
 
 		// ... parse device identifier...
-		FString DeviceIdString;
-		FTargetDeviceId DeviceId;
+		FString DeviceName;
 		
-		if (!FParse::Value(*ServiceString, TEXT("DeviceId="), DeviceIdString) || !FTargetDeviceId::Parse(DeviceIdString, DeviceId))
+		if (!FParse::Value(*ServiceString, TEXT("DeviceName="), DeviceName))
 		{
-			UE_LOG(TargetDeviceServicesLog, Warning, TEXT("[TargetDeviceServices] failed to parse DeviceId in configuration setting: StartupServices=%s"), *ServiceString);
+			UE_LOG(TargetDeviceServicesLog, Warning, TEXT("[TargetDeviceServices] failed to parse DeviceName in configuration setting: StartupServices=%s"), *ServiceString);
 
 			continue;
 		}
 
-		if (DeviceServices.Contains(DeviceId))
+		if (DeviceServices.Contains(DeviceName))
 		{
 			UE_LOG(TargetDeviceServicesLog, Warning, TEXT("[TargetDeviceServices] duplicate entry for : StartupServices=%s"), *ServiceString);
 
 			continue;
 		}
 
-		// ... parse device name...
-		FString DeviceName;
-
-		if (!FParse::Value(*ServiceString, TEXT("DeviceName="), DeviceName))
-		{
-			DeviceName = DeviceIdString;
-		}
-
 		// ... parse sharing state...
 		bool Shared;
 
@@ -209,12 +217,12 @@ void FTargetDeviceServiceManager::LoadSettings( )
 			Shared = false;
 		}
 
-		StartupServices.Add(DeviceId, Shared);
+		StartupServices.Add(DeviceName, Shared);
 
 		// ... create and start device service
 		ITargetDeviceServicePtr DeviceService;
 
-		if (!AddService(DeviceId, DeviceName))
+		if (!AddService(DeviceName).IsValid())
 		{
 			UE_LOG(TargetDeviceServicesLog, Warning, TEXT("[TargetDeviceServices] failed to create service for: StartupServices=%s"), *ServiceString);
 		}
@@ -222,9 +230,9 @@ void FTargetDeviceServiceManager::LoadSettings( )
 }
 
 
-void FTargetDeviceServiceManager::RemoveService( const FTargetDeviceId& DeviceId )
+void FTargetDeviceServiceManager::RemoveService(const FString& DeviceName)
 {
-	ITargetDeviceServicePtr DeviceService = DeviceServices.FindRef(DeviceId);
+	ITargetDeviceServicePtr DeviceService = DeviceServices.FindRef(DeviceName);
 
 	if (!DeviceService.IsValid())
 	{
@@ -234,13 +242,32 @@ void FTargetDeviceServiceManager::RemoveService( const FTargetDeviceId& DeviceId
 	DeviceService->Stop();
 
 	// only truly remove if not startup device or physical device not available
-	if (!StartupServices.Contains(DeviceId) && !DeviceService->GetDevice().IsValid())
+	if (!StartupServices.Contains(DeviceName) && !DeviceService->GetDevice().IsValid())
 	{
-		DeviceServices.Remove(DeviceId);
+		DeviceServices.Remove(DeviceName);
 		ServiceRemovedDelegate.Broadcast(DeviceService.ToSharedRef());
 	}
 }
 
+void FTargetDeviceServiceManager::RemoveTargetDevice(ITargetDevicePtr InDevice)
+{
+	ITargetDeviceServicePtr DeviceService = DeviceServices.FindRef(InDevice->GetName());
+
+	if (!DeviceService.IsValid())
+	{
+		return;
+	}
+
+	DeviceService->RemoveTargetDevice(InDevice);
+
+	if (DeviceService->NumTargetDevices() <= 0)
+	{
+		RemoveService(InDevice->GetName());
+	}
+
+	return;
+}
+
 
 void FTargetDeviceServiceManager::SaveSettings( )
 {
@@ -254,9 +281,9 @@ void FTargetDeviceServiceManager::SaveSettings( )
 	TArray<FString> ServiceStrings;
 
 	// for each device service...
-	for (TMap<FTargetDeviceId, ITargetDeviceServicePtr>::TConstIterator It(DeviceServices); It; ++It)
+	for (TMap<FString, ITargetDeviceServicePtr>::TConstIterator It(DeviceServices); It; ++It)
 	{
-		const FTargetDeviceId& DeviceId = It.Key();
+		const FString& DeviceName= It.Key();
 		const ITargetDeviceServicePtr& DeviceService = It.Value();
 
 		// ... that is not managing a default device...
@@ -268,14 +295,13 @@ void FTargetDeviceServiceManager::SaveSettings( )
 		}
 
 		// ... that should be automatically restarted next time...
-		const bool* Shared = StartupServices.Find(DeviceId);
+		const bool* Shared = StartupServices.Find(DeviceName);
 
 		if ((Shared != nullptr) || DeviceService->IsRunning())
 		{
 			// ... generate an entry in the INI file
-			FString ServiceString = FString::Printf(TEXT("DeviceId=\"%s\",DeviceName=\"%s\",Shared=%s"),
-				*DeviceId.ToString(),
-				*DeviceService->GetCachedDeviceName(),
+			FString ServiceString = FString::Printf(TEXT("DeviceName=\"%s\",Shared=%s"),
+				*DeviceName,
 				((Shared != nullptr) && *Shared) ? TEXT("true") : TEXT("false")
 			);
 
@@ -316,19 +342,19 @@ void FTargetDeviceServiceManager::HandleMessageBusShutdown( )
 }
 
 
-void FTargetDeviceServiceManager::HandleTargetPlatformDeviceDiscovered( const ITargetDeviceRef& DiscoveredDevice )
+void FTargetDeviceServiceManager::HandleTargetPlatformDeviceDiscovered(ITargetDeviceRef DiscoveredDevice)
 {
 	FScopeLock Lock(&CriticalSection);
 	{
-		AddService(DiscoveredDevice->GetId(), DiscoveredDevice->GetName());
+		AddTargetDevice(DiscoveredDevice);
 	}
 }
 
 
-void FTargetDeviceServiceManager::HandleTargetPlatformDeviceLost( const ITargetDeviceRef& LostDevice )
+void FTargetDeviceServiceManager::HandleTargetPlatformDeviceLost(ITargetDeviceRef LostDevice)
 {
 	FScopeLock Lock(&CriticalSection);
 	{
-		RemoveService(LostDevice->GetId());
+		RemoveTargetDevice(LostDevice);
 	}
 }
diff --git a/Engine/Source/Developer/TargetDeviceServices/Private/TargetDeviceServiceManager.h b/Engine/Source/Developer/TargetDeviceServices/Private/TargetDeviceServiceManager.h
--- a/Engine/Source/Developer/TargetDeviceServices/Private/TargetDeviceServiceManager.h
+++ b/Engine/Source/Developer/TargetDeviceServices/Private/TargetDeviceServiceManager.h
@@ -27,79 +27,93 @@ public:
 
 	// ITargetDeviceServiceManager interface
 
-	virtual bool AddStartupService( const FTargetDeviceId& DeviceId, const FString& PreliminaryDeviceName ) override;
-	virtual int32 GetServices( TArray<ITargetDeviceServicePtr>& OutServices ) override;
+	virtual bool AddStartupService(const FString& DeviceName) override;
+	virtual int32 GetServices(TArray<ITargetDeviceServicePtr>& OutServices) override;
 
 	DECLARE_DERIVED_EVENT(FTargetDeviceServiceManager, ITargetDeviceServiceManager::FOnTargetDeviceServiceAdded, FOnTargetDeviceServiceAdded);
-	virtual FOnTargetDeviceServiceAdded& OnServiceAdded( ) override
+	virtual FOnTargetDeviceServiceAdded& OnServiceAdded() override
 	{
 		return ServiceAddedDelegate;
 	}
 
 	DECLARE_DERIVED_EVENT(FTargetDeviceServiceManager, ITargetDeviceServiceManager::FOnTargetDeviceServiceRemoved, FOnTargetDeviceServiceRemoved);
-	virtual FOnTargetDeviceServiceRemoved& OnServiceRemoved( ) override
+	virtual FOnTargetDeviceServiceRemoved& OnServiceRemoved() override
 	{
 		return ServiceRemovedDelegate;
 	}
 
-	virtual void RemoveStartupService( const FTargetDeviceId& DeviceId ) override;
+	virtual void RemoveStartupService(const FString& DeviceName) override;
 
 protected:
 
 	/**
-	 * Adds a device service for the given target device.
+	 * Adds a device service for the given device.
 	 *
-	 * @param DeviceId The identifier of the device to add a service for.
-	 * @param PreliminaryDeviceName The preliminary name to assign to this device.
-	 * @return true if the service was created, false otherwise.
+	 * @param DeviceName The name to of the device to add.
+	 * @return the created or preexisting service, nullptr if creation failed.
 	 */
-	bool AddService( const FTargetDeviceId& DeviceId, const FString& PreliminaryDeviceName );
+	ITargetDeviceServicePtr AddService(const FString& DeviceName);
+
+	/**
+	 * Adds a target device, a device service will be created if needed
+	 *
+	 * @param InDevice The target device to add
+	 * @return true if the device was added to the device service, false if the service failed to create.
+	 */
+	bool AddTargetDevice(ITargetDevicePtr InDevice);
 
 	/**
 	 * Initializes target platforms callbacks.
 	 *
 	 * @see ShutdownTargetPlatforms
 	 */
-	void InitializeTargetPlatforms( );
+	void InitializeTargetPlatforms();
 
 	/**
 	 * Loads all settings from disk.
 	 *
 	 * @see SaveSettings
 	 */
-	void LoadSettings( );
+	void LoadSettings();
 
 	/**
 	 * Removes the device service for the given target device.
 	 *
-	 * @param DeviceId The identifier of the device to remove the service for.
+	 * @param DeviceName The name of the device to remove the service for.
+	 */
+	void RemoveService(const FString& DeviceName);
+
+	/**
+	 * Removes the given target device, if that is the last target device in the device service the entire service is removed.
+	 *
+	 * @param InDevice The target device to remove.
 	 */
-	void RemoveService( const FTargetDeviceId& DeviceId );
+	void RemoveTargetDevice(ITargetDevicePtr InDevice);
 
 	/**
 	 * Saves all settings to disk.
 	 *
 	 * @see LoadSettings
 	 */
-	void SaveSettings( );
+	void SaveSettings();
 
 	/**
 	 * Shuts down target platforms callbacks.
 	 *
 	 * @see InitializeTargetPlatforms
 	 */
-	void ShutdownTargetPlatforms( );
+	void ShutdownTargetPlatforms();
 
 private:
 
 	// Callback for handling message bus shutdowns.
-	void HandleMessageBusShutdown( );
+	void HandleMessageBusShutdown();
 
 	// Callback for discovered target devices.
-	void HandleTargetPlatformDeviceDiscovered( const ITargetDeviceRef& DiscoveredDevice );
+	void HandleTargetPlatformDeviceDiscovered(ITargetDeviceRef DiscoveredDevice);
 
 	// Callback for lost target devices.
-	void HandleTargetPlatformDeviceLost( const ITargetDeviceRef& LostDevice );
+	void HandleTargetPlatformDeviceLost(ITargetDeviceRef LostDevice);
 
 private:
 
@@ -107,13 +121,13 @@ private:
 	FCriticalSection CriticalSection;
 
 	// Holds the list of locally managed device services.
-	TMap<FTargetDeviceId, ITargetDeviceServicePtr> DeviceServices;
+	TMap<FString, ITargetDeviceServicePtr> DeviceServices;
 
 	// Holds a weak pointer to the message bus.
 	IMessageBusWeakPtr MessageBusPtr;
 
 	// Holds the collection of identifiers for devices that start automatically (shared/unshared).
-	TMap<FTargetDeviceId, bool> StartupServices;
+	TMap<FString, bool> StartupServices;
 
 private:
 
diff --git a/Engine/Source/Developer/TargetDeviceServices/Public/Interfaces/ITargetDeviceProxy.h b/Engine/Source/Developer/TargetDeviceServices/Public/Interfaces/ITargetDeviceProxy.h
--- a/Engine/Source/Developer/TargetDeviceServices/Public/Interfaces/ITargetDeviceProxy.h
+++ b/Engine/Source/Developer/TargetDeviceServices/Public/Interfaces/ITargetDeviceProxy.h
@@ -40,7 +40,6 @@ DECLARE_MULTICAST_DELEGATE_OneParam(FOnTargetDeviceProxyLaunchFailed, const FStr
  */
 DECLARE_MULTICAST_DELEGATE_TwoParams(FOnTargetDeviceProxyLaunchSucceeded, const FString&, uint32)
 
-
 /**
  * Interface for target device proxies.
  */
@@ -53,112 +52,184 @@ public:
 	 *
 	 * @return true if launching multiple instances is supported, false otherwise.
 	 */
-	virtual const bool CanMultiLaunch( ) const = 0;
+	virtual const bool CanMultiLaunch() const = 0;
 
 	/**
 	 * Checks whether this device can be powered off remotely.
 	 *
 	 * @return true if the device can be powered off, false otherwise.
 	 */
-	virtual bool CanPowerOff( ) const = 0;
+	virtual bool CanPowerOff() const = 0;
 
 	/**
 	 * Checks whether this device can be powered on remotely.
 	 *
 	 * @return true if the device can be powered on, false otherwise.
 	 */
-	virtual bool CanPowerOn( ) const = 0;
+	virtual bool CanPowerOn() const = 0;
 
 	/**
 	 * Checks whether this device can be rebooted remotely.
 	 *
 	 * @return true if the device can be rebooted, false otherwise.
 	 */
-	virtual bool CanReboot( ) const = 0;
+	virtual bool CanReboot() const = 0;
+
+	/**
+	 * Checks whether this device's target platform supports variants.
+	 *
+	 * @return true if the device's target platform supports variants, false otherwise.
+	 */
+	virtual bool CanSupportVariants() const = 0;
+
+	/**
+	 * Gets the number of variants this device supports
+	 *
+	 * @return The number of variants supported.
+	 */
+	virtual int32 GetNumVariants() const = 0;
+
+	/**
+	 * Gets the list of variants this device supports
+	 *
+	 * @param OutVariants This list will be populated with the unique variant names. 
+	 * @return The number of variants in the list.
+	 */
+	virtual int32 GetVariants(TArray<FName>& OutVariants) const = 0;
+
+	/**
+	 * Gets the variant name of the target device.
+	 *
+	 * @param InDeviceId target device id to get the variant for.
+	 * @return The target device variant name.
+	 */
+	virtual FName GetTargetDeviceVariant(const FString& InDeviceId) const = 0;
+
+	/**
+	 * Checks whether this device proxy contains a variant for the provided target device.
+	 *
+	 * @param InDeviceId DeviceId to check for a variant.
+	 * @return true if this device has a variant of the provided DeviceId.
+	 */
+	virtual bool HasDeviceId(const FString& InDeviceId) const = 0;
 
 	/**
 	 * Gets the identifier of the device.
 	 *
-	 * @return The device identifier.
+	 * @param InVariant name of the variant to return the device identifier for, If InVariant == NAME_None you'll get back the default target device identifier.
+	 * @return The target device identifier.
 	 */
-	virtual const FString& GetDeviceId( ) = 0;
+	virtual const FString& GetTargetDeviceId(FName InVariant) const = 0;
+
+	/**
+	* Checks whether this device proxy contains a variant for the provided platform.
+	*
+	* @param InTargetPlatformName Target Platform Name for the target device to check for as a variant of this device.
+	* @return true if this device has a variant of the provided Target Platform .
+	*/
+	virtual bool HasTargetPlatform(FName InTargetPlatformId) const = 0;
+
+	/**
+	 * Gets the target platform of the device variant.
+	 *
+	 * @param InVariant Variant to get the platform name for.
+	 * @return Target Platform Name string.
+	 */
+	virtual FString GetTargetPlatformName(FName InVariant) const = 0;
+
+	/**
+	 * Gets the target platform of the device variant.
+	 *
+	 * @param InVariant Variant to get the platform id for.
+	 * @return Target Platform Id.
+	 */
+	virtual FName GetTargetPlatformId(FName InVariant) const = 0;
+
+	/**
+	 * Gets the Vanilla platform of the device variant.
+	 *
+	 * @param InVariant Variant to get the platform id for.
+	 * @return Vanilla Platform Id.
+	 */
+	virtual FName GetVanillaPlatformId(FName InVariant) const = 0;
+
+	/**
+	 * Gets the Vanilla platform of the device variant.
+	 *
+	 * @param InVariant Variant to get the vanilla platform name for.
+	 * @return Platform Display Name string.
+	 */
+	virtual FText GetPlatformDisplayName(FName InVariant) const = 0;
 
 	/**
 	 * Gets the name of the host machine that claimed the device.
 	 *
 	 * @return Host name string.
 	 */
-	virtual const FString& GetHostName( ) const = 0;
+	virtual const FString& GetHostName() const = 0;
 	 
 	/**
 	 * Gets the name of the user that claimed the device.
 	 *
 	 * @return User name.
 	 */
-	virtual const FString& GetHostUser( ) const = 0;
+	virtual const FString& GetHostUser() const = 0;
 
 	/**
 	 * Gets the name of the (device's) user that is logged in on a device.
 	 *
 	 * @return Device user name.
 	 */
-	virtual const FString& GetDeviceUser( ) const = 0;
+	virtual const FString& GetDeviceUser() const = 0;
 
 	/**
 	 * Gets the password of the (device's) user that is logged in on a device.
 	 *
 	 * @return Device user password.
 	 */
-	virtual const FString& GetDeviceUserPassword( ) const = 0;
+	virtual const FString& GetDeviceUserPassword() const = 0;
 
 	/**
 	 * Gets the device make (i.e. Apple or Sony).
 	 *
 	 * @return Device make.
 	 */
-	virtual const FString& GetMake( ) const = 0;
+	virtual const FString& GetMake() const = 0;
 
 	/**
 	 * Gets the device model (i.e. PS3 or XBox).
 	 *
 	 * @return Device model.
 	 */
-	virtual const FString& GetModel( ) const = 0;
+	virtual const FString& GetModel() const = 0;
 	 
 	/**
 	 * Gets the name of the device (i.e. network name or IP address).
 	 *
 	 * @return Device name string.
 	 */
-	virtual const FString& GetName( ) const = 0;
-
-	/**
-	 * Gets the platform of the device (i.e. Windows or Mac).
-	 *
-	 * @return Platform string.
-	 */
-	virtual const FString& GetPlatformName( ) const = 0;
+	virtual const FString& GetName() const = 0;
 
 	/**
 	 * Gets the device type (i.e. Console, PC or Mobile).
 	 *
 	 * @return Device type.
 	 */
-	virtual const FString& GetType( ) const = 0;
+	virtual const FString& GetType() const = 0;
 
 	/**
 	 * Checks whether the device is currently connected.
 	 *
 	 * @return true if the device is connected, false otherwise.
 	 */
-	virtual bool IsConnected( ) const = 0;
+	virtual bool IsConnected() const = 0;
 
 	/**
 	 * Checks whether this device is being shared with other users.
 	 *
 	 * @return true if the device is being shared, false otherwise.
 	 */
-	virtual bool IsShared( ) const = 0;
+	virtual bool IsShared() const = 0;
 
 public:
 
@@ -169,6 +240,7 @@ public:
 	 * application identifier (AppId) that can be used to launch to the deployed
 	 * application. If deploymend fails, the OnDeployFailed() delegate will be invoked.
 	 *
+	 * @param InVariant Variant to deploy files on.
 	 * @param Files The files to deploy.
 	 * @param TransactionId The transaction identifier for the deployment.
 	 * @return true if deployment has been started, false otherwise.
@@ -176,13 +248,14 @@ public:
 	 * @see OnDeployCommitted
 	 * @see OnDeployFailed
 	 */
-	virtual bool DeployApp( const TMap<FString, FString>& Files, const FGuid& TransactionId ) = 0;
+	virtual bool DeployApp(FName InVariant, const TMap<FString, FString>& Files, const FGuid& TransactionId) = 0;
 
 	/**
 	 * Launches a previously deployed build on the target device.
 	 *
 	 * If the launching fails, the OnLaunchFailed() delegate will be invoked.
 	 *
+	 * @param InVariant Variant to launch the app with.
 	 * @param AppId The identifier of the application to launch.
 	 * @param BuildConfiguration The build configuration to launch with.
 	 * @param Params The command line parameters to launch with.
@@ -190,7 +263,7 @@ public:
 	 * @see Deploy
 	 * @see OnLaunchFailed
 	 */
-	virtual bool LaunchApp( const FString& AppId, EBuildConfigurations::Type BuildConfiguration, const FString& Params ) = 0;
+	virtual bool LaunchApp(FName InVariant, const FString& AppId, EBuildConfigurations::Type BuildConfiguration, const FString& Params) = 0;
 
 	/**
 	 * Powers off the device.
@@ -199,7 +272,7 @@ public:
 	 * @see PowerOn
 	 * @see Reboot
 	 */
-	virtual void PowerOff( bool Force ) = 0;
+	virtual void PowerOff(bool Force) = 0;
 
 	/**
 	 * Powers on the device.
@@ -207,7 +280,7 @@ public:
 	 * @see PowerOff
 	 * @see Reboot
 	 */
-	virtual void PowerOn( ) = 0;
+	virtual void PowerOn() = 0;
 
 	/**
 	 * Reboots the device.
@@ -215,15 +288,16 @@ public:
 	 * @see PowerOff
 	 * @see PowerOn
 	 */
-	virtual void Reboot( ) = 0;
+	virtual void Reboot() = 0;
 
 	/**
 	 * Attempts to run an executable on the device.
 	 *
+	 * @param InVariant Variant to run the executable with.
 	 * @param ExecutablePath The path to the executable to run.
 	 * @param Params The command line parameters.
 	 */
-	virtual void Run( const FString& ExecutablePath, const FString& Params ) = 0;
+	virtual void Run(FName InVariant, const FString& ExecutablePath, const FString& Params) = 0;
 
 public:
 
@@ -233,7 +307,7 @@ public:
 	 * @return The delegate.
 	 * @see DeployApp
 	 */
-	virtual FOnTargetDeviceProxyDeployCommitted& OnDeployCommitted( ) = 0;
+	virtual FOnTargetDeviceProxyDeployCommitted& OnDeployCommitted() = 0;
 
 	/**
 	 * Returns a delegate to be invoked when a build failed to deploy to the target device.
@@ -241,7 +315,7 @@ public:
 	 * @return The delegate.
 	 * @see DeployApp
 	 */
-	virtual FOnTargetDeviceProxyDeployFailed& OnDeployFailed( ) = 0;
+	virtual FOnTargetDeviceProxyDeployFailed& OnDeployFailed() = 0;
 
 	/**
 	 * Returns a delegate to be invoked when a build failed to launch on the target device.
@@ -249,12 +323,12 @@ public:
 	 * @return The delegate.
 	 * @see LaunchApp
 	 */
-	virtual FOnTargetDeviceProxyLaunchFailed& OnLaunchFailed( ) = 0;
+	virtual FOnTargetDeviceProxyLaunchFailed& OnLaunchFailed() = 0;
 
 public:
 
 	/**
 	 * Virtual destructor.
 	 */
-	virtual ~ITargetDeviceProxy( ) { }
+	virtual ~ITargetDeviceProxy() { }
 };
diff --git a/Engine/Source/Developer/TargetDeviceServices/Public/Interfaces/ITargetDeviceProxyManager.h b/Engine/Source/Developer/TargetDeviceServices/Public/Interfaces/ITargetDeviceProxyManager.h
--- a/Engine/Source/Developer/TargetDeviceServices/Public/Interfaces/ITargetDeviceProxyManager.h
+++ b/Engine/Source/Developer/TargetDeviceServices/Public/Interfaces/ITargetDeviceProxyManager.h
@@ -9,7 +9,6 @@ typedef TSharedPtr<class ITargetDeviceProxyManager> ITargetDeviceProxyManagerPtr
 /** Type definition for shared references to instances of ITargetDeviceProxyManager. */
 typedef TSharedRef<class ITargetDeviceProxyManager> ITargetDeviceProxyManagerRef;
 
-
 /**
  * Interface for device proxy managers.
  */
@@ -18,20 +17,28 @@ class ITargetDeviceProxyManager
 public:
 	
 	/**
-	 * Finds or adds a proxy for the specified device identifier.
+	 * Finds or adds a device proxy for the specified device name.
+	 *
+	 * @param Name The name of the device to create the proxy for.
+	 * @return The device proxy.
+	 */
+	virtual ITargetDeviceProxyRef FindOrAddProxy(const FString& Name) = 0;
+
+	/**
+	 * Finds the device proxy for the specified device name.
 	 *
-	 * @param DeviceId The identifier of the device to create the proxy for.
-	 * @return The proxy.
+	 * @param Name The name of the device to create the proxy for.
+	 * @return The device proxy, or nullptr if it couldn't be found.
 	 */
-	virtual ITargetDeviceProxyRef FindOrAddProxy( const FString& DeviceId ) = 0;
+	virtual ITargetDeviceProxyPtr FindProxy(const FString& Name) = 0;
 
 	/**
-	 * Finds the proxy for the specified device identifier.
+	 * Finds the device proxy for the specified target device id.
 	 *
-	 * @param DeviceId The identifier of the device to create the proxy for.
-	 * @return The proxy, or nullptr if it couldn't be found.
+	 * @param DeviceId The identifier of the target device to create the proxy for.
+	 * @return The device proxy, or nullptr if it couldn't be found.
 	 */
-	virtual ITargetDeviceProxyPtr FindProxy( const FString& DeviceId ) = 0;
+	virtual ITargetDeviceProxyPtr FindProxyDeviceForTargetDevice(const FString& DeviceId) = 0;
 
 	/**
 	 * Gets a list of devices found by the device discovery.
@@ -40,7 +47,7 @@ public:
 	 * @param IncludeUnshared Whether to include devices that are not being shared with the local user.
 	 * @param OutProxies Will hold the list of devices found by the locator
 	 */
-	virtual void GetProxies( const FString& PlatformName, bool IncludeUnshared, TArray<ITargetDeviceProxyPtr>& OutProxies ) = 0;
+	virtual void GetProxies(FName TargetPlatformName, bool IncludeUnshared, TArray<ITargetDeviceProxyPtr>& OutProxies) = 0;
 
 public:
 
diff --git a/Engine/Source/Developer/TargetDeviceServices/Public/Interfaces/ITargetDeviceService.h b/Engine/Source/Developer/TargetDeviceServices/Public/Interfaces/ITargetDeviceService.h
--- a/Engine/Source/Developer/TargetDeviceServices/Public/Interfaces/ITargetDeviceService.h
+++ b/Engine/Source/Developer/TargetDeviceServices/Public/Interfaces/ITargetDeviceService.h
@@ -23,6 +23,27 @@ class ITargetDeviceService
 {
 public:
 
+	/**
+	 * Adds a flavor to this device.
+	 *
+	 * @return true if the service can start, false otherwise.
+	 */
+	virtual void AddTargetDevice(ITargetDevicePtr InDevice) = 0;
+
+	/**
+	 * Removes a flavor from this device.
+	 *
+	 * @return true if the service can start, false otherwise.
+	 */
+	virtual void RemoveTargetDevice(ITargetDevicePtr InDevice) = 0;
+
+	/**
+	 * Gets number of target devices.
+	 *
+	 * @return number of target devices in this device service.
+	 */
+	virtual int32 NumTargetDevices() = 0;
+
 	/**
 	 * Checks whether this service can start.
 	 *
@@ -31,7 +52,7 @@ public:
 	 *
 	 * @return true if the service can start, false otherwise.
 	 */
-	virtual bool CanStart( ) const = 0;
+	virtual bool CanStart(FName InFlavor = NAME_None) const = 0;
 
 	/**
 	 * Gets the name of the host that has a claim on the device.
@@ -48,25 +69,32 @@ public:
 	virtual const FString& GetClaimUser( ) = 0;
 
 	/**
-	 * Gets the target device that this service exposes.
+	 * Gets the target device for the supplied flavor, if no flavor is specified then the default device flavor is returned.
 	 *
 	 * @return A pointer to the target device.
 	 */
-	virtual ITargetDevicePtr GetDevice( ) = 0;
+	virtual ITargetDevicePtr GetDevice(FName InFlavor = NAME_None) const = 0;
 
 	/**
-	 * Gets the identifier of the device that this service exposes.
+	 * Gets the name of the device that this service exposes.
 	 *
-	 * @return The device identifier.
+	 * @return The human readable device name.
 	 */
-	virtual const FTargetDeviceId& GetDeviceId( ) const = 0;
+	virtual FString GetDeviceName() const = 0;
 
 	/**
-	 * Gets the name of the device that this service exposes.
+	 * Gets the name of the platform of device that this service exposes.
 	 *
-	 * @return The human readable device name.
+	 * @return The platform name.
+	 */
+	virtual FName GetDevicePlatformName() const = 0;
+
+	/**
+	 * Gets the display name of the platform of device that this service exposes.
+	 *
+	 * @return The human readable platform name.
 	 */
-	virtual FString GetCachedDeviceName( ) const = 0;
+	virtual FString GetDevicePlatformDisplayName() const = 0;
 
 	/**
 	 * Checks whether the service is currently running.
diff --git a/Engine/Source/Developer/TargetDeviceServices/Public/Interfaces/ITargetDeviceServiceManager.h b/Engine/Source/Developer/TargetDeviceServices/Public/Interfaces/ITargetDeviceServiceManager.h
--- a/Engine/Source/Developer/TargetDeviceServices/Public/Interfaces/ITargetDeviceServiceManager.h
+++ b/Engine/Source/Developer/TargetDeviceServices/Public/Interfaces/ITargetDeviceServiceManager.h
@@ -24,13 +24,12 @@ public:
 	 * could not be discovered automatically or are currently unavailable. This name will
 	 * be replaced with the actual device name once the physical device becomes available.
 	 *
-	 * @param DeviceId The device identifier for the service to add.
-	 * @param PreliminaryDeviceName The preliminary device name.
+	 * @param DeviceName The device name for the service to add to the startup list
 	 * @return true if the service added, false otherwise.
 	 * @see GetServices
 	 * @see RemoveStartupService
 	 */
-	virtual bool AddStartupService( const FTargetDeviceId& DeviceId, const FString& PreliminaryDeviceName ) = 0;
+	virtual bool AddStartupService(const FString& DeviceName) = 0;
 
 	/**
 	 * Gets the collection of target device services managed by this instance.
@@ -43,11 +42,11 @@ public:
 	/**
 	 * Removes a service from the list of services that are started automatically.
 	 *
-	 * @param DeviceId The device identifier for the service to remove.
+	 * @param DeviceName The device name for the service to remove from the startup list.
 	 * @see AddStartupService
 	 * @see GetServices
 	 */
-	virtual void RemoveStartupService( const FTargetDeviceId& DeviceId ) = 0;
+	virtual void RemoveStartupService(const FString& DeviceName) = 0;
 
 public:
 
diff --git a/Engine/Source/Developer/TargetDeviceServices/TargetDeviceServices.Build.cs b/Engine/Source/Developer/TargetDeviceServices/TargetDeviceServices.Build.cs
--- a/Engine/Source/Developer/TargetDeviceServices/TargetDeviceServices.Build.cs
+++ b/Engine/Source/Developer/TargetDeviceServices/TargetDeviceServices.Build.cs
@@ -11,6 +11,7 @@ namespace UnrealBuildTool.Rules
 				{
 					"Core",
 					"TargetPlatform",
+                    "DesktopPlatform",
 				}
 			);
 
diff --git a/Engine/Source/Developer/TargetPlatform/Public/Common/TargetPlatformBase.h b/Engine/Source/Developer/TargetPlatform/Public/Common/TargetPlatformBase.h
--- a/Engine/Source/Developer/TargetPlatform/Public/Common/TargetPlatformBase.h
+++ b/Engine/Source/Developer/TargetPlatform/Public/Common/TargetPlatformBase.h
@@ -180,6 +180,26 @@ public:
 		return bReadyToBuild;
 	}
 
+	virtual bool SupportsVariants() const override
+	{
+		return false;
+	}
+
+	virtual FText GetVariantDisplayName() const override
+	{
+		return FText();
+	}
+
+	virtual FText GetVariantTitle() const override
+	{
+		return FText();
+	}
+
+	virtual float GetVariantPriority() const override
+	{
+		return 0.0f;
+	}
+
 	// End ITargetPlatform interface
 
 protected:
diff --git a/Engine/Source/Developer/TargetPlatform/Public/Interfaces/ITargetPlatform.h b/Engine/Source/Developer/TargetPlatform/Public/Interfaces/ITargetPlatform.h
--- a/Engine/Source/Developer/TargetPlatform/Public/Interfaces/ITargetPlatform.h
+++ b/Engine/Source/Developer/TargetPlatform/Public/Interfaces/ITargetPlatform.h
@@ -133,6 +133,14 @@ public:
 	 */
 	virtual bool AddDevice( const FString& DeviceName, bool bDefault ) = 0;
 
+	/**
+	* Returns the name of this platform
+	*
+	* @return Platform name.
+	* @see DisplayName
+	*/
+	virtual FString PlatformName() const = 0;
+
 	/**
 	 * Gets the platform's display name.
 	 *
@@ -242,14 +250,6 @@ public:
 	 */
 	virtual uint32 MaxGpuSkinBones( ) const = 0;
 
-	/**
-	 * Returns the name of this platform
-	 *
-	 * @return Platform name.
-	 * @see DisplayName
-	 */
-	virtual FString PlatformName( ) const = 0;
-
 	/**
 	 * Checks whether this platform requires cooked data (typically console platforms).
 	 *
@@ -376,19 +376,47 @@ public:
 	 */ 
 	virtual bool PackageBuild(const FString& InPackgeDirectory) = 0;
 
+	/**
+	 * Returns true if the platform is part of a family of variants
+	 */
+	virtual bool SupportsVariants() const = 0;
+
+	/**
+	 * Gets the variant display name of this platform.
+	 * eg. For Android: "ETC1", "ETC2", ...
+	 *
+	 * @return FText display name for this platform variant.
+	 */
+	virtual FText GetVariantDisplayName() const = 0;
+
+	/**
+	 * Gets the variant title of this platform family
+	 * eg. For Android: "Texture Format"
+	 *
+	 * @return FText title for what variants mean to this family of platforms.
+	 */
+	virtual FText GetVariantTitle() const = 0;
+
+	/**
+	 * Gets the variant priority of this platform
+	 *
+	 * @return float priority for this platform variant.
+	 */
+	virtual float GetVariantPriority() const = 0;
+
 public:
 
 	/**
 	 * Gets an event delegate that is executed when a new target device has been discovered.
 	 */
-	DECLARE_EVENT_OneParam(ITargetPlatform, FOnTargetDeviceDiscovered, const ITargetDeviceRef& /*DiscoveredDevice*/);
-	virtual FOnTargetDeviceDiscovered& OnDeviceDiscovered( ) = 0;
+	DECLARE_EVENT_OneParam(ITargetPlatform, FOnTargetDeviceDiscovered, ITargetDeviceRef /*DiscoveredDevice*/);
+	virtual FOnTargetDeviceDiscovered& OnDeviceDiscovered() = 0;
 
 	/**
 	 * Gets an event delegate that is executed when a target device has been lost, i.e. disconnected or timed out.
 	 */
-	DECLARE_EVENT_OneParam(ITargetPlatform, FOnTargetDeviceLost, const ITargetDeviceRef& /*LostDevice*/);
-	virtual FOnTargetDeviceLost& OnDeviceLost( ) = 0;
+	DECLARE_EVENT_OneParam(ITargetPlatform, FOnTargetDeviceLost, ITargetDeviceRef /*LostDevice*/);
+	virtual FOnTargetDeviceLost& OnDeviceLost() = 0;
 
 public:
 
diff --git a/Engine/Source/Developer/Windows/WindowsTargetPlatform/Private/GenericWindowsTargetPlatform.h b/Engine/Source/Developer/Windows/WindowsTargetPlatform/Private/GenericWindowsTargetPlatform.h
--- a/Engine/Source/Developer/Windows/WindowsTargetPlatform/Private/GenericWindowsTargetPlatform.h
+++ b/Engine/Source/Developer/Windows/WindowsTargetPlatform/Private/GenericWindowsTargetPlatform.h
@@ -7,6 +7,7 @@
 	#include "StaticMeshResources.h"
 #endif // WITH_ENGINE
 
+#define LOCTEXT_NAMESPACE "TGenericWindowsTargetPlatform"
 
 /**
  * Template for Windows target platforms
@@ -16,8 +17,8 @@ class TGenericWindowsTargetPlatform
 	: public TTargetPlatformBase<FWindowsPlatformProperties<HAS_EDITOR_DATA, IS_DEDICATED_SERVER, IS_CLIENT_ONLY> >
 {
 public:
-
-	typedef TTargetPlatformBase<FWindowsPlatformProperties<HAS_EDITOR_DATA, IS_DEDICATED_SERVER, IS_CLIENT_ONLY> > TSuper;
+	typedef FWindowsPlatformProperties<HAS_EDITOR_DATA, IS_DEDICATED_SERVER, IS_CLIENT_ONLY> TProperties;
+	typedef TTargetPlatformBase<TProperties> TSuper;
 
 	/**
 	 * Default constructor.
@@ -163,8 +164,44 @@ public:
 
 		return NAME_OGG;
 	}
+
 #endif //WITH_ENGINE
 
+	virtual bool SupportsVariants() const override
+	{
+		return true;
+	}
+
+	virtual FText GetVariantDisplayName() const override
+	{
+		if (IS_DEDICATED_SERVER)
+		{
+			return LOCTEXT("WindowsServerVariantTitle", "Dedicated Server");
+		}
+
+		if (HAS_EDITOR_DATA)
+		{
+			return LOCTEXT("WindowsClientEditorDataVariantTitle", "Client with Editor Data");
+		}
+
+		if (IS_CLIENT_ONLY)
+		{
+			return LOCTEXT("WindowsClientOnlyVariantTitle", "Client only");
+		}
+
+		return LOCTEXT("WindowsClientVariantTitle", "Client");
+	}
+
+	virtual FText GetVariantTitle() const override
+	{
+		return LOCTEXT("WindowsVariantTitle", "Build Type");
+	}
+
+	virtual float GetVariantPriority() const override
+	{
+		return TProperties::GetVariantPriority();
+	}
+
 	DECLARE_DERIVED_EVENT(TGenericWindowsTargetPlatform, ITargetPlatform::FOnTargetDeviceDiscovered, FOnTargetDeviceDiscovered);
 	virtual FOnTargetDeviceDiscovered& OnDeviceDiscovered( ) override
 	{
@@ -222,3 +259,5 @@ private:
 	// Holds an event delegate that is executed when a target device has been lost, i.e. disconnected or timed out.
 	FOnTargetDeviceLost DeviceLostEvent;
 };
+
+#undef LOCTEXT_NAMESPACE
diff --git a/Engine/Source/Editor/EditorStyle/Private/SlateEditorStyle.cpp b/Engine/Source/Editor/EditorStyle/Private/SlateEditorStyle.cpp
--- a/Engine/Source/Editor/EditorStyle/Private/SlateEditorStyle.cpp
+++ b/Engine/Source/Editor/EditorStyle/Private/SlateEditorStyle.cpp
@@ -945,10 +945,10 @@ void FSlateEditorStyle::FStyle::SetupGeneralStyles()
 		Set("SessionFrontEnd.Tabs.Tools", new IMAGE_BRUSH( "/Icons/icon_tab_Tools_16x", Icon16x16 ) );
 	}
 
-	// Session Launcher Window
+	// Launcher Window
 	{
-		Set("SessionLauncher.TabIcon", new IMAGE_BRUSH( "Icons/icon_tab_SessionLauncher_16x", Icon16x16 ) );
-		Set("SessionLauncher.Tabs.Tools", new IMAGE_BRUSH( "/Icons/icon_tab_Tools_16x", Icon16x16 ) );
+		Set("Launcher.TabIcon", new IMAGE_BRUSH( "Icons/icon_tab_SessionLauncher_16x", Icon16x16 ) );
+		Set("Launcher.Tabs.Tools", new IMAGE_BRUSH( "/Icons/icon_tab_Tools_16x", Icon16x16 ) );
 	}
 
 	// Undo History Window
@@ -5553,6 +5553,11 @@ void FSlateEditorStyle::FStyle::SetupAutomationStyles()
 
 	// Launcher
 	{
+		Set( "Launcher.Run", new IMAGE_BRUSH("Launcher/Launcher_Run", Icon40x40) );
+		Set( "Launcher.EditSettings", new IMAGE_BRUSH("Launcher/Launcher_EditSettings", Icon40x40) );
+		Set( "Launcher.Back", new IMAGE_BRUSH("Launcher/Launcher_Back", Icon32x32) );
+		Set( "Launcher.Delete", new IMAGE_BRUSH("Launcher/Launcher_Delete", Icon32x32) );
+
 		Set( "Launcher.Instance_Commandlet", new IMAGE_BRUSH( "Launcher/Instance_Commandlet", Icon25x25 ) );
 		Set( "Launcher.Instance_Editor", new IMAGE_BRUSH( "Launcher/Instance_Editor", Icon25x25 ) );
 		Set( "Launcher.Instance_Game", new IMAGE_BRUSH( "Launcher/Instance_Game", Icon25x25 ) );
@@ -5563,6 +5568,8 @@ void FSlateEditorStyle::FStyle::SetupAutomationStyles()
 		Set( "LauncherCommand.QuickLaunch", new IMAGE_BRUSH( "Launcher/Launcher_Launch", Icon40x40 ) );
 		Set( "LauncherCommand.CreateBuild", new IMAGE_BRUSH( "Launcher/Launcher_Build", Icon40x40 ) );
 		Set( "LauncherCommand.AdvancedBuild", new IMAGE_BRUSH( "Launcher/Launcher_Advanced", Icon40x40 ) );
+		Set( "LauncherCommand.AdvancedBuild.Medium", new IMAGE_BRUSH("Launcher/Launcher_Advanced", Icon25x25) );
+		Set( "LauncherCommand.AdvancedBuild.Small", new IMAGE_BRUSH("Launcher/Launcher_Advanced", Icon20x20) );
 
 		Set( "Launcher.Platform.Warning", new IMAGE_BRUSH( "Icons/alert", Icon24x24) );
 
@@ -5578,6 +5585,13 @@ void FSlateEditorStyle::FStyle::SetupAutomationStyles()
 			Set( PlatformInfo.GetIconStyleName(PlatformInfo::EPlatformIconSize::XLarge), new IMAGE_BRUSH( *PlatformInfo.GetIconPath(PlatformInfo::EPlatformIconSize::XLarge), Icon128x128 ) );
 		}
 #endif
+
+		Set("Launcher.NoHoverTableRow", FTableRowStyle(NormalTableRowStyle)
+			.SetEvenRowBackgroundHoveredBrush(FSlateNoResource())
+			.SetOddRowBackgroundHoveredBrush(FSlateNoResource())
+			.SetActiveHoveredBrush(FSlateNoResource())
+			.SetInactiveHoveredBrush(FSlateNoResource())
+			);
 	}
 
 	// Device Manager
diff --git a/Engine/Source/Editor/MainFrame/Private/Menus/MainMenu.cpp b/Engine/Source/Editor/MainFrame/Private/Menus/MainMenu.cpp
--- a/Engine/Source/Editor/MainFrame/Private/Menus/MainMenu.cpp
+++ b/Engine/Source/Editor/MainFrame/Private/Menus/MainMenu.cpp
@@ -182,7 +182,8 @@ void FMainMenu::FillWindowMenu( FMenuBuilder& MenuBuilder, const TSharedRef< FEx
 
 		if (GetDefault<UEditorExperimentalSettings>()->bGameLauncher)
 		{
-			FGlobalTabmanager::Get()->PopulateTabSpawnerMenu(MenuBuilder, "SessionLauncher");
+			//FGlobalTabmanager::Get()->PopulateTabSpawnerMenu(MenuBuilder, "SessionLauncher"); // Will be deleted soon
+			FGlobalTabmanager::Get()->PopulateTabSpawnerMenu(MenuBuilder, "ProjectLauncher");
 		}
 
 		if (GetDefault<UEditorExperimentalSettings>()->bEnableTranslationEditor)
diff --git a/Engine/Source/Editor/UnrealEd/Private/Editor.cpp b/Engine/Source/Editor/UnrealEd/Private/Editor.cpp
--- a/Engine/Source/Editor/UnrealEd/Private/Editor.cpp
+++ b/Engine/Source/Editor/UnrealEd/Private/Editor.cpp
@@ -550,6 +550,7 @@ void UEditorEngine::Init(IEngineLoop* InEngineLoop)
 //		FModuleManager::Get().LoadModule(TEXT("Search"));
 		FModuleManager::Get().LoadModule(TEXT("SessionFrontend"));
 		FModuleManager::Get().LoadModule(TEXT("SessionLauncher"));
+		FModuleManager::Get().LoadModule(TEXT("ProjectLauncher"));
 		FModuleManager::Get().LoadModule(TEXT("SettingsEditor"));
 		FModuleManager::Get().LoadModule(TEXT("EditorSettingsViewer"));
 		FModuleManager::Get().LoadModule(TEXT("ProjectSettingsViewer"));
diff --git a/Engine/Source/Editor/UnrealEd/Private/Kismet2/DebuggerCommands.cpp b/Engine/Source/Editor/UnrealEd/Private/Kismet2/DebuggerCommands.cpp
--- a/Engine/Source/Editor/UnrealEd/Private/Kismet2/DebuggerCommands.cpp
+++ b/Engine/Source/Editor/UnrealEd/Private/Kismet2/DebuggerCommands.cpp
@@ -60,8 +60,8 @@ public:
 
 	// Launch On
 	static void HandleLaunchOnDeviceActionExecute(FString DevicedId, FString DeviceName);
-	static bool HandleLaunchOnDeviceActionCanExecute(FString DevicedId);
-	static bool HandleLaunchOnDeviceActionIsChecked(FString DeviceId);
+	static bool HandleLaunchOnDeviceActionCanExecute(FString DeviceName);
+	static bool HandleLaunchOnDeviceActionIsChecked(FString DeviceName);
 
 	// No Device
 	static void HandleNoDeviceFoundActionExecute() {}
@@ -117,7 +117,7 @@ protected:
 	 *
 	 * @return true if the platform can be played on, false otherwise.
 	 */
-	static bool CanLaunchOnDevice(const FString& DeviceId);
+	static bool CanLaunchOnDevice(const FString& DeviceName);
 
 	/**
 	 * Starts a game session on the default device of the specified platform.
@@ -163,7 +163,7 @@ FPlayWorldCommands::FPlayWorldCommands()
 	ULevelEditorPlaySettings* PlaySettings = GetMutableDefault<ULevelEditorPlaySettings>();
 
 	// initialize default Play device
-	if (PlaySettings->LastExecutedLaunchDevice.IsEmpty())
+	if (PlaySettings->LastExecutedLaunchName.IsEmpty())
 	{
 		FString RunningPlatformName = GetTargetPlatformManagerRef().GetRunningTargetPlatform()->PlatformName();
 		FString PlayPlatformName;
@@ -187,7 +187,7 @@ FPlayWorldCommands::FPlayWorldCommands()
 
 				if (PlayDevice.IsValid())
 				{
-					PlaySettings->LastExecutedLaunchDevice = PlayDevice->GetId().ToString();
+					PlaySettings->LastExecutedLaunchDevice = PlayDevice->GetName();
 					PlaySettings->SaveConfig();
 				}
 			}
@@ -596,90 +596,73 @@ TSharedRef< SWidget > FPlayWorldCommands::GenerateLaunchMenuContent( TSharedRef<
 
 			if (VanillaPlatform.PlatformInfo->SDKStatus == PlatformInfo::EPlatformSDKStatus::Installed)
 			{
-				TArray<FString> PlatformFlavors;
-				if (VanillaPlatform.PlatformFlavors.Num() && PlatformsToCheckFlavorsFor.Find(VanillaPlatform.PlatformInfo->TargetPlatformName.ToString()) != INDEX_NONE)
+				// for each platform...
+				TArray<ITargetDeviceProxyPtr> DeviceProxies;
+				TargetDeviceServicesModule->GetDeviceProxyManager()->GetProxies(VanillaPlatform.PlatformInfo->VanillaPlatformName, false, DeviceProxies);
+					
+				// if this platform had no devices, but we want to show an extra option if not installed right
+				if (DeviceProxies.Num() == 0)
 				{
-					for (const PlatformInfo::FPlatformInfo* SubPlatformInfo : VanillaPlatform.PlatformFlavors)
+					if (PlatformsWithNoDevices.Find(VanillaPlatform.PlatformInfo->VanillaPlatformName) == INDEX_NONE)
 					{
-						PlatformFlavors.Add(SubPlatformInfo->TargetPlatformName.ToString());
+						// add an entry with a no devices found
+						PlatformsWithNoDevices.Add(VanillaPlatform.PlatformInfo->VanillaPlatformName);
 					}
 				}
 				else
 				{
-					PlatformFlavors.Add(VanillaPlatform.PlatformInfo->TargetPlatformName.ToString());
-				}
-				for (int32 PlatformIndex = 0; PlatformIndex < PlatformFlavors.Num(); ++PlatformIndex)
-				{
-					ITargetPlatform* Platform = GetTargetPlatformManager()->FindTargetPlatform(PlatformFlavors[PlatformIndex]);
+					// for each proxy...
+					for (auto DeviceProxyIt = DeviceProxies.CreateIterator(); DeviceProxyIt; ++DeviceProxyIt)
+					{
+						ITargetDeviceProxyPtr DeviceProxy = *DeviceProxyIt;
 
-					// for each platform...
-					TArray<ITargetDeviceProxyPtr> DeviceProxies;
-					TargetDeviceServicesModule->GetDeviceProxyManager()->GetProxies(Platform->PlatformName(), false, DeviceProxies);
+						// ... create an action...
+						FUIAction LaunchDeviceAction(
+							FExecuteAction::CreateStatic(&FInternalPlayWorldCommandCallbacks::HandleLaunchOnDeviceActionExecute, DeviceProxy->GetTargetDeviceId(NAME_None), DeviceProxy->GetName()),
+							FCanExecuteAction::CreateStatic(&FInternalPlayWorldCommandCallbacks::HandleLaunchOnDeviceActionCanExecute, DeviceProxy->GetName()),
+							FIsActionChecked::CreateStatic(&FInternalPlayWorldCommandCallbacks::HandleLaunchOnDeviceActionIsChecked, DeviceProxy->GetName())
+							);
 
-					// if this platform had no devices, but we want to show an extra option if not installed right
-					if (DeviceProxies.Num() == 0)
-					{
-						if (PlatformsToMaybeInstallLinksFor.Find(PlatformFlavors[PlatformIndex]) != INDEX_NONE && PlatformsWithNoDevices.Find(VanillaPlatform.PlatformInfo->VanillaPlatformName) == INDEX_NONE)
+						// ... generate display label...
+						FFormatNamedArguments LabelArguments;
+						LabelArguments.Add(TEXT("DeviceName"), FText::FromString(DeviceProxy->GetName()));
+
+						if (!DeviceProxy->IsConnected())
 						{
-							// add an entry with a no devices found
-							PlatformsWithNoDevices.Add(VanillaPlatform.PlatformInfo->VanillaPlatformName);
+							LabelArguments.Add(TEXT("HostUser"), LOCTEXT("DisconnectedHint", " [Disconnected]"));
 						}
-					}
-					else
-					{
-						// for each proxy...
-						for (auto DeviceProxyIt = DeviceProxies.CreateIterator(); DeviceProxyIt; ++DeviceProxyIt)
+						else if (DeviceProxy->GetHostUser() != FPlatformProcess::UserName(false))
 						{
-							ITargetDeviceProxyPtr DeviceProxy = *DeviceProxyIt;
-
-							// ... create an action...
-							FUIAction LaunchDeviceAction(
-								FExecuteAction::CreateStatic(&FInternalPlayWorldCommandCallbacks::HandleLaunchOnDeviceActionExecute, DeviceProxy->GetDeviceId(), DeviceProxy->GetName()),
-								FCanExecuteAction::CreateStatic(&FInternalPlayWorldCommandCallbacks::HandleLaunchOnDeviceActionCanExecute, DeviceProxy->GetDeviceId()),
-								FIsActionChecked::CreateStatic(&FInternalPlayWorldCommandCallbacks::HandleLaunchOnDeviceActionIsChecked, DeviceProxy->GetDeviceId())
-								);
-
-							// ... generate display label...
-							FFormatNamedArguments LabelArguments;
-							LabelArguments.Add(TEXT("DeviceName"), FText::FromString(DeviceProxy->GetName()));
-
-							if (!DeviceProxy->IsConnected())
-							{
-								LabelArguments.Add(TEXT("HostUser"), LOCTEXT("DisconnectedHint", " [Disconnected]"));
-							}
-							else if (DeviceProxy->GetHostUser() != FPlatformProcess::UserName(false))
-							{
-								LabelArguments.Add(TEXT("HostUser"), FText::FromString(DeviceProxy->GetHostUser()));
-							}
-							else
-							{
-								LabelArguments.Add(TEXT("HostUser"), FText::GetEmpty());
-							}
-
-							FText Label = FText::Format(LOCTEXT("LaunchDeviceLabel", "{DeviceName}{HostUser}"), LabelArguments);
-
-							// ... generate tooltip text
-							FFormatNamedArguments TooltipArguments;
-							TooltipArguments.Add(TEXT("DeviceID"), FText::FromString(DeviceProxy->GetName()));
-							TooltipArguments.Add(TEXT("DisplayName"), VanillaPlatform.PlatformInfo->DisplayName);
-							FText Tooltip = FText::Format(LOCTEXT("LaunchDeviceToolTipText", "Launch the game on this {DisplayName} device ({DeviceID})"), TooltipArguments);
-
-							FProjectStatus ProjectStatus;
-							if (IProjectManager::Get().QueryStatusForCurrentProject(ProjectStatus) && !ProjectStatus.IsTargetPlatformSupported(VanillaPlatform.PlatformInfo->VanillaPlatformName)) 
-							{
-								FText TooltipLine2 = FText::Format(LOCTEXT("LaunchDevicePlatformWarning", "{DisplayName} is not listed as a target platform for this project, so may not run as expected."), TooltipArguments);
-								Tooltip = FText::Format(FText::FromString(TEXT("{0}\n\n{1}")), Tooltip, TooltipLine2);
-							}
-
-							// ... and add a menu entry
-							MenuBuilder.AddMenuEntry(
-								LaunchDeviceAction, 
-								ProjectTargetPlatformEditorModule.MakePlatformMenuItemWidget(Platform->GetPlatformInfo(), true, Label), 
-								NAME_None, 
-								Tooltip, 
-								EUserInterfaceActionType::Check
-								);
+							LabelArguments.Add(TEXT("HostUser"), FText::FromString(DeviceProxy->GetHostUser()));
 						}
+						else
+						{
+							LabelArguments.Add(TEXT("HostUser"), FText::GetEmpty());
+						}
+
+						FText Label = FText::Format(LOCTEXT("LaunchDeviceLabel", "{DeviceName}{HostUser}"), LabelArguments);
+
+						// ... generate tooltip text
+						FFormatNamedArguments TooltipArguments;
+						TooltipArguments.Add(TEXT("DeviceID"), FText::FromString(DeviceProxy->GetName()));
+						TooltipArguments.Add(TEXT("DisplayName"), VanillaPlatform.PlatformInfo->DisplayName);
+						FText Tooltip = FText::Format(LOCTEXT("LaunchDeviceToolTipText", "Launch the game on this {DisplayName} device ({DeviceID})"), TooltipArguments);
+
+						FProjectStatus ProjectStatus;
+						if (IProjectManager::Get().QueryStatusForCurrentProject(ProjectStatus) && !ProjectStatus.IsTargetPlatformSupported(VanillaPlatform.PlatformInfo->VanillaPlatformName)) 
+						{
+							FText TooltipLine2 = FText::Format(LOCTEXT("LaunchDevicePlatformWarning", "{DisplayName} is not listed as a target platform for this project, so may not run as expected."), TooltipArguments);
+							Tooltip = FText::Format(FText::FromString(TEXT("{0}\n\n{1}")), Tooltip, TooltipLine2);
+						}
+
+						// ... and add a menu entry
+						MenuBuilder.AddMenuEntry(
+							LaunchDeviceAction, 
+							ProjectTargetPlatformEditorModule.MakePlatformMenuItemWidget(*VanillaPlatform.PlatformInfo, true, Label),
+							NAME_None, 
+							Tooltip, 
+							EUserInterfaceActionType::Check
+							);
 					}
 				}
 			}
@@ -1324,7 +1307,7 @@ bool FInternalPlayWorldCommandCallbacks::RepeatLastLaunch_CanExecute()
 	switch (PlaySettings->LastExecutedLaunchModeType)
 	{
 	case LaunchMode_OnDevice:
-		return CanLaunchOnDevice(PlaySettings->LastExecutedLaunchDevice);
+		return CanLaunchOnDevice(PlaySettings->LastExecutedLaunchName);
 
 	default:
 		return false;
@@ -1339,10 +1322,10 @@ FText FInternalPlayWorldCommandCallbacks::GetRepeatLastLaunchToolTip()
 	switch (PlaySettings->LastExecutedLaunchModeType)
 	{
 	case LaunchMode_OnDevice:
-		if (CanLaunchOnDevice(PlaySettings->LastExecutedLaunchDevice))
+		if (CanLaunchOnDevice(PlaySettings->LastExecutedLaunchName))
 		{
 			FFormatNamedArguments Arguments;
-			Arguments.Add(TEXT("DeviceName"), FText::FromString(PlaySettings->LastExecutedLaunchDevice));
+			Arguments.Add(TEXT("DeviceName"), FText::FromString(PlaySettings->LastExecutedLaunchName));
 
 			return FText::Format( LOCTEXT("RepeatLaunchTooltip", "Launch this level on {DeviceName}"), Arguments);
 		}
@@ -1458,15 +1441,15 @@ void FInternalPlayWorldCommandCallbacks::HandleLaunchOnDeviceActionExecute( FStr
 }
 
 
-bool FInternalPlayWorldCommandCallbacks::HandleLaunchOnDeviceActionCanExecute( FString DeviceId )
+bool FInternalPlayWorldCommandCallbacks::HandleLaunchOnDeviceActionCanExecute(FString DeviceName)
 {
-	return CanLaunchOnDevice(DeviceId);
+	return CanLaunchOnDevice(DeviceName);
 }
 
 
-bool FInternalPlayWorldCommandCallbacks::HandleLaunchOnDeviceActionIsChecked(FString DeviceId)
+bool FInternalPlayWorldCommandCallbacks::HandleLaunchOnDeviceActionIsChecked(FString DeviceName)
 {
-	return (DeviceId == GetDefault<ULevelEditorPlaySettings>()->LastExecutedLaunchDevice);
+	return (DeviceName == GetDefault<ULevelEditorPlaySettings>()->LastExecutedLaunchName);
 }
 
 
@@ -1668,7 +1651,7 @@ bool FInternalPlayWorldCommandCallbacks::CanPossessEjectPlayer()
 }
 
 
-bool FInternalPlayWorldCommandCallbacks::CanLaunchOnDevice( const FString& DeviceId )
+bool FInternalPlayWorldCommandCallbacks::CanLaunchOnDevice(const FString& DeviceName)
 {
 	if (!GUnrealEd->IsPlayingViaLauncher())
 	{
@@ -1676,7 +1659,7 @@ bool FInternalPlayWorldCommandCallbacks::CanLaunchOnDevice( const FString& Devic
 
 		if (TargetDeviceServicesModule.IsValid())
 		{
-			ITargetDeviceProxyPtr DeviceProxy = TargetDeviceServicesModule->GetDeviceProxyManager()->FindProxy(DeviceId);
+			ITargetDeviceProxyPtr DeviceProxy = TargetDeviceServicesModule->GetDeviceProxyManager()->FindProxy(DeviceName);
 
 			return (DeviceProxy.IsValid() && DeviceProxy->IsConnected());
 		}
diff --git a/Engine/Source/Editor/UnrealEd/Private/PlayLevel.cpp b/Engine/Source/Editor/UnrealEd/Private/PlayLevel.cpp
--- a/Engine/Source/Editor/UnrealEd/Private/PlayLevel.cpp
+++ b/Engine/Source/Editor/UnrealEd/Private/PlayLevel.cpp
@@ -1256,18 +1256,14 @@ void UEditorEngine::PlayUsingLauncher()
 		FGameProjectGenerationModule& GameProjectModule = FModuleManager::LoadModuleChecked<FGameProjectGenerationModule>(TEXT("GameProjectGeneration"));
 		bool bHasCode = GameProjectModule.Get().GetProjectCodeFileCount() > 0;
 
+		// Setup launch profile, keep the setting here to a minimum.
 		ILauncherProfileRef LauncherProfile = LauncherServicesModule.CreateProfile(TEXT("Play On Device"));
 		LauncherProfile->SetBuildGame(bHasCode && FSourceCodeNavigation::IsCompilerAvailable());
 		LauncherProfile->SetCookMode(ELauncherProfileCookModes::ByTheBook);
 		LauncherProfile->AddCookedPlatform(PlayUsingLauncherDeviceId.Left(PlayUsingLauncherDeviceId.Find(TEXT("@"))));
-		LauncherProfile->SetForceClose(true);
-		LauncherProfile->SetTimeout(60);
-		LauncherProfile->SetCookConfiguration(FApp::GetBuildConfiguration());
-		LauncherProfile->SetBuildConfiguration(EBuildConfigurations::Development);
-		LauncherProfile->SetDeploymentMode(ELauncherProfileDeploymentModes::CopyToDevice);
 		LauncherProfile->SetDeployedDeviceGroup(DeviceGroup);
-		LauncherProfile->SetHideFileServerWindow(false);
 		LauncherProfile->SetEditorExe(FUnrealEdMisc::Get().GetExecutableForCommandlets());
+
 		const FString DummyDeviceName(FString::Printf(TEXT("All_iOS_On_%s"), FPlatformProcess::ComputerName()));
 		if (PlayUsingLauncherDeviceId.Left(PlayUsingLauncherDeviceId.Find(TEXT("@"))) != TEXT("IOS") || !PlayUsingLauncherDeviceName.Contains(DummyDeviceName))
 		{
diff --git a/Engine/Source/Editor/UnrealEd/UnrealEd.Build.cs b/Engine/Source/Editor/UnrealEd/UnrealEd.Build.cs
--- a/Engine/Source/Editor/UnrealEd/UnrealEd.Build.cs
+++ b/Engine/Source/Editor/UnrealEd/UnrealEd.Build.cs
@@ -145,6 +145,7 @@ public class UnrealEd : ModuleRules
 				"PackagesDialog",
 				"Persona",
 				"PhAT",
+                "ProjectLauncher",
 				"DeviceManager",
 				"SettingsEditor",
 				"SessionFrontend",
diff --git a/Engine/Source/Programs/UnrealFrontend/Private/UnrealFrontendMain.cpp b/Engine/Source/Programs/UnrealFrontend/Private/UnrealFrontendMain.cpp
--- a/Engine/Source/Programs/UnrealFrontend/Private/UnrealFrontendMain.cpp
+++ b/Engine/Source/Programs/UnrealFrontend/Private/UnrealFrontendMain.cpp
@@ -365,6 +365,7 @@ void RunUI()
 	FModuleManager::Get().LoadModule("SessionFrontend");
 	FModuleManager::Get().LoadModule("SessionLauncher");
 	FModuleManager::Get().LoadModule("SettingsEditor");
+	FModuleManager::Get().LoadModule("ProjectLauncher");
 
 	// Create developer tools menu with widget reflector.
 	FModuleManager::LoadModuleChecked<ISlateReflectorModule>("SlateReflector").RegisterTabSpawner(WorkspaceMenu::DeveloperTools);
@@ -387,6 +388,7 @@ void RunUI()
 						->AddTab(FName("MessagingDebugger"), ETabState::ClosedTab)
 						->AddTab(FName("SessionLauncher"), ETabState::OpenedTab)
 						->AddTab(FName("SessionFrontend"), ETabState::OpenedTab)
+						->AddTab(FName("ProjectLauncher"), ETabState::OpenedTab)
 				)
 		)
 		->AddArea
diff --git a/Engine/Source/Runtime/Core/Public/Linux/LinuxPlatformProperties.h b/Engine/Source/Runtime/Core/Public/Linux/LinuxPlatformProperties.h
--- a/Engine/Source/Runtime/Core/Public/Linux/LinuxPlatformProperties.h
+++ b/Engine/Source/Runtime/Core/Public/Linux/LinuxPlatformProperties.h
@@ -108,4 +108,24 @@ struct FLinuxPlatformProperties
 	{
 		return true;
 	}
+
+	static FORCEINLINE float GetVariantPriority()
+	{
+		if (IS_DEDICATED_SERVER)
+		{
+			return 0.0f;
+		}
+
+		if (HAS_EDITOR_DATA)
+		{
+			return 0.0f;
+		}
+
+		if (IS_CLIENT_ONLY)
+		{
+			return 0.0f;
+		}
+
+		return 1.0f;
+	}
 };
diff --git a/Engine/Source/Runtime/Core/Public/Mac/MacPlatformProperties.h b/Engine/Source/Runtime/Core/Public/Mac/MacPlatformProperties.h
--- a/Engine/Source/Runtime/Core/Public/Mac/MacPlatformProperties.h
+++ b/Engine/Source/Runtime/Core/Public/Mac/MacPlatformProperties.h
@@ -92,4 +92,24 @@ struct FMacPlatformProperties
 	{
 		return true;
 	}
+
+	static FORCEINLINE float GetVariantPriority()
+	{
+		if (IS_DEDICATED_SERVER)
+		{
+			return 0.0f;
+		}
+
+		if (HAS_EDITOR_DATA)
+		{
+			return 0.0f;
+		}
+
+		if (IS_CLIENT_ONLY)
+		{
+			return 0.0f;
+		}
+
+		return 1.0f;
+	}
 };
diff --git a/Engine/Source/Runtime/Core/Public/Windows/WindowsPlatformProperties.h b/Engine/Source/Runtime/Core/Public/Windows/WindowsPlatformProperties.h
--- a/Engine/Source/Runtime/Core/Public/Windows/WindowsPlatformProperties.h
+++ b/Engine/Source/Runtime/Core/Public/Windows/WindowsPlatformProperties.h
@@ -9,6 +9,10 @@
 
 #pragma once
 
+//#include "Core.h"
+//#include "AssertionMacros.h"
+//#include "Internationalization.h"
+
 #include "GenericPlatformProperties.h"
 
 
@@ -103,4 +107,24 @@ struct FWindowsPlatformProperties
 	{
 		return true;
 	}
+
+	static FORCEINLINE float GetVariantPriority()
+	{
+		if (IS_DEDICATED_SERVER)
+		{
+			return 0.0f;
+		}
+
+		if (HAS_EDITOR_DATA)
+		{
+			return 0.0f;
+		}
+
+		if (IS_CLIENT_ONLY)
+		{
+			return 0.0f;
+		}
+
+		return 1.0f;
+	}
 };
diff --git a/Engine/Source/Runtime/LauncherAutomatedService/Private/LauncherAutomatedServiceProvider.cpp b/Engine/Source/Runtime/LauncherAutomatedService/Private/LauncherAutomatedServiceProvider.cpp
--- a/Engine/Source/Runtime/LauncherAutomatedService/Private/LauncherAutomatedServiceProvider.cpp
+++ b/Engine/Source/Runtime/LauncherAutomatedService/Private/LauncherAutomatedServiceProvider.cpp
@@ -205,7 +205,7 @@ void FLauncherAutomatedServiceProvider::SetupProfileAndGroupSettings( const TCHA
 				}
 			}
 
-			bHasValidDeviceGroup = AutomatedDeviceGroup.IsValid() && ( AutomatedDeviceGroup->GetDevices().Num() > 0 );
+			bHasValidDeviceGroup = AutomatedDeviceGroup.IsValid() && (AutomatedDeviceGroup->GetNumDevices() > 0);
 
 			// A reference to the proxy manager responsible for device activity here.
 			DeviceProxyManager = FModuleManager::LoadModuleChecked<ITargetDeviceServicesModule>(TEXT("TargetDeviceServices")).GetDeviceProxyManager();
@@ -243,11 +243,11 @@ void FLauncherAutomatedServiceProvider::HandleDeviceProxyManagerProxyAdded( cons
 		// @todo gmp: fix automated service provider; must use ILauncher interface here
 		//AddedProxy->Run(AutomatedProfile->GetBuildConfiguration(), AutomatedProfile->GetBuildGame(), CommandLine);
 
-		UE_LOG(LauncherAutomatedService, Display, TEXT("%s"), *FText::Format(FText::FromString("Deploying to: {0}"), FText::FromString(AddedProxy->GetDeviceId())).ToString());
+		UE_LOG(LauncherAutomatedService, Display, TEXT("%s"), *FText::Format(FText::FromString("Deploying to: {0}"), FText::FromString(AddedProxy->GetName())).ToString());
 
 		DeployedInstances.Add( AddedProxy );
 
-		if( AutomatedDeviceGroup->GetDevices().Num() == DeployedInstances.Num() )
+		if (AutomatedDeviceGroup->GetNumDevices() == DeployedInstances.Num())
 		{
 			bHasLaunchedAllInstances = true;
 			// All devices have been found. Start the automation test process
diff --git a/Engine/Source/Runtime/SlateCore/Private/Styling/CoreStyle.cpp b/Engine/Source/Runtime/SlateCore/Private/Styling/CoreStyle.cpp
--- a/Engine/Source/Runtime/SlateCore/Private/Styling/CoreStyle.cpp
+++ b/Engine/Source/Runtime/SlateCore/Private/Styling/CoreStyle.cpp
@@ -146,6 +146,8 @@ TSharedRef<ISlateStyle> FCoreStyle::Create( const FName& InStyleSetName )
 
 	Style->Set("NormalFont", TTF_FONT("Fonts/Roboto-Regular", 9));
 
+	Style->Set("SmallFont", TTF_FONT("Fonts/Roboto-Regular", 8));
+
 	// Normal Text
 	FTextBlockStyle NormalText = FTextBlockStyle()
 		.SetFont(TTF_FONT("Fonts/Roboto-Regular", 9))
@@ -155,6 +157,10 @@ TSharedRef<ISlateStyle> FCoreStyle::Create( const FName& InStyleSetName )
 		.SetHighlightColor(FLinearColor(0.02f, 0.3f, 0.0f))
 		.SetHighlightShape(BOX_BRUSH("Common/TextBlockHighlightShape", FMargin(3.f /8.f)));
 
+	// Small Text
+	FTextBlockStyle SmallText = FTextBlockStyle(NormalText)
+		.SetFont(TTF_FONT("Fonts/Roboto-Regular", 8));
+
 	// Embossed Text
 	Style->Set("EmbossedText", FTextBlockStyle(NormalText)
 		.SetFont(TTF_FONT("Fonts/Roboto-Regular", 24))
@@ -420,6 +426,7 @@ TSharedRef<ISlateStyle> FCoreStyle::Create( const FName& InStyleSetName )
 	// STextBlock defaults...
 	{
 		Style->Set("NormalText", NormalText);
+		Style->Set("SmallText", SmallText);
 	}
 
 	// SInlineEditableTextBlock
@@ -429,6 +436,9 @@ TSharedRef<ISlateStyle> FCoreStyle::Create( const FName& InStyleSetName )
 			.SetShadowOffset( FVector2D::ZeroVector )
 			.SetShadowColorAndOpacity( FLinearColor::Black );
 
+		FTextBlockStyle InlineEditableTextBlockSmallReadOnly = FTextBlockStyle(InlineEditableTextBlockReadOnly)
+			.SetFont(SmallText.Font);
+
 		FEditableTextBoxStyle InlineEditableTextBlockEditable = FEditableTextBoxStyle()
 			.SetFont(NormalText.Font)
 			.SetBackgroundImageNormal( BOX_BRUSH( "Common/TextBox", FMargin(4.0f/16.0f) ) )
@@ -437,10 +447,20 @@ TSharedRef<ISlateStyle> FCoreStyle::Create( const FName& InStyleSetName )
 			.SetBackgroundImageReadOnly( BOX_BRUSH( "Common/TextBox_ReadOnly", FMargin(4.0f/16.0f) ) )
 			.SetScrollBarStyle( ScrollBar );
 
+		FEditableTextBoxStyle InlineEditableTextBlockSmallEditable = FEditableTextBoxStyle(InlineEditableTextBlockEditable)
+			.SetFont(SmallText.Font);
+
 		FInlineEditableTextBlockStyle InlineEditableTextBlockStyle = FInlineEditableTextBlockStyle()
 			.SetTextStyle(InlineEditableTextBlockReadOnly)
 			.SetEditableTextBoxStyle(InlineEditableTextBlockEditable);
+
 		Style->Set( "InlineEditableTextBlockStyle", InlineEditableTextBlockStyle );
+
+		FInlineEditableTextBlockStyle InlineEditableTextBlockSmallStyle = FInlineEditableTextBlockStyle()
+			.SetTextStyle(InlineEditableTextBlockSmallReadOnly)
+			.SetEditableTextBoxStyle(InlineEditableTextBlockSmallEditable);
+
+		Style->Set("InlineEditableTextBlockSmallStyle", InlineEditableTextBlockSmallStyle);
 	}
 
 	// SSuggestionTextBox defaults...
