Set lightmap index to auto generated dst

diff --git a/Engine/Source/Editor/UnrealEd/Private/Fbx/FbxStaticMeshImport.cpp b/Engine/Source/Editor/UnrealEd/Private/Fbx/FbxStaticMeshImport.cpp
--- a/Engine/Source/Editor/UnrealEd/Private/Fbx/FbxStaticMeshImport.cpp
+++ b/Engine/Source/Editor/UnrealEd/Private/Fbx/FbxStaticMeshImport.cpp
@@ -1167,8 +1167,13 @@ UStaticMesh* UnFbx::FFbxImporter::ImportStaticMeshAsSingle(UObject* InParent, TA
 		SrcModel.BuildSettings.bRemoveDegenerates = ImportOptions->bRemoveDegenerates;
 		SrcModel.BuildSettings.bRecomputeNormals = ImportOptions->NormalImportMethod == FBXNIM_ComputeNormals;
 		SrcModel.BuildSettings.bRecomputeTangents = ImportOptions->NormalImportMethod != FBXNIM_ImportNormalsAndTangents;
-		SrcModel.BuildSettings.bGenerateLightmapUVs = ImportOptions->bGenerateLightmapUVs;
-		SrcModel.BuildSettings.DstLightmapIndex = FirstOpenUVChannel;
+
+		if( ImportOptions->bGenerateLightmapUVs )
+		{
+			SrcModel.BuildSettings.bGenerateLightmapUVs = true;
+			SrcModel.BuildSettings.DstLightmapIndex = FirstOpenUVChannel;
+			StaticMesh->LightMapCoordinateIndex = FirstOpenUVChannel;
+		}
 
 		StaticMesh->LODGroup = ImportOptions->StaticMeshLODGroup;
 		StaticMesh->Build(false);
