Automatically generate 18-dop for imported static mesh

TTP#345831

diff --git a/Engine/Source/Editor/UnrealEd/Private/Fbx/FbxStaticMeshImport.cpp b/Engine/Source/Editor/UnrealEd/Private/Fbx/FbxStaticMeshImport.cpp
--- a/Engine/Source/Editor/UnrealEd/Private/Fbx/FbxStaticMeshImport.cpp
+++ b/Engine/Source/Editor/UnrealEd/Private/Fbx/FbxStaticMeshImport.cpp
@@ -1200,6 +1200,20 @@ UStaticMesh* UnFbx::FFbxImporter::ImportStaticMeshAsSingle(UObject* InParent, TA
 		StaticMesh = NULL;
 	}
 
+
+	if (StaticMesh)
+	{
+
+		FKAggregateGeom & AggGeom = StaticMesh->BodySetup->AggGeom;
+		AggGeom.ConvexElems.Empty(1);	//when reimport happens we remove any existing convex mesh
+
+		const int32 NumDirs = 18;
+		TArray<FVector> Dirs;
+		Dirs.AddUninitialized(NumDirs);
+		for (int32 DirIdx = 0; DirIdx < NumDirs; ++DirIdx) { Dirs[DirIdx] = KDopDir18[DirIdx]; }
+		GenerateKDopAsSimpleCollision(StaticMesh, Dirs);
+	}
+
 	return StaticMesh;
 }
 
diff --git a/Engine/Source/Editor/UnrealEd/Private/GeomFitUtils.cpp b/Engine/Source/Editor/UnrealEd/Private/GeomFitUtils.cpp
--- a/Engine/Source/Editor/UnrealEd/Private/GeomFitUtils.cpp
+++ b/Engine/Source/Editor/UnrealEd/Private/GeomFitUtils.cpp
@@ -63,7 +63,7 @@ static bool PromptToRemoveExistingCollision(UStaticMesh* StaticMesh)
 // THIS FUNCTION REPLACES EXISTING SIMPLE COLLISION MODEL WITH KDOP
 #define MY_FLTMAX (3.402823466e+38F)
 
-int32 GenerateKDopAsSimpleCollision(UStaticMesh* StaticMesh, TArray<FVector> &Dirs)
+int32 GenerateKDopAsSimpleCollision(UStaticMesh* StaticMesh, const TArray<FVector> &Dirs)
 {
 	// Make sure rendering is done - so we are not changing data being used by collision drawing.
 	FlushRenderingCommands();
diff --git a/Engine/Source/Editor/UnrealEd/Private/GeomFitUtils.h b/Engine/Source/Editor/UnrealEd/Private/GeomFitUtils.h
--- a/Engine/Source/Editor/UnrealEd/Private/GeomFitUtils.h
+++ b/Engine/Source/Editor/UnrealEd/Private/GeomFitUtils.h
@@ -105,7 +105,7 @@ const FVector KDopDir26[26] =
 };
 
 // Utilities
-UNREALED_API int32 GenerateKDopAsSimpleCollision(UStaticMesh* StaticMesh, TArray<FVector> &dirs);
+UNREALED_API int32 GenerateKDopAsSimpleCollision(UStaticMesh* StaticMesh, const TArray<FVector> &dirs);
 UNREALED_API int32 GenerateBoxAsSimpleCollision(UStaticMesh* StaticMesh);
 UNREALED_API int32 GenerateSphereAsSimpleCollision(UStaticMesh* StaticMesh);
 UNREALED_API int32 GenerateSphylAsSimpleCollision(UStaticMesh* StaticMesh);
