New Rendering project setting 'bUseNormalMapsForStaticLighting'
* Disabled by default as it causes artifacts with detailed normal maps (when normal map features approach lightmap or material export resolution, some details get filtered out while others survive)

diff --git a/Engine/Config/BaseEngine.ini b/Engine/Config/BaseEngine.ini
--- a/Engine/Config/BaseEngine.ini
+++ b/Engine/Config/BaseEngine.ini
@@ -154,9 +154,6 @@ InitialButtonRepeatDelay=0.2
 ButtonRepeatDelay=0.1
 NetDriverDefinitions=(DefName="GameNetDriver",DriverClassName="/Script/OnlineSubsystemUtils.IpNetDriver",DriverClassNameFallback="/Script/OnlineSubsystemUtils.IpNetDriver")
 
-; Enables normal map sampling when Lightmass is generating 'simple' light maps.  This increases lighting build time, but may improve quality when normal maps are used to represent curvature over a large surface area.  When this setting is disabled, 'simple' light maps will not take normal maps into account.
-bUseNormalMapsForSimpleLightMaps=true
-
 bSmoothFrameRate=true
 SmoothedFrameRateRange=(LowerBound=(Type="ERangeBoundTypes::Inclusive",Value=22),UpperBound=(Type="ERangeBoundTypes::Exclusive",Value=62))
 bCheckForMultiplePawnsSpawnedInAFrame=false
diff --git a/Engine/Source/Editor/UnrealEd/Private/Lightmass/Lightmass.cpp b/Engine/Source/Editor/UnrealEd/Private/Lightmass/Lightmass.cpp
--- a/Engine/Source/Editor/UnrealEd/Private/Lightmass/Lightmass.cpp
+++ b/Engine/Source/Editor/UnrealEd/Private/Lightmass/Lightmass.cpp
@@ -1118,9 +1118,11 @@ void FLightmassExporter::BuildMaterialMap(UMaterialInterface* Material)
 		{
 			if( ErrorCode == NSwarm::SWARM_ERROR_FILE_FOUND_NOT )
 			{
+				static const auto CVar = IConsoleManager::Get().FindTConsoleVariableDataInt(TEXT("r.NormalMapsForStaticLighting"));
+				const bool bUseNormalMapsForLighting = CVar->GetValueOnGameThread() != 0;
+
 				// Only generate normal maps if we'll actually need them for lighting
-				const bool bWantNormals = GEngine->bUseNormalMapsForSimpleLightMaps;
-				MaterialRenderer.BeginGenerateMaterialData(Material, bWantNormals, NewChannelName, MaterialExportData);
+				MaterialRenderer.BeginGenerateMaterialData(Material, bUseNormalMapsForLighting, NewChannelName, MaterialExportData);
 			}
 			else
 			{
@@ -1825,7 +1827,8 @@ void FLightmassExporter::WriteSceneSettings( Lightmass::FSceneFileHeader& Scene
 
 		Scene.MaterialSettings.EnvironmentColor = FLinearColor(LevelSettings.EnvironmentColor) * LevelSettings.EnvironmentIntensity;
 
-		Scene.MaterialSettings.bUseNormalMapsForSimpleLightMaps = GEngine->bUseNormalMapsForSimpleLightMaps;
+		static const auto CVar = IConsoleManager::Get().FindTConsoleVariableDataInt(TEXT("r.NormalMapsForStaticLighting"));
+		Scene.MaterialSettings.bUseNormalMapsForLighting = CVar->GetValueOnGameThread() != 0;
 	}
 	{
 		verify(GConfig->GetBool(TEXT("DevOptions.MeshAreaLights"), TEXT("bVisualizeMeshAreaLightPrimitives"), bConfigBool, GLightmassIni));
diff --git a/Engine/Source/Programs/UnrealLightmass/Private/Lighting/LightingMesh.h b/Engine/Source/Programs/UnrealLightmass/Private/Lighting/LightingMesh.h
--- a/Engine/Source/Programs/UnrealLightmass/Private/Lighting/LightingMesh.h
+++ b/Engine/Source/Programs/UnrealLightmass/Private/Lighting/LightingMesh.h
@@ -314,6 +314,7 @@ public:
 	inline float GetFullyOccludedSamplesFraction(int32 ElementIndex) const { return MaterialElements[ElementIndex].FullyOccludedSamplesFraction; }
 	inline int32 GetNumElements() const { return MaterialElements.Num(); }
 	inline bool ShouldColorInvalidTexels() const { return bColorInvalidTexels; }
+	inline bool HasImportedNormal(int32 ElementIndex) const { return MaterialElements[ElementIndex].Material->NormalSize > 0; }
 
 	/**
 	 *	Returns the Guid for the object associated with this lighting mesh.
diff --git a/Engine/Source/Programs/UnrealLightmass/Private/Lighting/LightingSystem.cpp b/Engine/Source/Programs/UnrealLightmass/Private/Lighting/LightingSystem.cpp
--- a/Engine/Source/Programs/UnrealLightmass/Private/Lighting/LightingSystem.cpp
+++ b/Engine/Source/Programs/UnrealLightmass/Private/Lighting/LightingSystem.cpp
@@ -2522,12 +2522,6 @@ FGatheredLightSample FStaticLightingSystem::CalculatePointLighting(
 		// Compute the incident lighting of the light on the vertex.
 		const FLinearColor LightIntensity = InLightIntensity * InTransmission;
 
-		// Figure out the tangent space normal of this point on the mapping's surface.  This normal is used
-		// when computing lighting for simple light maps.
-		const FVector4 TangentNormal = MaterialSettings.bUseNormalMapsForSimpleLightMaps ?
-			Mapping->Mesh->EvaluateNormal(Vertex.TextureCoordinates[0], ElementIndex) :
-			FVector4( 0.0f, 0.0f, 1.0f, 0.0f );
-
 		// Compute the light-map sample for the front-face of the vertex.
 		FGatheredLightSample FrontFaceSample = FGatheredLightSample::PointLightWorldSpace(LightIntensity, TangentLightVector, WorldLightVector.SafeNormal());
 
diff --git a/Engine/Source/Programs/UnrealLightmass/Private/Lighting/TextureMapping.cpp b/Engine/Source/Programs/UnrealLightmass/Private/Lighting/TextureMapping.cpp
--- a/Engine/Source/Programs/UnrealLightmass/Private/Lighting/TextureMapping.cpp
+++ b/Engine/Source/Programs/UnrealLightmass/Private/Lighting/TextureMapping.cpp
@@ -922,6 +922,23 @@ void FStaticLightingSystem::SetupTextureMapping(
 				}
 				// Mark the texel as mapped to some geometry in the scene
 				CurrentLightSample.bIsMapped = true;
+
+				if (MaterialSettings.bUseNormalMapsForLighting && TextureMapping->Mesh->HasImportedNormal(TexelToVertex.ElementIndex))
+				{
+					const FVector4 TangentNormal = TextureMapping->Mesh->EvaluateNormal(TexelToVertex.TextureCoordinates[0], TexelToVertex.ElementIndex);
+
+					const FVector4 WorldTangentRow0(TexelToVertex.WorldTangentX.X, TexelToVertex.WorldTangentY.X, TexelToVertex.WorldTangentZ.X);
+					const FVector4 WorldTangentRow1(TexelToVertex.WorldTangentX.Y, TexelToVertex.WorldTangentY.Y, TexelToVertex.WorldTangentZ.Y);
+					const FVector4 WorldTangentRow2(TexelToVertex.WorldTangentX.Z, TexelToVertex.WorldTangentY.Z, TexelToVertex.WorldTangentZ.Z);
+					const FVector4 WorldVector(
+						Dot3(WorldTangentRow0, TangentNormal),
+						Dot3(WorldTangentRow1, TangentNormal),
+						Dot3(WorldTangentRow2, TangentNormal)
+						);
+
+					TexelToVertex.WorldTangentZ = WorldVector;
+				}
+
 				// Normalize the tangent basis and ensure it is orthonormal
 				TexelToVertex.WorldTangentZ = TexelToVertex.WorldTangentZ.UnsafeNormal3();
 				TexelToVertex.TriangleNormal = TexelToVertex.TriangleNormal.UnsafeNormal3();
diff --git a/Engine/Source/Programs/UnrealLightmass/Public/SceneExport.h b/Engine/Source/Programs/UnrealLightmass/Public/SceneExport.h
--- a/Engine/Source/Programs/UnrealLightmass/Public/SceneExport.h
+++ b/Engine/Source/Programs/UnrealLightmass/Public/SceneExport.h
@@ -162,11 +162,8 @@ struct FSceneMaterialSettings
 	/** The size of the normal sample */
 	int32 NormalSize;
 
-
-	/** Enables normal map sampling when Lightmass is generating 'simple' light maps.  This increases lighting build time, but may improve quality when normal maps are used to represent curvature over a large surface area.  When this setting is disabled, 'simple' light maps will not take normal maps into account. */
-	/** Note: This setting is driven from the level's UEngine::bUseNormalMapsForSimpleLightMaps option */
-	bool bUseNormalMapsForSimpleLightMaps;
-
+	/** Whether to use the normal map for lighting, if false the smoothed vertex normal will be used. */
+	bool bUseNormalMapsForLighting;
 
 	/** 
 	 * Debugging - Amount of incoming light to reflect diffusely (equally in all directions). 
diff --git a/Engine/Source/Runtime/Core/Private/HAL/ConsoleManager.cpp b/Engine/Source/Runtime/Core/Private/HAL/ConsoleManager.cpp
--- a/Engine/Source/Runtime/Core/Private/HAL/ConsoleManager.cpp
+++ b/Engine/Source/Runtime/Core/Private/HAL/ConsoleManager.cpp
@@ -1567,6 +1567,11 @@ static TAutoConsoleVariable<int32> CVarAllowStaticLighting(
 	TEXT("Games that only use dynamic lighting should set this to 0 to save some static lighting overhead."),
 	ECVF_ReadOnly | ECVF_RenderThreadSafe);
 
+static TAutoConsoleVariable<int32> CVarNormalMaps(
+	TEXT("r.NormalMapsForStaticLighting"),
+	0,
+	TEXT("Whether to allow any static lighting to use normal maps for lighting computations."),
+	ECVF_Default);
 
 static TAutoConsoleVariable<int32> CVarNumBufferedOcclusionQueries(
 	TEXT("r.NumBufferedOcclusionQueries"),
diff --git a/Engine/Source/Runtime/Engine/Classes/Engine/Engine.h b/Engine/Source/Runtime/Engine/Classes/Engine/Engine.h
--- a/Engine/Source/Runtime/Engine/Classes/Engine/Engine.h
+++ b/Engine/Source/Runtime/Engine/Classes/Engine/Engine.h
@@ -1364,10 +1364,6 @@ private:
 
 public:
 
-	/** Enables normal map sampling when Lightmass is generating 'simple' light maps.  This increases lighting build time, but may improve quality when normal maps are used to represent curvature over a large surface area.  When this setting is disabled, 'simple' light maps will not take normal maps into account. */
-	UPROPERTY(globalconfig)
-	uint32 bUseNormalMapsForSimpleLightMaps:1;
-
 	/** determines if we should start the matinee capture as soon as the game loads */
 	UPROPERTY(transient)
 	uint32 bStartWithMatineeCapture:1;
diff --git a/Engine/Source/Runtime/Engine/Classes/Engine/RendererSettings.h b/Engine/Source/Runtime/Engine/Classes/Engine/RendererSettings.h
--- a/Engine/Source/Runtime/Engine/Classes/Engine/RendererSettings.h
+++ b/Engine/Source/Runtime/Engine/Classes/Engine/RendererSettings.h
@@ -135,6 +135,11 @@ class ENGINE_API URendererSettings
 		ToolTip="Whether to allow any static lighting to be generated and used, like lightmaps and shadowmaps. Games that only use dynamic lighting should set this to 0 to save some static lighting overhead. Changing this setting requires restarting the editor."))
 	uint32 bAllowStaticLighting:1;
 
+	UPROPERTY(config, EditAnywhere, Category=Lighting, meta=(
+		ConsoleVariable="r.NormalMapsForStaticLighting",
+		ToolTip="Whether to allow any static lighting to use normal maps for lighting computations."))
+	uint32 bUseNormalMapsForStaticLighting:1;
+
 	UPROPERTY(config, EditAnywhere, Category=Lighting, meta=(
 		ConsoleVariable="r.GBuffer",
 		ToolTip="1=Use GBuffer, 0=Don't use GBuffer (minimal limited renderer)."))
