Fix bug where warning/error banner won't display if multiple object parameters in message and a later one isn't a pin or node

diff --git a/Engine/Source/Editor/UnrealEd/Private/Kismet2/CompilerResultsLog.cpp b/Engine/Source/Editor/UnrealEd/Private/Kismet2/CompilerResultsLog.cpp
--- a/Engine/Source/Editor/UnrealEd/Private/Kismet2/CompilerResultsLog.cpp
+++ b/Engine/Source/Editor/UnrealEd/Private/Kismet2/CompilerResultsLog.cpp
@@ -112,7 +112,7 @@ UObject const* FCompilerResultsLog::FindSourceObject(UObject const* PossiblyDupl
 /** Create a tokenized message record from a message containing @@ indicating where each UObject* in the ArgPtr list goes and place it in the MessageLog. */
 void FCompilerResultsLog::InternalLogMessage(const EMessageSeverity::Type& Severity, const TCHAR* Message, va_list ArgPtr)
 {
-	UEdGraphNode* OwnerNode = NULL;
+	UEdGraphNode* OwnerNode = nullptr;
 
 	// Create the tokenized message
 	TSharedRef<FTokenizedMessage> Line = FTokenizedMessage::Create( Severity );
@@ -141,14 +141,13 @@ void FCompilerResultsLog::InternalLogMessage(const EMessageSeverity::Type& Sever
 
 				if (ObjectArgument)
 				{
-					const UEdGraphNode* Node = Cast<UEdGraphNode>(ObjectArgument);
-					const UEdGraphPin* Pin = Cast<UEdGraphPin>(ObjectArgument);
+					UEdGraphNode* Node = Cast<UEdGraphNode>(ObjectArgument);
+					const UEdGraphPin* Pin = (Node? nullptr : Cast<UEdGraphPin>(ObjectArgument));
 
-					//Get owner node reference
-					OwnerNode = Cast<UEdGraphNode>(ObjectArgument);
-					if (Pin)
+					//Get owner node reference, consider the first 
+					if (OwnerNode == nullptr)
 					{
-						OwnerNode = Pin->GetOwningNodeUnchecked();
+						OwnerNode = (Pin ? Pin->GetOwningNodeUnchecked() : Node);
 					}
 
 					if (ObjectArgument->GetOutermost() == GetTransientPackage())
