If a modal window is active in the editor when another modal window is opened, parent the new modal window to the existing window

diff --git a/Engine/Source/Editor/UnrealEd/Private/Editor.cpp b/Engine/Source/Editor/UnrealEd/Private/Editor.cpp
--- a/Engine/Source/Editor/UnrealEd/Private/Editor.cpp
+++ b/Engine/Source/Editor/UnrealEd/Private/Editor.cpp
@@ -3936,12 +3936,17 @@ void UEditorEngine::UpdateSkyCaptures()
 
 void UEditorEngine::EditorAddModalWindow( TSharedRef<SWindow> InModalWindow ) const
 {
-	TSharedPtr<SWindow> ParentWindow;
-	// Check if the main frame is loaded.  When using the old main frame it may not be.
-	if( FModuleManager::Get().IsModuleLoaded( "MainFrame" ) )
+	// If there is already a modal window active, parent this new modal window to the existing window so that it doesnt fall behind
+	TSharedPtr<SWindow> ParentWindow = FSlateApplication::Get().GetActiveModalWindow();
+
+	if( !ParentWindow.IsValid() )
 	{
-		IMainFrameModule& MainFrame = FModuleManager::LoadModuleChecked<IMainFrameModule>( "MainFrame" );
-		ParentWindow = MainFrame.GetParentWindow();
+		// Parent to the main frame window
+		if(FModuleManager::Get().IsModuleLoaded("MainFrame"))
+		{
+			IMainFrameModule& MainFrame = FModuleManager::LoadModuleChecked<IMainFrameModule>("MainFrame");
+			ParentWindow = MainFrame.GetParentWindow();
+		}
 	}
 
 	FSlateApplication::Get().AddModalWindow( InModalWindow, ParentWindow );
