Editor: Improve delete model/dialog handling when the only references remaining are in the undo buffer
- Updated AssetDeleteModel to treat in-memory(non-undo) and in-undo-buffer as mutually exclusive states
- Deleting an asset that only has a reference remaining in the undo buffer now looks like a normal delete with a note, instead of a prompt with extra text that requires using a scary Force Delete
- Fix a LOCTEXT collision that resulted in the wrong text being displayed for the "Assets Referencing the Pending Deleted Assets" prompt

diff --git a/Engine/Source/Editor/UnrealEd/Private/AssetDeleteModel.cpp b/Engine/Source/Editor/UnrealEd/Private/AssetDeleteModel.cpp
--- a/Engine/Source/Editor/UnrealEd/Private/AssetDeleteModel.cpp
+++ b/Engine/Source/Editor/UnrealEd/Private/AssetDeleteModel.cpp
@@ -12,7 +12,7 @@
 FAssetDeleteModel::FAssetDeleteModel( const TArray<UObject*>& InObjectsToDelete )
 	: State( StartScanning )
 	, bPendingObjectsCanBeReplaced(false)
-	, bIsAnythingReferencedInMemory(false)
+	, bIsAnythingReferencedInMemoryByNonUndo(false)
 	, bIsAnythingReferencedInMemoryByUndo(false)
 	, PendingDeleteIndex(0)
 	, ObjectsDeleted(0)
@@ -61,7 +61,7 @@ void FAssetDeleteModel::Tick( const float InDeltaTime )
 		break;
 	case StartScanning:
 		OnDiskReferences = TSet<FName>();
-		bIsAnythingReferencedInMemory = false;
+		bIsAnythingReferencedInMemoryByNonUndo = false;
 		bIsAnythingReferencedInMemoryByUndo = false;
 		PendingDeleteIndex = 0;
 
@@ -87,7 +87,8 @@ void FAssetDeleteModel::Tick( const float InDeltaTime )
 			PendingDelete->RemainingDiskReferences = NonPendingDeletedExternalOnDiskReferences;
 
 			// Count up all the external references
-			int NonPendingDeletedExternalInMemoryReferences = 0;
+			int32 NonPendingDeletedExternalInMemoryReferences = 0;
+
 			for ( const FReferencerInformation& Reference : PendingDelete->MemoryReferences.ExternalReferences )
 			{
 				if ( !IsObjectInPendingDeletes( Reference.Referencer ) )
@@ -97,7 +98,7 @@ void FAssetDeleteModel::Tick( const float InDeltaTime )
 			}
 			PendingDelete->RemainingMemoryReferences = NonPendingDeletedExternalInMemoryReferences;
 
-			bIsAnythingReferencedInMemory |= PendingDelete->RemainingMemoryReferences > 0;
+			bIsAnythingReferencedInMemoryByNonUndo |= PendingDelete->RemainingMemoryReferences > 0;
 			bIsAnythingReferencedInMemoryByUndo |= PendingDelete->IsReferencedInMemoryByUndo();
 
 			PendingDeleteIndex++;
@@ -142,7 +143,7 @@ bool FAssetDeleteModel::DoDelete()
 bool FAssetDeleteModel::CanForceDelete() const
 {
 	// We can only force delete when they are still referenced in memory or still referenced on disk.
-	return bIsAnythingReferencedInMemory || OnDiskReferences.Num() > 0;
+	return bIsAnythingReferencedInMemoryByNonUndo || (OnDiskReferences.Num() > 0);
 }
 
 bool FAssetDeleteModel::DoForceDelete()
@@ -415,7 +416,7 @@ FPendingDelete::FPendingDelete(UObject* InObject)
 	, RemainingMemoryReferences(0)
 	, Object(InObject)
 	, bReferencesChecked(false)
-	, bIsReferencedInMemory(false)
+	, bIsReferencedInMemoryByNonUndo(false)
 	, bIsReferencedInMemoryByUndo(false)
 	, bIsInternal(false)
 {
@@ -504,23 +505,20 @@ void FPendingDelete::CheckForReferences()
 
 	AssetRegistryModule.Get().GetReferencers(Object->GetOutermost()->GetFName(), DiskReferences);
 
-	// Check and see whether we are referenced by any objects that won't be garbage collected. 
-	bIsReferencedInMemory = IsReferenced(Object, GARBAGE_COLLECTION_KEEPFLAGS, true, &MemoryReferences);
-	bIsReferencedInMemoryByUndo = false;
-	if ( bIsReferencedInMemory )
-	{
-		FReferencerInformationList ReferencesExcludingUndo;
-		// determine whether the transaction buffer is holding a reference to the object
-		// and if so, offer the user the option to reset the transaction buffer.
-		GEditor->Trans->DisableObjectSerialization();
-		IsReferenced(Object, GARBAGE_COLLECTION_KEEPFLAGS, true, &ReferencesExcludingUndo);
-		GEditor->Trans->EnableObjectSerialization();
-		// only ref to this object is the transaction buffer - set a flag so we know we need to clear the undo stack
-		if ( MemoryReferences.InternalReferences.Num() > ReferencesExcludingUndo.InternalReferences.Num() )
-		{
-			bIsReferencedInMemoryByUndo = true;
-		}
-	}
+	// Check and see whether we are referenced by any objects that won't be garbage collected (*including* the undo buffer)
+	FReferencerInformationList ReferencesIncludingUndo;
+	bool bReferencedInMemoryOrUndoStack = IsReferenced(Object, GARBAGE_COLLECTION_KEEPFLAGS, true, &ReferencesIncludingUndo);
+
+	// Determine the in-memory references, *excluding* the undo buffer
+	GEditor->Trans->DisableObjectSerialization();
+	bIsReferencedInMemoryByNonUndo = IsReferenced(Object, GARBAGE_COLLECTION_KEEPFLAGS, true, &MemoryReferences);
+	GEditor->Trans->EnableObjectSerialization();
+
+	// see if this object is the transaction buffer - set a flag so we know we need to clear the undo stack
+	const int32 TotalReferenceCount = ReferencesIncludingUndo.ExternalReferences.Num() + ReferencesIncludingUndo.InternalReferences.Num();
+	const int32 NonUndoReferenceCount = MemoryReferences.ExternalReferences.Num() + MemoryReferences.InternalReferences.Num();
+
+	bIsReferencedInMemoryByUndo = TotalReferenceCount > NonUndoReferenceCount;
 }
 
 bool FPendingDelete::operator == ( const FPendingDelete& Other ) const
diff --git a/Engine/Source/Editor/UnrealEd/Private/Dialogs/SDeleteAssetsDialog.cpp b/Engine/Source/Editor/UnrealEd/Private/Dialogs/SDeleteAssetsDialog.cpp
--- a/Engine/Source/Editor/UnrealEd/Private/Dialogs/SDeleteAssetsDialog.cpp
+++ b/Engine/Source/Editor/UnrealEd/Private/Dialogs/SDeleteAssetsDialog.cpp
@@ -272,7 +272,7 @@ TSharedRef<SWidget> SDeleteAssetsDialog::BuildDeleteDialog()
 				.Padding( 3.0f )
 				[
 					SNew( STextBlock )
-					.Text( LOCTEXT( "AttemptingDelete", "Assets Referencing the Pending Deleted Assets" ) )
+					.Text( LOCTEXT( "AssetsReferencingPendingDeletedAssets", "Assets Referencing the Pending Deleted Assets" ) )
 					.Font( FEditorStyle::GetFontStyle( "BoldFont" ) )
 					.ShadowOffset( FVector2D( 1.0f, 1.0f ) )
 				]
@@ -549,7 +549,7 @@ TSharedRef<SWidget> SDeleteAssetsDialog::CreateThumbnailWidget()
 
 EVisibility SDeleteAssetsDialog::GetReferencesVisiblity() const
 {
-	return DeleteModel->IsAnythingReferencedInMemory() ? EVisibility::Visible : EVisibility::Collapsed;
+	return DeleteModel->IsAnythingReferencedInMemoryByNonUndo() ? EVisibility::Visible : EVisibility::Collapsed;
 }
 
 EVisibility SDeleteAssetsDialog::GetUndoVisiblity() const
@@ -566,6 +566,12 @@ TSharedRef<ITableRow> SDeleteAssetsDialog::HandleGenerateAssetRow( TSharedPtr<FP
 FReply SDeleteAssetsDialog::Delete()
 {
 	ParentWindow.Get()->RequestDestroyWindow();
+
+	if (DeleteModel->IsAnythingReferencedInMemoryByUndo())
+	{
+		GEditor->Trans->Reset(LOCTEXT("DeleteSelectedItem", "Delete Selected Item"));
+	}
+
 	DeleteModel->DoDelete();
 
 	return FReply::Handled();
diff --git a/Engine/Source/Editor/UnrealEd/Public/AssetDeleteModel.h b/Engine/Source/Editor/UnrealEd/Public/AssetDeleteModel.h
--- a/Engine/Source/Editor/UnrealEd/Public/AssetDeleteModel.h
+++ b/Engine/Source/Editor/UnrealEd/Public/AssetDeleteModel.h
@@ -23,7 +23,7 @@ public:
 	bool IsAssetContained(const FName& PackageName) const;
 
 	/** Is the pending deleted object referenced in memory by something other than the undo stack; INCLUDES PENDING DELETES */
-	bool IsReferencedInMemory() const { return bIsReferencedInMemory; }
+	bool IsReferencedInMemoryByNonUndo() const { return bIsReferencedInMemoryByNonUndo; }
 
 	/** Is the pending deleted object referenced in memory by the undo stack; INCLUDES PENDING DELETES */
 	bool IsReferencedInMemoryByUndo() const { return bIsReferencedInMemoryByUndo; }
@@ -40,7 +40,7 @@ public:
 	/** The on disk references to this object */
 	TArray<FName> DiskReferences;
 
-	/** In memory references to this object */
+	/** In memory references to this object (*excluding* the undo buffer) */
 	FReferencerInformationList MemoryReferences;
 
 	/** The remaining disk references; EXCLUDES PENDING DELETES */
@@ -58,8 +58,8 @@ private:
 	/** A flag indicating that references have been checked, so don't check again. */
 	bool bReferencesChecked;
 
-	/** flag indicating if this object is referenced in memory by the engine. */
-	bool bIsReferencedInMemory;
+	/** flag indicating if this object is referenced in memory by the engine (excluding the undo buffer). */
+	bool bIsReferencedInMemoryByNonUndo;
 
 	/** flag indicating if this object is referenced in memory by the undo stack. */
 	bool bIsReferencedInMemoryByUndo;
@@ -152,7 +152,7 @@ public:
 	FText GetProgressText() const;
 
 	/** Is any of the pending deleted assets being referenced in memory. */
-	bool IsAnythingReferencedInMemory() const { return bIsAnythingReferencedInMemory; }
+	bool IsAnythingReferencedInMemoryByNonUndo() const { return bIsAnythingReferencedInMemoryByNonUndo; }
 
 	/** Is any of the pending deleted assets being referenced in the undo stack. */
 	bool IsAnythingReferencedInMemoryByUndo() const { return bIsAnythingReferencedInMemoryByUndo; }
@@ -197,7 +197,7 @@ private:
 	bool bPendingObjectsCanBeReplaced;
 
 	/** Is any of the pending deleted assets being referenced in memory. */
-	bool bIsAnythingReferencedInMemory;
+	bool bIsAnythingReferencedInMemoryByNonUndo;
 
 	/** Is any of the pending deleted assets being referenced in the undo stack. */
 	bool bIsAnythingReferencedInMemoryByUndo;
