AutomationEditorCommon.cpp
--Made the editor performance test print the actual location of the csv file to the log.

diff --git a/Engine/Source/Editor/UnrealEd/Private/Tests/AutomationEditorCommon.cpp b/Engine/Source/Editor/UnrealEd/Private/Tests/AutomationEditorCommon.cpp
--- a/Engine/Source/Editor/UnrealEd/Private/Tests/AutomationEditorCommon.cpp
+++ b/Engine/Source/Editor/UnrealEd/Private/Tests/AutomationEditorCommon.cpp
@@ -535,16 +535,6 @@ bool FGenerateEditorPerformanceCharts::Update()
 		RAWCSVFile->Serialize(TCHAR_TO_ANSI(*RAWCSVLine), RAWCSVLine.Len());
 	}
 	
-	////Create the Performance .csv file	
-	//FArchive* FinalCSVFile = IFileManager::Get().CreateFileReader(*FinalCSVFilename);
-	//if (!FinalCSVFile)
-	//{
-	//	UE_LOG(LogEditorAutomationTests, Error, TEXT("Failed to create the csv file: %s"), *FinalCSVFilename);
-	//	return false;
-	//}
-
-
-
 	//Get the final data for the Performance csv file.
 	//This gets the averaged numbers using the raw data arrays.
 	float FPSAverage = ValueFromStringArrayofNumbers(FPSRawArray, true);
@@ -600,6 +590,7 @@ bool FGenerateEditorPerformanceCharts::Update()
 	UE_LOG(LogEditorAutomationTests, Display, TEXT("AVG FPS: '%.1f'"), FPSAverage);
 	UE_LOG(LogEditorAutomationTests, Display, TEXT("AVG Frame Time: '%.1f' ms"), FrameTimeAverage);
 	UE_LOG(LogEditorAutomationTests, Display, TEXT("AVG Used Physical Memory: '%.0f' kb"), MemoryAverage);
+	UE_LOG(LogEditorAutomationTests, Log, TEXT("Performance csv file is located here: %s"), *FPaths::ConvertRelativePathToFull(FinalCSVFilename));
 
 	//Close the csv files so we can use them while the editor is still open.
 	RAWCSVFile->Close();
