HotReload: Better fix for old hot-reload classes in the class viewer

#ttp 345524: HotReload: Strange temporary class names are visible in Placement Browser's "All Classes" list after a HotReload

diff --git a/Engine/Source/Editor/ClassViewer/Private/SClassViewer.cpp b/Engine/Source/Editor/ClassViewer/Private/SClassViewer.cpp
--- a/Engine/Source/Editor/ClassViewer/Private/SClassViewer.cpp
+++ b/Engine/Source/Editor/ClassViewer/Private/SClassViewer.cpp
@@ -1677,8 +1677,7 @@ void FClassHierarchy::AddChildren_NoFilter( TSharedPtr< FClassViewerNode >& InOu
 
 		// Ignore deprecated and temporary trash classes.
 		if (CurrentClass->HasAnyClassFlags(CLASS_Deprecated | CLASS_NewerVersionExists) || 
-			FKismetEditorUtilities::IsClassABlueprintSkeleton(CurrentClass) ||
-			CurrentClass->HasAnyFlags(RF_Transient))
+			FKismetEditorUtilities::IsClassABlueprintSkeleton(CurrentClass))
 		{
 			continue;
 		}
diff --git a/Engine/Source/Editor/UnrealEd/Private/Kismet2/KismetReinstanceUtilities.cpp b/Engine/Source/Editor/UnrealEd/Private/Kismet2/KismetReinstanceUtilities.cpp
--- a/Engine/Source/Editor/UnrealEd/Private/Kismet2/KismetReinstanceUtilities.cpp
+++ b/Engine/Source/Editor/UnrealEd/Private/Kismet2/KismetReinstanceUtilities.cpp
@@ -44,6 +44,10 @@ FBlueprintCompileReinstancer::FBlueprintCompileReinstancer(UClass* InNewClass, U
 			}
 		}
 	}
+
+	// Finally, remove the old class from Root so that it can get GC'd and mark it as CLASS_NewerVersionExists
+	InOldClass->RemoveFromRoot();
+	InOldClass->ClassFlags |= CLASS_NewerVersionExists;
 }
 
 FBlueprintCompileReinstancer::FBlueprintCompileReinstancer(UClass* InClassToReinstance, bool bIsBytecodeOnly, bool bSkipGC)
@@ -68,7 +72,6 @@ FBlueprintCompileReinstancer::FBlueprintCompileReinstancer(UClass* InClassToRein
 		// Add the class to the root set, so that it doesn't get prematurely garbage collected
 		DuplicatedClass->AddToRoot();
 		ClassToReinstance->ClassFlags &= ~CLASS_NewerVersionExists;
-		InClassToReinstance->RemoveFromRoot();
 		GIsDuplicatingClassForReinstancing = false;
 
 		// Bind and link the duplicate class, so that it has the proper duplicate property offsets
