Turn off notifies in animation thumbnails

diff --git a/Engine/Source/Editor/UnrealEd/Private/ThumbnailHelpers.cpp b/Engine/Source/Editor/UnrealEd/Private/ThumbnailHelpers.cpp
--- a/Engine/Source/Editor/UnrealEd/Private/ThumbnailHelpers.cpp
+++ b/Engine/Source/Editor/UnrealEd/Private/ThumbnailHelpers.cpp
@@ -539,7 +539,7 @@ bool FAnimationSequenceThumbnailScene::SetAnimation(UAnimSequenceBase* InAnimati
 
 				PreviewActor->SkeletalMeshComponent->PlayAnimation(InAnimation,false);
 				PreviewActor->SkeletalMeshComponent->Stop();
-				PreviewActor->SkeletalMeshComponent->SetPosition(AnimPosition);
+				PreviewActor->SkeletalMeshComponent->SetPosition(AnimPosition, false);
 
 				UAnimSingleNodeInstance* SingleNodeInstance = PreviewActor->SkeletalMeshComponent->GetSingleNodeInstance();
 				if (SingleNodeInstance)
diff --git a/Engine/Source/Runtime/Engine/Classes/Components/SkeletalMeshComponent.h b/Engine/Source/Runtime/Engine/Classes/Components/SkeletalMeshComponent.h
--- a/Engine/Source/Runtime/Engine/Classes/Components/SkeletalMeshComponent.h
+++ b/Engine/Source/Runtime/Engine/Classes/Components/SkeletalMeshComponent.h
@@ -548,7 +548,7 @@ public:
 	bool IsPlaying() const;
 
 	UFUNCTION(BlueprintCallable, Category="Components|SkeletalMesh")
-	void SetPosition(float InPos);
+	void SetPosition(float InPos, bool bFireNotifies = true);
 
 	UFUNCTION(BlueprintCallable, Category="Components|SkeletalMesh")
 	float GetPosition() const;
diff --git a/Engine/Source/Runtime/Engine/Private/SkeletalMeshComponent.cpp b/Engine/Source/Runtime/Engine/Private/SkeletalMeshComponent.cpp
--- a/Engine/Source/Runtime/Engine/Private/SkeletalMeshComponent.cpp
+++ b/Engine/Source/Runtime/Engine/Private/SkeletalMeshComponent.cpp
@@ -1552,12 +1552,12 @@ bool USkeletalMeshComponent::IsPlaying() const
 	return false;
 }
 
-void USkeletalMeshComponent::SetPosition(float InPos)
+void USkeletalMeshComponent::SetPosition(float InPos, bool bFireNotifies)
 {
 	UAnimSingleNodeInstance* SingleNodeInstance = GetSingleNodeInstance();
 	if (SingleNodeInstance)
 	{
-		SingleNodeInstance->SetPosition(InPos);
+		SingleNodeInstance->SetPosition(InPos, bFireNotifies);
 	}
 	else if( AnimScriptInstance != NULL )
 	{
