BP: Kismet compiler will patch CDO's subobjects.
ttp342812 external object warning prevents Blueprint'/Game/Characters/Enemies/Smasher/Blueprints/SmasherPlayer.SmasherPlayer' from beingà

diff --git a/Engine/Source/Editor/KismetCompiler/Private/KismetCompiler.cpp b/Engine/Source/Editor/KismetCompiler/Private/KismetCompiler.cpp
--- a/Engine/Source/Editor/KismetCompiler/Private/KismetCompiler.cpp
+++ b/Engine/Source/Editor/KismetCompiler/Private/KismetCompiler.cpp
@@ -3093,8 +3093,6 @@ void FKismetCompilerContext::Compile()
 	}
 
 	UObject* OldCDO = NULL;
-
-	int32 OldSkelLinkerIdx = INDEX_NONE;
 	int32 OldGenLinkerIdx = INDEX_NONE;
 	ULinkerLoad* OldLinker = Blueprint->GetLinker();
 
@@ -3107,17 +3105,9 @@ void FKismetCompilerContext::Compile()
 		for( int32 i = 0; i < OldLinker->ExportMap.Num(); i++ )
 		{
 			FObjectExport& ThisExport = OldLinker->ExportMap[i];
-			if( ThisExport.ObjectName == SkeletonName )
-			{
-				OldSkelLinkerIdx = i;
-			}
-			else if( ThisExport.ObjectName == GeneratedName )
+			if( ThisExport.ObjectName == GeneratedName )
 			{
 				OldGenLinkerIdx = i;
-			}
-
-			if( OldSkelLinkerIdx != INDEX_NONE && OldGenLinkerIdx != INDEX_NONE )
-			{
 				break;
 			}
 		}
@@ -3338,7 +3328,8 @@ void FKismetCompilerContext::Compile()
 						}
 					}
 
-					UEditorEngine::CopyPropertiesForUnrelatedObjects(OldCDO, NewCDO);			
+					UEditorEngine::CopyPropertiesForUnrelatedObjects(OldCDO, NewCDO);
+					FBlueprintEditorUtils::PatchCDOSubobjectsIntoExport(OldCDO, NewCDO);
 				}
 
 				// >>> Backwards Compatibility: Propagate data from the skel CDO to the gen CDO if we haven't already done so for this blueprint
diff --git a/Engine/Source/Editor/UnrealEd/Private/Kismet2/BlueprintEditorUtils.cpp b/Engine/Source/Editor/UnrealEd/Private/Kismet2/BlueprintEditorUtils.cpp
--- a/Engine/Source/Editor/UnrealEd/Private/Kismet2/BlueprintEditorUtils.cpp
+++ b/Engine/Source/Editor/UnrealEd/Private/Kismet2/BlueprintEditorUtils.cpp
@@ -1234,27 +1234,7 @@ UClass* FBlueprintEditorUtils::RegenerateBlueprintClass(UBlueprint* Blueprint, U
 
 					if (bRegenerated)
 					{
-						// Collect the instanced components in both the old and new CDOs
-						TArray<UObject*> OldComponents, NewComponents;
-						PreviousCDO->CollectDefaultSubobjects(OldComponents, true);
-						NewCDO->CollectDefaultSubobjects(NewComponents, true);
-
-						// For all components common to both, patch the linker table with the new version of the component, so things that reference the default (e.g. InternalArchetypes) will have the updated version
-						for (auto OldCompIt = OldComponents.CreateIterator(); OldCompIt; ++OldCompIt)
-						{
-							UObject* OldComponent = (*OldCompIt);
-							const FName OldComponentName = OldComponent->GetFName();
-							for (auto NewCompIt = NewComponents.CreateIterator(); NewCompIt; ++NewCompIt)
-							{
-								UObject* NewComponent = *NewCompIt;
-								if (NewComponent->GetFName() == OldComponentName)
-								{
-									ULinkerLoad::PRIVATE_PatchNewObjectIntoExport(OldComponent, NewComponent);
-									break;
-								}
-							}
-						}
-						NewCDO->CheckDefaultSubobjects();
+						PatchCDOSubobjectsIntoExport(PreviousCDO, NewCDO);
 						// We purposefully do not call post load here, it happens later on in the normal flow
 					}
 				}
@@ -1315,6 +1295,40 @@ UClass* FBlueprintEditorUtils::RegenerateBlueprintClass(UBlueprint* Blueprint, U
 	return bRegenerated ? Blueprint->GeneratedClass : NULL;
 }
 
+void FBlueprintEditorUtils::PatchCDOSubobjectsIntoExport(UObject* PreviousCDO, UObject* NewCDO)
+{
+	if (PreviousCDO && NewCDO)
+	{
+		// Collect the instanced components in both the old and new CDOs
+		TArray<UObject*> OldComponents, NewComponents;
+		PreviousCDO->CollectDefaultSubobjects(OldComponents, true);
+		NewCDO->CollectDefaultSubobjects(NewComponents, true);
+
+		TMap<FName, UObject*> NewComponentsMap;
+		for (auto NewCompIt = NewComponents.CreateIterator(); NewCompIt; ++NewCompIt)
+		{
+			UObject* NewComponent = *NewCompIt;
+			if (NewComponent)
+			{
+				NewComponentsMap.Add(NewComponent->GetFName(), NewComponent);
+			}
+		}
+
+		// For all components common to both, patch the linker table with the new version of the component, so things that reference the default (e.g. InternalArchetypes) will have the updated version
+		for (auto OldCompIt = OldComponents.CreateIterator(); OldCompIt; ++OldCompIt)
+		{
+			UObject* OldComponent = (*OldCompIt);
+			const FName OldComponentName = OldComponent->GetFName();
+			UObject** NewComponentPtr = NewComponentsMap.Find(OldComponentName);
+			if (NewComponentPtr && *NewComponentPtr)
+			{
+				ULinkerLoad::PRIVATE_PatchNewObjectIntoExport(OldComponent, *NewComponentPtr);
+			}
+		}
+		NewCDO->CheckDefaultSubobjects();
+	}
+}
+
 void FBlueprintEditorUtils::PropagateParentBlueprintDefaults(UClass* ClassToPropagate)
 {
 	check(ClassToPropagate);
diff --git a/Engine/Source/Editor/UnrealEd/Public/Kismet2/BlueprintEditorUtils.h b/Engine/Source/Editor/UnrealEd/Public/Kismet2/BlueprintEditorUtils.h
--- a/Engine/Source/Editor/UnrealEd/Public/Kismet2/BlueprintEditorUtils.h
+++ b/Engine/Source/Editor/UnrealEd/Public/Kismet2/BlueprintEditorUtils.h
@@ -78,6 +78,11 @@ public:
 	 */
 	static UClass* RegenerateBlueprintClass(UBlueprint* Blueprint, UClass* ClassToRegenerate, UObject* PreviousCDO, TArray<UObject*>& ObjLoaded);
 
+	/**
+	 * Replace subobjects of CDO in linker
+	 */
+	static void PatchCDOSubobjectsIntoExport(UObject* PreviousCDO, UObject* NewCDO);
+
 	/**
 	 * Copies the default properties of all parent blueprint classes in the chain to the specified blueprint's skeleton CDO
 	 */
