Remove large file import warning

diff --git a/Engine/Source/Editor/UnrealEd/Private/Factories/Factory.cpp b/Engine/Source/Editor/UnrealEd/Private/Factories/Factory.cpp
--- a/Engine/Source/Editor/UnrealEd/Private/Factories/Factory.cpp
+++ b/Engine/Source/Editor/UnrealEd/Private/Factories/Factory.cpp
@@ -269,18 +269,7 @@ UObject* UFactory::StaticImportObject
 				const int32 FileSize = IFileManager::Get().FileSize( Filename );
 				bool bValidFileSize = true;
 
-				// File size was found
-				if ( FileSize != INDEX_NONE )
-				{
-					if( ( MaxImportFileSize > 0 ) && ( FileSize > MaxImportFileSize ) )
-					{
-						// Prompt the user if they would like to proceed with large import, displaying the size of the file in MBs
-						// (File Size >> 20) is the same as dividing by (1024*1024) to convert to MBs
-						bValidFileSize = EAppReturnType::Yes == FMessageDialog::Open( EAppMsgType::YesNo,
-							FText::Format( NSLOCTEXT("UnrealEd", "Warning_LargeFileImport", "Attempting to import a very large file, proceed?\nFile Size: {0} MB"), FText::AsNumber(FileSize >> 20) ) );
-					}
-				}
-				else
+				if ( FileSize == INDEX_NONE )
 				{
 					UE_LOG(LogFactory, Error,TEXT("File '%s' does not exist"), Filename );
 					bValidFileSize = false;
