Do not prompt to reset transaction buffer after force delete

diff --git a/Engine/Source/Editor/UnrealEd/Private/ObjectTools.cpp b/Engine/Source/Editor/UnrealEd/Private/ObjectTools.cpp
--- a/Engine/Source/Editor/UnrealEd/Private/ObjectTools.cpp
+++ b/Engine/Source/Editor/UnrealEd/Private/ObjectTools.cpp
@@ -1699,17 +1699,10 @@ namespace ObjectTools
 				bIsReferenced = IsReferenced( ObjectToDelete, GARBAGE_COLLECTION_KEEPFLAGS, true, &Refs );
 				GEditor->Trans->EnableObjectSerialization();
 
-				// only ref to this object is the transaction buffer - let the user choose whether to clear the undo buffer
+				// only ref to this object is the transaction buffer, clear the transaction buffer
 				if ( !bIsReferenced )
 				{
-					if ( EAppReturnType::Yes == FMessageDialog::Open( EAppMsgType::YesNo, NSLOCTEXT( "UnrealEd", "ResetUndoBufferForObjectDeletionPrompt", "The only reference to this object is the undo history.  In order to delete this object, you must clear all undo history - would you like to clear undo history?" ) ) )
-					{
-						GEditor->Trans->Reset( NSLOCTEXT( "UnrealEd", "DeleteSelectedItem", "Delete Selected Item" ) );
-					}
-					else
-					{
-						bIsReferenced = true;
-					}
+					GEditor->Trans->Reset( NSLOCTEXT( "UnrealEd", "DeleteSelectedItem", "Delete Selected Item" ) );
 				}
 			}
 
