UE4: Fixed an automation error that happened because an asset was created in a temporary root and remained in memory after the root was unmounted. TTP#346216

diff --git a/Engine/Source/Editor/UnrealEd/Private/Tests/EditorAutomationTests.cpp b/Engine/Source/Editor/UnrealEd/Private/Tests/EditorAutomationTests.cpp
--- a/Engine/Source/Editor/UnrealEd/Private/Tests/EditorAutomationTests.cpp
+++ b/Engine/Source/Editor/UnrealEd/Private/Tests/EditorAutomationTests.cpp
@@ -23,6 +23,8 @@
 #include "AutomationEditorCommon.h"
 #include "AutomationTest.h"
 
+#include "PackageTools.h"
+
 
 //DEFINE_LOG_CATEGORY_STATIC(LogEditorAutomationTests, Log, All);
 
@@ -1157,6 +1159,17 @@ ABrush* ConvertTestFindNewBrush (const TArray<ABrush*> &PreviousBrushes)
 	return NewBrush;
 }
 
+/**
+ * Parameters to the Latent Automation command FCleanupConvertToValidation
+ */
+struct FCleanupConvertToValidationParameters
+{
+	TWeakObjectPtr<UWorld> TestWorld;
+	FString AssetPackageName;
+};
+
+DEFINE_LATENT_AUTOMATION_COMMAND_ONE_PARAMETER(FCleanupConvertToValidation, FCleanupConvertToValidationParameters, Parameters);
+
 bool FConvertToValidation::RunTest(const FString& Parameters)
 {
 	UWorld* World = AutomationEditorCommonUtils::CreateNewMap();
@@ -1218,6 +1231,7 @@ bool FConvertToValidation::RunTest(const FString& Parameters)
 	}
 
 	//convert to static mesh
+	FString AssetPackageName;
 	{
 		TArray<ABrush*> PreviousBrushes;
 		ConvertTestFindAllBrushes(PreviousBrushes);
@@ -1239,12 +1253,12 @@ bool FConvertToValidation::RunTest(const FString& Parameters)
 
 		//generate static mesh package name. Temporarily mount /Automation.
 		FPackageName::RegisterMountPoint(TEXT("/Automation/"), FPaths::AutomationTransientDir());
-		FString PackageName = TEXT("/Automation/ConvertToBSPToStaticMesh");
+		AssetPackageName = TEXT("/Automation/ConvertToBSPToStaticMesh");
 		//Convert brush to specific package name
-		GEditor->DoConvertActors(ToStaticMeshActors, AStaticMeshActor::StaticClass(), TSet<FString>(), true, PackageName);
+		GEditor->DoConvertActors(ToStaticMeshActors, AStaticMeshActor::StaticClass(), TSet<FString>(), true, AssetPackageName);
 
 		//find the package
-		UPackage* NewPackage = FindPackage(NULL, *PackageName);
+		UPackage* NewPackage = FindPackage(NULL, *AssetPackageName);
 		if (NewPackage)
 		{
 			TArray<UPackage*> PackagesToSave;
@@ -1259,9 +1273,6 @@ bool FConvertToValidation::RunTest(const FString& Parameters)
 		{
 			UE_LOG(LogEditorAutomationTests, Error, TEXT("Failed to save ConvertToBSPToStaticMesh."));
 		}
-
-		// Unmount /Automation now that the asset was successfully saved
-		FPackageName::UnRegisterMountPoint(TEXT("/Automation/"), FPaths::AutomationTransientDir());
 	}
 
 	TakeLatentAutomationScreenshot(ConvertMeshParameters, BaseFileName, FString::Printf(TEXT("FinalConvertMesh")), FString::Printf(TEXT("04_Final")), false);
@@ -1279,6 +1290,58 @@ bool FConvertToValidation::RunTest(const FString& Parameters)
 
 	GLevelEditorModeTools().MapChangeNotify();
 
+	// Add a latent action to clean up the static mesh actor we created and unload the temporary asset AFTER we take the screenshot
+	FCleanupConvertToValidationParameters CleanupParameters;
+	CleanupParameters.AssetPackageName = AssetPackageName;
+	CleanupParameters.TestWorld = World;
+	ADD_LATENT_AUTOMATION_COMMAND(FCleanupConvertToValidation(CleanupParameters));
+
+	return true;
+}
+
+bool FCleanupConvertToValidation::Update()
+{
+	const FString& AssetPackageName = Parameters.AssetPackageName;
+	UWorld* TestWorld = Parameters.TestWorld.Get();
+
+	// Attempt to unload the asset we created temporarily.
+	UPackage* NewPackage = FindPackage(NULL, *AssetPackageName);
+	if (NewPackage)
+	{
+		if (TestWorld)
+		{
+			// First find the static mesh we made in this package
+			UStaticMesh* GeneratedMesh = FindObject<UStaticMesh>(NewPackage, *FPackageName::GetLongPackageAssetName(AssetPackageName));
+
+			// If we found the mesh, find and delete the static mesh actor we added to the level to clear the reference to it.
+			if (GeneratedMesh)
+			{
+				for (FActorIterator ActorIt(TestWorld); ActorIt; ++ActorIt)
+				{
+					AStaticMeshActor* StaticMeshActor = Cast<AStaticMeshActor>(*ActorIt);
+					if (StaticMeshActor)
+					{
+						if ( StaticMeshActor->StaticMeshComponent->StaticMesh == GeneratedMesh )
+						{
+							TestWorld->DestroyActor(StaticMeshActor);
+						}
+					}
+				}
+			}
+		}
+
+		// Clear the transaction buffer to remove the last reference
+		GEditor->Trans->Reset( NSLOCTEXT( "UnrealEd.Test", "ConvertToValidationClear", "ConvertToValidation Clear" ) );
+
+		// Now unload the package
+		TArray<UPackage*> PackagesToUnload;
+		PackagesToUnload.Add(NewPackage);
+		PackageTools::UnloadPackages(PackagesToUnload);
+	}
+
+	// Unmount /Automation.
+	FPackageName::UnRegisterMountPoint(TEXT("/Automation/"), FPaths::AutomationTransientDir());
+
 	return true;
 }
 
