Add warning when importing fbx that is too small

diff --git a/Engine/Source/Editor/UnrealEd/Private/Fbx/FbxStaticMeshImport.cpp b/Engine/Source/Editor/UnrealEd/Private/Fbx/FbxStaticMeshImport.cpp
--- a/Engine/Source/Editor/UnrealEd/Private/Fbx/FbxStaticMeshImport.cpp
+++ b/Engine/Source/Editor/UnrealEd/Private/Fbx/FbxStaticMeshImport.cpp
@@ -806,6 +806,20 @@ UStaticMesh* UnFbx::FFbxImporter::ReimportStaticMesh(UStaticMesh* Mesh, UFbxStat
 	return NewMesh;
 }
 
+void UnFbx::FFbxImporter::VerifyGeometry(UStaticMesh* StaticMesh)
+{
+	// Calculate bounding box to check if too small
+	{
+		FVector Center, Extents;
+		ComputeBoundingBox(StaticMesh, Center, Extents);
+
+		if (Extents.GetAbsMax() < 5.f)
+		{
+			AddTokenizedErrorMessage(FTokenizedMessage::Create(EMessageSeverity::Warning, LOCTEXT("Prompt_MeshVerySmall", "Warning: The imported mesh is very small. This is most likely an issue with the units used when exporting to FBX. Please import again and increase the Scale to make the mesh bigger.")), FFbxErrors::Generic_Mesh_SmallGeometry);
+		}
+	}
+}
+
 UStaticMesh* UnFbx::FFbxImporter::ImportStaticMeshAsSingle(UObject* InParent, TArray<FbxNode*>& MeshNodeArray, const FName InName, EObjectFlags Flags, UFbxStaticMeshImportData* TemplateImportData, UStaticMesh* InStaticMesh, int LODIndex)
 {
 	bool bBuildStatus = true;
@@ -1212,17 +1226,24 @@ UStaticMesh* UnFbx::FFbxImporter::ImportStaticMeshAsSingle(UObject* InParent, TA
 		StaticMesh = NULL;
 	}
 
-	if (StaticMesh && StaticMesh->bCustomizedCollision == false && ImportOptions->bAutoGenerateCollision)
+	if (StaticMesh)
 	{
 
-		FKAggregateGeom & AggGeom = StaticMesh->BodySetup->AggGeom;
-		AggGeom.ConvexElems.Empty(1);	//if no custom collision is setup we just regenerate collision when reimport
-
-		const int32 NumDirs = 18;
-		TArray<FVector> Dirs;
-		Dirs.AddUninitialized(NumDirs);
-		for (int32 DirIdx = 0; DirIdx < NumDirs; ++DirIdx) { Dirs[DirIdx] = KDopDir18[DirIdx]; }
-		GenerateKDopAsSimpleCollision(StaticMesh, Dirs);
+		//collision generation
+		if (StaticMesh->bCustomizedCollision == false && ImportOptions->bAutoGenerateCollision)
+		{
+			FKAggregateGeom & AggGeom = StaticMesh->BodySetup->AggGeom;
+			AggGeom.ConvexElems.Empty(1);	//if no custom collision is setup we just regenerate collision when reimport
+
+			const int32 NumDirs = 18;
+			TArray<FVector> Dirs;
+			Dirs.AddUninitialized(NumDirs);
+			for (int32 DirIdx = 0; DirIdx < NumDirs; ++DirIdx) { Dirs[DirIdx] = KDopDir18[DirIdx]; }
+			GenerateKDopAsSimpleCollision(StaticMesh, Dirs);
+		}
+		
+		//warnings based on geometry
+		VerifyGeometry(StaticMesh);
 	}
 
 	return StaticMesh;
diff --git a/Engine/Source/Editor/UnrealEd/Private/FbxImporter.h b/Engine/Source/Editor/UnrealEd/Private/FbxImporter.h
--- a/Engine/Source/Editor/UnrealEd/Private/FbxImporter.h
+++ b/Engine/Source/Editor/UnrealEd/Private/FbxImporter.h
@@ -599,6 +599,12 @@ private:
 	*/
 	void InsertNewLODToBaseSkeletalMesh(USkeletalMesh* InSkeletalMesh, USkeletalMesh* BaseSkeletalMesh, int32 DesiredLOD);
 
+	/**
+	* Method used to verify if the geometry is valid. For example, if the bounding box is tiny we should warn
+	* @param StaticMesh - The imported static mesh which we'd like to verify
+	*/
+	void VerifyGeometry(UStaticMesh* StaticMesh);
+
 public:
 	// current Fbx scene we are importing. Make sure to release it after import
 	FbxScene* Scene;
diff --git a/Engine/Source/Editor/UnrealEd/Private/GeomFitUtils.cpp b/Engine/Source/Editor/UnrealEd/Private/GeomFitUtils.cpp
--- a/Engine/Source/Editor/UnrealEd/Private/GeomFitUtils.cpp
+++ b/Engine/Source/Editor/UnrealEd/Private/GeomFitUtils.cpp
@@ -193,6 +193,17 @@ static void CalcBoundingBox(const FRawMesh& RawMesh, FVector& Center, FVector& E
 	Box.GetCenterAndExtents(Center, Extents);
 }
 
+void ComputeBoundingBox(UStaticMesh* StaticMesh, FVector& Center, FVector& Extents)
+{
+	// Calculate bounding Box.
+	FRawMesh RawMesh;
+	FStaticMeshSourceModel& SrcModel = StaticMesh->SourceModels[0];
+	SrcModel.RawMeshBulkData->LoadRawMesh(RawMesh);
+
+	FVector unitVec = FVector(1.f);
+	CalcBoundingBox(RawMesh, Center, Extents, unitVec);
+}
+
 int32 GenerateBoxAsSimpleCollision(UStaticMesh* StaticMesh)
 {
 	if (!PromptToRemoveExistingCollision(StaticMesh))
diff --git a/Engine/Source/Editor/UnrealEd/Private/GeomFitUtils.h b/Engine/Source/Editor/UnrealEd/Private/GeomFitUtils.h
--- a/Engine/Source/Editor/UnrealEd/Private/GeomFitUtils.h
+++ b/Engine/Source/Editor/UnrealEd/Private/GeomFitUtils.h
@@ -109,6 +109,7 @@ UNREALED_API int32 GenerateKDopAsSimpleCollision(UStaticMesh* StaticMesh, const
 UNREALED_API int32 GenerateBoxAsSimpleCollision(UStaticMesh* StaticMesh);
 UNREALED_API int32 GenerateSphereAsSimpleCollision(UStaticMesh* StaticMesh);
 UNREALED_API int32 GenerateSphylAsSimpleCollision(UStaticMesh* StaticMesh);
+UNREALED_API void ComputeBoundingBox(UStaticMesh* StaticMesh, FVector& Center, FVector& Extents);
 
 /**
  * Refresh Collision Change
diff --git a/Engine/Source/Runtime/Core/Private/Misc/FbxErrors.cpp b/Engine/Source/Runtime/Core/Private/Misc/FbxErrors.cpp
--- a/Engine/Source/Runtime/Core/Private/Misc/FbxErrors.cpp
+++ b/Engine/Source/Runtime/Core/Private/Misc/FbxErrors.cpp
@@ -22,6 +22,8 @@ FName FFbxErrors::Generic_MeshNotFound(TEXT("Generic_MeshNotFound"));
 
 FName FFbxErrors::Generic_Mesh_NoGeometry(TEXT("Generic_Mesh_NoGeometry"));
 
+FName FFbxErrors::Generic_Mesh_SmallGeometry(TEXT("Generic_Mesh_SmallGeometry"));
+
 FName FFbxErrors::Generic_Mesh_TriangulationFailed(TEXT("Generic_Mesh_TriangulationFailed"));
 
 FName FFbxErrors::Generic_Mesh_ConvertSmoothingGroupFailed(TEXT("Generic_Mesh_ConvertSmoothingGroupFailed"));
diff --git a/Engine/Source/Runtime/Core/Public/Misc/FbxErrors.h b/Engine/Source/Runtime/Core/Public/Misc/FbxErrors.h
--- a/Engine/Source/Runtime/Core/Public/Misc/FbxErrors.h
+++ b/Engine/Source/Runtime/Core/Public/Misc/FbxErrors.h
@@ -29,6 +29,8 @@ struct CORE_API FFbxErrors
 
 	static FName Generic_Mesh_NoGeometry;
 
+	static FName Generic_Mesh_SmallGeometry;
+
 	static FName Generic_Mesh_TriangulationFailed;
 
 	static FName Generic_Mesh_ConvertSmoothingGroupFailed;
