Editor Performance Test
--Fixed how it handled getting the url for maps in the engine folder
--Other minor fixes.

diff --git a/Engine/Source/Editor/UnrealEd/Private/Tests/AutomationEditorCommon.cpp b/Engine/Source/Editor/UnrealEd/Private/Tests/AutomationEditorCommon.cpp
--- a/Engine/Source/Editor/UnrealEd/Private/Tests/AutomationEditorCommon.cpp
+++ b/Engine/Source/Editor/UnrealEd/Private/Tests/AutomationEditorCommon.cpp
@@ -497,8 +497,12 @@ bool FGenerateEditorPerformanceCharts::Update()
 	//Get the current changelist number
 	FString Changelist = GEngineVersion.ToString(EVersionComponent::Changelist);
 
+	//Gets the main frame module to get the name of our current level.
+	const IMainFrameModule& MainFrameModule = FModuleManager::GetModuleChecked< IMainFrameModule >("MainFrame");
+	FString LoadedMapName = MainFrameModule.GetLoadedLevelName();
+
 	//The locations of the raw data.
-	FString FileSaveLocation = FPaths::AutomationLogDir() + TEXT("Performance/") + MapName + TEXT("/");
+	FString FileSaveLocation = FPaths::AutomationLogDir() + TEXT("Performance/") + LoadedMapName + TEXT("/");
 	FString FPSFileSaveLocation = FileSaveLocation + TEXT("RAWAverageFPS.txt");
 	FString FrameTimeSaveLocation = FileSaveLocation + TEXT("RAWFrameTime.txt");
 	FString MemoryFileSaveLocation = FileSaveLocation + TEXT("RAWMemoryUsedPhysical.txt");
@@ -514,15 +518,15 @@ bool FGenerateEditorPerformanceCharts::Update()
 	//The CSV files.
 	//RAW csv holds all of the information from the text files
 	//Final csv only holds the averaged information.
-	FString FinalCSVFilename = FString::Printf(TEXT("%s%s_Performance.csv"), *FileSaveLocation, *MapName);
-	FString RAWCSVFilename = FString::Printf(TEXT("%sRAW_%s_%s.csv"), *FileSaveLocation, *MapName, *FDateTime::Now().ToString());
+	FString FinalCSVFilename = FString::Printf(TEXT("%s%s_Performance.csv"), *FileSaveLocation, *LoadedMapName);
+	FString RAWCSVFilename = FString::Printf(TEXT("%sRAW_%s_%s.csv"), *FileSaveLocation, *LoadedMapName, *FDateTime::Now().ToString());
 
 	//Create the .csv file that will hold the raw data from the text files.
 	FArchive* RAWCSVFile = IFileManager::Get().CreateFileWriter(*RAWCSVFilename);
 	if (!RAWCSVFile)
 	{
 		UE_LOG(LogEditorAutomationTests, Error, TEXT("Failed to create the csv file: %s"), *RAWCSVFilename);
-		return false;
+		return true;
 	}
 
 	//Add the top title row to the raw csv file
@@ -531,7 +535,7 @@ bool FGenerateEditorPerformanceCharts::Update()
 	//Loop through the text files and add them to the raw csv file.
 	for (int32 IterLoop = 0; IterLoop < FPSRawArray.Num(); IterLoop++)
 	{
-		RAWCSVLine = MapName + TEXT(",") + Changelist + TEXT(",") + TestRunTimeRawArray[IterLoop] + TEXT(",") + MapLoadTimeRawArray[0] + TEXT(",") + FPSRawArray[IterLoop] + TEXT(",") + FrameTimeRawArray[IterLoop] + TEXT(",") + MemoryRawArray[IterLoop] + LINE_TERMINATOR;
+		RAWCSVLine = LoadedMapName + TEXT(",") + Changelist + TEXT(",") + TestRunTimeRawArray[IterLoop] + TEXT(",") + MapLoadTimeRawArray[0] + TEXT(",") + FPSRawArray[IterLoop] + TEXT(",") + FrameTimeRawArray[IterLoop] + TEXT(",") + MemoryRawArray[IterLoop] + LINE_TERMINATOR;
 		RAWCSVFile->Serialize(TCHAR_TO_ANSI(*RAWCSVLine), RAWCSVLine.Len());
 	}
 	
@@ -562,28 +566,20 @@ bool FGenerateEditorPerformanceCharts::Update()
 	if (IFileManager::Get().FileSize(FinalCSVFilename.GetCharArray().GetData()) <= 0)
 	{
 		FArchive* FinalCSVFile = IFileManager::Get().CreateFileWriter(*FinalCSVFilename);
-		if (!FinalCSVFile)
-		{
-			UE_LOG(LogEditorAutomationTests, Error, TEXT("Failed to create the csv file: %s"), *FinalCSVFilename);
-			return false;
-		}
+		//Create the .csv file that will hold the raw data from the text files.
 		FString FinalCSVLine = (TEXT("Date, Map Name, Changelist, Test Run Time, Map Load Time, Average FPS, Average MS, Used Physical KB\n"));
 		FinalCSVFile->Serialize(TCHAR_TO_ANSI(*FinalCSVLine), FinalCSVLine.Len());
+		FinalCSVFile->Close();
 	}
 
 	FFileHelper::LoadFileToString(OldPerformanceCSVFile, *FinalCSVFilename);
 	FArchive* FinalCSVFile = IFileManager::Get().CreateFileWriter(*FinalCSVFilename);
-	if (!FinalCSVFile)
-	{
-		UE_LOG(LogEditorAutomationTests, Error, TEXT("Failed to create the csv file: %s"), *FinalCSVFilename);
-		return false;
-	}
 	FinalCSVFile->Serialize(TCHAR_TO_ANSI(*OldPerformanceCSVFile), OldPerformanceCSVFile.Len());
 
 
 
 	//Write the averaged numbers to the Performance CSV file.
-	FString FinalCSVLine = FString::Printf(TEXT("%s,%s,%s,%.0f,%s,%.3f,%.3f,%.0f%s"), *FDateTime::Today().ToString(), *MapName, *Changelist, TestRunTime, *MapLoadTime, FPSAverage, FrameTimeAverage, MemoryAverage, LINE_TERMINATOR);
+	FString FinalCSVLine = FString::Printf(TEXT("%s,%s,%s,%.0f,%s,%.3f,%.3f,%.0f%s"), *FDateTime::Today().ToString(), *LoadedMapName, *Changelist, TestRunTime, *MapLoadTime, FPSAverage, FrameTimeAverage, MemoryAverage, LINE_TERMINATOR);
 	FinalCSVFile->Serialize(TCHAR_TO_ANSI(*FinalCSVLine), FinalCSVLine.Len());
 
 	//Display the averaged numbers.
@@ -596,7 +592,7 @@ bool FGenerateEditorPerformanceCharts::Update()
 	RAWCSVFile->Close();
 	FinalCSVFile->Close();
 
-	//	 the text files.
+	//Delete the text files.
 	IFileManager::Get().Delete(*FPSFileSaveLocation);
 	IFileManager::Get().Delete(*FrameTimeSaveLocation);
 	IFileManager::Get().Delete(*MemoryFileSaveLocation);
@@ -613,7 +609,7 @@ bool FEditorLoadMap::Update()
 {
 	//Get the base filename for the map that will be used.
 	FString ShortMapName = FPaths::GetBaseFilename(MapName);
-
+	
 	double MapLoadStartTime = 0.0f;
 	double MapLoadTime = 0.0f;
 
@@ -631,12 +627,16 @@ bool FEditorLoadMap::Update()
 	//This is the time it took to load the map in the editor.
 	MapLoadTime = FPlatformTime::Seconds() - MapLoadStartTime;
 
+	//Gets the main frame module to get the name of our current level.
+	const IMainFrameModule& MainFrameModule = FModuleManager::GetModuleChecked< IMainFrameModule >("MainFrame");
+	FString LoadedMapName = MainFrameModule.GetLoadedLevelName();
+
 	UE_LOG(LogEditorAutomationTests, Log, TEXT("%s has been loaded."), *ShortMapName);
 
 	//Log out to a text file the time it takes to load the map.
-	WriteToTextFile(TestName, ShortMapName, FileName = TEXT("RAWMapLoadTime"), MapLoadTime, Delimiter);
+	WriteToTextFile(TestName, LoadedMapName, FileName = TEXT("RAWMapLoadTime"), MapLoadTime, Delimiter);
 	
-	UE_LOG(LogEditorAutomationTests, Display, TEXT("%s took %.3f to load."), *ShortMapName, MapLoadTime);
+	UE_LOG(LogEditorAutomationTests, Display, TEXT("%s took %.3f to load."), *LoadedMapName, MapLoadTime);
 	
 	return true;
 }
diff --git a/Engine/Source/Editor/UnrealEd/Private/Tests/EditorPerformanceTests.cpp b/Engine/Source/Editor/UnrealEd/Private/Tests/EditorPerformanceTests.cpp
--- a/Engine/Source/Editor/UnrealEd/Private/Tests/EditorPerformanceTests.cpp
+++ b/Engine/Source/Editor/UnrealEd/Private/Tests/EditorPerformanceTests.cpp
@@ -35,20 +35,40 @@ void FMapPerformanceInEditor::GetTests(TArray<FString>& OutBeautifiedNames, TArr
 		{
 			//Get the location of the map being used.
 			FString Filename = FPaths::ConvertRelativePathToFull(PerformanceMapName[indexloop]);
-			
-			//If true it will proceed to add the asset to the test list.
-			//This will be false if the map is on a different drive.
-			if (FPaths::MakePathRelativeTo(Filename, *FPaths::GameContentDir()))
+
+			if (Filename.Contains(TEXT("/Engine/"), ESearchCase::IgnoreCase, ESearchDir::FromStart))
 			{
-				FString ShortName = FPaths::GetBaseFilename(Filename);
-				FString PathName = FPaths::GetPath(Filename);
-				OutBeautifiedNames.Add(ShortName);
-				FString AssetName = FString::Printf(TEXT("/Game/%s/%s.%s"), *PathName, *ShortName, *ShortName);
-				OutTestCommands.Add(AssetName);
+				//If true it will proceed to add the asset to the test list.
+				//This will be false if the map is on a different drive.
+				if (FPaths::MakePathRelativeTo(Filename, *FPaths::EngineContentDir()))
+				{
+					FString ShortName = FPaths::GetBaseFilename(Filename);
+					FString PathName = FPaths::GetPath(Filename);
+					OutBeautifiedNames.Add(ShortName);
+					FString AssetName = FString::Printf(TEXT("/Engine/%s/%s.%s"), *PathName, *ShortName, *ShortName);
+					OutTestCommands.Add(AssetName);
+				}
+				else
+				{
+					UE_LOG(LogEditorAutomationTests, Error, TEXT("Invalid asset path: %s."), *Filename);
+				}
 			}
 			else
 			{
-				UE_LOG(LogEditorAutomationTests, Error, TEXT("Invalid asset path: %s."), *Filename);
+				//If true it will proceed to add the asset to the test list.
+				//This will be false if the map is on a different drive.
+				if (FPaths::MakePathRelativeTo(Filename, *FPaths::GameContentDir()))
+				{
+					FString ShortName = FPaths::GetBaseFilename(Filename);
+					FString PathName = FPaths::GetPath(Filename);
+					OutBeautifiedNames.Add(ShortName);
+					FString AssetName = FString::Printf(TEXT("/Game/%s/%s.%s"), *PathName, *ShortName, *ShortName);
+					OutTestCommands.Add(AssetName);
+				}
+				else
+				{
+					UE_LOG(LogEditorAutomationTests, Error, TEXT("Invalid asset path: %s."), *Filename);
+				}
 			}
 		}
 
