Fix constant FName initialization in places that execute often

diff --git a/Engine/Source/Developer/CollisionAnalyzer/Public/CollisionAnalyzerModule.h b/Engine/Source/Developer/CollisionAnalyzer/Public/CollisionAnalyzerModule.h
--- a/Engine/Source/Developer/CollisionAnalyzer/Public/CollisionAnalyzerModule.h
+++ b/Engine/Source/Developer/CollisionAnalyzer/Public/CollisionAnalyzerModule.h
@@ -17,7 +17,8 @@ public:
 	/** Gets the debugger singleton or returns NULL */
 	static ICollisionAnalyzer* Get()
 	{
-		FCollisionAnalyzerModule& DebuggerModule = FModuleManager::Get().LoadModuleChecked<FCollisionAnalyzerModule>("CollisionAnalyzer");
+		static FName CollisionAnalyzerModuleName("CollisionAnalyzer");
+		FCollisionAnalyzerModule& DebuggerModule = FModuleManager::Get().LoadModuleChecked<FCollisionAnalyzerModule>(CollisionAnalyzerModuleName);
 		return DebuggerModule.GetSingleton();
 	}
 
diff --git a/Engine/Source/Developer/GameplayDebugger/Public/GameplayDebugger.h b/Engine/Source/Developer/GameplayDebugger/Public/GameplayDebugger.h
--- a/Engine/Source/Developer/GameplayDebugger/Public/GameplayDebugger.h
+++ b/Engine/Source/Developer/GameplayDebugger/Public/GameplayDebugger.h
@@ -20,7 +20,8 @@ public:
 	 */
 	static inline IGameplayDebugger& Get()
 	{
-		return FModuleManager::LoadModuleChecked< IGameplayDebugger >("GameplayDebugger");
+		static FName GamePlayDebuggerModuleName("GameplayDebugger");
+		return FModuleManager::LoadModuleChecked< IGameplayDebugger >(GamePlayDebuggerModuleName);
 	}
 
 	/**
@@ -30,7 +31,8 @@ public:
 	 */
 	static inline bool IsAvailable()
 	{
-		return FModuleManager::Get().IsModuleLoaded( "GameplayDebugger" );
+		static FName GamePlayDebuggerModuleName("GameplayDebugger");
+		return FModuleManager::Get().IsModuleLoaded( GamePlayDebuggerModuleName );
 	}
 
 	// Each player controller that wants to use gameplay debugging must call this function
diff --git a/Engine/Source/Developer/SourceControl/Private/SSourceControlStatus.cpp b/Engine/Source/Developer/SourceControl/Private/SSourceControlStatus.cpp
--- a/Engine/Source/Developer/SourceControl/Private/SSourceControlStatus.cpp
+++ b/Engine/Source/Developer/SourceControl/Private/SSourceControlStatus.cpp
@@ -50,7 +50,8 @@ const FSlateBrush* SSourceControlStatus::GetSourceControlIconImage() const
 { 
 	if(QueryState == EQueryState::Querying)
 	{
-		return FEditorStyle::GetBrush("SourceControl.StatusIcon.Unknown");
+		static FName UnknownIconName("SourceControl.StatusIcon.Unknown");
+		return FEditorStyle::GetBrush(UnknownIconName);
 	}
 	else 
 	{
@@ -59,16 +60,19 @@ const FSlateBrush* SSourceControlStatus::GetSourceControlIconImage() const
 		{
 			if (!SourceControlModule.GetProvider().IsAvailable())
 			{
-				return FEditorStyle::GetBrush("SourceControl.StatusIcon.Error");
+				static FName ErrorIconName("SourceControl.StatusIcon.Error");
+				return FEditorStyle::GetBrush(ErrorIconName);
 			}
 			else
 			{
-				return FEditorStyle::GetBrush("SourceControl.StatusIcon.On");
+				static FName OnIconName("SourceControl.StatusIcon.On");
+				return FEditorStyle::GetBrush(OnIconName);
 			}
 		}
 		else
 		{
-			return FEditorStyle::GetBrush("SourceControl.StatusIcon.Off");
+			static FName OffIconName("SourceControl.StatusIcon.Off");
+			return FEditorStyle::GetBrush(OffIconName);
 		}
 	}
 }
diff --git a/Engine/Source/Developer/SourceControl/Public/ISourceControlModule.h b/Engine/Source/Developer/SourceControl/Public/ISourceControlModule.h
--- a/Engine/Source/Developer/SourceControl/Public/ISourceControlModule.h
+++ b/Engine/Source/Developer/SourceControl/Public/ISourceControlModule.h
@@ -123,6 +123,7 @@ public:
 	 */
 	static inline ISourceControlModule& Get()
 	{
-		return FModuleManager::LoadModuleChecked<ISourceControlModule>("SourceControl");
+		static FName SourceControlModule("SourceControl");
+		return FModuleManager::LoadModuleChecked<ISourceControlModule>(SourceControlModule);
 	}
 };
diff --git a/Engine/Source/Editor/ContentBrowser/Private/AssetViewWidgets.cpp b/Engine/Source/Editor/ContentBrowser/Private/AssetViewWidgets.cpp
--- a/Engine/Source/Editor/ContentBrowser/Private/AssetViewWidgets.cpp
+++ b/Engine/Source/Editor/ContentBrowser/Private/AssetViewWidgets.cpp
@@ -1320,11 +1320,13 @@ FSlateFontInfo SAssetTileItem::GetThumbnailFont() const
 		float Size = ThumbSize.Get();
 		if ( Size < 85 )
 		{
-			return FEditorStyle::GetFontStyle("ContentBrowser.AssetTileViewNameFontSmall");
+			static FName SmallFontName("ContentBrowser.AssetTileViewNameFontSmall");
+			return FEditorStyle::GetFontStyle(SmallFontName);
 		}
 	}
 
-	return FEditorStyle::GetFontStyle("ContentBrowser.AssetTileViewNameFont");
+	static FName RegularFont("ContentBrowser.AssetTileViewNameFont");
+	return FEditorStyle::GetFontStyle(RegularFont);
 }
 
 float SAssetTileItem::GetNameTextWrapWidth() const
diff --git a/Engine/Source/Editor/EditorLiveStreaming/Public/IEditorLiveStreaming.h b/Engine/Source/Editor/EditorLiveStreaming/Public/IEditorLiveStreaming.h
--- a/Engine/Source/Editor/EditorLiveStreaming/Public/IEditorLiveStreaming.h
+++ b/Engine/Source/Editor/EditorLiveStreaming/Public/IEditorLiveStreaming.h
@@ -46,7 +46,9 @@ public:
 	 */
 	static inline IEditorLiveStreaming& Get()
 	{
-		return FModuleManager::LoadModuleChecked< IEditorLiveStreaming >( "EditorLiveStreaming" );
+		static FName LiveStreamingModule("EditorLiveStreaming");
+
+		return FModuleManager::LoadModuleChecked< IEditorLiveStreaming >( LiveStreamingModule );
 	}
 
 };
diff --git a/Engine/Source/Editor/IntroTutorials/Private/STutorialOverlay.cpp b/Engine/Source/Editor/IntroTutorials/Private/STutorialOverlay.cpp
--- a/Engine/Source/Editor/IntroTutorials/Private/STutorialOverlay.cpp
+++ b/Engine/Source/Editor/IntroTutorials/Private/STutorialOverlay.cpp
@@ -110,7 +110,8 @@ int32 STutorialOverlay::TraverseWidgets(TSharedRef<SWidget> InWidget, const FGeo
 		// if we are picking, we need to draw an outline here
 		FName WidgetNameToHighlight = NAME_None;
 		bool bIsPicking = false;
-		FIntroTutorials& IntroTutorials = FModuleManager::Get().GetModuleChecked<FIntroTutorials>("IntroTutorials");
+		static FName IntroTutorialsModuleName("IntroTutorials");
+		FIntroTutorials& IntroTutorials = FModuleManager::Get().GetModuleChecked<FIntroTutorials>(IntroTutorialsModuleName);
 		if(IntroTutorials.OnIsPicking().IsBound())
 		{
 			bIsPicking = IntroTutorials.OnIsPicking().Execute(WidgetNameToHighlight);
diff --git a/Engine/Source/Editor/LevelEditor/Private/LevelEditorActions.cpp b/Engine/Source/Editor/LevelEditor/Private/LevelEditorActions.cpp
--- a/Engine/Source/Editor/LevelEditor/Private/LevelEditorActions.cpp
+++ b/Engine/Source/Editor/LevelEditor/Private/LevelEditorActions.cpp
@@ -54,6 +54,8 @@ DEFINE_LOG_CATEGORY_STATIC(LevelEditorActions, Log, All);
 
 #define LOCTEXT_NAMESPACE "LevelEditorActions"
 
+const FName HotReloadModule("HotReload");
+
 namespace LevelEditorActionsHelpers
 {
 	/** 
@@ -866,7 +868,7 @@ void FLevelEditorActionCallbacks::MapCheck_Execute()
 
 bool FLevelEditorActionCallbacks::CanShowSourceCodeActions()
 {
-	IHotReloadInterface& HotReloadSupport = FModuleManager::LoadModuleChecked<IHotReloadInterface>("HotReload");
+	IHotReloadInterface& HotReloadSupport = FModuleManager::LoadModuleChecked<IHotReloadInterface>(HotReloadModule);
 	// If there is at least one loaded game module, source code actions should be available.
 	return HotReloadSupport.IsAnyGameModuleLoaded();
 }
@@ -884,7 +886,7 @@ void FLevelEditorActionCallbacks::ReloadLevelEditor_Clicked()
 void FLevelEditorActionCallbacks::RecompileGameCode_Clicked()
 {
 	// Don't allow a recompile while already compiling!
-	IHotReloadInterface& HotReloadSupport = FModuleManager::LoadModuleChecked<IHotReloadInterface>("HotReload");
+	IHotReloadInterface& HotReloadSupport = FModuleManager::LoadModuleChecked<IHotReloadInterface>(HotReloadModule);
 	if( !HotReloadSupport.IsCurrentlyCompiling() )
 	{
 		HotReloadSupport.DoHotReloadFromEditor();
@@ -894,14 +896,15 @@ void FLevelEditorActionCallbacks::RecompileGameCode_Clicked()
 bool FLevelEditorActionCallbacks::Recompile_CanExecute()
 {
 	// We're not able to recompile if a compile is already in progress!
-	IHotReloadInterface& HotReloadSupport = FModuleManager::LoadModuleChecked<IHotReloadInterface>("HotReload");
+
+	IHotReloadInterface& HotReloadSupport = FModuleManager::LoadModuleChecked<IHotReloadInterface>(HotReloadModule);
 	return !HotReloadSupport.IsCurrentlyCompiling() && !(GEngineVersion.IsPromotedBuild() && FEngineBuildSettings::IsPerforceBuild());
 }
 
 bool FLevelEditorActionCallbacks::Reload_CanExecute()
 {
 	// We're not able to reload if a compile is already in progress!
-	IHotReloadInterface& HotReloadSupport = FModuleManager::LoadModuleChecked<IHotReloadInterface>("HotReload");
+	IHotReloadInterface& HotReloadSupport = FModuleManager::LoadModuleChecked<IHotReloadInterface>(HotReloadModule);
 	return !HotReloadSupport.IsCurrentlyCompiling();
 }
 
diff --git a/Engine/Source/Editor/LevelEditor/Private/SLevelViewportToolBar.cpp b/Engine/Source/Editor/LevelEditor/Private/SLevelViewportToolBar.cpp
--- a/Engine/Source/Editor/LevelEditor/Private/SLevelViewportToolBar.cpp
+++ b/Engine/Source/Editor/LevelEditor/Private/SLevelViewportToolBar.cpp
@@ -299,22 +299,27 @@ const FSlateBrush* SLevelViewportToolBar::GetCameraMenuLabelIcon() const
 	TSharedPtr< SLevelViewport > PinnedViewport( Viewport.Pin() );
 	if( PinnedViewport.IsValid() )
 	{
+		static FName PerspectiveIcon("EditorViewport.Perspective");
+		static FName TopIcon("EditorViewport.Top");
+		static FName SideIcon("EditorViewport.Side");
+		static FName FrontIcon("EditorViewport.Front");
+
 		switch( PinnedViewport->GetLevelViewportClient().ViewportType )
 		{
 			case LVT_Perspective:
-				Icon = FName( "EditorViewport.Perspective" );
+				Icon = PerspectiveIcon;
 				break;
 
 			case LVT_OrthoXY:
-				Icon = FName( "EditorViewport.Top" );
+				Icon = TopIcon;
 				break;
 
 			case LVT_OrthoYZ:
-				Icon = FName( "EditorViewport.Side" );
+				Icon = SideIcon;
 				break;
 
 			case LVT_OrthoXZ:
-				Icon = FName( "EditorViewport.Front" );
+				Icon = FrontIcon;
 				break;
 		}
 	}
diff --git a/Engine/Source/Editor/SceneOutliner/Private/SceneOutlinerActorInfoColumn.cpp b/Engine/Source/Editor/SceneOutliner/Private/SceneOutlinerActorInfoColumn.cpp
--- a/Engine/Source/Editor/SceneOutliner/Private/SceneOutlinerActorInfoColumn.cpp
+++ b/Engine/Source/Editor/SceneOutliner/Private/SceneOutlinerActorInfoColumn.cpp
@@ -19,7 +19,8 @@ FSceneOutlinerActorInfoColumn::FSceneOutlinerActorInfoColumn( const TWeakPtr< IS
 
 FName FSceneOutlinerActorInfoColumn::GetColumnID()
 {
-	return FName( "ActorInfo" );
+	static FName ActorInfo("ActorInfo");
+	return ActorInfo;
 }
 
 
diff --git a/Engine/Source/Editor/UnrealEd/Private/Editor.cpp b/Engine/Source/Editor/UnrealEd/Private/Editor.cpp
--- a/Engine/Source/Editor/UnrealEd/Private/Editor.cpp
+++ b/Engine/Source/Editor/UnrealEd/Private/Editor.cpp
@@ -982,7 +982,7 @@ void UEditorEngine::Tick( float DeltaSeconds, bool bIdleMode )
 	}
 
 	// Find out if the editor has focus. Audio should only play if the editor has focus.
-	const bool bHasFocus = FPlatformProcess::IsThisApplicationForeground();
+	const bool bHasFocus = FSlateApplication::Get().GetActiveTopLevelWindow().IsValid() || FPlatformProcess::IsThisApplicationForeground();
 
 	if (bHasFocus || GetDefault<ULevelEditorMiscSettings>()->bAllowBackgroundAudio)
 	{
@@ -1004,12 +1004,14 @@ void UEditorEngine::Tick( float DeltaSeconds, bool bIdleMode )
 	// Tick the asset registry
 	FAssetRegistryModule::TickAssetRegistry(DeltaSeconds);
 
-	ISourceCodeAccessModule& SourceCodeAccessModule = FModuleManager::LoadModuleChecked<ISourceCodeAccessModule>("SourceCodeAccess");
+	static FName SourceCodeAccessName("SourceCodeAccess");
+	ISourceCodeAccessModule& SourceCodeAccessModule = FModuleManager::LoadModuleChecked<ISourceCodeAccessModule>(SourceCodeAccessName);
 	SourceCodeAccessModule.GetAccessor().Tick(DeltaSeconds);
 
 	// tick the directory watcher
 	// @todo: Put me into an FTicker that is created when the DW module is loaded
-	FDirectoryWatcherModule& DirectoryWatcherModule = FModuleManager::Get().LoadModuleChecked<FDirectoryWatcherModule>(TEXT("DirectoryWatcher"));
+	static FName DirectoryWatcherName("DirectoryWatcher");
+	FDirectoryWatcherModule& DirectoryWatcherModule = FModuleManager::Get().LoadModuleChecked<FDirectoryWatcherModule>(DirectoryWatcherName);
 	DirectoryWatcherModule.Get()->Tick(DeltaSeconds);
 
 	if( bShouldTickEditorWorld )
@@ -5980,7 +5982,7 @@ bool UEditorEngine::ShouldThrottleCPUUsage() const
 {
 	bool bShouldThrottle = false;
 
-	bool bIsForeground = FPlatformProcess::IsThisApplicationForeground();
+	bool bIsForeground = FSlateApplication::Get().GetActiveTopLevelWindow().IsValid() || FPlatformProcess::IsThisApplicationForeground();
 
 	if( !bIsForeground )
 	{
diff --git a/Engine/Source/Editor/UnrealEd/Private/EditorComponents.cpp b/Engine/Source/Editor/UnrealEd/Private/EditorComponents.cpp
--- a/Engine/Source/Editor/UnrealEd/Private/EditorComponents.cpp
+++ b/Engine/Source/Editor/UnrealEd/Private/EditorComponents.cpp
@@ -141,15 +141,20 @@ void FGridWidget::DrawNewGrid(const FSceneView* View, FPrimitiveDrawInterface* P
 	// to get a light grid in a black level but use a high opacity value to be able to see it in a bright level
 	static float Darken = 0.11f;
 
+	static FName GridColorName("GridColor");
+	static FName SnapColorName("SnapColor");
+	static FName ExponentName("Exponent");
+	static FName AlphaBiasName("AlphaBias");
+	
 	if(bIsPerspective)
 	{
-		MaterialInst->SetVectorParameterValue("GridColor", FLinearColor(0.6f * Darken, 0.6f * Darken, 0.6f * Darken, CVarEditor3DGridFade.GetValueOnGameThread()));
-		MaterialInst->SetVectorParameterValue("SnapColor", FLinearColor(0.5f, 0.0f, 0.0f, SnapAlphaMultiplier * CVarEditor3DSnapFade.GetValueOnGameThread()));
+		MaterialInst->SetVectorParameterValue(GridColorName, FLinearColor(0.6f * Darken, 0.6f * Darken, 0.6f * Darken, CVarEditor3DGridFade.GetValueOnGameThread()));
+		MaterialInst->SetVectorParameterValue(SnapColorName, FLinearColor(0.5f, 0.0f, 0.0f, SnapAlphaMultiplier * CVarEditor3DSnapFade.GetValueOnGameThread()));
 	}
 	else
 	{
-		MaterialInst->SetVectorParameterValue("GridColor", FLinearColor(0.6f * Darken, 0.6f * Darken, 0.6f * Darken, CVarEditor2DGridFade.GetValueOnGameThread()));
-		MaterialInst->SetVectorParameterValue("SnapColor", FLinearColor(0.5f, 0.0f, 0.0f, SnapAlphaMultiplier * CVarEditor2DSnapFade.GetValueOnGameThread()));
+		MaterialInst->SetVectorParameterValue(GridColorName, FLinearColor(0.6f * Darken, 0.6f * Darken, 0.6f * Darken, CVarEditor2DGridFade.GetValueOnGameThread()));
+		MaterialInst->SetVectorParameterValue(SnapColorName, FLinearColor(0.5f, 0.0f, 0.0f, SnapAlphaMultiplier * CVarEditor2DSnapFade.GetValueOnGameThread()));
 	}
 
 	// true:1m, false:1dm ios smallest grid size
@@ -158,10 +163,10 @@ void FGridWidget::DrawNewGrid(const FSceneView* View, FPrimitiveDrawInterface* P
 	const int Exponent = 10;
 
 	// 2 is the default so we need to set it
-	MaterialInst->SetScalarParameterValue("Exponent", (float)Exponent);
+	MaterialInst->SetScalarParameterValue(ExponentName, (float)Exponent);
 
 	// without MSAA we need the grid to be more see through so lines behind it can be recognized
-	MaterialInst->SetScalarParameterValue("AlphaBias", bMSAA ? 0.0f : 0.05f);
+	MaterialInst->SetScalarParameterValue(AlphaBiasName, bMSAA ? 0.0f : 0.05f);
 
 	// grid for size
 	float GridSplit = 0.5f;
diff --git a/Engine/Source/Editor/UnrealEd/Private/Kismet2/DebuggerCommands.cpp b/Engine/Source/Editor/UnrealEd/Private/Kismet2/DebuggerCommands.cpp
--- a/Engine/Source/Editor/UnrealEd/Private/Kismet2/DebuggerCommands.cpp
+++ b/Engine/Source/Editor/UnrealEd/Private/Kismet2/DebuggerCommands.cpp
@@ -1351,7 +1351,9 @@ FSlateIcon FInternalPlayWorldCommandCallbacks::GetRepeatLastLaunchIcon()
 		break;
 	}
 
-	return FSlateIcon(FEditorStyle::GetStyleSetName(), "PlayWorld.RepeatLastLaunch");
+	static FName RepeatLastLaunchIcon("PlayWorld.RepeatLastLaunch");
+
+	return FSlateIcon(FEditorStyle::GetStyleSetName(), RepeatLastLaunchIcon );
 }
 
 bool FInternalPlayWorldCommandCallbacks::IsReadyToLaunchOnDevice(FString DeviceId)
diff --git a/Engine/Source/Editor/UnrealEd/Private/SEditorViewportViewMenu.cpp b/Engine/Source/Editor/UnrealEd/Private/SEditorViewportViewMenu.cpp
--- a/Engine/Source/Editor/UnrealEd/Private/SEditorViewportViewMenu.cpp
+++ b/Engine/Source/Editor/UnrealEd/Private/SEditorViewportViewMenu.cpp
@@ -103,62 +103,76 @@ const FSlateBrush* SEditorViewportViewMenu::GetViewMenuLabelIcon() const
 	TSharedPtr< SEditorViewport > PinnedViewport = Viewport.Pin();
 	if( PinnedViewport.IsValid() )
 	{
+		static FName WireframeIcon( "EditorViewport.WireframeMode" );
+		static FName UnlitIcon( "EditorViewport.UnlitMode" );
+		static FName LitIcon( "EditorViewport.LitMode" );
+		static FName DetailLightingIcon("EditorViewport.DetailLightingMode");
+		static FName LightingOnlyIcon("EditorViewport.LightingOnlyMode");
+		static FName LightComplexityIcon("EditorViewport.LightComplexityMode");
+		static FName ShaderComplexityIcon("EditorViewport.ShaderComplexityMode");
+		static FName LightOverlapIcon("EditorViewport.StationaryLightOverlapMode");
+		static FName LightmapDensityIcon("EditorViewport.LightmapDensityMode");
+		static FName ReflectionModeIcon("EditorViewport.ReflectionOverrideMode");
+		static FName VisualizeBufferIcon("EditorViewport.VisualizeBufferMode");
+		static FName CollisionPawnIcon("EditorViewport.CollisionPawn");
+		static FName CollisionVisibilityIcon("EditorViewport.CollisionVisibility");
+
 		switch( PinnedViewport->GetViewportClient()->GetViewMode() )
 		{
 			case VMI_BrushWireframe:
-				Icon = FName( "EditorViewport.WireframeMode" );
+				Icon = WireframeIcon;
 				break;
 
 			case VMI_Wireframe:
-				Icon = FName( "EditorViewport.WireframeMode" );
+				Icon = WireframeIcon;
 				break;
 
 			case VMI_Unlit:
-				Icon = FName( "EditorViewport.UnlitMode" );
+				Icon = UnlitIcon;
 				break;
 
 			case VMI_Lit:
-				Icon = FName( "EditorViewport.LitMode" );
+				Icon = LitIcon;
 				break;
 
 			case VMI_Lit_DetailLighting:
-				Icon = FName( "EditorViewport.DetailLightingMode" );
+				Icon = DetailLightingIcon;
 				break;
 
 			case VMI_LightingOnly:
-				Icon = FName( "EditorViewport.LightingOnlyMode" );
+				Icon = LightingOnlyIcon;
 				break;
 
 			case VMI_LightComplexity:
-				Icon = FName( "EditorViewport.LightComplexityMode" );
+				Icon = LightComplexityIcon;
 				break;
 
 			case VMI_ShaderComplexity:
-				Icon = FName( "EditorViewport.ShaderComplexityMode" );
+				Icon = ShaderComplexityIcon;
 				break;
 
 			case VMI_StationaryLightOverlap:
-				Icon = FName( "EditorViewport.StationaryLightOverlapMode" );
+				Icon = LightOverlapIcon;
 				break;
 
 			case VMI_LightmapDensity:
-				Icon = FName( "EditorViewport.LightmapDensityMode" );
+				Icon = LightmapDensityIcon;
 				break;
 
 			case VMI_ReflectionOverride:
-				Icon = FName( "EditorViewport.ReflectionOverrideMode" );
+				Icon = ReflectionModeIcon;
 				break;
 
 			case VMI_VisualizeBuffer:
-				Icon = FName( "EditorViewport.VisualizeBufferMode" );
+				Icon = VisualizeBufferIcon;
 				break;
 
 			case VMI_CollisionPawn:
-				Icon = FName( "EditorViewport.CollisionPawn" );
+				Icon = CollisionPawnIcon;
 				break;
 
 			case VMI_CollisionVisibility:
-				Icon = FName( "EditorViewport.CollisionVisibility" );
+				Icon = CollisionVisibilityIcon;
 				break;
 		}
 	}
diff --git a/Engine/Source/Runtime/AssetRegistry/Private/AssetRegistryModule.cpp b/Engine/Source/Runtime/AssetRegistry/Private/AssetRegistryModule.cpp
--- a/Engine/Source/Runtime/AssetRegistry/Private/AssetRegistryModule.cpp
+++ b/Engine/Source/Runtime/AssetRegistry/Private/AssetRegistryModule.cpp
@@ -41,9 +41,12 @@ void FAssetRegistryModule::ShutdownModule()
 	}
 
 #if WITH_UNREAL_DEVELOPER_TOOLS
-	// unregister message log
-	FMessageLogModule& MessageLogModule = FModuleManager::LoadModuleChecked<FMessageLogModule>("MessageLog");
-	MessageLogModule.UnregisterLogListing("AssetRegistry");
+	if( FModuleManager::Get().IsModuleLoaded("MessageLog") )
+	{
+		// unregister message log
+		FMessageLogModule& MessageLogModule = FModuleManager::GetModuleChecked<FMessageLogModule>("MessageLog");
+		MessageLogModule.UnregisterLogListing("AssetRegistry");
+	}
 #endif
 }
 
diff --git a/Engine/Source/Runtime/AssetRegistry/Public/AssetRegistryModule.h b/Engine/Source/Runtime/AssetRegistry/Public/AssetRegistryModule.h
--- a/Engine/Source/Runtime/AssetRegistry/Public/AssetRegistryModule.h
+++ b/Engine/Source/Runtime/AssetRegistry/Public/AssetRegistryModule.h
@@ -6,6 +6,11 @@
 #include "ModuleInterface.h"
 #include "IAssetRegistry.h"
 
+namespace AssetRegistryConstants
+{
+	const FName ModuleName("AssetRegistry");
+}
+
 /**
  * Asset registry module
  */
@@ -26,7 +31,7 @@ public:
 	/** Tick the asset registry with the supplied timestep */
 	static void TickAssetRegistry(float DeltaTime)
 	{
-		FAssetRegistryModule& AssetRegistryModule = FModuleManager::LoadModuleChecked<FAssetRegistryModule>(TEXT("AssetRegistry"));
+		FAssetRegistryModule& AssetRegistryModule = FModuleManager::LoadModuleChecked<FAssetRegistryModule>(AssetRegistryConstants::ModuleName);
 		AssetRegistryModule.Get().Tick(DeltaTime);
 	}
 
@@ -35,7 +40,7 @@ public:
 	{
 		if ( GIsEditor )
 		{
-			FAssetRegistryModule& AssetRegistryModule = FModuleManager::LoadModuleChecked<FAssetRegistryModule>(TEXT("AssetRegistry"));
+			FAssetRegistryModule& AssetRegistryModule = FModuleManager::LoadModuleChecked<FAssetRegistryModule>(AssetRegistryConstants::ModuleName);
 			AssetRegistryModule.Get().AssetCreated(NewAsset);
 		}
 	}
@@ -45,7 +50,7 @@ public:
 	{
 		if ( GIsEditor )
 		{
-			FAssetRegistryModule& AssetRegistryModule = FModuleManager::LoadModuleChecked<FAssetRegistryModule>(TEXT("AssetRegistry"));
+			FAssetRegistryModule& AssetRegistryModule = FModuleManager::LoadModuleChecked<FAssetRegistryModule>(AssetRegistryConstants::ModuleName);
 			AssetRegistryModule.Get().AssetDeleted(DeletedAsset);
 		}
 	}
@@ -55,7 +60,7 @@ public:
 	{
 		if ( GIsEditor )
 		{
-			FAssetRegistryModule& AssetRegistryModule = FModuleManager::LoadModuleChecked<FAssetRegistryModule>(TEXT("AssetRegistry"));
+			FAssetRegistryModule& AssetRegistryModule = FModuleManager::LoadModuleChecked<FAssetRegistryModule>(AssetRegistryConstants::ModuleName);
 			AssetRegistryModule.Get().AssetRenamed(RenamedAsset, OldObjectPath);
 		}
 	}
diff --git a/Engine/Source/Runtime/CoreUObject/Public/Misc/HotReloadInterface.h b/Engine/Source/Runtime/CoreUObject/Public/Misc/HotReloadInterface.h
--- a/Engine/Source/Runtime/CoreUObject/Public/Misc/HotReloadInterface.h
+++ b/Engine/Source/Runtime/CoreUObject/Public/Misc/HotReloadInterface.h
@@ -17,7 +17,8 @@ public:
 	 */
 	static inline IHotReloadInterface* GetPtr()
 	{
-		return FModuleManager::GetModulePtr<IHotReloadInterface>("HotReload");
+		static FName HotReload("HotReload");
+		return FModuleManager::GetModulePtr<IHotReloadInterface>(HotReload);
 	}
 
 	/**
diff --git a/Engine/Source/Runtime/Launch/Private/LaunchEngineLoop.cpp b/Engine/Source/Runtime/Launch/Private/LaunchEngineLoop.cpp
--- a/Engine/Source/Runtime/Launch/Private/LaunchEngineLoop.cpp
+++ b/Engine/Source/Runtime/Launch/Private/LaunchEngineLoop.cpp
@@ -2137,10 +2137,11 @@ void FEngineLoop::Tick()
 #if WITH_EDITOR
 		if( GIsEditor )
 		{
+			static FName AutomationController("AutomationController");
 			//Check if module loaded to support the change to allow this to be hot compilable.
-			if (FModuleManager::Get().IsModuleLoaded(TEXT("AutomationController")))
+			if (FModuleManager::Get().IsModuleLoaded(AutomationController))
 			{
-				FModuleManager::GetModuleChecked<IAutomationControllerModule>("AutomationController").Tick();
+				FModuleManager::GetModuleChecked<IAutomationControllerModule>(AutomationController).Tick();
 			}
 		}
 #endif
diff --git a/Engine/Source/Runtime/Slate/Public/Widgets/Views/STableRow.h b/Engine/Source/Runtime/Slate/Public/Widgets/Views/STableRow.h
--- a/Engine/Source/Runtime/Slate/Public/Widgets/Views/STableRow.h
+++ b/Engine/Source/Runtime/Slate/Public/Widgets/Views/STableRow.h
@@ -590,8 +590,10 @@ public:
 
 		const bool bIsActive = OwnerWidget->AsWidget()->HasKeyboardFocus();
 
+		static FName GenericWhiteBoxBrush("GenericWhiteBox");
+
 		// @todo: Slate Style - make this part of the widget style
-		const FSlateBrush* WhiteBox = FCoreStyle::Get().GetBrush("GenericWhiteBox");
+		const FSlateBrush* WhiteBox = FCoreStyle::Get().GetBrush(GenericWhiteBoxBrush);
 
 		const ItemType* MyItem = OwnerWidget->Private_ItemFromWidget( this );
 
