Fixed WASD navigation in viewport when using trackpad (broken by addign bIsTracking dependency to FEditorViewportClient::IsFlightCameraInputModeActive() in CL #2289372)

diff --git a/Engine/Source/Editor/UnrealEd/Private/EditorViewportClient.cpp b/Engine/Source/Editor/UnrealEd/Private/EditorViewportClient.cpp
--- a/Engine/Source/Editor/UnrealEd/Private/EditorViewportClient.cpp
+++ b/Engine/Source/Editor/UnrealEd/Private/EditorViewportClient.cpp
@@ -1077,7 +1077,7 @@ void FEditorViewportClient::UpdateCameraMovement( float DeltaTime )
 		const bool bIsUsingTrackpad = FSlateApplication::Get().IsUsingTrackpad();
 		
 		// Certain keys are only available while the flight camera input mode is active
-		const bool bUsingFlightInput = IsFlightCameraInputModeActive();
+		const bool bUsingFlightInput = IsFlightCameraInputModeActive() || bIsUsingTrackpad;
 
 		// Do we want to use the regular arrow keys for flight input?
 		// Because the arrow keys are used for nudging actors, we'll only do this while the ALT key is up
