More stats on the blueprint action database (to help identify the memory impact).

diff --git a/Engine/Source/Editor/UnrealEd/Private/Commandlets/DumpBlueprintsInfoCommandlet.cpp b/Engine/Source/Editor/UnrealEd/Private/Commandlets/DumpBlueprintsInfoCommandlet.cpp
--- a/Engine/Source/Editor/UnrealEd/Private/Commandlets/DumpBlueprintsInfoCommandlet.cpp
+++ b/Engine/Source/Editor/UnrealEd/Private/Commandlets/DumpBlueprintsInfoCommandlet.cpp
@@ -28,6 +28,7 @@
 #include "AssetRegistryModule.h"
 #include "PackageName.h"
 #include "PackageHelperFunctions.h" // for NORMALIZE_ExcludeMapPackages
+#include "K2Node_FunctionEntry.h"
 
 DEFINE_LOG_CATEGORY_STATIC(LogBlueprintInfoDump, Log, All);
 
@@ -114,6 +115,9 @@ DumpBlueprintsInfo commandlet params: \n\
                         the specified Blueprint. The blueprint name can contain\n\
                         wildcards (to match multiple blueprints). Can also be  \n\
                         set to \"all\", to load every non-developer blueprint. \n\
+\n\
+    -dbInfo             Will dump info regarding the blueprint action database \n\
+                        (which is used to consturct blueprint action menus).   \n\
 \n\
     -help, -h, -?       Display this message and then exit.                    \n\
 \n");
@@ -134,6 +138,8 @@ DumpBlueprintsInfo commandlet params: \n\
 		BPDUMP_SelectAllObjTypes    = (1<<9), 
 		BPDUMP_UseLegacyMenuBuilder = (1<<10),
 
+		BPDUMP_ActionDatabaseInfo   = (1<<11),
+
 		BPDUMP_PaletteMask = (BPDUMP_UnfilteredPalette | BPDUMP_FilteredPalette),
 		BPDUMP_ContextMask = (BPDUMP_GraphContextActions | BPDUMP_PinContextActions),
 	};
@@ -254,10 +260,12 @@ DumpBlueprintsInfo commandlet params: \n\
 	 * to the blueprint).
 	 */
 	static void GetComponentProperties(UBlueprint* Blueprint, TArray<UObjectProperty*>& PropertiesOut);
-	
 
 	/**
-	 * 
+	 * Takes the given size (in bytes) and returns a formatted string in either
+	 * bytes, kilobytes, megabytes, or gigabytes.
+	 *
+	 * @return A formatted string, representing the specified size in the best units possible.
 	 */
 	static FString BuildByteSizeString(int32 ByteSize);
 
@@ -617,6 +625,10 @@ DumpBlueprintInfoUtils::CommandletOptions::CommandletOptions(TArray<FString> con
 			bool const bAllowDevBlueprints = BlueprintName.IsEmpty() ? false : true;
 			LoadBlueprints(BlueprintName, bAllowDevBlueprints);
 		}
+		else if (!Switch.Compare("dbInfo", ESearchCase::IgnoreCase))
+		{
+			NewDumpFlags |= BPDUMP_ActionDatabaseInfo;
+		}
 		else if (!Switch.StartsWith("run=")) // account for the first switch, which invoked this commandlet
 		{
 			UE_LOG(LogBlueprintInfoDump, Warning, TEXT("Unrecognized command switch '%s', use -help for a listing of all accepted params"), *Switch);
@@ -1068,24 +1080,35 @@ static FString DumpBlueprintInfoUtils::BuildByteSizeString(int32 ByteSize)
 	int32 UnitsIndex = 0;
 	float ConvertedSize = ByteSize;
 
+	TCHAR const* Format = TEXT("%.0f %s");
+
 	float const MetricStepSize = 1024.0f;
 	while ((ConvertedSize > MetricStepSize) && (UnitsIndex < ByteUnits_MAX))
 	{
 		ConvertedSize /= MetricStepSize;
 		++UnitsIndex;
+
+		Format = TEXT("%.2f %s");
 	}
 
-	return FString::Printf(TEXT("%.2f %s"), ConvertedSize, ByteUnits[UnitsIndex]);
+	return FString::Printf(Format, ConvertedSize, ByteUnits[UnitsIndex]);
 }
 
+#include "BlueprintFunctionNodeSpawner.h"
+#include "BlueprintVariableNodeSpawner.h"
+
 //------------------------------------------------------------------------------
 static bool DumpBlueprintInfoUtils::DumpActionDatabaseInfo(uint32 Indent, FArchive* FileOutWriter)
 {
 	bool bWroteToFile = false;
-	if ((CommandOptions.DumpFlags & BPDUMP_UseLegacyMenuBuilder) == 0)
+
+	uint32 const DbInfoMask = (BPDUMP_UseLegacyMenuBuilder | BPDUMP_ActionDatabaseInfo);
+	if ((CommandOptions.DumpFlags & DbInfoMask) == BPDUMP_ActionDatabaseInfo)
 	{
 		UE_LOG(LogBlueprintInfoDump, Display, TEXT("%sDumping Database info..."), *BuildIndentString(Indent, true));
-
+		//--------------------------------------
+		// Composing Data
+		//--------------------------------------
 		double DatabaseBuildTime = 0.0;
 		{
 			FScopedDurationTimer DatabaseBuildTimer(DatabaseBuildTime);
@@ -1094,26 +1117,52 @@ static bool DumpBlueprintInfoUtils::DumpActionDatabaseInfo(uint32 Indent, FArchi
 		}
 		FBlueprintActionDatabase& Database = FBlueprintActionDatabase::Get();
 		
+		FBlueprintActionDatabase::FActionRegistry const& ActionRegistry = Database.GetAllActions();
+		int32 EstimatedDatabaseSize = sizeof(Database) + ActionRegistry.GetAllocatedSize();
+		int32 EstimatedSystemSize   = EstimatedDatabaseSize;
+
 		TSet<UBlueprint*> TemplateOuters;
 		int32 DatabaseCount = 0;
-		double NodeCachingDuration = 0.0;
+		double TotalPrimingTime = 0.0;
 		int32 TemplateCount = 0;
+		int32 UnknownAssetActions = 0;
 
-		for (auto const& DbEntry : Database.GetAllActions())
+		struct FSpawnerInfo
 		{
+			int32  Count;
+			int32  TemplateNodeCount;
+			double TotalPrimingTime;
+
+			FSpawnerInfo() : Count(0), TemplateNodeCount(0), TotalPrimingTime(0.0) {}
+		};
+		TMap<UClass*, FSpawnerInfo> DatabaseBreakdown;
+
+		for (auto const& DbEntry : ActionRegistry)
+		{
+			UObject const* ActionSetKey = DbEntry.Key;
+			bool const bIsUnknownAssetEntry = ActionSetKey->IsAsset() &&
+				!ActionSetKey->IsA<UBlueprint>() &&
+				!ActionSetKey->IsA<UUserDefinedStruct>() &&
+				!ActionSetKey->IsA<UUserDefinedEnum>();
+
 			for (UBlueprintNodeSpawner* BpAction : DbEntry.Value)
 			{
-				FDurationTimer NodeCacheTimer(NodeCachingDuration);
-				NodeCacheTimer.Start();
-				UEdGraphNode* TemplateNode = BpAction->GetTemplateNode();
-				NodeCacheTimer.Stop();
+				++DatabaseCount;
+				// @TODO: doesn't account for any allocated memory (for delegates, text strings, etc.)
+				EstimatedDatabaseSize += sizeof(*BpAction);
 
-				if (TemplateNode != nullptr)
+				FSpawnerInfo& SpawnerInfo = DatabaseBreakdown.FindOrAdd(BpAction->GetClass());
+				SpawnerInfo.Count += 1;
+
+				int32 OldPrimingTime = TotalPrimingTime;
 				{
-					NodeCacheTimer.Start();
-					TemplateNode->AllocateDefaultPins();
-					NodeCacheTimer.Stop();
+					FScopedDurationTimer PrimingTimer(TotalPrimingTime);
+					BpAction->Prime();
+				}
+				SpawnerInfo.TotalPrimingTime += (TotalPrimingTime - OldPrimingTime);
 
+				if (UEdGraphNode* TemplateNode = BpAction->GetTemplateNode(NoInit))
+				{
 					UObject* TemplateOuter = TemplateNode->GetOuter();
 					while ((TemplateOuter != nullptr) && (Cast<UBlueprint>(TemplateOuter) == nullptr))
 					{
@@ -1123,45 +1172,262 @@ static bool DumpBlueprintInfoUtils::DumpActionDatabaseInfo(uint32 Indent, FArchi
 					TemplateOuters.Add(OuterBlueprint);
 
 					++TemplateCount;
+					SpawnerInfo.TemplateNodeCount += 1;
+				}
+
+				if (bIsUnknownAssetEntry)
+				{
+					++UnknownAssetActions;
 				}
-				++DatabaseCount;
 			}
 		}
 
-		int32 EstimatedCacheKbSize = 0;
+		int32 SpawnerCount = 0;
+		for (TObjectIterator<UBlueprintNodeSpawner> NodeSpawnerIt; NodeSpawnerIt; ++NodeSpawnerIt)
+		{
+			++SpawnerCount;
+			// @TODO: doesn't account for any allocated memory (for delegates, text strings, etc.)
+			EstimatedSystemSize += sizeof(**NodeSpawnerIt);
+		}
+
+		FString const OriginalIndent = BuildIndentString(Indent);
+		FString const IndentedNewline = "\n" + BuildIndentString(Indent + 1);
+
+		FString DatabaseInfoHeading = FString::Printf(TEXT("%s\"ActionDatabaseInfo\" : {%s\"TotalNodeSpawnerCount\" : %d,"), 
+			*OriginalIndent, *IndentedNewline, SpawnerCount);
+		FileOutWriter->Serialize(TCHAR_TO_ANSI(*DatabaseInfoHeading), DatabaseInfoHeading.Len());
+
+		//--------------------------------------
+		// Dumping Database Stats
+		//--------------------------------------
+		UE_LOG(LogBlueprintInfoDump, Display, TEXT("%sDumping raw action stats..."), *BuildIndentString(Indent+1, true));
+
+		FString const DoubleIndent = BuildIndentString(Indent + 2);
+		FString const DblIndentedNewline = "\n" + DoubleIndent;
+		FString const SubDictEndingBrace = IndentedNewline + "}";
+
+		FString const DatabaseSizeStr = BuildByteSizeString(EstimatedDatabaseSize);
+		FString const AvgActionSizeStr = BuildByteSizeString(EstimatedDatabaseSize / DatabaseCount);
+
+		FString DatabaseStats = FString::Printf(TEXT("%s\"Database Stats\" : {"), *IndentedNewline);
+		DatabaseStats += FString::Printf(TEXT("%s\"DatabaseBuildTime\"     : %.3f seconds,"), *DblIndentedNewline, DatabaseBuildTime);
+		DatabaseStats += FString::Printf(TEXT("%s\"NodeSpawnerCount\"      : %d,"), *DblIndentedNewline, DatabaseCount);
+		DatabaseStats += FString::Printf(TEXT("%s\"EstimatedDatabaseSize\" : %s,"), *DblIndentedNewline, *DatabaseSizeStr);
+		DatabaseStats += FString::Printf(TEXT("%s\"AvgSizePerEntry\"       : %s"), *DblIndentedNewline, *AvgActionSizeStr);
+		DatabaseStats += SubDictEndingBrace + ",";
+
+		FileOutWriter->Serialize(TCHAR_TO_ANSI(*DatabaseStats), DatabaseStats.Len());
+
+		//--------------------------------------
+		// Dumping Template Cache Stats
+		//--------------------------------------
+		UE_LOG(LogBlueprintInfoDump, Display, TEXT("%sDumping template-cache stats..."), *BuildIndentString(Indent+1, true));
+
+		int32 EstimatedCacheSize = 0;
 		for (UBlueprint* CacheBlueprint : TemplateOuters)
 		{
 			TArray<UObject*> ChildObjs;
 			GetObjectsWithOuter(CacheBlueprint, ChildObjs);
 
-			EstimatedCacheKbSize += sizeof(*CacheBlueprint);
+			EstimatedCacheSize += sizeof(*CacheBlueprint);
 			for (UObject* ChildObj : ChildObjs)
 			{
-				EstimatedCacheKbSize += sizeof(*ChildObj);
+				// @TODO: doesn't account for any allocated memory (for member TArrays, etc.)
+				EstimatedCacheSize += sizeof(*ChildObj);
 			}
 		}
+		EstimatedSystemSize += EstimatedCacheSize;
 
-		int32 const DatabaseSize = Database.EstimatedSize();
-		FString const DatabaseSizeStr = BuildByteSizeString(DatabaseSize);
-		FString const NodeCacheSizeStr = BuildByteSizeString(EstimatedCacheKbSize);
-		FString const TotalSystemSizeStr = BuildByteSizeString(DatabaseSize + EstimatedCacheKbSize);
-		FString const AvgDatabaseEntrySizeStr = BuildByteSizeString(DatabaseSize / DatabaseCount);
-		FString const AvgCachedNodeSizeStr = BuildByteSizeString(EstimatedCacheKbSize / TemplateCount);
+		FString const NodeCacheSizeStr = BuildByteSizeString(EstimatedCacheSize);
+		FString const AvgNodeSizeStr = BuildByteSizeString(EstimatedCacheSize / TemplateCount);
+		FString NodeCacheStats = FString::Printf(TEXT("%s\"Template-Cache Stats\" : {"), *IndentedNewline);
+		NodeCacheStats += FString::Printf(TEXT("%s\"TotalPrimingDuration\" : %.3f seconds,"), *DblIndentedNewline, TotalPrimingTime);
+		NodeCacheStats += FString::Printf(TEXT("%s\"CachedNodeCount\"      : %d,"), *DblIndentedNewline, TemplateCount);
+		NodeCacheStats += FString::Printf(TEXT("%s\"EstimatedCacheSize\"   : %s,"), *DblIndentedNewline, *NodeCacheSizeStr);
+		NodeCacheStats += FString::Printf(TEXT("%s\"AvgSizePerEntry\"      : %s"), *DblIndentedNewline, *AvgNodeSizeStr);
+		NodeCacheStats += SubDictEndingBrace + ",";
 
-		FString const OriginalIndent  = BuildIndentString(Indent);
-		FString const IndentedNewline = "\n" + BuildIndentString(Indent + 1);
+		FileOutWriter->Serialize(TCHAR_TO_ANSI(*NodeCacheStats), NodeCacheStats.Len());
+
+		//--------------------------------------
+		// Dumping Database Breakdown
+		//--------------------------------------
+		TCHAR const* LineEnding = TEXT(",\n");
+
+		FString BreakdownStats = FString::Printf(TEXT("%s\"Database Breakdown\" : {\n"), *IndentedNewline);
+		for (auto const& SpawnerEntry : DatabaseBreakdown)
+		{
+			int32 SpawnerCount = SpawnerEntry.Value.Count;
+
+			static FString const TripleIndent = BuildIndentString(Indent + 3);
+			BreakdownStats += FString::Printf(TEXT("%s\"%s\" : {\n"), *DoubleIndent, *SpawnerEntry.Key->GetName());
+			BreakdownStats += FString::Printf(TEXT("%s\"Total\"               : %d,\n"), *TripleIndent, SpawnerCount);
+			BreakdownStats += FString::Printf(TEXT("%s\"TemplateNodesPrimed\" : %d,\n"), *TripleIndent, SpawnerEntry.Value.TemplateNodeCount);
+			BreakdownStats += FString::Printf(TEXT("%s\"AvgPrimingDuration\"  : %.03f seconds\n"), *TripleIndent, SpawnerEntry.Value.TotalPrimingTime / SpawnerCount);
+			BreakdownStats += DoubleIndent + "}" + LineEnding;
+		}
+		BreakdownStats.RemoveFromEnd(LineEnding);
+		BreakdownStats += IndentedNewline + "},";
+
+		FileOutWriter->Serialize(TCHAR_TO_ANSI(*BreakdownStats), BreakdownStats.Len());
+
+		//--------------------------------------
+		// Dumping Blueprint Stats
+		//--------------------------------------
+		UE_LOG(LogBlueprintInfoDump, Display, TEXT("%sDumping blueprint related stats..."), *BuildIndentString(Indent+1, true));
+
+		int32 BlueprintsCount = 0;
+		int32 BlueprintsDbCount = 0;
+		int32 TotalBlueprintActionCount = 0;
+		int32 BlueprintFunctionCount = 0;
+		int32 BlueprintVariableCount = 0;
+		int32 BlueprintDelegateCount = 0;
+		int32 BlueprintLocalVarCount = 0;
+
+		for (TObjectIterator<UBlueprint> BlueprintIt; BlueprintIt; ++BlueprintIt)
+		{
+			if (BlueprintIt->IsAsset())
+			{
+				++BlueprintsCount;
+				if (FBlueprintActionDatabase::FActionList const* ActionList = ActionRegistry.Find(*BlueprintIt))
+				{
+					if (ActionList->Num() > 0)
+					{
+						++BlueprintsDbCount;
+					}
+					TotalBlueprintActionCount += ActionList->Num();
+				}
+
+				for (auto GraphIt = BlueprintIt->FunctionGraphs.CreateConstIterator(); GraphIt; ++GraphIt)
+				{
+					UEdGraph* FunctionGraph = (*GraphIt);
+
+					TArray<UK2Node_FunctionEntry*> GraphEntryNodes;
+					FunctionGraph->GetNodesOfClass<UK2Node_FunctionEntry>(GraphEntryNodes);
+
+					for (UK2Node_FunctionEntry* FunctionEntry : GraphEntryNodes)
+					{
+						for (FBPVariableDescription const& LocalVar : FunctionEntry->LocalVariables)
+						{
+							++BlueprintLocalVarCount;
+						}
+					}
+				}
+
+				UClass* BlueprintClass = BlueprintIt->GeneratedClass;
+				if (BlueprintClass == nullptr)
+				{
+					continue;
+				}
+
+				for (TFieldIterator<UFunction> FunctionIt(BlueprintClass, EFieldIteratorFlags::ExcludeSuper); FunctionIt; ++FunctionIt)
+				{
+					++BlueprintFunctionCount;
+				}
+
+				for (TFieldIterator<UProperty> PropertyIt(BlueprintClass, EFieldIteratorFlags::ExcludeSuper); PropertyIt; ++PropertyIt)
+				{
+					bool const bIsDelegate = PropertyIt->IsA(UMulticastDelegateProperty::StaticClass());
+					if (bIsDelegate)
+					{
+						++BlueprintDelegateCount;
+					}
+					else
+					{
+						++BlueprintVariableCount;
+					}
+				}
+			}
+		};
+
+		FString BlueprintStats = FString::Printf(TEXT("%s\"Blueprint Stats\" : {"), *IndentedNewline);
+		BlueprintStats += FString::Printf(TEXT("%s\"BlueprintsLoaded\"         : %d,"), *DblIndentedNewline, BlueprintsCount);
+		BlueprintStats += FString::Printf(TEXT("%s\"BlueprintsWithActions\"    : %d,"), *DblIndentedNewline, BlueprintsDbCount);
+		BlueprintStats += FString::Printf(TEXT("%s\"AvgFunctionsPerBlueprint\" : %d,"), *DblIndentedNewline, FMath::RoundToInt((float)BlueprintFunctionCount / BlueprintsCount));
+		BlueprintStats += FString::Printf(TEXT("%s\"AvgVariablesPerBlueprint\" : %d,"), *DblIndentedNewline, FMath::RoundToInt((float)BlueprintVariableCount / BlueprintsCount));
+		BlueprintStats += FString::Printf(TEXT("%s\"AvgDelegatesPerBlueprint\" : %d,"), *DblIndentedNewline, FMath::RoundToInt((float)BlueprintDelegateCount / BlueprintsCount));
+		BlueprintStats += FString::Printf(TEXT("%s\"AvgLocalVarsPerBlueprint\" : %d,"), *DblIndentedNewline, FMath::RoundToInt((float)BlueprintLocalVarCount / BlueprintsCount));
+		BlueprintStats += FString::Printf(TEXT("%s\"TotalBlueprintActions\"    : %d,"), *DblIndentedNewline, TotalBlueprintActionCount);
+		BlueprintStats += FString::Printf(TEXT("%s\"AvgActionsPerBlueprint\"   : %d"), *DblIndentedNewline, FMath::RoundToInt((float)TotalBlueprintActionCount / BlueprintsDbCount));
+		BlueprintStats += SubDictEndingBrace + ",";
+		
+		FileOutWriter->Serialize(TCHAR_TO_ANSI(*BlueprintStats), BlueprintStats.Len());
+
+		//--------------------------------------
+		// Additional Asset Stats
+		//--------------------------------------
+		UE_LOG(LogBlueprintInfoDump, Display, TEXT("%sDumping additional asset stats..."), *BuildIndentString(Indent+1, true));
+
+		int32 EnumAssetCount = 0;
+		int32 EnumsWithActionsCount = 0;
+		int32 TotalEnumActions = 0;
+
+		for (TObjectIterator<UUserDefinedEnum> EnumIt; EnumIt; ++EnumIt)
+		{
+			if (EnumIt->IsAsset())
+			{
+				++EnumAssetCount;
+				if (FBlueprintActionDatabase::FActionList const* ActionList = ActionRegistry.Find(*EnumIt))
+				{
+					if (ActionList->Num() > 0)
+					{
+						++EnumsWithActionsCount;
+					}
+					TotalEnumActions += ActionList->Num();
+				}
+			}
+			else
+			{
+				// @TODO: bad assumption? all UUserDefinedEnum should be assets
+			}
+		};
+
+		int32 StructAssetCount = 0;
+		int32 StructsWithActionsCount = 0;
+		int32 TotalStructActions = 0;
+
+		for (TObjectIterator<UUserDefinedStruct> StructIt; StructIt; ++StructIt)
+		{
+			if (StructIt->IsAsset())
+			{
+				++StructAssetCount;
+				if (FBlueprintActionDatabase::FActionList const* ActionList = ActionRegistry.Find(*StructIt))
+				{
+					if (ActionList->Num() > 0)
+					{
+						++StructsWithActionsCount;
+					}
+					TotalStructActions += ActionList->Num();
+				}
+			}
+			else
+			{
+				// @TODO: bad assumption? all UUserDefinedStructs should be assets
+			}
+		};
+
+		FString OtherAssetStats = FString::Printf(TEXT("%s\"Other Asset Stats\"  : {"), *IndentedNewline);
+		OtherAssetStats += FString::Printf(TEXT("%s\"EnumAssetsLoaded\"        : %d,"), *DblIndentedNewline, EnumAssetCount);
+		OtherAssetStats += FString::Printf(TEXT("%s\"EnumAssetsWithActions\"   : %d,"), *DblIndentedNewline, EnumsWithActionsCount);
+		OtherAssetStats += FString::Printf(TEXT("%s\"TotalEnumActionCount\"    : %d,"), *DblIndentedNewline, TotalEnumActions);
+		OtherAssetStats += FString::Printf(TEXT("%s\"StructAssetsLoaded\"      : %d,"), *DblIndentedNewline, StructAssetCount);
+		OtherAssetStats += FString::Printf(TEXT("%s\"StructAssetsWithActions\" : %d,"), *DblIndentedNewline, StructsWithActionsCount);
+		OtherAssetStats += FString::Printf(TEXT("%s\"TotalStructActionCount\"  : %d,"), *DblIndentedNewline, TotalStructActions);
+		OtherAssetStats += FString::Printf(TEXT("%s\"OtherAssetActions\"       : %d"), *DblIndentedNewline, UnknownAssetActions);
+		OtherAssetStats += SubDictEndingBrace + ",";
+
+		FileOutWriter->Serialize(TCHAR_TO_ANSI(*OtherAssetStats), OtherAssetStats.Len());
+		
+		//--------------------------------------
+		// Stats Closing
+		//--------------------------------------
+		
+		FString const TotalSystemSizeStr = BuildByteSizeString(EstimatedSystemSize);
+
+		FString DatabaseInfoClosing = FString::Printf(TEXT("%s\"EstimatedSystemSize\" : \"%s\"\n%s}"),
+			*IndentedNewline, *TotalSystemSizeStr, *OriginalIndent);
+		FileOutWriter->Serialize(TCHAR_TO_ANSI(*DatabaseInfoClosing), DatabaseInfoClosing.Len());
 
-		FString DatabaseInfoEntry = FString::Printf(TEXT("%s\"ActionDatabaseInfo\" : {"), *OriginalIndent);
-		DatabaseInfoEntry += FString::Printf(TEXT("%s\"DatabaseBuildTime\"      : \"%.3f seconds\","), *IndentedNewline, DatabaseBuildTime);
-		DatabaseInfoEntry += FString::Printf(TEXT("%s\"CachingNodesDuration\"   : \"%.3f seconds\","), *IndentedNewline, NodeCachingDuration);
-		DatabaseInfoEntry += FString::Printf(TEXT("%s\"DatabaseEntryCount\"     : %d,"), *IndentedNewline, DatabaseCount);
-		DatabaseInfoEntry += FString::Printf(TEXT("%s\"NodeCacheCount\"         : %d,"), *IndentedNewline, TemplateCount);
-		DatabaseInfoEntry += FString::Printf(TEXT("%s\"DatabaseSize\"           : \"%s\","), *IndentedNewline, *DatabaseSizeStr);
-		DatabaseInfoEntry += FString::Printf(TEXT("%s\"AvgSizePerDatabaseEntry\": \"%s\","), *IndentedNewline, *AvgDatabaseEntrySizeStr);
-		DatabaseInfoEntry += FString::Printf(TEXT("%s\"EstimatedNodeCacheSize\" : \"%s\","), *IndentedNewline, *NodeCacheSizeStr);
-		DatabaseInfoEntry += FString::Printf(TEXT("%s\"AvgSizePerCachedNode\"   : \"%s\","), *IndentedNewline, *AvgCachedNodeSizeStr);
-		DatabaseInfoEntry += FString::Printf(TEXT("%s\"EstimatedSystemSize\"    : \"%s\"\n%s}"), *IndentedNewline, *TotalSystemSizeStr, *OriginalIndent);
-		FileOutWriter->Serialize(TCHAR_TO_ANSI(*DatabaseInfoEntry), DatabaseInfoEntry.Len());
 		bWroteToFile = true;
 	}
 	return bWroteToFile;
@@ -2078,11 +2344,14 @@ static int32 DumpBlueprintInfoUtils::LoadBlueprints(FString AssetName, bool bAll
  	FString DevelopersRoot = FPaths::GameDevelopersDir().LeftChop(1);
 	FPackageName::TryConvertFilenameToLongPackageName(DevelopersRoot, DevelopersRoot);
 
-	bool const bLoadPackageSubset = (PackagesToLoad.Num() > 0);
+	bool const bLoadPackageSubset = (PackagesToLoad.Num() > 0);	
+	int32 AlreadyLoadedCount = 0;
+
 	for (FAssetData const& Asset : BlueprintAssetList)
 	{
 		if (Asset.IsAssetLoaded())
 		{
+			++AlreadyLoadedCount;
 			continue;
 		}
 
@@ -2116,6 +2385,13 @@ static int32 DumpBlueprintInfoUtils::LoadBlueprints(FString AssetName, bool bAll
 		}
  	}
 
+	if (!bLoadPackageSubset || (PackagesToLoad.Num() > 0))
+	{
+		int32 const TotalLoaded = LoadedCount + AlreadyLoadedCount;
+		int32 const AttemptedLoadCount = bLoadPackageSubset ? PackagesToLoad.Num() : BlueprintAssetList.Num();
+		UE_LOG(LogBlueprintInfoDump, Warning, TEXT("Successfully loaded %d/%d Blueprints (%.1f%%)."), TotalLoaded, AttemptedLoadCount, 100.f * (float)TotalLoaded / AttemptedLoadCount);
+	}
+
 	return LoadedCount;
 }
 
