Backout experimental changes to detect if an application is foreground

diff --git a/Engine/Source/Editor/UnrealEd/Private/Editor.cpp b/Engine/Source/Editor/UnrealEd/Private/Editor.cpp
--- a/Engine/Source/Editor/UnrealEd/Private/Editor.cpp
+++ b/Engine/Source/Editor/UnrealEd/Private/Editor.cpp
@@ -982,7 +982,7 @@ void UEditorEngine::Tick( float DeltaSeconds, bool bIdleMode )
 	}
 
 	// Find out if the editor has focus. Audio should only play if the editor has focus.
-	const bool bHasFocus = FSlateApplication::Get().GetActiveTopLevelWindow().IsValid() || FPlatformProcess::IsThisApplicationForeground();
+	const bool bHasFocus = FPlatformProcess::IsThisApplicationForeground();
 
 	if (bHasFocus || GetDefault<ULevelEditorMiscSettings>()->bAllowBackgroundAudio)
 	{
@@ -5982,7 +5982,7 @@ bool UEditorEngine::ShouldThrottleCPUUsage() const
 {
 	bool bShouldThrottle = false;
 
-	bool bIsForeground = FSlateApplication::Get().GetActiveTopLevelWindow().IsValid() || FPlatformProcess::IsThisApplicationForeground();
+	bool bIsForeground = FPlatformProcess::IsThisApplicationForeground();
 
 	if( !bIsForeground )
 	{
