Fix drag tool display text formatting so that it is more accurate.

diff --git a/Engine/Source/Editor/UnrealEd/Private/DragTool_Measure.cpp b/Engine/Source/Editor/UnrealEd/Private/DragTool_Measure.cpp
--- a/Engine/Source/Editor/UnrealEd/Private/DragTool_Measure.cpp
+++ b/Engine/Source/Editor/UnrealEd/Private/DragTool_Measure.cpp
@@ -69,8 +69,12 @@ void FDragTool_Measure::Render(const FSceneView* View, FCanvas* Canvas)
 
 		const FVector2D PixelMid = FVector2D(PixelStart + ((PixelEnd - PixelStart) / 2));
 
-		const FString LengthStr = FEditorViewportClient::UnrealUnitsToSiUnits(Length);
-		FCanvasTextItem TextItem( FVector2D( FMath::FloorToFloat(PixelMid.X), FMath::FloorToFloat(PixelMid.Y) ), FText::FromString( LengthStr ), GEngine->GetSmallFont(), FLinearColor::White );
+		FNumberFormattingOptions Options;
+		Options.UseGrouping = false;
+		const FText LengthStr = FText::AsNumber( FMath::RoundToInt(Length), &Options );
+
+
+		FCanvasTextItem TextItem( FVector2D( FMath::FloorToFloat(PixelMid.X), FMath::FloorToFloat(PixelMid.Y) ), LengthStr, GEngine->GetSmallFont(), FLinearColor::White );
 		TextItem.bCentreX = true;
 		Canvas->DrawItem( TextItem );
 	}
