Missed one file in the Username whitspace removal change.

diff --git a/Engine/Source/Editor/UnrealEd/Private/Kismet2/DebuggerCommands.cpp b/Engine/Source/Editor/UnrealEd/Private/Kismet2/DebuggerCommands.cpp
--- a/Engine/Source/Editor/UnrealEd/Private/Kismet2/DebuggerCommands.cpp
+++ b/Engine/Source/Editor/UnrealEd/Private/Kismet2/DebuggerCommands.cpp
@@ -631,7 +631,7 @@ TSharedRef< SWidget > FPlayWorldCommands::GenerateLaunchMenuContent( TSharedRef<
 						{
 							LabelArguments.Add(TEXT("HostUser"), LOCTEXT("DisconnectedHint", " [Disconnected]"));
 						}
-						else if (DeviceProxy->GetHostUser() != FPlatformProcess::UserName(false))
+						else if (DeviceProxy->GetHostUser() != FPlatformProcess::UserName(true))
 						{
 							LabelArguments.Add(TEXT("HostUser"), FText::FromString(DeviceProxy->GetHostUser()));
 						}
@@ -1277,7 +1277,7 @@ void FInternalPlayWorldCommandCallbacks::OpenDeviceManager_Clicked()
 
 void FInternalPlayWorldCommandCallbacks::OpenLauncher_Clicked()
 {
-	//FGlobalTabmanager::Get()->InvokeTab(FName("SessionLauncher"));
+	//FGlobalTabmanager::Get()->InvokeTab(FName("ProjectLauncher"));
 }
 
 
