[342533] Blueprints - Deleting a parent blueprint while a child blueprint still exists will no longer cause a crash.

- Modified ForceDeleteObjects() to reparent and recompile child BPs prior to replacing references to the BP that's being deleted. The user is already notified in the dialog that issues may arise due to data loss.

diff --git a/Engine/Source/Editor/UnrealEd/Private/ObjectTools.cpp b/Engine/Source/Editor/UnrealEd/Private/ObjectTools.cpp
--- a/Engine/Source/Editor/UnrealEd/Private/ObjectTools.cpp
+++ b/Engine/Source/Editor/UnrealEd/Private/ObjectTools.cpp
@@ -18,6 +18,7 @@
 #include "AssetToolsModule.h"
 #include "Editor/UnrealEd/Public/Toolkits/AssetEditorManager.h"
 #include "ISourceControlModule.h"
+#include "Kismet2/KismetEditorUtilities.h"
 #include "Kismet2/KismetReinstanceUtilities.h"
 #include "FbxImporter.h"
 #include "PackageHelperFunctions.h"
@@ -1773,23 +1774,65 @@ namespace ObjectTools
 		ReloadEditorWorldForReferenceReplacementIfNecessary(ObjectsToDelete);
 
 		{
-			// Replacing references inside already loaded objects could cause rendering issues, so globally detach all components from their scenes for now
-			FGlobalComponentReregisterContext ReregisterContext;
-
 			int32 ReplaceableObjectsNum = 0;
 			{
+				bool bNeedsGarbageCollection = false;
 				TArray<UObject*> ObjectsToReplace = ObjectsToDelete;
 				for (TArray<UObject*>::TIterator ObjectItr(ObjectsToReplace); ObjectItr; ++ObjectItr)
 				{
 					UObject* CurObject = *ObjectItr;
-					// If we're a blueprint add our generated class as well
+					
 					UBlueprint *BlueprintObject = Cast<UBlueprint>(CurObject);
-					if (BlueprintObject && BlueprintObject->GeneratedClass)
+					if (BlueprintObject)
 					{
-						ObjectsToReplace.AddUnique(BlueprintObject->GeneratedClass);
+						// If we're a blueprint add our generated class as well
+						if (BlueprintObject->GeneratedClass)
+						{
+							ObjectsToReplace.AddUnique(BlueprintObject->GeneratedClass);
+						}
+
+						// Clear BP ref so we're not warned about the skeleton class being transient during replacement
+						if (BlueprintObject->SkeletonGeneratedClass)
+						{
+							BlueprintObject->SkeletonGeneratedClass->ClassGeneratedBy = NULL;
+						}
+
+						// Reparent any direct children to the parent class of the blueprint that's about to be deleted
+						if(BlueprintObject->ParentClass != nullptr)
+						{
+							for(TObjectIterator<UClass> ClassIt; ClassIt; ++ClassIt)
+							{
+								UClass* ChildClass = *ClassIt;
+								if(ChildClass->GetSuperStruct() == BlueprintObject->GeneratedClass)
+								{
+									UBlueprint* ChildBlueprint = Cast<UBlueprint>(ChildClass->ClassGeneratedBy);
+									if(ChildBlueprint != nullptr)
+									{
+										ChildBlueprint->Modify();
+										ChildBlueprint->ParentClass = BlueprintObject->ParentClass;
+
+										// Recompile the child blueprint to fix up the generated class
+										FKismetEditorUtilities::CompileBlueprint(ChildBlueprint, false, true);
+
+										// Defer garbage collection until after we're done processing the list of objects
+										bNeedsGarbageCollection = true;
+									}
+								}
+							}
+						}
 					}
 				}
 
+				// Handle deferred garbage collection
+				if(bNeedsGarbageCollection)
+				{
+					CollectGarbage(GARBAGE_COLLECTION_KEEPFLAGS);
+					bNeedsGarbageCollection = false;
+				}
+
+				// Replacing references inside already loaded objects could cause rendering issues, so globally detach all components from their scenes for now
+				FGlobalComponentReregisterContext ReregisterContext;
+
 				// UserDefinedStructs (probably all SctiptStructs) should be replaced with the FallbackStruct
 				{
 					TArray<UObject*> UDStructToReplace;
