Do not prompt to reset the transaction buffer during package cleanup after object deletion.

diff --git a/Engine/Source/Editor/UnrealEd/Private/ObjectTools.cpp b/Engine/Source/Editor/UnrealEd/Private/ObjectTools.cpp
--- a/Engine/Source/Editor/UnrealEd/Private/ObjectTools.cpp
+++ b/Engine/Source/Editor/UnrealEd/Private/ObjectTools.cpp
@@ -1361,17 +1361,10 @@ namespace ObjectTools
 					bIsReferenced = IsReferenced(Package, GARBAGE_COLLECTION_KEEPFLAGS, true, &FoundReferences);
 					GEditor->Trans->EnableObjectSerialization();
 
-					// only ref to this object is the transaction buffer - let the user choose whether to clear the undo buffer
+					// only ref to this object is the transaction buffer, clear the transaction buffer
 					if ( !bIsReferenced )
 					{
-						if ( EAppReturnType::Yes == FMessageDialog::Open(EAppMsgType::YesNo, NSLOCTEXT("UnrealEd", "ResetUndoBufferForObjectDeletionPrompt", "The only reference to this object is the undo history.  In order to delete this object, you must clear all undo history - would you like to clear undo history?")) )
-						{
-							GEditor->Trans->Reset(NSLOCTEXT("UnrealEd", "DeleteSelectedItem", "Delete Selected Item"));
-						}
-						else
-						{
-							bIsReferenced = true;
-						}
+						GEditor->Trans->Reset(NSLOCTEXT("UnrealEd", "DeleteSelectedItem", "Delete Selected Item"));
 					}
 				}
 			}
