Fixed GatherTextFromSource commandlet gathering text in its escaped form - it now unescapes gathered text. Fixed escape/unescape handling in PO file import/export. Fixed escaping of double quotes when storing manifest/archive strings from their data structures into a JSON object. The manifest/archive JSON serializers already properly escape text consistently before writing to file.

diff --git a/Engine/Config/Localization/RepairData.ini b/Engine/Config/Localization/RepairData.ini
new file mode 100644
--- /dev/null
+++ b/Engine/Config/Localization/RepairData.ini
@@ -0,0 +1,67 @@
+[CommonSettings]
+CulturesToGenerate=en
+CulturesToGenerate=es
+CulturesToGenerate=ja
+CulturesToGenerate=ko
+CulturesToGenerate=pt-BR
+CulturesToGenerate=zh-Hans-CN
+
+;Repair localization data for Engine
+[GatherTextStep0]
+CommandletClass=RepairLocalizationData
+DestinationPath=./Content/Localization/Engine
+ManifestName=Engine.manifest
+ArchiveName=Engine.archive
+
+;Write Localized Text Resource for Engine
+[GatherTextStep1]
+CommandletClass=GenerateTextLocalizationResource
+SourcePath=./Content/Localization/Engine
+DestinationPath=./Content/Localization/Engine
+ManifestName=Engine.manifest
+ResourceName=Engine.locres
+
+;Repair localization data for Editor
+[GatherTextStep2]
+CommandletClass=RepairLocalizationData
+DestinationPath=./Content/Localization/Editor
+ManifestName=Editor.manifest
+ArchiveName=Editor.archive
+
+;Write Localized Text Resource for Editor
+[GatherTextStep3]
+CommandletClass=GenerateTextLocalizationResource
+SourcePath=./Content/Localization/Editor
+DestinationPath=./Content/Localization/Editor
+ManifestName=Editor.manifest
+ResourceName=Editor.locres
+
+;Repair localization data for Editor Tool Tips
+[GatherTextStep4]
+CommandletClass=RepairLocalizationData
+DestinationPath=./Content/Localization/ToolTips
+ManifestName=ToolTips.manifest
+ArchiveName=ToolTips.archive
+
+;Write Localized Text Resource for Editor Tool Tips
+[GatherTextStep5]
+CommandletClass=GenerateTextLocalizationResource
+SourcePath=./Content/Localization/ToolTips
+DestinationPath=./Content/Localization/ToolTips
+ManifestName=ToolTips.manifest
+ResourceName=ToolTips.locres
+
+;Repair localization data for Editor Property Names
+[GatherTextStep6]
+CommandletClass=RepairLocalizationData
+DestinationPath=./Content/Localization/PropertyNames
+ManifestName=PropertyNames.manifest
+ArchiveName=PropertyNames.archive
+
+;Write Localized Text Resource for Editor Property Names
+[GatherTextStep7]
+CommandletClass=GenerateTextLocalizationResource
+SourcePath=./Content/Localization/PropertyNames
+DestinationPath=./Content/Localization/PropertyNames
+ManifestName=PropertyNames.manifest
+ResourceName=PropertyNames.locres
\ No newline at end of file
diff --git a/Engine/Source/Editor/UnrealEd/Classes/Commandlets/RepairLocalizationDataCommandlet.h b/Engine/Source/Editor/UnrealEd/Classes/Commandlets/RepairLocalizationDataCommandlet.h
new file mode 100644
--- /dev/null
+++ b/Engine/Source/Editor/UnrealEd/Classes/Commandlets/RepairLocalizationDataCommandlet.h
@@ -0,0 +1,18 @@
+// Copyright 1998-2014 Epic Games, Inc. All Rights Reserved.
+
+#pragma once
+
+#include "Commandlets/GatherTextCommandletBase.h"
+#include "RepairLocalizationDataCommandlet.generated.h"
+
+
+UCLASS()
+class URepairLocalizationDataCommandlet : public UGatherTextCommandletBase
+{
+	GENERATED_UCLASS_BODY()
+
+public:
+	// Begin UCommandlet Interface
+	virtual int32 Main(const FString& Params) override;
+	// End UCommandlet Interface
+};
\ No newline at end of file
diff --git a/Engine/Source/Editor/UnrealEd/Private/Commandlets/GatherTextFromSourceCommandlet.cpp b/Engine/Source/Editor/UnrealEd/Private/Commandlets/GatherTextFromSourceCommandlet.cpp
--- a/Engine/Source/Editor/UnrealEd/Private/Commandlets/GatherTextFromSourceCommandlet.cpp
+++ b/Engine/Source/Editor/UnrealEd/Private/Commandlets/GatherTextFromSourceCommandlet.cpp
@@ -524,7 +524,7 @@ bool UGatherTextFromSourceCommandlet::FSourceFileParseContext::AddManifestText(
 		*Filename, 
 		LineNumber, 
 		*LineText);
-	FLocItem Source( SourceText );
+	FLocItem Source( SourceText.ReplaceEscapedCharWithChar() );
 	return ManifestInfo->AddEntry(EntryDescription, InNamespace, Source, Context);
 }
 
diff --git a/Engine/Source/Editor/UnrealEd/Private/Commandlets/InternationalizationExportCommandlet.cpp b/Engine/Source/Editor/UnrealEd/Private/Commandlets/InternationalizationExportCommandlet.cpp
--- a/Engine/Source/Editor/UnrealEd/Private/Commandlets/InternationalizationExportCommandlet.cpp
+++ b/Engine/Source/Editor/UnrealEd/Private/Commandlets/InternationalizationExportCommandlet.cpp
@@ -181,23 +181,103 @@ static const TMap< FString, FString > POCulturePluralForms = TMapBuilder< FStrin
 */
 FString ConditionArchiveStrForPo( const FString& InStr )
 {
-	FString Result = InStr;
-	Result.ReplaceInline(TEXT("\\'"), TEXT("\\\\'"));
-	Result.ReplaceInline(TEXT("\\x"), TEXT("\\\\x"));
-	Result.ReplaceInline(TEXT("\\u"), TEXT("\\\\u"));
-	Result.ReplaceInline(TEXT("\r"), TEXT("\\r"));
-	Result.ReplaceInline(TEXT("\n"), TEXT("\\n"));
-
-
+	FString Result;
+	for (const TCHAR C : InStr)
+	{
+		switch(C)
+		{
+		case TEXT('\\'):	Result += TEXT("\\\\");	break;
+		case TEXT('"'):		Result += TEXT("\\\"");	break;
+		case TEXT('\r'):	Result += TEXT("\\r");	break;
+		case TEXT('\n'):	Result += TEXT("\\n");	break;
+		case TEXT('\t'):	Result += TEXT("\\t");	break;
+		default:			Result += C;			break;
+		}
+	}
 	return Result;
 }
 
 FString ConditionPoStringForArchive( const FString& InStr )
 {
-	FString Result = InStr;
-	Result.ReplaceInline(TEXT("\\\\u"), TEXT("\\u"));
-	Result.ReplaceInline(TEXT("\\\\x"), TEXT("\\x"));
-	Result.ReplaceInline(TEXT("\\\\'"), TEXT("\\'"));
+	FString Result;
+	FString EscapeSequenceBuffer;
+
+	auto HandleEscapeSequenceBuffer = [&]()
+	{
+		// Insert unescaped sequence if needed.
+		if(!EscapeSequenceBuffer.IsEmpty())
+		{
+			bool EscapeSequenceIdentified = true;
+
+			// Identify escape sequence
+			TCHAR UnescapedCharacter = 0;
+			if(EscapeSequenceBuffer == TEXT("\\\\"))
+			{
+				UnescapedCharacter = '\\';
+			}
+			else if(EscapeSequenceBuffer == TEXT("\\\""))
+			{
+				UnescapedCharacter = '"';
+			}
+			else if(EscapeSequenceBuffer == TEXT("\\r"))
+			{
+				UnescapedCharacter = '\r';
+			}
+			else if(EscapeSequenceBuffer == TEXT("\\n"))
+			{
+				UnescapedCharacter = '\n';
+			}
+			else if(EscapeSequenceBuffer == TEXT("\\t"))
+			{
+				UnescapedCharacter = '\t';
+			}
+			else
+			{
+				EscapeSequenceIdentified = false;
+			}
+
+			// If identified, append the processed sequence as the unescaped character.
+			if(EscapeSequenceIdentified)
+			{
+				Result += UnescapedCharacter;
+			}
+			// If it was not identified, preserve the escape sequence and append it.
+			else
+			{
+				Result += EscapeSequenceBuffer;
+			}
+			// Either way, we've appended something based on the buffer and it should be reset.
+			EscapeSequenceBuffer.Empty();
+		}
+	};
+
+	for (const TCHAR C : InStr)
+	{
+		// Not in an escape sequence.
+		if(EscapeSequenceBuffer.IsEmpty())
+		{
+			// Regular character, just copy over.
+			if(C != TEXT('\\'))
+			{
+				Result += C;
+			}
+			// Start of an escape sequence, put in escape sequence buffer.
+			else
+			{
+				EscapeSequenceBuffer += C;
+			}
+		}
+		// If already in an escape sequence.
+		else
+		{
+			// Append to escape sequence buffer.
+			EscapeSequenceBuffer += C;
+
+			HandleEscapeSequenceBuffer();
+		}
+	}
+	// Catch any trailing backslashes.
+	HandleEscapeSequenceBuffer();
 	return Result;
 }
 
diff --git a/Engine/Source/Editor/UnrealEd/Private/Commandlets/RepairLocalizationDataCommandlet.cpp b/Engine/Source/Editor/UnrealEd/Private/Commandlets/RepairLocalizationDataCommandlet.cpp
new file mode 100644
--- /dev/null
+++ b/Engine/Source/Editor/UnrealEd/Private/Commandlets/RepairLocalizationDataCommandlet.cpp
@@ -0,0 +1,303 @@
+// Copyright 1998-2014 Epic Games, Inc. All Rights Reserved.
+
+#include "UnrealEd.h"
+#include "Commandlets/RepairLocalizationDataCommandlet.h"
+#include "InternationalizationManifest.h"
+#include "InternationalizationArchive.h"
+#include "Regex.h"
+
+DEFINE_LOG_CATEGORY_STATIC(LogRepairLocalizationDataCommandlet, Log, All);
+
+void RepairManifestAndArchives(TSharedRef<FInternationalizationManifest> Manifest, TArray< TSharedRef<FInternationalizationArchive> > Archives)
+{
+	// Update source text if this manifest was saved before the escape sequence fixes.
+	if (Manifest->GetFormatVersion() < FInternationalizationManifest::EFormatVersion::EscapeFixes)
+	{
+		TManifestEntryBySourceTextContainer::TConstIterator Iterator = Manifest->GetEntriesBySourceTextIterator();
+		for(Iterator; Iterator; ++Iterator)
+		{
+			const FString& Source = Iterator.Key();
+			const TSharedRef<FManifestEntry>& ManifestEntry = Iterator.Value();
+
+			// Find double quotes and unescape them once.
+			FLocItem UpdatedSource(ManifestEntry->Source);
+			UpdatedSource.Text.ReplaceInline(TEXT("\\\""), TEXT("\""));
+
+			TSharedRef<FManifestEntry> UpdatedManifestEntry = MakeShareable(new FManifestEntry(ManifestEntry->Namespace, UpdatedSource));
+			UpdatedManifestEntry->Contexts = ManifestEntry->Contexts;
+
+			Manifest->UpdateEntry(ManifestEntry, UpdatedManifestEntry);
+		}
+	}
+
+	for (const TSharedRef<FInternationalizationArchive>& Archive : Archives)
+	{
+		// Update source text if this archive was saved before the escape sequence fixes.
+		if (Archive->GetFormatVersion() < FInternationalizationManifest::EFormatVersion::EscapeFixes)
+		{
+			TArchiveEntryContainer::TConstIterator Iterator = Archive->GetEntryIterator();
+
+			for (Iterator; Iterator; ++Iterator)
+			{
+				const FString& Source = Iterator.Key();
+				const TSharedRef<FArchiveEntry>& ArchiveEntry = Iterator.Value();
+
+				// Find double quotes and unescape them once.
+				FLocItem UpdatedSource(ArchiveEntry->Source);
+				UpdatedSource.Text.ReplaceInline(TEXT("\\\""), TEXT("\""));
+
+				// Find double quotes and unescape them once.
+				FLocItem UpdatedTranslation(ArchiveEntry->Translation);
+				UpdatedTranslation.Text.ReplaceInline(TEXT("\\\""), TEXT("\""));
+
+				TSharedRef<FArchiveEntry> UpdatedArchiveEntry = MakeShareable(new FArchiveEntry(ArchiveEntry->Namespace, UpdatedSource, UpdatedTranslation, ArchiveEntry->KeyMetadataObj, ArchiveEntry->bIsOptional));
+
+				Archive->UpdateEntry(ArchiveEntry, UpdatedArchiveEntry);
+			}
+		}
+	}
+
+	const FRegexPattern Pattern(TEXT(".* - line \\d+"));
+
+	TManifestEntryBySourceTextContainer::TConstIterator Iterator = Manifest->GetEntriesBySourceTextIterator();
+	for(Iterator; Iterator; ++Iterator)
+	{
+		const FString& Source = Iterator.Key();
+		const TSharedRef<FManifestEntry>& ManifestEntry = Iterator.Value();
+
+		// Identify if this entry comes from source text only.
+		bool AreAllContextsFromSource = ManifestEntry->Contexts.Num() > 0;
+		for (const FContext& Context : ManifestEntry->Contexts)
+		{
+			FRegexMatcher Matcher(Pattern, Context.SourceLocation);
+			AreAllContextsFromSource = Matcher.FindNext();
+
+			if(!AreAllContextsFromSource)
+			{
+				break;
+			}
+		}
+
+		// No updates needed for this entry if it isn't all from source.
+		if(!AreAllContextsFromSource)
+		{
+			continue;
+		}
+
+		for (const TSharedRef<FInternationalizationArchive>& Archive : Archives)
+		{
+			// Update source text if this manifest was saved before the escape sequence fixes.
+			if (Archive->GetFormatVersion() < FInternationalizationManifest::EFormatVersion::EscapeFixes)
+			{
+				const TSharedPtr<FArchiveEntry> ArchiveEntry = Archive->FindEntryBySource(ManifestEntry->Namespace, ManifestEntry->Source, /*KeyMetadataObj*/ nullptr);
+				if (!ArchiveEntry.IsValid())
+				{
+					continue;
+				}
+
+				// Replace escape sequences with their associated character, once.
+				FLocItem UpdatedSource(ArchiveEntry->Source);
+				UpdatedSource.Text = UpdatedSource.Text.ReplaceEscapedCharWithChar();
+
+				// Replace escape sequences with their associated character, once.
+				FLocItem UpdatedTranslation(ArchiveEntry->Translation);
+				UpdatedTranslation.Text = UpdatedTranslation.Text.ReplaceEscapedCharWithChar();
+
+				TSharedRef<FArchiveEntry> UpdatedArchiveEntry = MakeShareable(new FArchiveEntry(ArchiveEntry->Namespace, UpdatedSource, UpdatedTranslation, ArchiveEntry->KeyMetadataObj, ArchiveEntry->bIsOptional));
+
+				Archive->UpdateEntry(ArchiveEntry.ToSharedRef(), UpdatedArchiveEntry);
+			}
+		}
+
+		// Update source text if this manifest was saved before the escape sequence fixes.
+		if (Manifest->GetFormatVersion() < FInternationalizationManifest::EFormatVersion::EscapeFixes)
+		{
+			// Replace escape sequences with their associated character, once.
+			FLocItem UpdatedSource(ManifestEntry->Source);
+			UpdatedSource.Text = UpdatedSource.Text.ReplaceEscapedCharWithChar();
+
+			TSharedRef<FManifestEntry> UpdatedManifestEntry = MakeShareable(new FManifestEntry(ManifestEntry->Namespace, UpdatedSource));
+			UpdatedManifestEntry->Contexts = ManifestEntry->Contexts;
+
+			Manifest->UpdateEntry(ManifestEntry, UpdatedManifestEntry);
+		}
+	}
+
+	Manifest->SetFormatVersion(FInternationalizationManifest::EFormatVersion::Latest);
+	for (const TSharedRef<FInternationalizationArchive>& Archive : Archives)
+	{
+		Archive->SetFormatVersion(FInternationalizationArchive::EFormatVersion::Latest);
+	}
+}
+
+URepairLocalizationDataCommandlet::URepairLocalizationDataCommandlet(const class FPostConstructInitializeProperties& PCIP)
+	: Super(PCIP)
+{
+}
+
+int32 URepairLocalizationDataCommandlet::Main(const FString& Params)
+{
+	FInternationalization& I18N = FInternationalization::Get();
+
+	// Parse command line.
+	TArray<FString> Tokens;
+	TArray<FString> Switches;
+	TMap<FString, FString> ParamVals;
+	UCommandlet::ParseCommandLine(*Params, Tokens, Switches, ParamVals);
+
+	//Set config file
+	const FString* ParamVal = ParamVals.Find(FString(TEXT("Config")));
+	FString GatherTextConfigPath;
+
+	if ( ParamVal )
+	{
+		GatherTextConfigPath = *ParamVal;
+	}
+	else
+	{
+		UE_LOG(LogRepairLocalizationDataCommandlet, Error, TEXT("No config specified."));
+		return -1;
+	}
+
+	//Set config section
+	ParamVal = ParamVals.Find(FString(TEXT("Section")));
+	FString SectionName;
+
+	if ( ParamVal )
+	{
+		SectionName = *ParamVal;
+	}
+	else
+	{
+		UE_LOG(LogRepairLocalizationDataCommandlet, Error, TEXT("No config section specified."));
+		return -1;
+	}
+
+	// Get destination path.
+	FString DestinationPath;
+	if( !GetConfigString( *SectionName, TEXT("DestinationPath"), DestinationPath, GatherTextConfigPath ) )
+	{
+		UE_LOG( LogRepairLocalizationDataCommandlet, Error, TEXT("No destination path specified.") );
+		return -1;
+	}
+
+	if (FPaths::IsRelative(DestinationPath))
+	{
+		if (!FPaths::GameDir().IsEmpty())
+		{
+			DestinationPath = FPaths::Combine( *( FPaths::GameDir() ), *DestinationPath );
+		}
+		else
+		{
+			DestinationPath = FPaths::Combine( *( FPaths::EngineDir() ), *DestinationPath );
+		}
+	}
+
+	// Get manifest name.
+	FString ManifestName;
+	if( !GetConfigString( *SectionName, TEXT("ManifestName"), ManifestName, GatherTextConfigPath ) )
+	{
+		UE_LOG( LogRepairLocalizationDataCommandlet, Error, TEXT("No manifest name specified.") );
+		return -1;
+	}
+
+	// Get archive name.
+	FString ArchiveName;
+	if( !( GetConfigString(* SectionName, TEXT("ArchiveName"), ArchiveName, GatherTextConfigPath ) ) )
+	{
+		UE_LOG(LogGenerateArchiveCommandlet, Error, TEXT("No archive name specified."));
+		return -1;
+	}
+
+	// Get cultures to generate.
+	TArray<FString> CulturesToGenerate;
+	GetConfigArray(*SectionName, TEXT("CulturesToGenerate"), CulturesToGenerate, GatherTextConfigPath);
+
+	if( CulturesToGenerate.Num() == 0 )
+	{
+		UE_LOG(LogRepairLocalizationDataCommandlet, Error, TEXT("No cultures specified for generation."));
+		return -1;
+	}
+
+	for(int32 i = 0; i < CulturesToGenerate.Num(); ++i)
+	{
+		if( I18N.GetCulture( CulturesToGenerate[i] ).IsValid() )
+		{
+			UE_LOG(LogRepairLocalizationDataCommandlet, Verbose, TEXT("Specified culture is not a valid runtime culture, but may be a valid base language: %s"), *(CulturesToGenerate[i]) );
+		}
+	}
+
+
+	//////////////////////////////////////////////////////////////////////////
+
+
+	// Read the damaged manifest.
+	const FString ManifestFilePath = DestinationPath / ManifestName;
+	const TSharedPtr<FJsonObject> ManifestJsonObject = ReadJSONTextFile( ManifestFilePath );
+	if( !ManifestJsonObject.IsValid() )
+	{
+		UE_LOG(LogRepairLocalizationDataCommandlet, Error, TEXT("Could not read manifest file %s."), *ManifestFilePath);
+		return -1;
+	}
+	FInternationalizationManifestJsonSerializer ManifestSerializer;
+	const TSharedRef< FInternationalizationManifest > InternationalizationManifest = MakeShareable( new FInternationalizationManifest );
+	ManifestSerializer.DeserializeManifest( ManifestJsonObject.ToSharedRef(), InternationalizationManifest );
+
+	// Read the damaged archives.
+	TArray< TSharedRef<FInternationalizationArchive> > InternationalizationArchives;
+	TArray<FString> ArchiveCultures;
+	for(int32 Culture = 0; Culture < CulturesToGenerate.Num(); Culture++)
+	{
+		// Read in any existing archive for this culture.
+		const FString CulturePath = DestinationPath / CulturesToGenerate[Culture];
+		const FString ArchiveFilePath = CulturePath / ArchiveName;
+		if( FPaths::FileExists(ArchiveFilePath) )
+		{
+			TSharedPtr< FJsonObject > ArchiveJsonObject = ReadJSONTextFile( ArchiveFilePath );
+			if( !ArchiveJsonObject.IsValid() )
+			{
+				UE_LOG(LogRepairLocalizationDataCommandlet, Error, TEXT("Could not read archive file %s."), *ArchiveFilePath);
+				return -1;
+			}
+			FInternationalizationArchiveJsonSerializer ArchiveSerializer;
+			const TSharedRef< FInternationalizationArchive > InternationalizationArchive = MakeShareable( new FInternationalizationArchive );
+			ArchiveSerializer.DeserializeArchive( ArchiveJsonObject.ToSharedRef(), InternationalizationArchive );
+
+			InternationalizationArchives.Add(InternationalizationArchive);
+			ArchiveCultures.Add(CulturesToGenerate[Culture]);
+		}
+	}
+
+	// Repair.
+	RepairManifestAndArchives(InternationalizationManifest, InternationalizationArchives);
+
+	// Write the repaired manifests.
+	TSharedRef< FJsonObject > OutputManifestJsonObj = MakeShareable( new FJsonObject );
+	ManifestSerializer.SerializeManifest( InternationalizationManifest, OutputManifestJsonObj );
+	if (!WriteJSONToTextFile( OutputManifestJsonObj, DestinationPath / ManifestName, SourceControlInfo ))
+	{
+		UE_LOG( LogRepairLocalizationDataCommandlet, Error,TEXT("Failed to write manifest to %s."), *DestinationPath );				
+		return -1;
+	}
+
+	// Write the repaired archives.
+	{
+		int Index = 0;
+		for (const TSharedRef<FInternationalizationArchive>& InternationalizationArchive : InternationalizationArchives)
+		{
+			TSharedRef< FJsonObject > OutputArchiveJsonObj = MakeShareable( new FJsonObject );
+			FInternationalizationArchiveJsonSerializer ArchiveSerializer;
+			ArchiveSerializer.SerializeArchive( InternationalizationArchive, OutputArchiveJsonObj );
+
+			int32 Culture = 0;
+			if (!WriteJSONToTextFile( OutputArchiveJsonObj, (DestinationPath / *ArchiveCultures[Index] / ArchiveName), SourceControlInfo ))
+			{
+				UE_LOG( LogRepairLocalizationDataCommandlet, Error,TEXT("Failed to write archive to %s."), *DestinationPath );				
+				return -1;
+			}
+			++Index;
+		}
+	}
+
+	return 0;
+}
\ No newline at end of file
diff --git a/Engine/Source/Runtime/Core/Private/Internationalization/InternationalizationArchive.cpp b/Engine/Source/Runtime/Core/Private/Internationalization/InternationalizationArchive.cpp
--- a/Engine/Source/Runtime/Core/Private/Internationalization/InternationalizationArchive.cpp
+++ b/Engine/Source/Runtime/Core/Private/Internationalization/InternationalizationArchive.cpp
@@ -84,6 +84,12 @@ TSharedPtr< FArchiveEntry > FInternationalizationArchive::FindEntryBySource( con
 	return NULL;
 }
 
+void FInternationalizationArchive::UpdateEntry(const TSharedRef<FArchiveEntry>& OldEntry, const TSharedRef<FArchiveEntry>& NewEntry)
+{
+	EntriesBySourceText.RemoveSingle(OldEntry->Source.Text, OldEntry);
+	EntriesBySourceText.Add(NewEntry->Source.Text, NewEntry);
+}
+
 TArchiveEntryContainer::TConstIterator FInternationalizationArchive::GetEntryIterator() const
 {
 	return EntriesBySourceText.CreateConstIterator();
diff --git a/Engine/Source/Runtime/Core/Private/Internationalization/InternationalizationArchiveJsonSerializer.cpp b/Engine/Source/Runtime/Core/Private/Internationalization/InternationalizationArchiveJsonSerializer.cpp
--- a/Engine/Source/Runtime/Core/Private/Internationalization/InternationalizationArchiveJsonSerializer.cpp
+++ b/Engine/Source/Runtime/Core/Private/Internationalization/InternationalizationArchiveJsonSerializer.cpp
@@ -1,4 +1,4 @@
-// Copyright 1998-2014 Epic Games, Inc. All Rights Reserved.
+﻿// Copyright 1998-2014 Epic Games, Inc. All Rights Reserved.
 
 #include "Core.h"
 #include "InternationalizationArchive.h"
@@ -8,6 +8,7 @@
 
 DEFINE_LOG_CATEGORY_STATIC(LogInternationalizationArchiveSerializer, Log, All);
 
+const FString FInternationalizationArchiveJsonSerializer::TAG_FORMATVERSION = TEXT("FormatVersion");
 const FString FInternationalizationArchiveJsonSerializer::TAG_NAMESPACE = TEXT("Namespace");
 const FString FInternationalizationArchiveJsonSerializer::TAG_CHILDREN = TEXT("Children");
 const FString FInternationalizationArchiveJsonSerializer::TAG_SUBNAMESPACES = TEXT("Subnamespaces");
@@ -122,6 +123,15 @@ bool FInternationalizationArchiveJsonSerializer::SerializeArchive( TSharedRef< c
 
 bool FInternationalizationArchiveJsonSerializer::DeserializeInternal( TSharedRef< FJsonObject > InJsonObj, TSharedRef< FInternationalizationArchive > InternationalizationArchive )
 {
+	if( InJsonObj->HasField( TAG_FORMATVERSION ) )
+	{
+		InternationalizationArchive->FormatVersion = static_cast<int>(InJsonObj->GetNumberField( TAG_FORMATVERSION ));
+	}
+	else
+	{
+		InternationalizationArchive->FormatVersion = FInternationalizationArchive::EFormatVersion::Initial;
+	}
+
 	return JsonObjToArchive( InJsonObj, TEXT(""), InternationalizationArchive );
 }
 
@@ -137,6 +147,9 @@ bool FInternationalizationArchiveJsonSerializer::SerializeInternal( TSharedRef<
 	// Clear anything that may be in the JSON object
 	JsonObj->Values.Empty();
 
+	// Set format version.
+	JsonObj->SetNumberField(TAG_FORMATVERSION, static_cast<double>(InInternationalizationArchive->FormatVersion));
+
 	// Setup the JSON object using the structured data created
 	StructuredDataToJsonObj( RootElement, JsonObj );
 	return true;
@@ -365,9 +378,8 @@ void FInternationalizationArchiveJsonSerializer::StructuredDataToJsonObj( TShare
 		const TSharedPtr< FArchiveEntry > Entry = *Iter;
 		TSharedPtr< FJsonObject > EntryNode = MakeShareable( new FJsonObject );
 
-		// Add escapes for special chars - doesn't do backslash
-		FString ProcessedSourceText = Entry->Source.Text.ReplaceQuotesWithEscapedQuotes();
-		FString ProcessedTranslation = Entry->Translation.Text.ReplaceQuotesWithEscapedQuotes();
+		FString ProcessedSourceText = Entry->Source.Text;
+		FString ProcessedTranslation = Entry->Translation.Text;
 
 		TSharedPtr< FJsonObject > SourceNode;
 		if( Entry->Source.MetadataObj.IsValid() )
diff --git a/Engine/Source/Runtime/Core/Private/Internationalization/InternationalizationManifest.cpp b/Engine/Source/Runtime/Core/Private/Internationalization/InternationalizationManifest.cpp
--- a/Engine/Source/Runtime/Core/Private/Internationalization/InternationalizationManifest.cpp
+++ b/Engine/Source/Runtime/Core/Private/Internationalization/InternationalizationManifest.cpp
@@ -263,6 +263,21 @@ TSharedPtr< FManifestEntry > FInternationalizationManifest::FindEntryByContext(
 	return NULL;
 }
 
+void FInternationalizationManifest::UpdateEntry(const TSharedRef<FManifestEntry>& OldEntry, TSharedRef<FManifestEntry>& NewEntry)
+{
+	for (const FContext& Context : OldEntry->Contexts)
+	{
+		EntriesByContextId.RemoveSingle(Context.Key, OldEntry);
+	}
+	for (const FContext& Context : NewEntry->Contexts)
+	{
+		EntriesByContextId.Add(Context.Key, NewEntry);
+	}
+
+	EntriesBySourceText.RemoveSingle(OldEntry->Source.Text, OldEntry);
+	EntriesBySourceText.Add(NewEntry->Source.Text, NewEntry);
+}
+
 FContext* FManifestEntry::FindContext( const FString& ContextKey, const TSharedPtr<FLocMetadataObject>& KeyMetadata /*= NULL */ )
 {
 	for( auto ContextIter = Contexts.CreateIterator(); ContextIter; ++ContextIter )
@@ -283,5 +298,4 @@ FContext* FManifestEntry::FindContext( const FString& ContextKey, const TSharedP
 		}
 	}
 	return NULL;
-}
-
+}
\ No newline at end of file
diff --git a/Engine/Source/Runtime/Core/Private/Internationalization/InternationalizationManifestJsonSerializer.cpp b/Engine/Source/Runtime/Core/Private/Internationalization/InternationalizationManifestJsonSerializer.cpp
--- a/Engine/Source/Runtime/Core/Private/Internationalization/InternationalizationManifestJsonSerializer.cpp
+++ b/Engine/Source/Runtime/Core/Private/Internationalization/InternationalizationManifestJsonSerializer.cpp
@@ -1,4 +1,4 @@
-// Copyright 1998-2014 Epic Games, Inc. All Rights Reserved.
+﻿// Copyright 1998-2014 Epic Games, Inc. All Rights Reserved.
 
 #include "Core.h"
 #include "InternationalizationManifest.h"
@@ -9,6 +9,7 @@
 
 DEFINE_LOG_CATEGORY_STATIC(LogInternationalizationManifestSerializer, Log, All);
 
+const FString FInternationalizationManifestJsonSerializer::TAG_FORMATVERSION = TEXT("FormatVersion");
 const FString FInternationalizationManifestJsonSerializer::TAG_NAMESPACE = TEXT("Namespace");
 const FString FInternationalizationManifestJsonSerializer::TAG_CHILDREN = TEXT("Children");
 const FString FInternationalizationManifestJsonSerializer::TAG_SUBNAMESPACES = TEXT("Subnamespaces");
@@ -115,6 +116,15 @@ bool FInternationalizationManifestJsonSerializer::SerializeManifest( TSharedRef<
 
 bool FInternationalizationManifestJsonSerializer::DeserializeInternal( TSharedRef< FJsonObject > InJsonObj, TSharedRef< FInternationalizationManifest > Manifest )
 {
+	if( InJsonObj->HasField( TAG_FORMATVERSION ) )
+	{
+		Manifest->FormatVersion = static_cast<int>(InJsonObj->GetNumberField( TAG_FORMATVERSION ));
+	}
+	else
+	{
+		Manifest->FormatVersion = FInternationalizationManifest::EFormatVersion::Initial;
+	}
+
 	return JsonObjToManifest(InJsonObj, TEXT(""), Manifest );
 }
 
@@ -132,6 +142,9 @@ bool FInternationalizationManifestJsonSerializer::SerializeInternal( TSharedRef<
 	//Clear out anything that may be in the JSON object
 	JsonObj->Values.Empty();
 
+	// Set format version.
+	JsonObj->SetNumberField(TAG_FORMATVERSION, static_cast<double>(InManifest->FormatVersion));
+
 	// Setup the JSON object using the structured data created
 	StructuredDataToJsonObj( RootElement, JsonObj );
 
@@ -262,7 +275,7 @@ bool FInternationalizationManifestJsonSerializer::JsonObjToManifest( TSharedRef<
 						{
 							UE_LOG( LogInternationalizationManifestSerializer, Warning,TEXT("Could not add JSON entry to the Internationalization manifest: Namespace:%s SourceText:%s SourceData:%s"), 
 								*AccumulatedNamespace, 
-								*SourceText, 
+								*Source.Text, 
 								*FInternationalizationMetaDataJsonSerializer::MetadataToString(Source.MetadataObj) );
 						}
 					}
@@ -406,7 +419,7 @@ void FInternationalizationManifestJsonSerializer::StructuredDataToJsonObj( TShar
 
 		// Add escapes for special chars - doesn't do backslash
 		{
-			FString ProcessedText = Entry->Source.Text.ReplaceQuotesWithEscapedQuotes();
+			FString ProcessedText = Entry->Source.Text;
 			SourceNode->SetStringField(TAG_SOURCE_TEXT, ProcessedText);
 		}
 
diff --git a/Engine/Source/Runtime/Core/Public/Internationalization/InternationalizationArchive.h b/Engine/Source/Runtime/Core/Public/Internationalization/InternationalizationArchive.h
--- a/Engine/Source/Runtime/Core/Public/Internationalization/InternationalizationArchive.h
+++ b/Engine/Source/Runtime/Core/Public/Internationalization/InternationalizationArchive.h
@@ -4,7 +4,7 @@
 
 #include "Internationalization/InternationalizationManifest.h"
 
-class FArchiveEntry
+class CORE_API FArchiveEntry
 {
 public:
 	FArchiveEntry( const FString& InNamespace, const FLocItem& InSource, const FLocItem& InTranslation, TSharedPtr<FLocMetadataObject> InKeyMetadataObj = NULL, bool IsOptional = false );
@@ -37,9 +37,23 @@ typedef TMultiMap< FString, TSharedRef< FArchiveEntry >, FDefaultSetAllocator, F
 class CORE_API FInternationalizationArchive 
 {
 	friend class IInternationalizationArchiveSerializer;
-	friend class FInternationalizationArchivetJsonSerializer;
-
+	friend class FInternationalizationArchiveJsonSerializer;
 public:
+	enum EFormatVersion
+	{
+		Initial = 0,
+		EscapeFixes,
+
+		LatestPlusOne,
+		Latest = LatestPlusOne - 1
+	};
+
+	FInternationalizationArchive()
+		: FormatVersion(static_cast<int32>(EFormatVersion::Latest))
+	{
+
+	}
+
 	TArchiveEntryContainer::TConstIterator GetEntryIterator() const;
 
 	bool AddEntry( const FString& Namespace, const FLocItem& Source, const FLocItem& Translation, const TSharedPtr<FLocMetadataObject> KeyMetadataObj, const bool bOptional );
@@ -54,13 +68,19 @@ public:
 		return EntriesBySourceText.Num();
 	}
 
-private:
-
-	TArchiveEntryContainer EntriesBySourceText;
-};
-
-
-
+	void UpdateEntry(const TSharedRef<FArchiveEntry>& OldEntry, const TSharedRef<FArchiveEntry>& NewEntry);
 
+	void SetFormatVersion(const EFormatVersion Version)
+	{
+		FormatVersion = static_cast<int>(Version);
+	}
 
+	EFormatVersion GetFormatVersion() const
+	{
+		return EFormatVersion(FormatVersion);
+	}
 
+private:
+	int32 FormatVersion;
+	TArchiveEntryContainer EntriesBySourceText;
+};
\ No newline at end of file
diff --git a/Engine/Source/Runtime/Core/Public/Internationalization/InternationalizationArchiveJsonSerializer.h b/Engine/Source/Runtime/Core/Public/Internationalization/InternationalizationArchiveJsonSerializer.h
--- a/Engine/Source/Runtime/Core/Public/Internationalization/InternationalizationArchiveJsonSerializer.h
+++ b/Engine/Source/Runtime/Core/Public/Internationalization/InternationalizationArchiveJsonSerializer.h
@@ -146,6 +146,8 @@ protected:
 
 
 public:
+
+	static const FString TAG_FORMATVERSION;
 	static const FString TAG_NAMESPACE;
 	static const FString TAG_CHILDREN;
 	static const FString TAG_SUBNAMESPACES;
diff --git a/Engine/Source/Runtime/Core/Public/Internationalization/InternationalizationManifest.h b/Engine/Source/Runtime/Core/Public/Internationalization/InternationalizationManifest.h
--- a/Engine/Source/Runtime/Core/Public/Internationalization/InternationalizationManifest.h
+++ b/Engine/Source/Runtime/Core/Public/Internationalization/InternationalizationManifest.h
@@ -1,7 +1,7 @@
 // Copyright 1998-2014 Epic Games, Inc. All Rights Reserved.
 #pragma once
 
-class FLocMetadataObject;
+;class FLocMetadataObject;
 
 struct CORE_API FContext
 {
@@ -112,11 +112,19 @@ typedef TMultiMap< FString, TSharedRef< FManifestEntry >, FDefaultSetAllocator,
 
 class CORE_API FInternationalizationManifest 
 {
-
 public:
+	enum EFormatVersion
+	{
+		Initial = 0,
+		EscapeFixes,
+
+		LatestPlusOne,
+		Latest = LatestPlusOne - 1
+	};
 
 	//Default constructor
 	FInternationalizationManifest( )
+		: FormatVersion(static_cast<int32>(EFormatVersion::Latest))
 	{
 	}
 
@@ -151,11 +159,24 @@ public:
 		return EntriesBySourceText.Num();
 	}
 
+	void UpdateEntry(const TSharedRef<FManifestEntry>& OldEntry, TSharedRef<FManifestEntry>& NewEntry);
+
+	void SetFormatVersion(const EFormatVersion Version)
+	{
+		FormatVersion = static_cast<int>(Version);
+	}
+
+	EFormatVersion GetFormatVersion() const
+	{
+		return EFormatVersion(FormatVersion);
+	}
+
 	friend class IInternationalizationManifestSerializer;
 	friend class FInternationalizationManifestJsonSerializer;
 
 private:
 
+	int32 FormatVersion;
 	TManifestEntryBySourceTextContainer EntriesBySourceText;
 	TManifestEntryByContextIdContainer EntriesByContextId;
 };
diff --git a/Engine/Source/Runtime/Core/Public/Internationalization/InternationalizationManifestJsonSerializer.h b/Engine/Source/Runtime/Core/Public/Internationalization/InternationalizationManifestJsonSerializer.h
--- a/Engine/Source/Runtime/Core/Public/Internationalization/InternationalizationManifestJsonSerializer.h
+++ b/Engine/Source/Runtime/Core/Public/Internationalization/InternationalizationManifestJsonSerializer.h
@@ -144,6 +144,7 @@ protected:
 
 public:
 
+	static const FString TAG_FORMATVERSION;
 	static const FString TAG_NAMESPACE;
 	static const FString TAG_CHILDREN;
 	static const FString TAG_SUBNAMESPACES;
