Fixed typo with enum comparison of wrong enum type and a bad log statement category.

diff --git a/Engine/Source/Editor/UnrealEd/Private/Commandlets/RepairLocalizationDataCommandlet.cpp b/Engine/Source/Editor/UnrealEd/Private/Commandlets/RepairLocalizationDataCommandlet.cpp
--- a/Engine/Source/Editor/UnrealEd/Private/Commandlets/RepairLocalizationDataCommandlet.cpp
+++ b/Engine/Source/Editor/UnrealEd/Private/Commandlets/RepairLocalizationDataCommandlet.cpp
@@ -33,7 +33,7 @@ void RepairManifestAndArchives(TSharedRef<FInternationalizationManifest> Manifes
 	for (const TSharedRef<FInternationalizationArchive>& Archive : Archives)
 	{
 		// Update source text if this archive was saved before the escape sequence fixes.
-		if (Archive->GetFormatVersion() < FInternationalizationManifest::EFormatVersion::EscapeFixes)
+		if (Archive->GetFormatVersion() < FInternationalizationArchive::EFormatVersion::EscapeFixes)
 		{
 			TArchiveEntryContainer::TConstIterator Iterator = Archive->GetEntryIterator();
 
@@ -87,7 +87,7 @@ void RepairManifestAndArchives(TSharedRef<FInternationalizationManifest> Manifes
 		for (const TSharedRef<FInternationalizationArchive>& Archive : Archives)
 		{
 			// Update source text if this manifest was saved before the escape sequence fixes.
-			if (Archive->GetFormatVersion() < FInternationalizationManifest::EFormatVersion::EscapeFixes)
+			if (Archive->GetFormatVersion() < FInternationalizationArchive::EFormatVersion::EscapeFixes)
 			{
 				const TSharedPtr<FArchiveEntry> ArchiveEntry = Archive->FindEntryBySource(ManifestEntry->Namespace, ManifestEntry->Source, /*KeyMetadataObj*/ nullptr);
 				if (!ArchiveEntry.IsValid())
@@ -205,7 +205,7 @@ int32 URepairLocalizationDataCommandlet::Main(const FString& Params)
 	FString ArchiveName;
 	if( !( GetConfigString(* SectionName, TEXT("ArchiveName"), ArchiveName, GatherTextConfigPath ) ) )
 	{
-		UE_LOG(LogGenerateArchiveCommandlet, Error, TEXT("No archive name specified."));
+		UE_LOG(LogRepairLocalizationDataCommandlet, Error, TEXT("No archive name specified."));
 		return -1;
 	}
 
