Resolved possible non-unity build compilation issues in localization data repair commandlet.

diff --git a/Engine/Source/Editor/UnrealEd/Private/Commandlets/RepairLocalizationDataCommandlet.cpp b/Engine/Source/Editor/UnrealEd/Private/Commandlets/RepairLocalizationDataCommandlet.cpp
--- a/Engine/Source/Editor/UnrealEd/Private/Commandlets/RepairLocalizationDataCommandlet.cpp
+++ b/Engine/Source/Editor/UnrealEd/Private/Commandlets/RepairLocalizationDataCommandlet.cpp
@@ -3,7 +3,9 @@
 #include "UnrealEd.h"
 #include "Commandlets/RepairLocalizationDataCommandlet.h"
 #include "InternationalizationManifest.h"
+#include "InternationalizationManifestJsonSerializer.h"
 #include "InternationalizationArchive.h"
+#include "InternationalizationArchiveJsonSerializer.h"
 #include "Regex.h"
 
 DEFINE_LOG_CATEGORY_STATIC(LogRepairLocalizationDataCommandlet, Log, All);
