Launching PIE-V as Client, with a Dedicated Server no longer crashes

#ttp 342070 - CRITICAL: Regression: SHOOTERGAME: EDITOR: CRASH: Launching in editor with dedicated server enabled

#branch UE4

#note was fixed with CL#2236478, however testing the repro steps with latest threw up another unrelated crash, so fixed that instead.

#change with those settings; a listen server was being launched separately, not a dedicated server, so that was modified to behave correctly.

#change the PIE-V instance was crashing as it itself was launching as a dedicated server, instead of a client, so this was modified in order to correctly connect to the server.

diff --git a/Engine/Source/Editor/UnrealEd/Private/PlayLevel.cpp b/Engine/Source/Editor/UnrealEd/Private/PlayLevel.cpp
--- a/Engine/Source/Editor/UnrealEd/Private/PlayLevel.cpp
+++ b/Engine/Source/Editor/UnrealEd/Private/PlayLevel.cpp
@@ -686,9 +686,9 @@ void GetWindowSizeForInstanceType(FIntPoint &WindowSize, const ULevelEditorPlayS
  * @param	WinX			Window X position. This will contain the X position to use for the next window. (Not changed for dedicated server window).
  * @param	WinY			Window Y position. This will contain the X position to use for the next window. (Not changed for dedicated server window).
  * @param	InstanceNum		PIE instance index.
- * @param	IsServer		Is this instance a server. true if so else false.
+ * @param	bIsDedicatedServer	Is this instance a dedicate server. true if so else false.
  */
-FString GenerateCmdLineForNextPieInstance(int32 &WinX, int32 &WinY, int32 &InstanceNum, bool IsServer)
+FString GenerateCmdLineForNextPieInstance(int32 &WinX, int32 &WinY, int32 &InstanceNum, bool bIsDedicatedServer)
 {
 	const ULevelEditorPlaySettings* PlayInSettings = GetDefault<ULevelEditorPlaySettings>();
 	// Get GameSettings INI override
@@ -700,7 +700,7 @@ FString GenerateCmdLineForNextPieInstance(int32 &WinX, int32 &WinY, int32 &Insta
 	//	-Allow saving of config files (since we are giving them an override INI)
 	FString CmdLine = FString::Printf(TEXT("GameUserSettingsINI=%s -MultiprocessSaveConfig %s -MultiprocessOSS "), *GameUserSettingsOverride, *PlayInSettings->AdditionalLaunchOptions);
 
-	if (IsServer && PlayInSettings->PlayNetDedicated)
+	if (bIsDedicatedServer)
 	{
 		// Append dedicated server options
 		CmdLine += TEXT("-server -log ");
@@ -985,11 +985,11 @@ void UEditorEngine::PlayStandaloneLocalPc(FString MapNameOverride, FIntPoint* Wi
 	FString CmdLine;
 	if (WindowPos != NULL)	// If WindowPos == NULL, we're just launching one instance
 	{
-		CmdLine = GenerateCmdLineForNextPieInstance(WindowPos->X, WindowPos->Y, PIENum, bIsServer);
+		CmdLine = GenerateCmdLineForNextPieInstance(WindowPos->X, WindowPos->Y, PIENum, bIsServer && PlayInSettings->PlayNetDedicated);
 	}
 	
-	const FString URLParms = bIsServer == true ? TEXT("?Listen") : FString();
-	
+	const FString URLParms = bIsServer && !PlayInSettings->PlayNetDedicated ? TEXT("?Listen") : FString();
+
 	// select map to play
 	TArray<FString> SavedMapNames;
 	if (MapNameOverride.IsEmpty())
@@ -1906,7 +1906,7 @@ void UEditorEngine::PlayInEditor( UWorld* InWorld, bool bInSimulateInEditor )
 	if (bInSimulateInEditor || (PlayInSettings->PlayNetMode == EPlayNetMode::PIE_Standalone && !bSupportsOnlinePIE) || !PlayInSettings->RunUnderOneProcess)
 	{
 		// Only spawning 1 PIE instance under this process
-		UGameInstance* const GameInstance = CreatePIEGameInstance(0, bInSimulateInEditor, bAnyBlueprintErrors, bStartInSpectatorMode, PlayInSettings->PlayNetDedicated, PIEStartTime);
+		UGameInstance* const GameInstance = CreatePIEGameInstance(0, bInSimulateInEditor, bAnyBlueprintErrors, bStartInSpectatorMode, false, PIEStartTime);
 
 		if (bInSimulateInEditor)
 		{
