Exposed different Blueprint editor contexts for tutorials

Added new GetToolkitContextFName() function to IToolKit, most implementations of which return GetToolkitFName().
Tutorials button now exposes a different context depending on the type of Blueprint being edited. Updated BP-editor derived classes so that 'downstream' editors (e.g. Persona) are not affected.

Also edited padding & icon for tutorials.

diff --git a/Engine/Source/Editor/GameplayAbilitiesEditor/Private/GameplayAbilitiesEditor.cpp b/Engine/Source/Editor/GameplayAbilitiesEditor/Private/GameplayAbilitiesEditor.cpp
--- a/Engine/Source/Editor/GameplayAbilitiesEditor/Private/GameplayAbilitiesEditor.cpp
+++ b/Engine/Source/Editor/GameplayAbilitiesEditor/Private/GameplayAbilitiesEditor.cpp
@@ -53,6 +53,11 @@ void FGameplayAbilitiesEditor::EnsureGameplayAbilityBlueprintIsUpToDate(UBluepri
 #endif
 }
 
+FName FGameplayAbilitiesEditor::GetToolkitContextFName() const
+{
+	return GetToolkitFName();
+}
+
 FName FGameplayAbilitiesEditor::GetToolkitFName() const
 {
 	return FName("GameplayAbilitiesEditor");
diff --git a/Engine/Source/Editor/GameplayAbilitiesEditor/Public/GameplayAbilitiesEditor.h b/Engine/Source/Editor/GameplayAbilitiesEditor/Public/GameplayAbilitiesEditor.h
--- a/Engine/Source/Editor/GameplayAbilitiesEditor/Public/GameplayAbilitiesEditor.h
+++ b/Engine/Source/Editor/GameplayAbilitiesEditor/Public/GameplayAbilitiesEditor.h
@@ -41,6 +41,7 @@ public:
 
 public:
 	// IToolkit interface
+	virtual FName GetToolkitContextFName() const override;
 	virtual FName GetToolkitFName() const override;
 	virtual FText GetBaseToolkitName() const override;
 	virtual FText GetToolkitName() const override;
diff --git a/Engine/Source/Editor/Kismet/Private/BlueprintEditor.cpp b/Engine/Source/Editor/Kismet/Private/BlueprintEditor.cpp
--- a/Engine/Source/Editor/Kismet/Private/BlueprintEditor.cpp
+++ b/Engine/Source/Editor/Kismet/Private/BlueprintEditor.cpp
@@ -5954,6 +5954,26 @@ FName FBlueprintEditor::GetToolkitFName() const
 	return FName("BlueprintEditor");
 }
 
+FName FBlueprintEditor::GetToolkitContextFName() const
+{
+	if(GetBlueprintObj())
+	{
+		switch (GetBlueprintObj()->BlueprintType)
+		{
+		case BPTYPE_Normal:
+			return FName("BlueprintEditor");
+		case BPTYPE_MacroLibrary:
+			return FName("BlueprintEditor.MacroLibrary");
+		case BPTYPE_Interface:
+			return FName("BlueprintEditor.Interface");
+		case BPTYPE_LevelScript:
+			return FName("BlueprintEditor.LevelScript");
+		}
+	}
+
+	return FName("BlueprintEditor");
+}
+
 FText FBlueprintEditor::GetBaseToolkitName() const
 {
 	return LOCTEXT( "AppLabel", "Blueprint Editor" );
diff --git a/Engine/Source/Editor/Kismet/Public/BlueprintEditor.h b/Engine/Source/Editor/Kismet/Public/BlueprintEditor.h
--- a/Engine/Source/Editor/Kismet/Public/BlueprintEditor.h
+++ b/Engine/Source/Editor/Kismet/Public/BlueprintEditor.h
@@ -140,6 +140,7 @@ public:
 	// End of FAssetEditorToolkit 
 
 	// IToolkit interface
+	virtual FName GetToolkitContextFName() const override;
 	virtual FName GetToolkitFName() const override;
 	virtual FText GetBaseToolkitName() const override;
 	virtual FText GetToolkitName() const override;
diff --git a/Engine/Source/Editor/LevelEditor/Private/LevelEditor.cpp b/Engine/Source/Editor/LevelEditor/Private/LevelEditor.cpp
--- a/Engine/Source/Editor/LevelEditor/Private/LevelEditor.cpp
+++ b/Engine/Source/Editor/LevelEditor/Private/LevelEditor.cpp
@@ -125,7 +125,7 @@ TSharedRef<SDockTab> FLevelEditorModule::SpawnLevelEditor( const FSpawnTabArgs&
 				]
 				+ SHorizontalBox::Slot()
 				.AutoWidth()
-				.Padding(8.0f, 0.0f, 0.0f, 0.0f)
+				.Padding(8.0f, 0.0f, 8.0f, 0.0f)
 				.VAlign(VAlign_Center)
 				[
 					TutorialWidget
diff --git a/Engine/Source/Editor/Persona/Private/Persona.cpp b/Engine/Source/Editor/Persona/Private/Persona.cpp
--- a/Engine/Source/Editor/Persona/Private/Persona.cpp
+++ b/Engine/Source/Editor/Persona/Private/Persona.cpp
@@ -1463,6 +1463,11 @@ FString FPersona::GetDefaultEditorTitle()
 	return NSLOCTEXT("Kismet", "PreviewParamatersTabTitle", "Preview Parameters").ToString();
 }
 
+FName FPersona::GetToolkitContextFName() const
+{
+	return GetToolkitFName();
+}
+
 FName FPersona::GetToolkitFName() const
 {
 	return FName("Persona");
diff --git a/Engine/Source/Editor/Persona/Private/Persona.h b/Engine/Source/Editor/Persona/Private/Persona.h
--- a/Engine/Source/Editor/Persona/Private/Persona.h
+++ b/Engine/Source/Editor/Persona/Private/Persona.h
@@ -169,6 +169,7 @@ public:
 
 public:
 	// IToolkit interface
+	virtual FName GetToolkitContextFName() const override;
 	virtual FName GetToolkitFName() const override;
 	virtual FText GetBaseToolkitName() const override;
 	virtual FText GetToolkitName() const override;
diff --git a/Engine/Source/Editor/UnrealEd/Private/Toolkits/AssetEditorToolkit.cpp b/Engine/Source/Editor/UnrealEd/Private/Toolkits/AssetEditorToolkit.cpp
--- a/Engine/Source/Editor/UnrealEd/Private/Toolkits/AssetEditorToolkit.cpp
+++ b/Engine/Source/Editor/UnrealEd/Private/Toolkits/AssetEditorToolkit.cpp
@@ -104,7 +104,7 @@ void FAssetEditorToolkit::InitAssetEditor( const EToolkitMode::Type Mode, const
 		TSharedRef<SWidget> UserFeedbackWidget = UserFeedback.CreateFeedbackWidget(GetBaseToolkitName());
 
 		IIntroTutorials& IntroTutorials = FModuleManager::LoadModuleChecked<IIntroTutorials>(TEXT("IntroTutorials"));
-		TSharedRef<SWidget> TutorialWidget = IntroTutorials.CreateTutorialsWidget(GetToolkitFName(), NewMajorTab->GetParentWindow());
+		TSharedRef<SWidget> TutorialWidget = IntroTutorials.CreateTutorialsWidget(GetToolkitContextFName(), NewMajorTab->GetParentWindow());
 
 		NewMajorTab->SetRightContent(
 				SNew(SHorizontalBox)
@@ -117,7 +117,7 @@ void FAssetEditorToolkit::InitAssetEditor( const EToolkitMode::Type Mode, const
 				]
 				+SHorizontalBox::Slot()
 				.AutoWidth()
-				.Padding(8.0f, 0.0f, 0.0f, 0.0f)
+				.Padding(8.0f, 0.0f, 8.0f, 0.0f)
 				.VAlign(VAlign_Center)
 				[
 					TutorialWidget
diff --git a/Engine/Source/Editor/UnrealEd/Private/Toolkits/BaseToolkit.cpp b/Engine/Source/Editor/UnrealEd/Private/Toolkits/BaseToolkit.cpp
--- a/Engine/Source/Editor/UnrealEd/Private/Toolkits/BaseToolkit.cpp
+++ b/Engine/Source/Editor/UnrealEd/Private/Toolkits/BaseToolkit.cpp
@@ -61,6 +61,12 @@ const TMap< EToolkitTabSpot::Type, TArray< TWeakPtr< SDockableTab > > >& FBaseTo
 }
 
 
+FName FBaseToolkit::GetToolkitContextFName() const
+{
+	return GetToolkitFName();
+}
+
+
 bool FBaseToolkit::ProcessCommandBindings( const FKeyboardEvent& InKeyboardEvent ) const
 {
 	if( ToolkitCommands->ProcessCommandBindings( InKeyboardEvent ) )
diff --git a/Engine/Source/Editor/UnrealEd/Public/Toolkits/BaseToolkit.h b/Engine/Source/Editor/UnrealEd/Public/Toolkits/BaseToolkit.h
--- a/Engine/Source/Editor/UnrealEd/Public/Toolkits/BaseToolkit.h
+++ b/Engine/Source/Editor/UnrealEd/Public/Toolkits/BaseToolkit.h
@@ -20,6 +20,7 @@ public:
 	virtual ~FBaseToolkit();
 
 	/** IToolkit interface */
+	virtual FName GetToolkitContextFName() const override;
 	virtual bool ProcessCommandBindings( const FKeyboardEvent& InKeyboardEvent ) const override;
 	virtual bool IsHosted() const override;
 	virtual const TSharedRef< class IToolkitHost > GetToolkitHost() const override;
diff --git a/Engine/Source/Editor/UnrealEd/Public/Toolkits/IToolkit.h b/Engine/Source/Editor/UnrealEd/Public/Toolkits/IToolkit.h
--- a/Engine/Source/Editor/UnrealEd/Public/Toolkits/IToolkit.h
+++ b/Engine/Source/Editor/UnrealEd/Public/Toolkits/IToolkit.h
@@ -63,6 +63,9 @@ public:
 	/** Returns the invariant name of this toolkit type */
 	virtual FName GetToolkitFName() const = 0;
 
+	/** Returns the invariant name of this toolkit type, used for establishing the context for help, documentation and tutorials */
+	virtual FName GetToolkitContextFName() const = 0;
+
 	/** Returns the localized name of this toolkit type (typically just "<ClassName> editor") */
 	virtual FText GetBaseToolkitName() const = 0;
 
