Multiple clients are now spawned correctly when launching a listen server

#ttp 338126 - SHOOTERGAME: EDITOR: Attempting to PIE as a listen server with multiple clients will only PIE with one player

#branch UE4

diff --git a/Engine/Source/Editor/UnrealEd/Classes/Settings/LevelEditorPlaySettings.h b/Engine/Source/Editor/UnrealEd/Classes/Settings/LevelEditorPlaySettings.h
--- a/Engine/Source/Editor/UnrealEd/Classes/Settings/LevelEditorPlaySettings.h
+++ b/Engine/Source/Editor/UnrealEd/Classes/Settings/LevelEditorPlaySettings.h
@@ -191,7 +191,7 @@ public:
 	UPROPERTY(config, EditAnywhere, Category=MultiplayerOptions)
 	bool PlayNetDedicated;
 
-	/** Number of player windows that should be opened. Listen server will count as a player, a dedicated server will not. */
+	/** Number of player windows that should be opened. The editor will count as a player, a dedicated server will not. */
 	UPROPERTY(config, EditAnywhere, Category=MultiplayerOptions, meta=(ClampMin = "1", UIMin = "1", UIMax = "64"))
 	int32 PlayNumberOfClients;
 
diff --git a/Engine/Source/Editor/UnrealEd/Private/PlayLevel.cpp b/Engine/Source/Editor/UnrealEd/Private/PlayLevel.cpp
--- a/Engine/Source/Editor/UnrealEd/Private/PlayLevel.cpp
+++ b/Engine/Source/Editor/UnrealEd/Private/PlayLevel.cpp
@@ -786,21 +786,7 @@ void UEditorEngine::StartQueuedPlayMapRequest()
 			}
 		}
 		
-		if (!bPlayOnLocalPcSession && !bStartMovieCapture && PlayInSettings->PlayNetMode == PIE_Client)
-		{
-			// Editor counts as a client
-			NumClients++;
-		}
-
-		if (bPlayOnLocalPcSession)
-		{
-			// Spawn number of clients
-			for (int32 i = NumClients; i < PlayInSettings->PlayNumberOfClients; ++i)
-			{
-				PlayStandaloneLocalPc(TEXT("127.0.0.1"), &WinPosition, NumClients + 1, false);
-			}
-		}
-		else 
+		if (!bPlayOnLocalPcSession)
 		{
 			if (bStartMovieCapture)
 			{
@@ -810,8 +796,17 @@ void UEditorEngine::StartQueuedPlayMapRequest()
 			else
 			{
 				PlayInEditor(GetEditorWorldContext().World(), bWantSimulateInEditor);
+
+				// Editor counts as a client
+				NumClients++;
 			}
 		}
+
+		// Spawn number of clients
+		for (int32 i = NumClients; i < PlayInSettings->PlayNumberOfClients; ++i)
+		{
+			PlayStandaloneLocalPc(TEXT("127.0.0.1"), &WinPosition, NumClients + 1, false);
+		}
 	}
 	else
 	{
