Change component picker combo to use more standard FEditorClassUtils for rich tooltips (files in Engine\\Documentation\\Source\\Shared\\Types\\)
Add doc w. link for Audio, SkelMesh and StaticMesh components to test
Polish component class picker combo (styling, padding, separators etc)

diff --git a/Engine/Source/Editor/GraphEditor/Private/SGraphActionMenu.cpp b/Engine/Source/Editor/GraphEditor/Private/SGraphActionMenu.cpp
--- a/Engine/Source/Editor/GraphEditor/Private/SGraphActionMenu.cpp
+++ b/Engine/Source/Editor/GraphEditor/Private/SGraphActionMenu.cpp
@@ -886,7 +886,7 @@ TSharedRef<ITableRow> SGraphActionMenu::MakeWidget( TSharedPtr<FGraphActionNode>
 					SNew( SBorder )
 
 					// We'll use the border's padding to actually create the horizontal line
-					.Padding( 1.f )
+					.Padding(FEditorStyle::GetMargin(TEXT("Menu.Separator.Padding")))
 
 					// Separator graphic
 					.BorderImage( FEditorStyle::GetBrush( TEXT( "Menu.Separator" ) ) )
@@ -912,7 +912,7 @@ TSharedRef<ITableRow> SGraphActionMenu::MakeWidget( TSharedPtr<FGraphActionNode>
 				SNew( SBorder )
 
 				// We'll use the border's padding to actually create the horizontal line
-				.Padding( 1.f )
+				.Padding(FEditorStyle::GetMargin(TEXT("Menu.Separator.Padding")))
 
 				// Separator graphic
 				.BorderImage( FEditorStyle::GetBrush( TEXT( "Menu.Separator" ) ) )
diff --git a/Engine/Source/Editor/Kismet/Private/SSCSEditor.cpp b/Engine/Source/Editor/Kismet/Private/SSCSEditor.cpp
--- a/Engine/Source/Editor/Kismet/Private/SSCSEditor.cpp
+++ b/Engine/Source/Editor/Kismet/Private/SSCSEditor.cpp
@@ -1985,20 +1985,24 @@ void SSCSEditor::Construct( const FArguments& InArgs, TSharedPtr<FBlueprintEdito
 	if( InKismet2.IsValid() )
 	{
 		Contents = SNew(SVerticalBox)
+
+		// Component picker
 		+ SVerticalBox::Slot()
+		.Padding(1.f)
 		.AutoHeight()
 		[
-			SNew(SComponentClassCombo)
-			.OnComponentClassSelected(this, &SSCSEditor::PerformComboAddClass)
-		]
-		+ SVerticalBox::Slot()
-		.AutoHeight()
-		[
-			SNew(SSpacer)
-			.Size(FVector2D(0.0f, 2.0f))
+			SNew(SHorizontalBox)
+			+ SHorizontalBox::Slot()
+			.AutoWidth()
+			[
+				SNew(SComponentClassCombo)
+				.OnComponentClassSelected(this, &SSCSEditor::PerformComboAddClass)
+			]
 		]
+
+		// Tree
 		+ SVerticalBox::Slot()
-		.FillHeight(.8f)
+		.Padding(0.f, 0.f, 0.f, 2.f)
 		[
 			SCSTreeWidget.ToSharedRef()
 		];
@@ -2010,7 +2014,9 @@ void SSCSEditor::Construct( const FArguments& InArgs, TSharedPtr<FBlueprintEdito
 
 	this->ChildSlot
 	[
-		SNew( SBox )
+		SNew(SBorder)
+		.Padding(2.0f)
+		.BorderImage(FEditorStyle::GetBrush("ToolPanel.GroupBorder"))
 		.AddMetaData<FTagMetaData>(FTagMetaData(TEXT("ComponentsPanel")))
 		[
 			Contents.ToSharedRef()
diff --git a/Engine/Source/Editor/UnrealEd/Private/SComponentClassCombo.cpp b/Engine/Source/Editor/UnrealEd/Private/SComponentClassCombo.cpp
--- a/Engine/Source/Editor/UnrealEd/Private/SComponentClassCombo.cpp
+++ b/Engine/Source/Editor/UnrealEd/Private/SComponentClassCombo.cpp
@@ -8,6 +8,7 @@
 #include "BlueprintGraphDefinitions.h"
 #include "IDocumentation.h"
 #include "SListViewSelectorDropdownMenu.h"
+#include "EditorClassUtils.h"
 
 #define LOCTEXT_NAMESPACE "ComponentClassCombo"
 
@@ -35,35 +36,56 @@ void SComponentClassCombo::Construct(const FArguments& InArgs)
 	SComboButton::FArguments Args;
 	Args.ButtonContent()
 	[
-		SNew(STextBlock)
-		.Text( LOCTEXT("AddComponentButtonLabel", "Add Component") )
+		SNew(SHorizontalBox)
+		+SHorizontalBox::Slot()
+		.VAlign(VAlign_Center)
+		.AutoWidth()
+		.Padding(2.f,1.f)
+		[
+			SNew(SImage)
+			.Image(FEditorStyle::GetBrush(TEXT("SCSEditor.NewComponentIcon")))
+		]
+		+ SHorizontalBox::Slot()
+		.VAlign(VAlign_Center)
+		.Padding(1.f)
+		[
+			SNew(STextBlock)
+			.Text(LOCTEXT("AddComponentButtonLabel", "Add Component"))
+			.TextStyle(FEditorStyle::Get(), "ContentBrowser.TopBar.Font")
+		]
 	]
 	.MenuContent()
 	[
 
 		SNew(SListViewSelectorDropdownMenu<FComponentClassComboEntryPtr>, SearchBox, ComponentClassListView)
 		[
-			SNew(SVerticalBox)
-			+SVerticalBox::Slot()
-			.AutoHeight()
-			[
-				SearchBox.ToSharedRef()
-			]
-			+SVerticalBox::Slot()
-			.AutoHeight()
-			[
-				SNew(SSeparator)
-				.Orientation(EOrientation::Orient_Horizontal)
-			]
-			+SVerticalBox::Slot()
-			.MaxHeight(400)
+			SNew(SBorder)
+			.BorderImage(FEditorStyle::GetBrush("Menu.Background"))
+			.Padding(2)
 			[
-				ComponentClassListView.ToSharedRef()
+				SNew(SBox)
+				.WidthOverride(250)
+				[				
+					SNew(SVerticalBox)
+					+SVerticalBox::Slot()
+					.Padding(1.f)
+					.AutoHeight()
+					[
+						SearchBox.ToSharedRef()
+					]
+					+SVerticalBox::Slot()
+					.MaxHeight(400)
+					[
+						ComponentClassListView.ToSharedRef()
+					]
+				]
 			]
 		]
 	]
 	.IsFocusable(true)
 	.ContentPadding(FMargin(2))
+	.ComboButtonStyle(FEditorStyle::Get(), "ContentBrowser.NewAsset.Style")
+	.ForegroundColor(FLinearColor::White)
 	.OnComboBoxOpened(this, &SComponentClassCombo::ClearSelection);
 
 	SComboButton::Construct(Args);
@@ -193,8 +215,13 @@ TSharedRef<ITableRow> SComponentClassCombo::GenerateAddComponentRow( FComponentC
 			SNew( STableRow< TSharedPtr<FString> >, OwnerTable )
 				.ShowSelection(false)
 			[
-				SNew(STextBlock)
-				.Text(Entry->GetHeadingText())
+				SNew(SBox)
+				.Padding(1.f)
+				[
+					SNew(STextBlock)
+					.Text(Entry->GetHeadingText())
+					.TextStyle(FEditorStyle::Get(), TEXT("Menu.Heading"))
+				]
 			];
 	}
 	else if ( Entry->IsSeparator() )
@@ -203,7 +230,13 @@ TSharedRef<ITableRow> SComponentClassCombo::GenerateAddComponentRow( FComponentC
 			SNew( STableRow< TSharedPtr<FString> >, OwnerTable )
 				.ShowSelection(false)
 			[
-				SNew(SSeparator)
+				SNew(SBox)
+				.Padding(1.f)
+				[
+					SNew(SBorder)
+					.Padding(FEditorStyle::GetMargin(TEXT("Menu.Separator.Padding")))
+					.BorderImage(FEditorStyle::GetBrush(TEXT("Menu.Separator")))
+				]
 			];
 	}
 	else
@@ -253,13 +286,9 @@ TSharedRef<ITableRow> SComponentClassCombo::GenerateAddComponentRow( FComponentC
 			FriendlyComponentName += FString(" (") + AssetName + FString(")");
 		}
 
-		// Get info for rich tooltip
-		FString DocLink = TEXT("Shared/Editor/Blueprint/VariableTypes");
-		FString DocExcerptName = FString::Printf(TEXT("object'%s'"), *Entry->GetComponentClass()->GetName());
-
 		return
 			SNew( SComboRow< TSharedPtr<FString> >, OwnerTable )
-			.ToolTip( IDocumentation::Get()->CreateToolTip(Entry->GetComponentClass()->GetToolTipText(), NULL, DocLink, DocExcerptName) )
+			.ToolTip( FEditorClassUtils::GetTooltip(Entry->GetComponentClass()) )
 			[
 				SNew(SHorizontalBox)
 				+SHorizontalBox::Slot()
@@ -285,6 +314,7 @@ TSharedRef<ITableRow> SComponentClassCombo::GenerateAddComponentRow( FComponentC
 				]
 				+SHorizontalBox::Slot()
 				.AutoWidth()
+				.VAlign(VAlign_Center)
 				[
 					SNew(STextBlock)
 					.HighlightText(this, &SComponentClassCombo::GetCurrentSearchString)
