Added setting to allow for 'transparent' box selection in perspective viewports.

Enabling this setting will allow marquee box selections in perspective viewports to select actors even if they are fully occluded by other actors.

This addresses TTP#333837 - Editor: Need an option for perspective view marquee select to allow selection of occluded actors.

diff --git a/Engine/Config/BaseEditorUserSettings.ini b/Engine/Config/BaseEditorUserSettings.ini
--- a/Engine/Config/BaseEditorUserSettings.ini
+++ b/Engine/Config/BaseEditorUserSettings.ini
@@ -342,6 +342,8 @@ bHighlightWithBrackets=False
 bUseLinkedOrthographicViewports=True
 ; If true, objects must be entirely encompassed by the selection box in ortho. viewports to be selected
 bStrictBoxSelection=False
+; True if viewport box selection also selects occluded objects, false if only objects with visible pixels are selected
+bTransparentBoxSelection=False
 ; If enabled, selected objects will have an outline around them
 bUseSelectionOutline=True
 ; Sets the intensity of the overlay displayed when an object is selected (defaults to be 0 so you can see the material of selected objects)
diff --git a/Engine/Source/Editor/LevelEditor/Private/LevelEditor.cpp b/Engine/Source/Editor/LevelEditor/Private/LevelEditor.cpp
--- a/Engine/Source/Editor/LevelEditor/Private/LevelEditor.cpp
+++ b/Engine/Source/Editor/LevelEditor/Private/LevelEditor.cpp
@@ -1423,6 +1423,13 @@ void FLevelEditorModule::BindGlobalLevelEditorCommands()
 		FIsActionChecked::CreateStatic( &FLevelEditorActionCallbacks::OnIsStrictBoxSelectEnabled ) 
 		);
 
+	ActionList.MapAction(
+		Commands.TransparentBoxSelect,
+		FExecuteAction::CreateStatic( &FLevelEditorActionCallbacks::OnToggleTransparentBoxSelect ),
+		FCanExecuteAction(),
+		FIsActionChecked::CreateStatic( &FLevelEditorActionCallbacks::OnIsTransparentBoxSelectEnabled ) 
+		);
+
 	ActionList.MapAction(
 		Commands.DrawBrushMarkerPolys,
 		FExecuteAction::CreateStatic( &FLevelEditorActionCallbacks::OnDrawBrushMarkerPolys ),
diff --git a/Engine/Source/Editor/LevelEditor/Private/LevelEditorActions.cpp b/Engine/Source/Editor/LevelEditor/Private/LevelEditorActions.cpp
--- a/Engine/Source/Editor/LevelEditor/Private/LevelEditorActions.cpp
+++ b/Engine/Source/Editor/LevelEditor/Private/LevelEditorActions.cpp
@@ -1931,6 +1931,18 @@ bool FLevelEditorActionCallbacks::OnIsStrictBoxSelectEnabled()
 	return GetDefault<ULevelEditorViewportSettings>()->bStrictBoxSelection;
 }
 
+void FLevelEditorActionCallbacks::OnToggleTransparentBoxSelect()
+{
+	ULevelEditorViewportSettings* ViewportSettings = GetMutableDefault<ULevelEditorViewportSettings>();
+	ViewportSettings->bTransparentBoxSelection = !ViewportSettings->bTransparentBoxSelection;
+	ViewportSettings->PostEditChange();
+}
+
+bool FLevelEditorActionCallbacks::OnIsTransparentBoxSelectEnabled() 
+{
+	return GetDefault<ULevelEditorViewportSettings>()->bTransparentBoxSelection;
+}
+
 void FLevelEditorActionCallbacks::OnDrawBrushMarkerPolys()
 {
 	GEditor->Exec( GetWorld(), *FString::Printf( TEXT("MODE SHOWBRUSHMARKERPOLYS=%d"), !GEditor->bShowBrushMarkerPolys ? 1 : 0 ) );
@@ -2756,6 +2768,7 @@ void FLevelEditorCommands::RegisterCommands()
 	UI_COMMAND( AllowTranslucentSelection, "Allow Translucent Selection", "Allows translucent objects to be selected", EUserInterfaceActionType::ToggleButton, FInputGesture(EKeys::T) );
 	UI_COMMAND( AllowGroupSelection, "Allow Group Selection", "Allows actor groups to be selected", EUserInterfaceActionType::ToggleButton, FInputGesture(EModifierKey::Control|EModifierKey::Shift, EKeys::G) );
 	UI_COMMAND( StrictBoxSelect, "Strict Box Selection", "When enabled an object must be entirely encompassed by the selection box when marquee box selecting", EUserInterfaceActionType::ToggleButton, FInputGesture() );
+	UI_COMMAND( TransparentBoxSelect, "Transparent Box Selection", "When enabled, marquee box select operations will also select objects that are occluded by other objects.", EUserInterfaceActionType::ToggleButton, FInputGesture() );
 	UI_COMMAND( DrawBrushMarkerPolys, "Draw Brush Polys", "Draws semi-transparent polygons around a brush when selected", EUserInterfaceActionType::ToggleButton, FInputGesture() );
 	UI_COMMAND( OnlyLoadVisibleInPIE, "Only Load Visible Levels in Game Preview", "If enabled, when game preview starts, only visible levels will be loaded", EUserInterfaceActionType::ToggleButton, FInputGesture() );
 	UI_COMMAND( ToggleSocketSnapping, "Enable Socket Snapping", "Enables or disables snapping to sockets", EUserInterfaceActionType::ToggleButton, FInputGesture() ); 
diff --git a/Engine/Source/Editor/LevelEditor/Private/LevelEditorToolBar.cpp b/Engine/Source/Editor/LevelEditor/Private/LevelEditorToolBar.cpp
--- a/Engine/Source/Editor/LevelEditor/Private/LevelEditorToolBar.cpp
+++ b/Engine/Source/Editor/LevelEditor/Private/LevelEditorToolBar.cpp
@@ -681,6 +681,7 @@ TSharedRef< SWidget > FLevelEditorToolBar::GenerateQuickSettingsMenu( TSharedRef
 		MenuBuilder.AddMenuEntry( FLevelEditorCommands::Get().AllowTranslucentSelection );
 		MenuBuilder.AddMenuEntry( FLevelEditorCommands::Get().AllowGroupSelection );
 		MenuBuilder.AddMenuEntry( FLevelEditorCommands::Get().StrictBoxSelect );
+		MenuBuilder.AddMenuEntry( FLevelEditorCommands::Get().TransparentBoxSelect );
 		MenuBuilder.AddMenuEntry( FLevelEditorCommands::Get().ShowTransformWidget );
 	}
 	MenuBuilder.EndSection();
diff --git a/Engine/Source/Editor/LevelEditor/Public/LevelEditorActions.h b/Engine/Source/Editor/LevelEditor/Public/LevelEditorActions.h
--- a/Engine/Source/Editor/LevelEditor/Public/LevelEditorActions.h
+++ b/Engine/Source/Editor/LevelEditor/Public/LevelEditorActions.h
@@ -520,6 +520,7 @@ public:
 	TSharedPtr< FUICommandInfo > AllowGroupSelection;
 
 	TSharedPtr< FUICommandInfo > StrictBoxSelect;
+	TSharedPtr< FUICommandInfo > TransparentBoxSelect;
 	TSharedPtr< FUICommandInfo > DrawBrushMarkerPolys;
 	TSharedPtr< FUICommandInfo > OnlyLoadVisibleInPIE;
 
@@ -1061,6 +1062,8 @@ public:
 	static bool OnIsAllowGroupSelectionEnabled(); 
 	static void OnToggleStrictBoxSelect();
 	static bool OnIsStrictBoxSelectEnabled(); 
+	static void OnToggleTransparentBoxSelect();
+	static bool OnIsTransparentBoxSelectEnabled();
 	static void OnDrawBrushMarkerPolys();
 	static bool OnIsDrawBrushMarkerPolysEnabled();
 	static void OnToggleOnlyLoadVisibleInPIE();
diff --git a/Engine/Source/Editor/UnrealEd/Classes/Settings/LevelEditorViewportSettings.h b/Engine/Source/Editor/UnrealEd/Classes/Settings/LevelEditorViewportSettings.h
--- a/Engine/Source/Editor/UnrealEd/Classes/Settings/LevelEditorViewportSettings.h
+++ b/Engine/Source/Editor/UnrealEd/Classes/Settings/LevelEditorViewportSettings.h
@@ -330,6 +330,10 @@ public:
 	UPROPERTY(config)
 	uint32 bStrictBoxSelection:1;
 
+	/** True if viewport box selection also selects occluded objects, false if only objects with visible pixels are selected */
+	UPROPERTY(config)
+	uint32 bTransparentBoxSelection:1;
+
 	/** Whether to show selection outlines for selected Actors */
 	UPROPERTY(EditAnywhere, config, Category=LookAndFeel, meta=(DisplayName = "Use Selection Outline"))
 	uint32 bUseSelectionOutline:1;
diff --git a/Engine/Source/Editor/UnrealEd/Private/DragTool_FrustumSelect.cpp b/Engine/Source/Editor/UnrealEd/Private/DragTool_FrustumSelect.cpp
--- a/Engine/Source/Editor/UnrealEd/Private/DragTool_FrustumSelect.cpp
+++ b/Engine/Source/Editor/UnrealEd/Private/DragTool_FrustumSelect.cpp
@@ -106,88 +106,106 @@ void FDragTool_ActorFrustumSelect::EndDrag()
 			Swap( Start.Y, End.Y );
 		}
 
-		// Extend the endpoint of the rect to get the actual line
-		FIntRect BoxRect( FIntPoint( FMath::Max( 0.0f, Start.X ), FMath::Max( 0.0f, Start.Y ) ), FIntPoint( FMath::Min(ViewportSizeX, FMath::TruncToInt(End.X+1)), FMath::Min( ViewportSizeY, FMath::TruncToInt(End.Y+1) ) ) );
+		const bool bTransparentBoxSelection = GetDefault<ULevelEditorViewportSettings>()->bTransparentBoxSelection;
+		if (bTransparentBoxSelection)
+		{
+			// Get a list of frustum-culled actors
+			for(FActorIterator It(ViewportClient->GetWorld()); It; ++It)
+			{
+				AActor* Actor = *It;
+				if (IntersectsFrustum( *Actor, Frustum, bStrictDragSelection))
+				{
+					// Select the actor if it intersected
+					GEditor->SelectActor( Actor, bShouldSelect, true);
+					bSelectionChanged = true;
+				}
+			}
+		}
+		else
+		{
+			// Extend the endpoint of the rect to get the actual line
+			FIntRect BoxRect( FIntPoint( FMath::Max( 0.0f, Start.X ), FMath::Max( 0.0f, Start.Y ) ), FIntPoint( FMath::Min(ViewportSizeX, FMath::TruncToInt(End.X+1)), FMath::Min( ViewportSizeY, FMath::TruncToInt(End.Y+1) ) ) );
 
-		const TArray<FColor>& RawHitProxyData = ViewportClient->Viewport->GetRawHitProxyData(BoxRect);
+			const TArray<FColor>& RawHitProxyData = ViewportClient->Viewport->GetRawHitProxyData(BoxRect);
 
-		TSet<AActor*> HitActors;
-		TSet<UModel*> HitModels;
+			TSet<AActor*> HitActors;
+			TSet<UModel*> HitModels;
 
 
-		// Lower the resolution with massive box selects
-		int32 Step = (BoxRect.Width() > 500 && BoxRect.Height() > 500) ? 4 : 1;
+			// Lower the resolution with massive box selects
+			int32 Step = (BoxRect.Width() > 500 && BoxRect.Height() > 500) ? 4 : 1;
 
 
-		for (int32 Y = BoxRect.Min.Y; Y < BoxRect.Max.Y; Y = Y < BoxRect.Max.Y - 1 ? FMath::Min(BoxRect.Max.Y-1, Y+Step) : ++Y )
-		{
-			const FColor* SourceData = &RawHitProxyData[Y * ViewportSizeX];
-			for (int32 X = BoxRect.Min.X; X < BoxRect.Max.X; X = X < BoxRect.Max.X-1 ? FMath::Min(BoxRect.Max.X-1, X + Step) : ++X )
+			for (int32 Y = BoxRect.Min.Y; Y < BoxRect.Max.Y; Y = Y < BoxRect.Max.Y - 1 ? FMath::Min(BoxRect.Max.Y-1, Y+Step) : ++Y )
 			{
-				FHitProxyId HitProxyId(SourceData[X]);
-				HHitProxy* HitProxy = GetHitProxyById(HitProxyId);
-
-				if (HitProxy)
+				const FColor* SourceData = &RawHitProxyData[Y * ViewportSizeX];
+				for (int32 X = BoxRect.Min.X; X < BoxRect.Max.X; X = X < BoxRect.Max.X-1 ? FMath::Min(BoxRect.Max.X-1, X + Step) : ++X )
 				{
-					if( HitProxy->IsA(HActor::StaticGetType()) )
+					FHitProxyId HitProxyId(SourceData[X]);
+					HHitProxy* HitProxy = GetHitProxyById(HitProxyId);
+
+					if (HitProxy)
 					{
-						AActor* Actor = ((HActor*)HitProxy)->Actor;
-						if (Actor)
+						if( HitProxy->IsA(HActor::StaticGetType()) )
 						{
-							HitActors.Add(Actor);
+							AActor* Actor = ((HActor*)HitProxy)->Actor;
+							if (Actor)
+							{
+								HitActors.Add(Actor);
+							}
 						}
-					}
-					else if( HitProxy->IsA(HModel::StaticGetType()) )
-					{
-						HitModels.Add( ((HModel*)HitProxy)->GetModel() );
-					}
-					else if( HitProxy->IsA(HBSPBrushVert::StaticGetType()) )
-					{
-						HBSPBrushVert* HitBSPBrushVert = ((HBSPBrushVert*)HitProxy);
-						if( HitBSPBrushVert->Brush.IsValid() )
+						else if( HitProxy->IsA(HModel::StaticGetType()) )
+						{
+							HitModels.Add( ((HModel*)HitProxy)->GetModel() );
+						}
+						else if( HitProxy->IsA(HBSPBrushVert::StaticGetType()) )
 						{
-							HitActors.Add( HitBSPBrushVert->Brush.Get() );
+							HBSPBrushVert* HitBSPBrushVert = ((HBSPBrushVert*)HitProxy);
+							if( HitBSPBrushVert->Brush.IsValid() )
+							{
+								HitActors.Add( HitBSPBrushVert->Brush.Get() );
+							}
 						}
 					}
-				}
 
+				}
 			}
-		}
 
 
-		if (HitModels.Num() > 0)
-		{
-			// Check every model to see if its BSP surfaces should be selected
-			for (auto It = HitModels.CreateConstIterator(); It; ++It)
+			if (HitModels.Num() > 0)
 			{
-				UModel& Model = **It;
-				// Check every node in the model
-				for (int32 NodeIndex = 0; NodeIndex < Model.Nodes.Num(); NodeIndex++)
+				// Check every model to see if its BSP surfaces should be selected
+				for (auto It = HitModels.CreateConstIterator(); It; ++It)
 				{
-					if (IntersectsFrustum(Model, NodeIndex, Frustum, bStrictDragSelection))
+					UModel& Model = **It;
+					// Check every node in the model
+					for (int32 NodeIndex = 0; NodeIndex < Model.Nodes.Num(); NodeIndex++)
 					{
-						uint32 SurfaceIndex = Model.Nodes[NodeIndex].iSurf;
-						FBspSurf& Surf = Model.Surfs[SurfaceIndex];
-						HitActors.Add( Surf.Actor );
+						if (IntersectsFrustum(Model, NodeIndex, Frustum, bStrictDragSelection))
+						{
+							uint32 SurfaceIndex = Model.Nodes[NodeIndex].iSurf;
+							FBspSurf& Surf = Model.Surfs[SurfaceIndex];
+							HitActors.Add( Surf.Actor );
+						}
 					}
 				}
 			}
-		}
 
-		if( HitActors.Num() > 0 )
-		{
-			for( auto It = HitActors.CreateConstIterator(); It; ++It )
+			if( HitActors.Num() > 0 )
 			{
-				AActor* Actor = *It;
-				if( bStrictDragSelection && IntersectsFrustum( *Actor, Frustum, bStrictDragSelection ) )
+				for( auto It = HitActors.CreateConstIterator(); It; ++It )
 				{
-					GEditor->SelectActor(Actor, bShouldSelect, false);
-					bSelectionChanged = true;
-				}
-				else if( !bStrictDragSelection )
-				{
-					GEditor->SelectActor(Actor, bShouldSelect, false);
-					bSelectionChanged = true;
+					AActor* Actor = *It;
+					if( bStrictDragSelection && IntersectsFrustum( *Actor, Frustum, bStrictDragSelection ) )
+					{
+						GEditor->SelectActor(Actor, bShouldSelect, false);
+						bSelectionChanged = true;
+					}
+					else if( !bStrictDragSelection )
+					{
+						GEditor->SelectActor(Actor, bShouldSelect, false);
+						bSelectionChanged = true;
+					}
 				}
 			}
 		}
@@ -353,9 +371,6 @@ bool FDragTool_ActorFrustumSelect::IntersectsFrustum( AActor& InActor, const FCo
 	// Never drag-select hidden actors or builder brushes. 
 	if( !bActorIsHiddenByShowFlags && !InActor.IsHiddenEd() && !FActorEditorUtils::IsABuilderBrush(&InActor) )
 	{
-		// Determine if the user would prefer to use strict selection or not
-		const bool bStrictDragSelection = GetDefault<ULevelEditorViewportSettings>()->bStrictBoxSelection;
-
 		// Iterate over all actor components, selecting out primitive components
 		TArray<UPrimitiveComponent*> Components;
 		InActor.GetComponents(Components);
@@ -404,7 +419,7 @@ bool FDragTool_ActorFrustumSelect::IntersectsFrustum( AActor& InActor, const FCo
 							const FVector Location = BrushComponent->ComponentToWorld.TransformPosition( Poly.Vertices[VertexIndex] );
 							const bool bIntersect = InFrustum.IntersectBox( Location, FVector::ZeroVector );
 
-							if( bIntersect && !bStrictDragSelection )
+							if( bIntersect && !bUseStrictSelection )
 							{
 								// If we intersected a vertex and we dont require the box to encompass the entire component
 								// then the actor should be selected and we can stop checking
@@ -412,7 +427,7 @@ bool FDragTool_ActorFrustumSelect::IntersectsFrustum( AActor& InActor, const FCo
 								bActorHitByBox = true;
 								break;
 							}
-							else if( !bIntersect && bStrictDragSelection )
+							else if( !bIntersect && bUseStrictSelection )
 							{
 								// If we didnt intersect a vertex but we require the box to encompass the entire component
 								// then this test failed and we can stop checking
@@ -423,7 +438,7 @@ bool FDragTool_ActorFrustumSelect::IntersectsFrustum( AActor& InActor, const FCo
 						}
 					}
 
-					if( bStrictDragSelection && !bAlreadyProcessed )
+					if( bUseStrictSelection && !bAlreadyProcessed )
 					{
 						// If we are here then every vert was intersected so we should select the actor
 						bActorHitByBox = true;
@@ -439,7 +454,7 @@ bool FDragTool_ActorFrustumSelect::IntersectsFrustum( AActor& InActor, const FCo
 				else
 				{
 					if ( InFrustum.IntersectBox( Origin, Extent, bActorFullyContained) 
-						&& (!bStrictDragSelection || bActorFullyContained))
+						&& (!bUseStrictSelection || bActorFullyContained))
 					{
 						if( PrimitiveComponent->IsA( UStaticMeshComponent::StaticClass() ) || PrimitiveComponent->IsA( USkeletalMeshComponent::StaticClass() ) )
 						{
