EditorPerformanceTests.cpp && AutomationEditorCommon.cpp
--Made it so that the test will stop if a csv file is not able to be written to.
--Added more logging.
--Code cleanup.

diff --git a/Engine/Source/Editor/UnrealEd/Private/Tests/AutomationEditorCommon.cpp b/Engine/Source/Editor/UnrealEd/Private/Tests/AutomationEditorCommon.cpp
--- a/Engine/Source/Editor/UnrealEd/Private/Tests/AutomationEditorCommon.cpp
+++ b/Engine/Source/Editor/UnrealEd/Private/Tests/AutomationEditorCommon.cpp
@@ -283,11 +283,11 @@ namespace AutomationEditorCommonUtils
 * Writes a number to a text file.
 * @param FolderNameForTest is the folder that has the same name as the test. (For Example: "Performance").
 * @param FolderNameForBeingTested is the name for the thing that is being tested. (For Example: "MapName").
-* @param FileName is the name of the file with an extension (meaning no need to add the .txt)
+* @param FileName is the name of the file with an extension (no need to add the .txt)
 * @param NumberToWriteToFile is the float number that is expected to be written to the file.
 * @param Delimiter is the delimiter to be used. TEXT(",")
 */
-void WriteToTextFile(FString FolderNameForTest ,FString FolderNameForBeingTested, FString FileName, float NumberToWriteToFile, FString Delimiter)
+void WriteToTextFile(const FString& FolderNameForTest ,const FString& FolderNameForBeingTested, const FString& FileName, const float& NumberToWriteToFile, const FString& Delimiter)
 {
 	//Performance file locations and setups.
 	FString FileSaveLocation = FPaths::AutomationLogDir() + FolderNameForTest + TEXT("/") + FolderNameForBeingTested + TEXT("/") + FileName + TEXT(".txt");
@@ -338,16 +338,15 @@ float ValueFromStringArrayofNumbers(const TArray<FString>& NumberArray, bool bAv
 
 }
 
-TArray<FString> CreateArrayFromFile(FString FullFileLocation)
+TArray<FString> CreateArrayFromFile(const FString& FullFileLocation)
 {
 
 	FString RawData;
 	TArray<FString> DataArray;
 
+	UE_LOG(LogEditorAutomationTests, Log, TEXT("Loading and parsing the data from '%s' into an array."), *FullFileLocation);
 	FFileHelper::LoadFileToString(RawData, *FullFileLocation);
-
 	RawData.ParseIntoArray(&DataArray, TEXT(","), false);
-		//ParseIntoArray(&DataArray, Delimiter, true);
 	
 	return DataArray;
 }
@@ -487,13 +486,14 @@ bool FEditorPerformanceCommand::Update()
 		return false;
 	}
 	
-	UE_LOG(LogEditorAutomationTests, Display, TEXT("Performance has been captured and saved to the '//Game/Saved/Automation/Log/Performance/%s' folder location."), *ShortMapName);
+	UE_LOG(LogEditorAutomationTests, Log, TEXT("Raw performance data has been captured and saved to the '//Game/Saved/Automation/Log/Performance/%s' folder location."), *ShortMapName);
 	
 	return true;
 }
 
 bool FGenerateEditorPerformanceCharts::Update()
 {
+	UE_LOG(LogEditorAutomationTests, Log, TEXT("Begin generating the editor performance charts."));
 	//Get the current changelist number
 	FString Changelist = GEngineVersion.ToString(EVersionComponent::Changelist);
 
@@ -515,6 +515,9 @@ bool FGenerateEditorPerformanceCharts::Update()
 	TArray<FString> MapLoadTimeRawArray = CreateArrayFromFile(MapLoadTimeFileSaveLocation);
 	TArray<FString> TestRunTimeRawArray = CreateArrayFromFile(DurationFileSaveLocation);
 
+	//Variable that will hold the old performance csv file data.
+	FString OldPerformanceCSVFile;
+
 	//The CSV files.
 	//RAW csv holds all of the information from the text files
 	//Final csv only holds the averaged information.
@@ -532,18 +535,31 @@ bool FGenerateEditorPerformanceCharts::Update()
 	//Add the top title row to the raw csv file
 	FString RAWCSVLine = (TEXT("Map Name, Changelist, Test Run Time, Map Load Time, Average FPS, Frame Time, Used Physical Memory\n"));
 	RAWCSVFile->Serialize(TCHAR_TO_ANSI(*RAWCSVLine), RAWCSVLine.Len());
+	
 	//Loop through the text files and add them to the raw csv file.
 	for (int32 IterLoop = 0; IterLoop < FPSRawArray.Num(); IterLoop++)
 	{
+		//If the raw file isn't available to write to then we'll fail back this test.
+		if (!RAWCSVFile)
+		{
+			UE_LOG(LogEditorAutomationTests, Error, TEXT("Failed to write to the csv file: '%s'.  \nIt may have been opened or deleted."), *RAWCSVFilename);
+			return true;
+		}
+
+		//Create the line that will be written and then write it to the csv file.
 		RAWCSVLine = LoadedMapName + TEXT(",") + Changelist + TEXT(",") + TestRunTimeRawArray[IterLoop] + TEXT(",") + MapLoadTimeRawArray[0] + TEXT(",") + FPSRawArray[IterLoop] + TEXT(",") + FrameTimeRawArray[IterLoop] + TEXT(",") + MemoryRawArray[IterLoop] + LINE_TERMINATOR;
 		RAWCSVFile->Serialize(TCHAR_TO_ANSI(*RAWCSVLine), RAWCSVLine.Len());
 	}
 	
+	//Close the raw csv file.
+	RAWCSVFile->Close();
+	
 	//Get the final data for the Performance csv file.
 	//This gets the averaged numbers using the raw data arrays.
 	float FPSAverage = ValueFromStringArrayofNumbers(FPSRawArray, true);
 	float FrameTimeAverage = ValueFromStringArrayofNumbers(FrameTimeRawArray, true);
 	float MemoryAverage = ValueFromStringArrayofNumbers(MemoryRawArray, true);
+	
 	//This gets the actual test run time by subtracting the last test run entry against the first test run entry.
 	FString TestRunStartTime = TestRunTimeRawArray[0];
 	FString TestRunEndTime;
@@ -559,40 +575,53 @@ bool FGenerateEditorPerformanceCharts::Update()
 	//Get the map load time.
 	FString MapLoadTime = MapLoadTimeRawArray[0];
 	
-	FString OldPerformanceCSVFile;
-
 	//If the Performance csv file does not have a size then it must be new.
 	//This adds the title row to the performance csv file.
 	if (IFileManager::Get().FileSize(FinalCSVFilename.GetCharArray().GetData()) <= 0)
 	{
 		FArchive* FinalCSVFile = IFileManager::Get().CreateFileWriter(*FinalCSVFilename);
+		//If the csv file can't be created or written to then we'll fail back this test.
+		if (!FinalCSVFile)
+		{
+			UE_LOG(LogEditorAutomationTests, Error, TEXT("Failed to create the performance csv file: '%s'."), *FPaths::ConvertRelativePathToFull(FinalCSVFilename));
+			return true;
+		}
+		
 		//Create the .csv file that will hold the raw data from the text files.
 		FString FinalCSVLine = (TEXT("Date, Map Name, Changelist, Test Run Time, Map Load Time, Average FPS, Average MS, Used Physical KB\n"));
 		FinalCSVFile->Serialize(TCHAR_TO_ANSI(*FinalCSVLine), FinalCSVLine.Len());
 		FinalCSVFile->Close();
 	}
 
+	//Load the currently existing performance csv data.
 	FFileHelper::LoadFileToString(OldPerformanceCSVFile, *FinalCSVFilename);
 	FArchive* FinalCSVFile = IFileManager::Get().CreateFileWriter(*FinalCSVFilename);
+	//If the performance file isn't available to write to then we'll fail back this test.
+	if (!FinalCSVFile)
+	{
+		UE_LOG(LogEditorAutomationTests, Error, TEXT("Failed to write to the Performance csv file: '%s'.  \nIt may have been opened or deleted while in use. \nData Text files will not be deleted at this time."), *FPaths::ConvertRelativePathToFull(FinalCSVFilename));
+		return true;
+	}
+	
+	//Write the old performance csv file data to the new csv file.
 	FinalCSVFile->Serialize(TCHAR_TO_ANSI(*OldPerformanceCSVFile), OldPerformanceCSVFile.Len());
 
-
-
 	//Write the averaged numbers to the Performance CSV file.
 	FString FinalCSVLine = FString::Printf(TEXT("%s,%s,%s,%.0f,%s,%.3f,%.3f,%.0f%s"), *FDateTime::Today().ToString(), *LoadedMapName, *Changelist, TestRunTime, *MapLoadTime, FPSAverage, FrameTimeAverage, MemoryAverage, LINE_TERMINATOR);
 	FinalCSVFile->Serialize(TCHAR_TO_ANSI(*FinalCSVLine), FinalCSVLine.Len());
 
+	//Close the csv files so we can use them while the editor is still open.
+	FinalCSVFile->Close();
+
 	//Display the averaged numbers.
 	UE_LOG(LogEditorAutomationTests, Display, TEXT("AVG FPS: '%.1f'"), FPSAverage);
 	UE_LOG(LogEditorAutomationTests, Display, TEXT("AVG Frame Time: '%.1f' ms"), FrameTimeAverage);
 	UE_LOG(LogEditorAutomationTests, Display, TEXT("AVG Used Physical Memory: '%.0f' kb"), MemoryAverage);
+	UE_LOG(LogEditorAutomationTests, Display, TEXT("Performance csv file is located here: %s"), *FPaths::ConvertRelativePathToFull(FinalCSVFilename));
 	UE_LOG(LogEditorAutomationTests, Log, TEXT("Performance csv file is located here: %s"), *FPaths::ConvertRelativePathToFull(FinalCSVFilename));
 
-	//Close the csv files so we can use them while the editor is still open.
-	RAWCSVFile->Close();
-	FinalCSVFile->Close();
-
 	//Delete the text files.
+	UE_LOG(LogEditorAutomationTests, Log, TEXT("Deleting temporary data text files."));
 	IFileManager::Get().Delete(*FPSFileSaveLocation);
 	IFileManager::Get().Delete(*FrameTimeSaveLocation);
 	IFileManager::Get().Delete(*MemoryFileSaveLocation);
diff --git a/Engine/Source/Editor/UnrealEd/Private/Tests/EditorPerformanceTests.cpp b/Engine/Source/Editor/UnrealEd/Private/Tests/EditorPerformanceTests.cpp
--- a/Engine/Source/Editor/UnrealEd/Private/Tests/EditorPerformanceTests.cpp
+++ b/Engine/Source/Editor/UnrealEd/Private/Tests/EditorPerformanceTests.cpp
@@ -115,6 +115,8 @@ bool FMapPerformanceInEditor::RunTest(const FString& Parameters)
 		}
 	}
 
+	UE_LOG(LogEditorAutomationTests, Log, TEXT("Running the performance capture test for %.0f seconds on %s"), Duration, *ShortMapName);
+
 	//Load Map and get the time it took to take to load the map.
 	ADD_LATENT_AUTOMATION_COMMAND(FEditorLoadMap(MapName));
 	ADD_LATENT_AUTOMATION_COMMAND(FWaitLatentCommand(1.0f));
