Fixes issues with Text pins stomping eachother's values after duplicating a Blueprint.

NodeGuids will be recreated when duplicating a Blueprint.

Version bump to handle giving Nodes a new Guid for all Blueprints.

diff --git a/Engine/Source/Editor/UnrealEd/Private/Kismet2/BlueprintEditorUtils.cpp b/Engine/Source/Editor/UnrealEd/Private/Kismet2/BlueprintEditorUtils.cpp
--- a/Engine/Source/Editor/UnrealEd/Private/Kismet2/BlueprintEditorUtils.cpp
+++ b/Engine/Source/Editor/UnrealEd/Private/Kismet2/BlueprintEditorUtils.cpp
@@ -1480,6 +1480,14 @@ void FBlueprintEditorUtils::PostDuplicateBlueprint(UBlueprint* Blueprint)
 		// Create a new blueprint guid
 		Blueprint->GenerateNewGuid();
 
+		// Give all nodes a new Guid
+		TArray< UEdGraphNode* > AllGraphNodes;
+		GetAllNodesOfClass(Blueprint, AllGraphNodes);
+		for(auto& Node : AllGraphNodes)
+		{
+			Node->CreateNewGuid();
+		}
+
 		// And compile again to make sure they go into the generated class, get cleaned up, etc...
 		FKismetEditorUtilities::CompileBlueprint(Blueprint, false, true);
 
diff --git a/Engine/Source/Runtime/Core/Public/UObject/ObjectVersion.h b/Engine/Source/Runtime/Core/Public/UObject/ObjectVersion.h
--- a/Engine/Source/Runtime/Core/Public/UObject/ObjectVersion.h
+++ b/Engine/Source/Runtime/Core/Public/UObject/ObjectVersion.h
@@ -608,6 +608,8 @@ enum EUnrealEngineObjectUE4Version
 	VER_UE4_RENAME_SM3_TO_ES3_1,
 	// Deprecated separate style assets for use in UMG
 	VER_UE4_DEPRECATE_UMG_STYLE_ASSETS,
+	// Duplicating Blueprints will regenerate NodeGuids after this version
+	VER_UE4_POST_DUPLICATE_NODE_GUID,
 
 	// -----<new versions can be added before this line>-------------------------------------------------
 	// - this needs to be the last line (see note below)
diff --git a/Engine/Source/Runtime/Engine/Private/EdGraph/EdGraphNode.cpp b/Engine/Source/Runtime/Engine/Private/EdGraph/EdGraphNode.cpp
--- a/Engine/Source/Runtime/Engine/Private/EdGraph/EdGraphNode.cpp
+++ b/Engine/Source/Runtime/Engine/Private/EdGraph/EdGraphNode.cpp
@@ -246,6 +246,13 @@ void UEdGraphNode::PostLoad()
 		// Generate new one
 		CreateNewGuid();
 	}
+
+	// Duplicating a Blueprint needs to have a new Node Guid generated, which was not occuring before this version
+	if(GetLinkerUE4Version() < VER_UE4_POST_DUPLICATE_NODE_GUID)
+	{
+		// Generate new one
+		CreateNewGuid();
+	}
 }
 
 void UEdGraphNode::CreateNewGuid()
