Generating archives will now fail if the archives are of the wrong version. This helps to avoid users losing data due to a resave with a different verison and not being able to repair the old data.

diff --git a/Engine/Source/Editor/UnrealEd/Private/Commandlets/GenerateGatherArchiveCommandlet.cpp b/Engine/Source/Editor/UnrealEd/Private/Commandlets/GenerateGatherArchiveCommandlet.cpp
--- a/Engine/Source/Editor/UnrealEd/Private/Commandlets/GenerateGatherArchiveCommandlet.cpp
+++ b/Engine/Source/Editor/UnrealEd/Private/Commandlets/GenerateGatherArchiveCommandlet.cpp
@@ -250,8 +250,15 @@ int32 UGenerateGatherArchiveCommandlet::Main( const FString& Params )
 			ArchiveSerializer.DeserializeArchive( ExistingArchiveJsonObject.ToSharedRef(), OutputInternationalizationArchive );
 		}
 
+		if (InternationalizationArchive->GetFormatVersion() < FInternationalizationArchive::EFormatVersion::Latest)
+		{
+			UE_LOG( LogGenerateArchiveCommandlet, Error,TEXT("Archive version is out of date. Repair the archives or manually set the version to %d."), static_cast<int32>(FInternationalizationArchive::EFormatVersion::Latest));
+			return -1;
+		}
+
 		// Combine the generated gather archive with the contents of the archive structure we will write out.
 		AppendArchiveData( InternationalizationArchive, OutputInternationalizationArchive );
+		InternationalizationArchive->SetFormatVersion(FInternationalizationArchive::EFormatVersion::Latest);
 
 		TSharedRef< FJsonObject > OutputArchiveJsonObj = MakeShareable( new FJsonObject );
 		ArchiveSerializer.SerializeArchive( OutputInternationalizationArchive, OutputArchiveJsonObj );
