Removed logic in PO import/export that failed to read/write from culture directories if only one culture was specified.

diff --git a/Engine/Source/Editor/UnrealEd/Private/Commandlets/InternationalizationExportCommandlet.cpp b/Engine/Source/Editor/UnrealEd/Private/Commandlets/InternationalizationExportCommandlet.cpp
--- a/Engine/Source/Editor/UnrealEd/Private/Commandlets/InternationalizationExportCommandlet.cpp
+++ b/Engine/Source/Editor/UnrealEd/Private/Commandlets/InternationalizationExportCommandlet.cpp
@@ -1143,16 +1143,7 @@ bool UInternationalizationExportCommandlet::DoExport( const FString& SourcePath,
 				// Write out the Portable Object to .po file.
 				{
 					FString OutputString = PortableObj.ToString();
-					FString OutputFileName;
-					
-					if (CulturesToGenerate.Num() > 1)
-					{
-						OutputFileName = DestinationPath / CultureName / Filename;
-					}
-					else
-					{
-						OutputFileName = DestinationPath / Filename;
-					}
+					FString OutputFileName = DestinationPath / CultureName / Filename;
 
 					if( SourceControlInfo.IsValid() )
 					{
@@ -1200,16 +1191,7 @@ bool UInternationalizationExportCommandlet::DoImport(const FString& SourcePath,
 	{
 		// Load the Portable Object file if found
 		const FString CultureName = CulturesToGenerate[Culture];
-		FString POFilePath;
-					
-		if (CulturesToGenerate.Num() > 1)
-		{
-			POFilePath = SourcePath / CultureName / Filename;
-		}
-		else
-		{
-			POFilePath = SourcePath / Filename;
-		}
+		const FString POFilePath = SourcePath / CultureName / Filename;
 
 		if( !FPaths::FileExists(POFilePath) )
 		{
