Retarget window cleanup
- Added viewport
- Added option to show only compatible ones - on by default

diff --git a/Engine/Source/Developer/AssetTools/Private/AssetTypeActions/AssetTypeActions_AnimBlueprint.cpp b/Engine/Source/Developer/AssetTools/Private/AssetTypeActions/AssetTypeActions_AnimBlueprint.cpp
--- a/Engine/Source/Developer/AssetTools/Private/AssetTypeActions/AssetTypeActions_AnimBlueprint.cpp
+++ b/Engine/Source/Developer/AssetTools/Private/AssetTypeActions/AssetTypeActions_AnimBlueprint.cpp
@@ -160,10 +160,11 @@ void FAssetTypeActions_AnimBlueprint::RetargetAssets(const TArray<UObject*> InAn
 	}
 
 	bool bConvertSpaces = OldSkeleton != NULL;
+	bool bShowOnlyCompatibleSkeletons = OldSkeleton != NULL;
 
 	const FText Message = LOCTEXT("RemapSkeleton_Warning", "Select the skeleton to remap this asset to.");
 
-	if (SAnimationRemapSkeleton::ShowModal(OldSkeleton, NewSkeleton, Message, OldSkeleton ? &bConvertSpaces : NULL, &bRemapReferencedAssets))
+	if (SAnimationRemapSkeleton::ShowModal(OldSkeleton, NewSkeleton, Message, OldSkeleton? &bShowOnlyCompatibleSkeletons : NULL, OldSkeleton ? &bConvertSpaces : NULL, &bRemapReferencedAssets))
 	{
 		EditorAnimUtils::RetargetAnimations(OldSkeleton, NewSkeleton, InAnimBlueprints, bRemapReferencedAssets, bDuplicateAssets, bConvertSpaces);
 	}
diff --git a/Engine/Source/Developer/AssetTools/Private/AssetTypeActions/AssetTypeActions_AnimationAsset.cpp b/Engine/Source/Developer/AssetTools/Private/AssetTypeActions/AssetTypeActions_AnimationAsset.cpp
--- a/Engine/Source/Developer/AssetTools/Private/AssetTypeActions/AssetTypeActions_AnimationAsset.cpp
+++ b/Engine/Source/Developer/AssetTools/Private/AssetTypeActions/AssetTypeActions_AnimationAsset.cpp
@@ -164,11 +164,12 @@ void FAssetTypeActions_AnimationAsset::RetargetAssets(const TArray<UObject*> InA
 		OldSkeleton = AnimAsset->GetSkeleton();
 	}
 
+	bool bShowOnlyCompatibleSkeletons = OldSkeleton != NULL;
 	bool bConvertSpaces = OldSkeleton != NULL;
 
 	const FText Message = LOCTEXT("SelectSkeletonToRemap", "Select the skeleton to remap this asset to.");
 
-	if (SAnimationRemapSkeleton::ShowModal(OldSkeleton, NewSkeleton, Message, OldSkeleton ? &bConvertSpaces : NULL, &bRemapReferencedAssets))
+	if (SAnimationRemapSkeleton::ShowModal(OldSkeleton, NewSkeleton, Message, OldSkeleton? &bShowOnlyCompatibleSkeletons : NULL, OldSkeleton ? &bConvertSpaces : NULL, &bRemapReferencedAssets))
 	{
 		EditorAnimUtils::RetargetAnimations(OldSkeleton, NewSkeleton, InAnimAssets, bRemapReferencedAssets, bDuplicateAssets, bConvertSpaces);
 	}
diff --git a/Engine/Source/Developer/AssetTools/Private/AssetTypeActions/AssetTypeActions_Skeleton.cpp b/Engine/Source/Developer/AssetTools/Private/AssetTypeActions/AssetTypeActions_Skeleton.cpp
--- a/Engine/Source/Developer/AssetTools/Private/AssetTypeActions/AssetTypeActions_Skeleton.cpp
+++ b/Engine/Source/Developer/AssetTools/Private/AssetTypeActions/AssetTypeActions_Skeleton.cpp
@@ -614,9 +614,10 @@ void FAssetTypeActions_Skeleton::ExecuteRetargetSkeleton(TArray<TWeakObjectPtr<U
 			const FText Message = LOCTEXT("RetargetSkeleton_Warning", "This only converts animation data -i.e. animation assets and Anim Blueprints. \nIf you'd like to convert SkeletalMesh, use the context menu (Assign Skeleton) for each mesh. \n\nIf you'd like to convert mesh as well, please do so before converting animation data. \nOtherwise you will lose any extra track that is in the new mesh.");
 
 			bool bConvertSpaces = OldSkeleton != NULL;
+			bool bShowOnlyCompatibleSkeletons = OldSkeleton != NULL;
 
 			// ask user what they'd like to change to 
-			if (SAnimationRemapSkeleton::ShowModal(OldSkeleton, NewSkeleton, Message, OldSkeleton ? &bConvertSpaces : NULL))
+			if (SAnimationRemapSkeleton::ShowModal(OldSkeleton, NewSkeleton, Message, OldSkeleton? &bShowOnlyCompatibleSkeletons : NULL, OldSkeleton ? &bConvertSpaces : NULL))
 			{
 				// find all assets who references old skeleton
 				TArray<FName> Packages;
diff --git a/Engine/Source/Editor/Persona/Private/SAnimationSequenceBrowser.cpp b/Engine/Source/Editor/Persona/Private/SAnimationSequenceBrowser.cpp
--- a/Engine/Source/Editor/Persona/Private/SAnimationSequenceBrowser.cpp
+++ b/Engine/Source/Editor/Persona/Private/SAnimationSequenceBrowser.cpp
@@ -317,12 +317,13 @@ void SAnimationSequenceBrowser::OnCreateCopy(TArray<FAssetData> Selected)
 		USkeleton * NewSkeleton = NULL;
 		bool		bRemapReferencedAssets = true;
 		bool		bConvertSpaces = true;
+		bool		bShowOnlyCompatibleSkeleton = true;
 		bool		bDuplicateAssets = true;
 
 		const FText Message = LOCTEXT("RemapSkeleton_Warning", "This will duplicate the asset and convert to new skeleton.");
 
 		// ask user what they'd like to change to 
-		if (SAnimationRemapSkeleton::ShowModal(OldSkeleton, NewSkeleton, Message, &bConvertSpaces, &bRemapReferencedAssets))
+		if (SAnimationRemapSkeleton::ShowModal(OldSkeleton, NewSkeleton, Message, &bShowOnlyCompatibleSkeleton, &bConvertSpaces, &bRemapReferencedAssets))
 		{
 			UObject* AssetToOpen = EditorAnimUtils::RetargetAnimations(OldSkeleton, NewSkeleton, Selected, bRemapReferencedAssets, bDuplicateAssets, bConvertSpaces);
 
diff --git a/Engine/Source/Editor/UnrealEd/Private/SSkeletonWidget.cpp b/Engine/Source/Editor/UnrealEd/Private/SSkeletonWidget.cpp
--- a/Engine/Source/Editor/UnrealEd/Private/SSkeletonWidget.cpp
+++ b/Engine/Source/Editor/UnrealEd/Private/SSkeletonWidget.cpp
@@ -7,6 +7,8 @@
 #include "MainFrame.h"
 #include "Editor/ContentBrowser/Public/ContentBrowserModule.h"
 #include "IDocumentation.h"
+#include "Animation/Rig.h"
+#include "Editor/AnimGraph/Classes/AnimPreviewInstance.h"
 
 #define LOCTEXT_NAMESPACE "SkeletonWidget"
 
@@ -324,22 +326,79 @@ void SSkeletonSelectorWindow::ConstructWindow()
 		];
 }
 
-void SAnimationRemapSkeleton::Construct( const FArguments& InArgs )
+bool SAnimationRemapSkeleton::OnShouldFilterAsset(const class FAssetData& AssetData)
 {
-	OldSkeleton = InArgs._CurrentSkeleton;
-	NewSkeleton = NULL;
-	bRemapReferencedAssets = true;
-	WidgetWindow = InArgs._WidgetWindow;
+	USkeleton * AssetSkeleton = NULL;
+	if (AssetData.IsAssetLoaded())
+	{
+		AssetSkeleton = Cast<USkeleton>(AssetData.GetAsset());
+	}
 
-	// button 
-	FContentBrowserModule& ContentBrowserModule = FModuleManager::Get().LoadModuleChecked<FContentBrowserModule>(TEXT("ContentBrowser"));
+	// do not show same skeleton
+	if (OldSkeleton && OldSkeleton == AssetSkeleton)
+	{
+		return true;
+	}
+
+	if (bShowOnlyCompatibleSkeletons)
+	{
+		if(OldSkeleton && OldSkeleton->GetRig())
+		{
+			URig * Rig = OldSkeleton->GetRig();
+
+			const FString * Value = AssetData.TagsAndValues.Find(USkeleton::RigTag);
+
+			if(Value && Rig->GetFullName() == *Value)
+			{
+				return false;
+			}
+
+			// if loaded, check to see if it has same rig
+			if (AssetData.IsAssetLoaded())
+			{
+				USkeleton * LoadedSkeleton = Cast<USkeleton>(AssetData.GetAsset());
+
+				if (LoadedSkeleton && LoadedSkeleton->GetRig() == Rig)
+				{
+					return false;
+				}
+			}
+		}
+
+		return true;
+	}
 
+	return false;
+}
+
+void SAnimationRemapSkeleton::UpdateAssetPicker()
+{
 	FAssetPickerConfig AssetPickerConfig;
 	AssetPickerConfig.Filter.ClassNames.Add(USkeleton::StaticClass()->GetFName());
 	AssetPickerConfig.OnAssetSelected = FOnAssetSelected::CreateSP(this, &SAnimationRemapSkeleton::OnAssetSelectedFromPicker);
 	AssetPickerConfig.bAllowNullSelection = false;
 	AssetPickerConfig.InitialAssetViewType = EAssetViewType::Column;
 	AssetPickerConfig.ThumbnailScale = 0.0f;
+	AssetPickerConfig.OnShouldFilterAsset = FOnShouldFilterAsset::CreateSP(this, &SAnimationRemapSkeleton::OnShouldFilterAsset);
+
+	FContentBrowserModule& ContentBrowserModule = FModuleManager::Get().LoadModuleChecked<FContentBrowserModule>(TEXT("ContentBrowser"));
+
+	if (AssetPickerBox.IsValid())
+	{
+		AssetPickerBox->SetContent(
+			ContentBrowserModule.Get().CreateAssetPicker(AssetPickerConfig)
+		);
+	}
+}
+
+void SAnimationRemapSkeleton::Construct( const FArguments& InArgs )
+{
+	OldSkeleton = InArgs._CurrentSkeleton;
+	NewSkeleton = NULL;
+	WidgetWindow = InArgs._WidgetWindow;
+	bRemapReferencedAssets = true;
+	bConvertSpaces = false;
+	bShowOnlyCompatibleSkeletons = false;
 
 	TSharedRef<SVerticalBox> Widget = SNew(SVerticalBox);
 	if(InArgs._ShowRemapOption)
@@ -353,6 +412,8 @@ void SAnimationRemapSkeleton::Construct( const FArguments& InArgs )
 				SNew(STextBlock).Text(LOCTEXT("RemapSkeleton_RemapAssets", "Remap referenced assets "))
 			]
 		];
+
+		bRemapReferencedAssets = true;
 	}
 
 	if (InArgs._ShowConvertSpacesOption)
@@ -369,127 +430,199 @@ void SAnimationRemapSkeleton::Construct( const FArguments& InArgs )
 					.ToolTip(ConvertSpaceTooltip)
 				]
 			];
+
+		bConvertSpaces = true;
+	}
+
+	if(InArgs._ShowCompatibleDisplayOption)
+	{
+		TSharedPtr<SToolTip> ConvertSpaceTooltip = IDocumentation::Get()->CreateToolTip(FText::FromString("Check if you'd like to show only the skeleton that uses the same rig."),
+			NULL, FString("Shared/Editors/Persona"), FString("AnimRemapSkeleton_ShowCompatbielSkeletons")); // @todo add tooltip
+		Widget->AddSlot()
+			[
+				SNew(SCheckBox)
+				.IsChecked(this, &SAnimationRemapSkeleton::IsShowOnlyCompatibleSkeletonsChecked)
+				.IsEnabled(this, &SAnimationRemapSkeleton::IsShowOnlyCompatibleSkeletonsEnabled)
+				.OnCheckStateChanged(this, &SAnimationRemapSkeleton::OnShowOnlyCompatibleSkeletonsCheckChanged)
+				[
+					SNew(STextBlock).Text(LOCTEXT("RemapSkeleton_ShowCompatible", "Show Only Compatible Skeletons"))
+					.ToolTip(ConvertSpaceTooltip)
+				]
+			];
+
+		bShowOnlyCompatibleSkeletons = true;
 	}
 
 	TSharedPtr<SToolTip> SkeletonTooltip = IDocumentation::Get()->CreateToolTip(FText::FromString("Pick a skeleton for this mesh"), NULL, FString("Shared/Editors/Persona"), FString("Skeleton"));
 
 	this->ChildSlot
 		[
-		SNew (SVerticalBox)
-
-		+SVerticalBox::Slot()
-		.AutoHeight() 
-		.Padding(2) 
-		.HAlign(HAlign_Fill)
-		[
-			SNew(SHorizontalBox)
-				 
+			SNew (SHorizontalBox)
 			+SHorizontalBox::Slot()
 			.AutoWidth()
-			.VAlign(VAlign_Center)
-			+SHorizontalBox::Slot()
-			[
-				SNew(STextBlock)
-				.Text(LOCTEXT("CurrentlySelectedSkeletonLabel_SelectSkeleton", "Select Skeleton"))
-				.Font(FSlateFontInfo(FPaths::EngineContentDir() / TEXT("Slate/Fonts/Roboto-Regular.ttf"), 16))
-				.ToolTip(SkeletonTooltip)
-			]
-			+SHorizontalBox::Slot()
-			.FillWidth(1)
-			.HAlign(HAlign_Left)
-			.VAlign(VAlign_Center)
 			[
-				IDocumentation::Get()->CreateAnchor(FString("Engine/Animation/Skeleton"))
-			]
-		]
+				SNew(SVerticalBox)
 
-		+SVerticalBox::Slot()
-		.AutoHeight() 
-		.Padding(2, 10)
-		.HAlign(HAlign_Fill)
-		[
-			SNew(SSeparator)
-			.Orientation(Orient_Horizontal)
-		]
+				+SVerticalBox::Slot()
+				.AutoHeight()
+				.Padding(2)
+				.HAlign(HAlign_Fill)
+				[
+					SNew(SHorizontalBox)
 
-		+SVerticalBox::Slot()
-		.AutoHeight()
-		.HAlign(HAlign_Center)
-		.Padding(5)
-		[
-			SNew(STextBlock)
-			.Font( FSlateFontInfo( FPaths::EngineContentDir() / TEXT("Slate/Fonts/Roboto-Regular.ttf"), 10 ) )
-			.ColorAndOpacity(FLinearColor::Red)
-			.Text( LOCTEXT("RemapSkeleton_Warning_Title", "[WARNING]") )
-		]
+					+SHorizontalBox::Slot()
+					.AutoWidth()
+					.VAlign(VAlign_Center)
+					+SHorizontalBox::Slot()
+					[
+						SNew(STextBlock)
+						.Text(LOCTEXT("CurrentlySelectedSkeletonLabel_SelectSkeleton", "Select Skeleton"))
+						.Font(FSlateFontInfo(FPaths::EngineContentDir() / TEXT("Slate/Fonts/Roboto-Regular.ttf"), 16))
+						.ToolTip(SkeletonTooltip)
+					]
+					+SHorizontalBox::Slot()
+					.FillWidth(1)
+					.HAlign(HAlign_Left)
+					.VAlign(VAlign_Center)
+					[
+						IDocumentation::Get()->CreateAnchor(FString("Engine/Animation/Skeleton"))
+					]
+				]
 
-		+SVerticalBox::Slot()
-		.AutoHeight()
-		.Padding(5)
-		[
-			SNew(STextBlock)
-			.Font( FSlateFontInfo( FPaths::EngineContentDir() / TEXT("Slate/Fonts/Roboto-Regular.ttf"), 10 ) )
-			.ColorAndOpacity(FLinearColor::Red)
-			.Text( InArgs._WarningMessage )
-		]
+				+SVerticalBox::Slot()
+				.AutoHeight()
+				.Padding(2, 10)
+				.HAlign(HAlign_Fill)
+				[
+					SNew(SSeparator)
+					.Orientation(Orient_Horizontal)
+				]
 
-		+SVerticalBox::Slot()
-		.AutoHeight() 
-		.Padding(5)
-		[
-			SNew(SSeparator)
-		]
+				+SVerticalBox::Slot()
+				.AutoHeight()
+				.HAlign(HAlign_Center)
+				.Padding(5)
+				[
+					SNew(STextBlock)
+					.Font(FSlateFontInfo(FPaths::EngineContentDir() / TEXT("Slate/Fonts/Roboto-Regular.ttf"), 10))
+					.ColorAndOpacity(FLinearColor::Red)
+					.Text(LOCTEXT("RemapSkeleton_Warning_Title", "[WARNING]"))
+				]
 
-		+SVerticalBox::Slot()
-		.MaxHeight(500)
-		[
-			SNew(SBox)
-			.WidthOverride(400)
-			[
-				ContentBrowserModule.Get().CreateAssetPicker(AssetPickerConfig)
-			]
-		]
+				+SVerticalBox::Slot()
+				.AutoHeight()
+				.Padding(5)
+				[
+					SNew(STextBlock)
+					.Font(FSlateFontInfo(FPaths::EngineContentDir() / TEXT("Slate/Fonts/Roboto-Regular.ttf"), 10))
+					.ColorAndOpacity(FLinearColor::Red)
+					.Text(InArgs._WarningMessage)
+				]
 
-		+SVerticalBox::Slot()
-		.AutoHeight()
-		[
-			Widget
-		]
+				+SVerticalBox::Slot()
+				.AutoHeight()
+				.Padding(5)
+				[
+					SNew(SSeparator)
+				]
 
-		+SVerticalBox::Slot()
-		.AutoHeight() 
-		.Padding(5)
-		[
-			SNew(SSeparator)
-		]
+				+SVerticalBox::Slot()
+				.MaxHeight(500)
+				[
+					SAssignNew(AssetPickerBox, SBox)
+					.WidthOverride(400)
+				]
 
-		+SVerticalBox::Slot()
-		.AutoHeight()
-		.HAlign(HAlign_Right)
-		.VAlign(VAlign_Bottom)
-		[
-			SNew(SUniformGridPanel)
-			.SlotPadding(FEditorStyle::GetMargin("StandardDialog.SlotPadding"))
-			.MinDesiredSlotWidth(FEditorStyle::GetFloat("StandardDialog.MinDesiredSlotWidth"))
-			.MinDesiredSlotHeight(FEditorStyle::GetFloat("StandardDialog.MinDesiredSlotHeight"))
-			+SUniformGridPanel::Slot(0,0)
+				+SVerticalBox::Slot()
+				.AutoHeight()
+				[
+					Widget
+				]
+
+				+SVerticalBox::Slot()
+				.AutoHeight()
+				.Padding(5)
+				[
+					SNew(SSeparator)
+				]
+
+				+SVerticalBox::Slot()
+				.AutoHeight()
+				.HAlign(HAlign_Right)
+				.VAlign(VAlign_Bottom)
+				[
+					SNew(SUniformGridPanel)
+					.SlotPadding(FEditorStyle::GetMargin("StandardDialog.SlotPadding"))
+					.MinDesiredSlotWidth(FEditorStyle::GetFloat("StandardDialog.MinDesiredSlotWidth"))
+					.MinDesiredSlotHeight(FEditorStyle::GetFloat("StandardDialog.MinDesiredSlotHeight"))
+					+SUniformGridPanel::Slot(0, 0)
+					[
+						SNew(SButton).HAlign(HAlign_Center)
+						.Text(LOCTEXT("RemapSkeleton_Apply", "Select"))
+						.IsEnabled(this, &SAnimationRemapSkeleton::CanApply)
+						.OnClicked(this, &SAnimationRemapSkeleton::OnApply)
+						.HAlign(HAlign_Center)
+						.ContentPadding(FEditorStyle::GetMargin("StandardDialog.ContentPadding"))
+					]
+					+SUniformGridPanel::Slot(1, 0)
+					[
+						SNew(SButton).HAlign(HAlign_Center)
+						.Text(LOCTEXT("RemapSkeleton_Cancel", "Cancel"))
+						.ContentPadding(FEditorStyle::GetMargin("StandardDialog.ContentPadding"))
+						.OnClicked(this, &SAnimationRemapSkeleton::OnCancel)
+					]
+				]
+			]
+
+			+SHorizontalBox::Slot()
+			.Padding(2)
+			.AutoWidth()
 			[
-				SNew(SButton) .HAlign(HAlign_Center)
-				.Text(LOCTEXT("RemapSkeleton_Apply", "Select"))
-				.IsEnabled(this, &SAnimationRemapSkeleton::CanApply)
-				.OnClicked(this, &SAnimationRemapSkeleton::OnApply)
-				.HAlign(HAlign_Center)
-				.ContentPadding(FEditorStyle::GetMargin("StandardDialog.ContentPadding"))
+				SNew(SSeparator)
+				.Orientation(Orient_Vertical)
 			]
-			+SUniformGridPanel::Slot(1,0)
+
+			+SHorizontalBox::Slot()
+			.HAlign(HAlign_Center)
+			.VAlign(VAlign_Center)
+			.AutoWidth()
 			[
-				SNew(SButton) .HAlign(HAlign_Center)
-				.Text(LOCTEXT("RemapSkeleton_Cancel", "Cancel"))
-				.ContentPadding(FEditorStyle::GetMargin("StandardDialog.ContentPadding"))
-				.OnClicked(this, &SAnimationRemapSkeleton::OnCancel)
+				SNew(SVerticalBox)
+				+SVerticalBox::Slot()
+				.AutoHeight()
+				.Padding(5, 5)
+				[
+					// put nice message here
+					SNew(STextBlock)
+					.AutoWrapText(true)
+					.Font(FEditorStyle::GetFontStyle("Persona.RetargetManager.BoldFont"))
+					.Text(LOCTEXT("RetargetBasePose_WarningMessage", "Make sure you have the similar retarget base pose. If they don't look alike here, you can edit your base pose in the Retarget Mangaer window to look alike."))
+				]
+
+				+SVerticalBox::Slot()
+				.FillHeight(1)
+				.Padding(0, 5)
+				[
+					SNew(SHorizontalBox)
+
+					+SHorizontalBox::Slot()
+					[
+						SAssignNew(SourceViewport, SBasePoseViewport)
+						.Title(TEXT("[Source]"))
+						.Skeleton(OldSkeleton)
+					]
+
+					+SHorizontalBox::Slot()
+					[
+						SAssignNew(TargetViewport, SBasePoseViewport)
+						.Title(TEXT("[Target]"))
+						.Skeleton(NULL)
+					]
+				]
 			]
-		]
-	];
+		];
+
+	UpdateAssetPicker();
 }
 
 ESlateCheckBoxState::Type SAnimationRemapSkeleton::IsRemappingReferencedAssets() const
@@ -512,6 +645,24 @@ void SAnimationRemapSkeleton::OnConvertSpacesCheckChanged(ESlateCheckBoxState::T
 	bConvertSpaces = (InNewRadioState == ESlateCheckBoxState::Checked);
 }
 
+bool SAnimationRemapSkeleton::IsShowOnlyCompatibleSkeletonsEnabled() const
+{
+	// if convert space is false, compatible skeletons won't matter either. 
+	return (bConvertSpaces == true);
+}
+
+ESlateCheckBoxState::Type SAnimationRemapSkeleton::IsShowOnlyCompatibleSkeletonsChecked() const
+{
+	return bShowOnlyCompatibleSkeletons ? ESlateCheckBoxState::Checked : ESlateCheckBoxState::Unchecked;
+}
+
+void SAnimationRemapSkeleton::OnShowOnlyCompatibleSkeletonsCheckChanged(ESlateCheckBoxState::Type InNewRadioState)
+{
+	bShowOnlyCompatibleSkeletons = (InNewRadioState == ESlateCheckBoxState::Checked);
+
+	UpdateAssetPicker();
+}
+
 bool SAnimationRemapSkeleton::CanApply() const
 {
 	return (NewSkeleton!=NULL && NewSkeleton!=OldSkeleton);
@@ -522,6 +673,8 @@ void SAnimationRemapSkeleton::OnAssetSelectedFromPicker(const FAssetData& AssetD
 	if (AssetData.GetAsset())
 	{
 		NewSkeleton = Cast<USkeleton>(AssetData.GetAsset());
+
+		TargetViewport->SetSkeleton(NewSkeleton);
 	}
 }
 
@@ -545,15 +698,22 @@ void SAnimationRemapSkeleton::CloseWindow()
 	}
 }
 
-bool SAnimationRemapSkeleton::ShowModal(USkeleton * OldSkeleton, USkeleton * & NewSkeleton, const FText& WarningMessage, bool * bConvertSpace, bool * bRemapReferencedAssets)
+bool SAnimationRemapSkeleton::ShowModal(USkeleton * OldSkeleton, USkeleton * & NewSkeleton, const FText& WarningMessage, bool *bShowOnlyCompatibleSkeletons, bool * bConvertSpace, bool * bRemapReferencedAssets)
 {
-	TSharedPtr<class SAnimationRemapSkeleton> DialogWidget;
+	static TSharedPtr<SWindow> DialogWindow;
 
-	TSharedPtr<SWindow> DialogWindow = SNew(SWindow)
+	if(DialogWindow.IsValid())
+	{
+		FSlateApplication::Get().DestroyWindowImmediately(DialogWindow.ToSharedRef());
+	}
+
+	DialogWindow = SNew(SWindow)
 		.Title( LOCTEXT("RemapSkeleton", "Select Skeleton") )
 		.SupportsMinimize(false) .SupportsMaximize(false)
 		.SizingRule( ESizingRule::Autosized );
 
+	TSharedPtr<class SAnimationRemapSkeleton> DialogWidget;
+
 	TSharedPtr<SBorder> DialogWrapper = 
 		SNew(SBorder)
 		.BorderImage(FEditorStyle::GetBrush("ToolPanel.GroupBorder"))
@@ -565,11 +725,12 @@ bool SAnimationRemapSkeleton::ShowModal(USkeleton * OldSkeleton, USkeleton * & N
 			.WarningMessage(WarningMessage)
 			.ShowRemapOption(bRemapReferencedAssets != NULL)
 			.ShowConvertSpacesOption(bConvertSpace != NULL)
+			.ShowCompatibleDisplayOption(bShowOnlyCompatibleSkeletons != NULL)
 		];
 
 	DialogWindow->SetContent(DialogWrapper.ToSharedRef());
 
-	GEditor->EditorAddModalWindow(DialogWindow.ToSharedRef());
+	FSlateApplication::Get().AddWindow(DialogWindow.ToSharedRef());
 
 	NewSkeleton = DialogWidget.Get()->NewSkeleton;
 
@@ -880,5 +1041,171 @@ bool SSkeletonBoneRemoval::ShowModal(const TArray<FName> BonesToRemove, const FT
 	return DialogWidget.Get()->bShouldContinue;
 }
 
+////////////////////////////////////////
+
+class FBasePoseViewportClient: public FEditorViewportClient
+{
+public:
+	FBasePoseViewportClient(FPreviewScene& InPreviewScene)
+			: FEditorViewportClient(GLevelEditorModeTools(), &InPreviewScene)
+	{
+		SetViewMode(VMI_Lit);
+
+		// Always composite editor objects after post processing in the editor
+		EngineShowFlags.CompositeEditorPrimitives = true;
+		EngineShowFlags.DisableAdvancedFeatures();
+
+		UpdateLighting();
+
+		// Setup defaults for the common draw helper.
+		DrawHelper.bDrawPivot = false;
+		DrawHelper.bDrawWorldBox = false;
+		DrawHelper.bDrawKillZ = false;
+		DrawHelper.bDrawGrid = true;
+		DrawHelper.GridColorAxis = FColor(70, 70, 70);
+		DrawHelper.GridColorMajor = FColor(40, 40, 40);
+		DrawHelper.GridColorMinor =  FColor(20, 20, 20);
+		DrawHelper.PerspectiveGridSize = HALF_WORLD_MAX1;
+
+		bDisableInput = true;
+	}
+
+
+	// FlEditorViewportClient interface
+	virtual FSceneInterface* GetScene() const
+	{
+		return PreviewScene->GetScene();
+	}
+
+	virtual FLinearColor GetBackgroundColor() const override { return FLinearColor::White; }
+
+	// End of FEditorViewportClient
+
+	void UpdateLighting()
+	{
+		const UDestructableMeshEditorSettings* Options = GetDefault<UDestructableMeshEditorSettings>();
+
+		PreviewScene->SetLightDirection(Options->AnimPreviewLightingDirection);
+		PreviewScene->GetScene()->UpdateDynamicSkyLight(Options->AnimPreviewSkyBrightness * FLinearColor(Options->AnimPreviewSkyColor), Options->AnimPreviewSkyBrightness * FLinearColor(Options->AnimPreviewFloorColor));
+		PreviewScene->SetLightColor(Options->AnimPreviewDirectionalColor);
+		PreviewScene->SetLightBrightness(Options->AnimPreviewLightBrightness);
+	}
+};
+
+////////////////////////////////
+// SBasePoseViewport
+void SBasePoseViewport::Construct(const FArguments& InArgs)
+{
+	this->ChildSlot
+	[
+		SNew(SVerticalBox)
+		+SVerticalBox::Slot()
+		.AutoHeight()
+		[
+			SNew(STextBlock)
+			.Text(FText::FromString(InArgs._Title))
+			.Font(FEditorStyle::GetFontStyle("Persona.RetargetManager.FilterFont"))
+			.AutoWrapText(true)
+		]
+		+SVerticalBox::Slot()
+		.AutoHeight()
+		[
+			SNew(SBorder)
+			.HAlign(HAlign_Center)
+			[
+				SAssignNew(ViewportWidget, SViewport)
+				.EnableGammaCorrection(false)
+			]
+		]
+	];
+
+	// Create a viewport client
+	LevelViewportClient	= MakeShareable(new FBasePoseViewportClient(PreviewScene));
+
+	LevelViewportClient->ViewportType = LVT_Perspective;
+	LevelViewportClient->bSetListenerPosition = false;
+	LevelViewportClient->SetViewLocation(EditorViewportDefs::DefaultPerspectiveViewLocation);
+	LevelViewportClient->SetViewRotation(EditorViewportDefs::DefaultPerspectiveViewRotation);
+
+	SceneViewport = MakeShareable(new FSceneViewport(LevelViewportClient.Get(), ViewportWidget));
+	LevelViewportClient->Viewport = SceneViewport.Get();
+	LevelViewportClient->SetRealtime(true);
+	LevelViewportClient->VisibilityDelegate.BindSP(this, &SBasePoseViewport::IsVisible);
+	LevelViewportClient->SetViewMode(VMI_Lit);
+
+	ViewportWidget->SetViewportInterface(SceneViewport.ToSharedRef());
+
+	PreviewComponent = ConstructObject<UDebugSkelMeshComponent>(UDebugSkelMeshComponent::StaticClass());
+	PreviewComponent->MeshComponentUpdateFlag = EMeshComponentUpdateFlag::AlwaysTickPoseAndRefreshBones;
+	PreviewScene.AddComponent(PreviewComponent, FTransform::Identity);
+
+	SetSkeleton(InArgs._Skeleton);
+}
+
+void SBasePoseViewport::SetSkeleton(USkeleton * Skeleton)
+{
+	if(Skeleton != TargetSkeleton)
+	{
+		TargetSkeleton = Skeleton;
+
+		if(TargetSkeleton)
+		{
+			USkeletalMesh* PreviewSkeletalMesh = Skeleton->GetPreviewMesh();
+			if(PreviewSkeletalMesh)
+			{
+				PreviewComponent->SetSkeletalMesh(PreviewSkeletalMesh);
+				PreviewComponent->EnablePreview(true, NULL, NULL);
+//				PreviewComponent->AnimScriptInstance = PreviewComponent->PreviewInstance;
+				PreviewComponent->PreviewInstance->bForceRetargetBasePose = true;
+				PreviewComponent->RefreshBoneTransforms(NULL);
+
+				//Place the camera at a good viewer position
+				FVector NewPosition = LevelViewportClient->GetViewLocation();
+				NewPosition.Normalize();
+				if(PreviewSkeletalMesh)
+				{
+					NewPosition *= (PreviewSkeletalMesh->Bounds.SphereRadius*1.5f);
+				}
+				LevelViewportClient->SetViewLocation(NewPosition);
+			}
+			else
+			{
+				PreviewComponent->SetSkeletalMesh(NULL);
+			}
+		}
+		else
+		{
+			PreviewComponent->SetSkeletalMesh(NULL);
+		}
+
+		LevelViewportClient->Invalidate();
+	}
+}
+
+SBasePoseViewport::~SBasePoseViewport()
+{
+	// Close viewport
+	if(LevelViewportClient.IsValid())
+	{
+		LevelViewportClient->Viewport = NULL;
+	}
+}
+
+SBasePoseViewport::SBasePoseViewport()
+: PreviewScene(FPreviewScene::ConstructionValues())
+{
+}
+
+bool SBasePoseViewport::IsVisible() const
+{
+	return true;
+}
+
+void SBasePoseViewport::Tick(const FGeometry& AllottedGeometry, const double InCurrentTime, const float InDeltaTime)
+{
+	PreviewScene.GetWorld()->Tick(LEVELTICK_All, InDeltaTime);
+	LevelViewportClient->Tick(InDeltaTime);
+	SCompoundWidget::Tick(AllottedGeometry, InCurrentTime, InDeltaTime);
+}
 #undef LOCTEXT_NAMESPACE 
 
diff --git a/Engine/Source/Editor/UnrealEd/Public/SSkeletonWidget.h b/Engine/Source/Editor/UnrealEd/Public/SSkeletonWidget.h
--- a/Engine/Source/Editor/UnrealEd/Public/SSkeletonWidget.h
+++ b/Engine/Source/Editor/UnrealEd/Public/SSkeletonWidget.h
@@ -5,6 +5,7 @@
 #define __SSkeletonWidget_h__
 
 #include "Slate.h"
+#include "PreviewScene.h"
 
 class UAnimSet;
 
@@ -249,6 +250,42 @@ private:
 };
 
 
+//////////////////////////
+class SBasePoseViewport: public SCompoundWidget
+{
+public:
+	SLATE_BEGIN_ARGS(SBasePoseViewport)
+	{}
+
+	SLATE_ARGUMENT(FString, Title)
+	SLATE_ARGUMENT(USkeleton*, Skeleton)
+SLATE_END_ARGS()
+
+public:
+	SBasePoseViewport();
+	virtual ~SBasePoseViewport();
+
+	void Construct(const FArguments& InArgs);
+	void SetSkeleton(USkeleton * Skeleton);
+private:
+	TSharedPtr<FEditorViewportClient> LevelViewportClient;
+
+	/** Slate viewport for rendering and I/O */
+	TSharedPtr<SViewport> ViewportWidget;
+
+	TSharedPtr<class FSceneViewport> SceneViewport;
+
+	/** Skeleton */
+	USkeleton* TargetSkeleton;
+
+	FPreviewScene PreviewScene;
+
+	class UDebugSkelMeshComponent* PreviewComponent;
+
+	bool IsVisible() const;
+
+	virtual void Tick(const FGeometry& AllottedGeometry, const double InCurrentTime, const float InDeltaTime) override;
+};
 
 /////////////////////////////////////////////
 /** 
@@ -271,6 +308,7 @@ public:
 		SLATE_ARGUMENT( FText, WarningMessage )
 		SLATE_ARGUMENT( bool, ShowRemapOption )
 		SLATE_ARGUMENT( bool, ShowConvertSpacesOption )
+		SLATE_ARGUMENT( bool, ShowCompatibleDisplayOption )
 
 	SLATE_END_ARGS()	
 
@@ -291,6 +329,7 @@ public:
 	 */
 	USkeleton *NewSkeleton;
 
+private:
 	/**
 	 * Whether we are remapping assets that are referenced by the assets the user selects to remap
 	 */
@@ -301,6 +340,16 @@ public:
 	 */
 	bool bConvertSpaces;
 
+	/**
+	 * Whether to show skeletons with the same rig set up
+	*/
+	bool bShowOnlyCompatibleSkeletons;
+
+	TSharedPtr<SBasePoseViewport> SourceViewport;
+	TSharedPtr<SBasePoseViewport> TargetViewport;
+
+	TSharedPtr<SBox> AssetPickerBox;
+
 	TWeakPtr<SWindow> WidgetWindow;
 
 	/** Handlers for check box for remapping assets option */
@@ -311,6 +360,14 @@ public:
 	ESlateCheckBoxState::Type IsConvertSpacesChecked() const;
 	void OnConvertSpacesCheckChanged(ESlateCheckBoxState::Type InNewRadioState);
 
+	/** Handlers for check box for converting spaces*/
+	ESlateCheckBoxState::Type IsShowOnlyCompatibleSkeletonsChecked() const;
+	bool IsShowOnlyCompatibleSkeletonsEnabled() const;
+	void OnShowOnlyCompatibleSkeletonsCheckChanged(ESlateCheckBoxState::Type InNewRadioState);
+
+	/** should filter asset */
+	bool OnShouldFilterAsset(const class FAssetData& AssetData);
+
 	/**
 	 * return true if it can apply 
 	 */
@@ -326,6 +383,12 @@ public:
 	 */
 	void OnAssetSelectedFromPicker(const FAssetData& AssetData);
 
+	/*
+	* Refreshes asset picker - call when asset picker option changes
+	*/
+	void UpdateAssetPicker();
+public:
+
 	/**
 	 *  Show Modal window
 	 *
@@ -334,7 +397,7 @@ public:
 	 *
 	 * @return true if successfully selected new skeleton
 	 */
-	static UNREALED_API bool ShowModal(USkeleton * OldSkeleton, USkeleton * & NewSkeleton, const FText& WarningMessage, bool * bConvertSpace = NULL, bool * bRemapReferencedAssets=NULL);
+	static UNREALED_API bool ShowModal(USkeleton * OldSkeleton, USkeleton * & NewSkeleton, const FText& WarningMessage, bool *bShowOnlyCompatibleSkeletons, bool * bConvertSpace = NULL, bool * bRemapReferencedAssets=NULL);
 };
 
 ////////////////////////////////////////////////////
diff --git a/Engine/Source/Runtime/Engine/Classes/Animation/Skeleton.h b/Engine/Source/Runtime/Engine/Classes/Animation/Skeleton.h
--- a/Engine/Source/Runtime/Engine/Classes/Animation/Skeleton.h
+++ b/Engine/Source/Runtime/Engine/Classes/Animation/Skeleton.h
@@ -211,9 +211,6 @@ private:
 	/** Guid for skeleton */
 	FGuid Guid;
 
-	UPROPERTY()
-	FRigConfiguration RigConfig;
-
 	/** Conversion function. Remove when VER_UE4_REFERENCE_SKELETON_REFACTOR is removed. */
 	void ConvertToFReferenceSkeleton();
 
@@ -262,6 +259,12 @@ private:
 	UPROPERTY()
 	TArray<FName> SlotGroupNames;
 
+	UPROPERTY()
+	FRigConfiguration RigConfig;
+
+	/** rig property will be saved separately */
+	virtual void GetAssetRegistryTags(TArray<FAssetRegistryTag>& OutTags) const override;
+
 public:
 
 	// Adds a new name to the smart name container and modifies the skeleton so it can be saved
@@ -625,6 +628,7 @@ public:
 	ENGINE_API static const TCHAR CurveTagDelimiter;
 
 	// rig Configs
+	ENGINE_API static const FName RigTag;
 	ENGINE_API void SetRigConfig(URig * Rig);
 	ENGINE_API FName GetRigBoneMapping(const FName & NodeName) const;
 	ENGINE_API bool SetRigBoneMapping(const FName & NodeName, FName BoneName);
diff --git a/Engine/Source/Runtime/Engine/Private/Animation/Skeleton.cpp b/Engine/Source/Runtime/Engine/Private/Animation/Skeleton.cpp
--- a/Engine/Source/Runtime/Engine/Private/Animation/Skeleton.cpp
+++ b/Engine/Source/Runtime/Engine/Private/Animation/Skeleton.cpp
@@ -20,6 +20,8 @@ const TCHAR USkeleton::AnimNotifyTagDelimiter = TEXT(';');
 
 const FName USkeleton::CurveTag = FName(TEXT("CurveUIDList"));
 const TCHAR USkeleton::CurveTagDelimiter = TEXT(';');
+
+const FName USkeleton::RigTag = FName(TEXT("Rig"));
 #endif 
 
 // Names of smartname containers for skeleton properties
@@ -1106,5 +1108,13 @@ URig * USkeleton::GetRig() const
 {
 	return RigConfig.Rig;
 }
+
+void USkeleton::GetAssetRegistryTags(TArray<FAssetRegistryTag>& OutTags) const
+{
+	Super::GetAssetRegistryTags(OutTags);
+	FString RigFullName = (RigConfig.Rig)? RigConfig.Rig->GetFullName() : TEXT("");
+
+	OutTags.Add(FAssetRegistryTag(USkeleton::RigTag, RigFullName, FAssetRegistryTag::TT_Hidden));
+}
 #endif //WITH_EDITOR
 #undef LOCTEXT_NAMESPACE 
\ No newline at end of file
