UE4: Dirty world packages are now properly discovered. The content browser now saves dirty world packages when clicking the save button on the browser.

diff --git a/Engine/Source/Editor/ContentBrowser/Private/ContentBrowserUtils.cpp b/Engine/Source/Editor/ContentBrowser/Private/ContentBrowserUtils.cpp
--- a/Engine/Source/Editor/ContentBrowser/Private/ContentBrowserUtils.cpp
+++ b/Engine/Source/Editor/ContentBrowser/Private/ContentBrowserUtils.cpp
@@ -638,7 +638,7 @@ bool ContentBrowserUtils::SavePackages(const TArray<UPackage*>& Packages)
 bool ContentBrowserUtils::SaveDirtyPackages()
 {
 	const bool bPromptUserToSave = true;
-	const bool bSaveMapPackages = false;
+	const bool bSaveMapPackages = true;
 	const bool bSaveContentPackages = true;
 	return FEditorFileUtils::SaveDirtyPackages( bPromptUserToSave, bSaveMapPackages, bSaveContentPackages );
 }
diff --git a/Engine/Source/Editor/UnrealEd/Private/FileHelpers.cpp b/Engine/Source/Editor/UnrealEd/Private/FileHelpers.cpp
--- a/Engine/Source/Editor/UnrealEd/Private/FileHelpers.cpp
+++ b/Engine/Source/Editor/UnrealEd/Private/FileHelpers.cpp
@@ -3115,13 +3115,9 @@ FString FEditorFileUtils::ExtractPackageName(const FString& ObjectPath)
 
 void FEditorFileUtils::GetDirtyWorldPackages(TArray<UPackage*>& OutDirtyPackages)
 {
-	// If we are saving map packages, collect all valid worlds and see if their package is dirty
-	TArray<UWorld*> Worlds;
-	EditorLevelUtils::GetWorlds(GWorld, Worlds, true);
-
-	for (int32 WorldIdx = 0; WorldIdx < Worlds.Num(); ++WorldIdx)
+	for (TObjectIterator<UWorld> WorldIt; WorldIt; ++WorldIt)
 	{
-		UPackage* WorldPackage = Worlds[WorldIdx]->GetOutermost();
+		UPackage* WorldPackage = WorldIt->GetOutermost();
 		if (WorldPackage->IsDirty() && (WorldPackage->PackageFlags & PKG_PlayInEditor) == 0
 			&& !WorldPackage->HasAnyFlags(RF_Transient))
 		{
