Fixed non-unity build

diff --git a/Engine/Source/Editor/UnrealEd/Private/SSkeletonWidget.cpp b/Engine/Source/Editor/UnrealEd/Private/SSkeletonWidget.cpp
--- a/Engine/Source/Editor/UnrealEd/Private/SSkeletonWidget.cpp
+++ b/Engine/Source/Editor/UnrealEd/Private/SSkeletonWidget.cpp
@@ -9,6 +9,7 @@
 #include "IDocumentation.h"
 #include "Animation/Rig.h"
 #include "Editor/AnimGraph/Classes/AnimPreviewInstance.h"
+#include "SceneViewport.h"
 
 #define LOCTEXT_NAMESPACE "SkeletonWidget"
 
