Standalone Listen Server now works correctly

#ttp 338115 - SHOOTERGAME: EDITOR: Attempting to Play in Standalone as a Listen Server will make all clients fail to join

#branch UE4

#change Previously a client was being launched instead of a listen server, so it was attempting to connect to something that wasn't there, rather than host.

diff --git a/Engine/Source/Editor/UnrealEd/Classes/Settings/LevelEditorPlaySettings.h b/Engine/Source/Editor/UnrealEd/Classes/Settings/LevelEditorPlaySettings.h
--- a/Engine/Source/Editor/UnrealEd/Classes/Settings/LevelEditorPlaySettings.h
+++ b/Engine/Source/Editor/UnrealEd/Classes/Settings/LevelEditorPlaySettings.h
@@ -191,7 +191,7 @@ public:
 	UPROPERTY(config, EditAnywhere, Category=MultiplayerOptions)
 	bool PlayNetDedicated;
 
-	/** Number of player windows that should be opened. The editor will count as a player, a dedicated server will not. */
+	/** Number of clients that should be spawned. The editor and listen server will count as clients, a dedicated server will not. */
 	UPROPERTY(config, EditAnywhere, Category=MultiplayerOptions, meta=(ClampMin = "1", UIMin = "1", UIMax = "64"))
 	int32 PlayNumberOfClients;
 
diff --git a/Engine/Source/Editor/UnrealEd/Private/PlayLevel.cpp b/Engine/Source/Editor/UnrealEd/Private/PlayLevel.cpp
--- a/Engine/Source/Editor/UnrealEd/Private/PlayLevel.cpp
+++ b/Engine/Source/Editor/UnrealEd/Private/PlayLevel.cpp
@@ -774,8 +774,8 @@ void UEditorEngine::StartQueuedPlayMapRequest()
 		FIntPoint WinPosition((int32)PreferredWorkArea.Right, (int32)PreferredWorkArea.Top);
 		AdvanceWindowPositionsForNextPIEWindow(WinPosition.X, WinPosition.Y, FIntPoint(0,0));
 
-		// Do we need to spawn a server?
-		if (PlayInSettings->PlayNetMode == PIE_Client)
+		// We'll need to spawn a server if we're playing outside the editor or the editor wants to run as a client
+		if (bPlayOnLocalPcSession || PlayInSettings->PlayNetMode == PIE_Client)
 		{			
 			PlayStandaloneLocalPc(TEXT(""), &WinPosition, NumClients, true);
 			
@@ -786,6 +786,7 @@ void UEditorEngine::StartQueuedPlayMapRequest()
 			}
 		}
 		
+		// If we're playing in the editor
 		if (!bPlayOnLocalPcSession)
 		{
 			if (bStartMovieCapture)
@@ -805,7 +806,7 @@ void UEditorEngine::StartQueuedPlayMapRequest()
 		// Spawn number of clients
 		for (int32 i = NumClients; i < PlayInSettings->PlayNumberOfClients; ++i)
 		{
-			PlayStandaloneLocalPc(TEXT("127.0.0.1"), &WinPosition, NumClients + 1, false);
+			PlayStandaloneLocalPc(TEXT("127.0.0.1"), &WinPosition, i, false);
 		}
 	}
 	else
