BP: Crash in FBlueprintEditorUtils::CheckIfGraphHasLatentFunctions

diff --git a/Engine/Source/Editor/UnrealEd/Private/Kismet2/BlueprintEditorUtils.cpp b/Engine/Source/Editor/UnrealEd/Private/Kismet2/BlueprintEditorUtils.cpp
--- a/Engine/Source/Editor/UnrealEd/Private/Kismet2/BlueprintEditorUtils.cpp
+++ b/Engine/Source/Editor/UnrealEd/Private/Kismet2/BlueprintEditorUtils.cpp
@@ -6459,7 +6459,8 @@ bool FBlueprintEditorUtils::CheckIfGraphHasLatentFunctions(UEdGraph* InGraph)
 					if(const UK2Node_CallFunction* CallFunctionNode = Cast<UK2Node_CallFunction>(Node))
 					{
 						// Check any function call nodes to see if they are latent.
-						if(CallFunctionNode->GetTargetFunction()->HasMetaData(FBlueprintMetadata::MD_Latent))
+						auto TargetFunction = CallFunctionNode->GetTargetFunction();
+						if (TargetFunction && TargetFunction->HasMetaData(FBlueprintMetadata::MD_Latent))
 						{
 							TunnelNode->MetaData.HasLatentFunctions = 1;
 							return true;
