Deleting local function variables in Blueprints will no longer delete all local variables of the same name in the Blueprint.

#ttp 346132 - LIVE: BP: Deleting a Local Variable in a function will delete nodes of a Local Variable in a separate function if both Local Variables share the same name

diff --git a/Engine/Source/Editor/Kismet/Private/SMyBlueprint.cpp b/Engine/Source/Editor/Kismet/Private/SMyBlueprint.cpp
--- a/Engine/Source/Editor/Kismet/Private/SMyBlueprint.cpp
+++ b/Engine/Source/Editor/Kismet/Private/SMyBlueprint.cpp
@@ -1740,7 +1740,7 @@ void SMyBlueprint::OnDeleteEntry()
 	}
 	else if ( FEdGraphSchemaAction_K2LocalVar* LocalVarAction = SelectionAsLocalVar() )
 	{
-		if(FBlueprintEditorUtils::IsVariableUsed(GetBlueprintObj(), LocalVarAction->GetVariableName()))
+		if(FBlueprintEditorUtils::IsVariableUsed(GetBlueprintObj(), LocalVarAction->GetVariableName(), FBlueprintEditorUtils::FindScopeGraph(GetBlueprintObj(), LocalVarAction->GetVariableScope())))
 		{
 			FText ConfirmDelete = FText::Format(LOCTEXT( "ConfirmDeleteLocalVariableInUse",
 				"Local Variable {0} is in use! Do you really want to delete it?"),
diff --git a/Engine/Source/Editor/UnrealEd/Private/Kismet2/BlueprintEditorUtils.cpp b/Engine/Source/Editor/UnrealEd/Private/Kismet2/BlueprintEditorUtils.cpp
--- a/Engine/Source/Editor/UnrealEd/Private/Kismet2/BlueprintEditorUtils.cpp
+++ b/Engine/Source/Editor/UnrealEd/Private/Kismet2/BlueprintEditorUtils.cpp
@@ -3307,7 +3307,7 @@ void FBlueprintEditorUtils::BulkRemoveMemberVariables(UBlueprint* Blueprint, con
 	}
 }
 
-void FBlueprintEditorUtils::RemoveVariableNodes(UBlueprint* Blueprint, const FName& VarName, bool const bForSelfOnly)
+void FBlueprintEditorUtils::RemoveVariableNodes(UBlueprint* Blueprint, const FName& VarName, bool const bForSelfOnly, UEdGraph* LocalGraphScope)
 {
 	TArray<UEdGraph*> AllGraphs;
 	Blueprint->GetAllGraphs(AllGraphs);
@@ -3328,9 +3328,12 @@ void FBlueprintEditorUtils::RemoveVariableNodes(UBlueprint* Blueprint, const FNa
 
 			if ((SelfClass == VariableParent) || !bForSelfOnly)
 			{
-				if(VarName == CurrentNode->GetVarName())
+				if(LocalGraphScope == CurrentNode->GetGraph() || LocalGraphScope == nullptr)
 				{
-					CurrentNode->DestroyNode();
+					if(VarName == CurrentNode->GetVarName())
+					{
+						CurrentNode->DestroyNode();
+					}
 				}
 			}
 		}
@@ -3606,7 +3609,7 @@ void FBlueprintEditorUtils::RemoveLocalVariable(UBlueprint* InBlueprint, const U
 			if(GraphNodes[0]->LocalVariables[VarIdx].VarName == InVarName)
 			{
 				GraphNodes[0]->LocalVariables.RemoveAt(VarIdx);
-				FBlueprintEditorUtils::RemoveVariableNodes(InBlueprint, InVarName);
+				FBlueprintEditorUtils::RemoveVariableNodes(InBlueprint, InVarName, true, ScopeGraph);
 				FBlueprintEditorUtils::MarkBlueprintAsStructurallyModified(InBlueprint);
 
 				bFoundLocalVariable = true;
@@ -3854,7 +3857,7 @@ bool FBlueprintEditorUtils::IsVariableComponent(const FBPVariableDescription& Va
 	return false;
 }
 
-bool FBlueprintEditorUtils::IsVariableUsed(const UBlueprint* Blueprint, const FName& Name )
+bool FBlueprintEditorUtils::IsVariableUsed(const UBlueprint* Blueprint, const FName& Name, UEdGraph* LocalGraphScope/* = nullptr*/ )
 {
 	TArray<UEdGraph*> AllGraphs;
 	Blueprint->GetAllGraphs(AllGraphs);
@@ -3862,15 +3865,18 @@ bool FBlueprintEditorUtils::IsVariableUsed(const UBlueprint* Blueprint, const FN
 	{
 		const UEdGraph* CurrentGraph = *it;
 
-		TArray<UK2Node_Variable*> GraphNodes;
-		CurrentGraph->GetNodesOfClass(GraphNodes);
-
-		for( TArray<UK2Node_Variable*>::TConstIterator NodeIt(GraphNodes); NodeIt; ++NodeIt )
+		if(CurrentGraph == LocalGraphScope || LocalGraphScope == nullptr)
 		{
-			UK2Node_Variable* CurrentNode = *NodeIt;
-			if(Name == CurrentNode->GetVarName())
+			TArray<UK2Node_Variable*> GraphNodes;
+			CurrentGraph->GetNodesOfClass(GraphNodes);
+
+			for( TArray<UK2Node_Variable*>::TConstIterator NodeIt(GraphNodes); NodeIt; ++NodeIt )
 			{
-				return true;
+				UK2Node_Variable* CurrentNode = *NodeIt;
+				if(Name == CurrentNode->GetVarName())
+				{
+					return true;
+				}
 			}
 		}
 	}
diff --git a/Engine/Source/Editor/UnrealEd/Public/Kismet2/BlueprintEditorUtils.h b/Engine/Source/Editor/UnrealEd/Public/Kismet2/BlueprintEditorUtils.h
--- a/Engine/Source/Editor/UnrealEd/Public/Kismet2/BlueprintEditorUtils.h
+++ b/Engine/Source/Editor/UnrealEd/Public/Kismet2/BlueprintEditorUtils.h
@@ -536,12 +536,13 @@ public:
 	/**
 	 * Removes the variable nodes associated with the specified var name
 	 *
-	 * @param	Blueprint		The blueprint you want variable nodes removed from.
-	 * @param	VarName			Name of the variable to be removed.
-	 * @param	bForSelfOnly	True if you only want to delete variable nodes that represent ones owned by this blueprint,
-	 *							false if you just want everything with the specified name removed (variables from other classes too).
+	 * @param	Blueprint			The blueprint you want variable nodes removed from.
+	 * @param	VarName				Name of the variable to be removed.
+	 * @param	bForSelfOnly		True if you only want to delete variable nodes that represent ones owned by this blueprint,
+	 * @param	LocalGraphScope		Local scope graph of variables
+	 *								false if you just want everything with the specified name removed (variables from other classes too).
 	 */
-	static void RemoveVariableNodes(UBlueprint* Blueprint, const FName& VarName, bool const bForSelfOnly = true);
+	static void RemoveVariableNodes(UBlueprint* Blueprint, const FName& VarName, bool const bForSelfOnly = true, UEdGraph* LocalGraphScope = nullptr);
 
 	/**Rename a member variable*/
 	static void RenameMemberVariable(UBlueprint* Blueprint, const FName& OldName, const FName& NewName);
@@ -783,7 +784,7 @@ public:
 	static bool IsVariableComponent(const FBPVariableDescription& Variable);
 
 	/** Indicates if the variable is used on any graphs in this Blueprint*/
-	static bool IsVariableUsed(const UBlueprint* Blueprint, const FName& Name);
+	static bool IsVariableUsed(const UBlueprint* Blueprint, const FName& Name, UEdGraph* LocalGraphScope = nullptr);
 
 	static void ImportKismetDefaultValueToProperty(UEdGraphPin* SourcePin, UProperty* DestinationProperty, uint8* DestinationAddress, UObject* OwnerObject);
 	
