Merge cook on the fly and cook by the book.
Support for cook by the book in the editor.

diff --git a/Engine/Source/Developer/LauncherServices/Private/Launcher/LauncherPackageCommands.h b/Engine/Source/Developer/LauncherServices/Private/Launcher/LauncherPackageCommands.h
--- a/Engine/Source/Developer/LauncherServices/Private/Launcher/LauncherPackageCommands.h
+++ b/Engine/Source/Developer/LauncherServices/Private/Launcher/LauncherPackageCommands.h
@@ -42,7 +42,7 @@ public:
 		CommandLine += CookCommand.IsValid() ? CookCommand->GetDependencyArguments(ChainState) : TEXT(" -skipcook");
         
 		// Determine if we should force PAK for the current configuration+platform+cook mode
-		if (ChainState.Profile->GetCookMode() == ELauncherProfileCookModes::ByTheBook)
+		if (ChainState.Profile->GetCookMode() == ELauncherProfileCookModes::ByTheBook || ChainState.Profile->GetCookMode() == ELauncherProfileCookModes::ByTheBookInEditor)
 		{
 			if (ChainState.Profile->IsPackingWithUnrealPak() || (TargetPlatform.PlatformName() == TEXT("LinuxServer")) || (TargetPlatform.PlatformName() == TEXT("PS4")) || TargetPlatform.PlatformName().StartsWith(TEXT("Android")))
 			{
@@ -62,7 +62,7 @@ public:
 			CookCommand.IsValid() ? *CookCommand->GetAdditionalArguments(ChainState) : TEXT(""),
 			CookCommand.IsValid() ? *(TEXT(" ") + CookCommand->GetAdditionalArguments(ChainState)) : TEXT(""),
 			((TargetPlatform.PlatformName() == TEXT("PS4") || ChainState.Profile->IsPackingWithUnrealPak())
-			&& ChainState.Profile->GetCookMode() == ELauncherProfileCookModes::ByTheBook) ? TEXT(" -pak") : TEXT(""),
+			&& (ChainState.Profile->GetCookMode() == ELauncherProfileCookModes::ByTheBook || ChainState.Profile->GetCookMode() == ELauncherProfileCookModes::ByTheBookInEditor)) ? TEXT(" -pak") : TEXT(""),
 			ChainState.Profile->GetLaunchRoles().Num() > 0 ? (ChainState.Profile->GetLaunchRoles()[0]->IsVsyncEnabled() ? TEXT(" -vsync") : TEXT("")) : TEXT(""));
 
 		return CommandLine;
@@ -136,7 +136,7 @@ public:
  
 		// if it is PS4 or Android and cook by the book, make sure to add -pak
 		if ((TargetPlatform.PlatformName() == TEXT("PS4") || TargetPlatform.PlatformName() == TEXT("LinuxServer") || ChainState.Profile->IsPackingWithUnrealPak())
-			&& ChainState.Profile->GetCookMode() == ELauncherProfileCookModes::ByTheBook)
+			&& (ChainState.Profile->GetCookMode() == ELauncherProfileCookModes::ByTheBook||ChainState.Profile->GetCookMode()==ELauncherProfileCookModes::ByTheBookInEditor))
 		{
 			CommandLine += TEXT(" -pak");
 		}
@@ -152,7 +152,7 @@ public:
 			*ChainState.Profile->GetName(),
 			CookCommand.IsValid() ? *(TEXT(" ") + CookCommand->GetAdditionalArguments(ChainState)) : TEXT(""),
 			((TargetPlatform.PlatformName() == TEXT("PS4") || ChainState.Profile->IsPackingWithUnrealPak())
-			&& ChainState.Profile->GetCookMode() == ELauncherProfileCookModes::ByTheBook) ? TEXT(" -pak") : TEXT(""),
+			&& (ChainState.Profile->GetCookMode() == ELauncherProfileCookModes::ByTheBook || ChainState.Profile->GetCookMode() == ELauncherProfileCookModes::ByTheBookInEditor)) ? TEXT(" -pak") : TEXT(""),
 			ChainState.Profile->GetLaunchRoles().Num() > 0 ? (ChainState.Profile->GetLaunchRoles()[0]->IsVsyncEnabled() ? TEXT(" -vsync") : TEXT("")) : TEXT(""));
 
 		return CommandLine;
@@ -211,7 +211,7 @@ public:
  
 		// if it is PS4 or Android and cook by the book, make sure to add -pak
 		if ((TargetPlatform.PlatformName() == TEXT("PS4") || TargetPlatform.PlatformName() == TEXT("Android") || ChainState.Profile->IsPackingWithUnrealPak())
-			&& ChainState.Profile->GetCookMode() == ELauncherProfileCookModes::ByTheBook)
+			&& (ChainState.Profile->GetCookMode() == ELauncherProfileCookModes::ByTheBook || ChainState.Profile->GetCookMode() == ELauncherProfileCookModes::ByTheBookInEditor))
 		{
 			CommandLine += TEXT(" -pak");
 		}
@@ -233,7 +233,7 @@ public:
 			*ChainState.Profile->GetName(),
 			CookCommand.IsValid() ? *(TEXT(" ") + CookCommand->GetAdditionalArguments(ChainState)) : TEXT(""),
 			((TargetPlatform.PlatformName() == TEXT("PS4") || ChainState.Profile->IsPackingWithUnrealPak())
-			&& ChainState.Profile->GetCookMode() == ELauncherProfileCookModes::ByTheBook) ? TEXT(" -pak") : TEXT(""),
+			&& (ChainState.Profile->GetCookMode() == ELauncherProfileCookModes::ByTheBook || ChainState.Profile->GetCookMode() == ELauncherProfileCookModes::ByTheBookInEditor)) ? TEXT(" -pak") : TEXT(""),
 			ChainState.Profile->GetLaunchRoles().Num() > 0 ? (ChainState.Profile->GetLaunchRoles()[0]->IsVsyncEnabled() ? TEXT(" -vsync") : TEXT("")) : TEXT(""));
 
 		return CommandLine;
@@ -322,7 +322,7 @@ public:
  
 		// if it is PS4 and cook by the book, make sure to add -pak
 		if ((TargetPlatform.PlatformName() == TEXT("PS4") || TargetPlatform.PlatformName() == TEXT("LinuxServer") || ChainState.Profile->IsPackingWithUnrealPak())
-			&& ChainState.Profile->GetCookMode() == ELauncherProfileCookModes::ByTheBook)
+			&& (ChainState.Profile->GetCookMode() == ELauncherProfileCookModes::ByTheBook || ChainState.Profile->GetCookMode() == ELauncherProfileCookModes::ByTheBookInEditor) )
 		{
 			CommandLine += TEXT(" -pak");
 		}
@@ -338,7 +338,7 @@ public:
 			*ChainState.Profile->GetName(),
 			CookCommand.IsValid() ? *(TEXT(" ") + CookCommand->GetAdditionalArguments(ChainState)) : TEXT(""),
 			((TargetPlatform.PlatformName() == TEXT("PS4") || ChainState.Profile->IsPackingWithUnrealPak())
-			&& ChainState.Profile->GetCookMode() == ELauncherProfileCookModes::ByTheBook) ? TEXT(" -pak") : TEXT(""),
+			&& (ChainState.Profile->GetCookMode() == ELauncherProfileCookModes::ByTheBook || ChainState.Profile->GetCookMode() == ELauncherProfileCookModes::ByTheBookInEditor)) ? TEXT(" -pak") : TEXT(""),
 			ChainState.Profile->GetLaunchRoles().Num() > 0 ? (ChainState.Profile->GetLaunchRoles()[0]->IsVsyncEnabled() ? TEXT(" -vsync") : TEXT("")) : TEXT(""));
 
 		return CommandLine;
diff --git a/Engine/Source/Developer/LauncherServices/Private/Launcher/LauncherWorker.cpp b/Engine/Source/Developer/LauncherServices/Private/Launcher/LauncherWorker.cpp
--- a/Engine/Source/Developer/LauncherServices/Private/Launcher/LauncherWorker.cpp
+++ b/Engine/Source/Developer/LauncherServices/Private/Launcher/LauncherWorker.cpp
@@ -226,6 +226,13 @@ void FLauncherWorker::CreateAndExecuteTasks( const ILauncherProfileRef& InProfil
 		}
 	}
 
+
+#if !WITH_EDITOR
+	// can't cook by the book in the editor if we are not in the editor...
+	check( InProfile->GetCookMode() != ELauncherProfileCookModes::ByTheBookInEditor );
+#endif
+
+
 	// for each desired platform...
 	for (int32 PlatformIndex = 0; PlatformIndex < Platforms.Num(); ++PlatformIndex)
 	{
@@ -312,7 +319,38 @@ void FLauncherWorker::CreateAndExecuteTasks( const ILauncherProfileRef& InProfil
                 CookCommand = Command;
             }
         }
+		else if ( InProfile->GetCookMode() == ELauncherProfileCookModes::ByTheBookInEditor )
+		{
+			// need a command which will wait for the cook to finish
+			class FWaitForCookInEditorToFinish : public FLauncherTask
+			{
+			public:
+				FWaitForCookInEditorToFinish() : FLauncherTask( FString(TEXT("CookByTheBookInEditor")), FString(TEXT("CookByTheBookInEditorDesk")), NULL, NULL)
+				{
+				}
+				virtual bool PerformTask( FLauncherTaskChainState& ChainState ) override
+				{
+					while ( !ChainState.Profile->OnIsCookFinished().Execute() )
+					{
+						FPlatformProcess::Sleep( 0.1f );
+					}
+					return true;
+				}
+			};
+
+			TSharedPtr<FLauncherTask> Command = MakeShareable( new FWaitForCookInEditorToFinish() );
 
+			if (!PerPlatformCookTask.IsValid())
+			{
+				PerPlatformCookTask = Command;
+				FirstPlatformCookTask = Command;
+			}
+			else
+			{
+				PerPlatformCookTask->AddContinuation(Command);
+				PerPlatformCookTask = Command;
+			}
+		}
 		// ... package the build...
 		if (InProfile->GetPackagingMode() != ELauncherProfilePackagingModes::DoNotPackage || ((TargetPlatform->PlatformName() == TEXT("IOS") || TargetPlatform->PlatformName() == TEXT("HTML5")) && InProfile->GetDeploymentMode() != ELauncherProfileDeploymentModes::CopyRepository && InProfile->GetDeploymentMode() != ELauncherProfileDeploymentModes::DoNotDeploy))
 		{
diff --git a/Engine/Source/Developer/LauncherServices/Private/Profiles/LauncherProfile.h b/Engine/Source/Developer/LauncherServices/Private/Profiles/LauncherProfile.h
--- a/Engine/Source/Developer/LauncherServices/Private/Profiles/LauncherProfile.h
+++ b/Engine/Source/Developer/LauncherServices/Private/Profiles/LauncherProfile.h
@@ -410,7 +410,7 @@ public:
 
 	virtual bool IsDeployablePlatform( const FString& PlatformName ) override
 	{
-		if (CookMode == ELauncherProfileCookModes::ByTheBook)
+		if (CookMode == ELauncherProfileCookModes::ByTheBook || CookMode == ELauncherProfileCookModes::ByTheBookInEditor)
 		{
 			return CookedPlatforms.Contains(PlatformName);
 		}
@@ -700,6 +700,11 @@ public:
 		Validate();
 	}
 
+	virtual FIsCookFinishedDelegate& OnIsCookFinished() override
+	{
+		return IsCookFinishedDelegate;
+	}
+
 	virtual void SetDeploymentMode( ELauncherProfileDeploymentModes::Type Mode ) override
 	{
 		if (DeploymentMode != Mode)
@@ -922,7 +927,7 @@ protected:
 		}
 
 		// Deploy: deployment by copying to devices requires cooking by the book
-		if ((DeploymentMode == ELauncherProfileDeploymentModes::CopyToDevice) && (CookMode != ELauncherProfileCookModes::ByTheBook))
+		if ((DeploymentMode == ELauncherProfileDeploymentModes::CopyToDevice) && ((CookMode != ELauncherProfileCookModes::ByTheBook)&&(CookMode!=ELauncherProfileCookModes::ByTheBookInEditor)))
 		{
 			ValidationErrors.Add(ELauncherProfileValidationErrors::CopyToDeviceRequiresCookByTheBook);
 		}
@@ -1136,9 +1141,11 @@ private:
 
 	// Path to the editor executable to pass to UAT, for cooking, etc... May be empty.
 	FString EditorExe;
-        
+
 private:
 
+	FIsCookFinishedDelegate IsCookFinishedDelegate;
+
 	// Holds a delegate to be invoked when changing the device group to deploy to.
 	FOnLauncherProfileDeployedDeviceGroupChanged DeployedDeviceGroupChangedDelegate;
 
diff --git a/Engine/Source/Developer/LauncherServices/Public/Interfaces/ILauncherProfile.h b/Engine/Source/Developer/LauncherServices/Public/Interfaces/ILauncherProfile.h
--- a/Engine/Source/Developer/LauncherServices/Public/Interfaces/ILauncherProfile.h
+++ b/Engine/Source/Developer/LauncherServices/Public/Interfaces/ILauncherProfile.h
@@ -17,7 +17,10 @@ namespace ELauncherProfileCookModes
 		ByTheBook,
 
 		/** Cook the build on the fly while the game is running. */
-		OnTheFly
+		OnTheFly,
+
+		/** Cook by the book in the editor process space */
+		ByTheBookInEditor
 	};
 }
 
@@ -272,6 +275,11 @@ DECLARE_MULTICAST_DELEGATE_OneParam(FOnLauncherProfileDeployedDeviceGroupChanged
 /** Delegate type for a change in project */
 DECLARE_MULTICAST_DELEGATE(FOnProfileProjectChanged);
 
+/** Delegate type for detecting if cook is finished
+ *	Used when cooking from the editor.  Specific cook task will wait for the cook to be finished by the editor
+ */
+DECLARE_DELEGATE_RetVal(bool, FIsCookFinishedDelegate);
+
 /**
  * Interface for launcher profile.
  */
@@ -370,6 +378,11 @@ public:
 	 */
 	virtual void SetDescription(const FString& NewDescription) = 0;
 
+	/**
+	 * Returns the cook delegate which can be used to query if the cook is finished
+	 *  Used by cook by the book in the editor
+	 */
+	virtual FIsCookFinishedDelegate& OnIsCookFinished() = 0;
 public:
 
 	/**
diff --git a/Engine/Source/Developer/TargetPlatform/Private/TargetPlatformManagerModule.cpp b/Engine/Source/Developer/TargetPlatform/Private/TargetPlatformManagerModule.cpp
--- a/Engine/Source/Developer/TargetPlatform/Private/TargetPlatformManagerModule.cpp
+++ b/Engine/Source/Developer/TargetPlatform/Private/TargetPlatformManagerModule.cpp
@@ -153,6 +153,10 @@ public:
 					Results = Platforms;
 				}
 			}
+			else
+			{
+				Results = Platforms;
+			}
 		}
 
 		return Results;
diff --git a/Engine/Source/Editor/UnrealEd/Classes/Commandlets/CookCommandlet.h b/Engine/Source/Editor/UnrealEd/Classes/Commandlets/CookCommandlet.h
--- a/Engine/Source/Editor/UnrealEd/Classes/Commandlets/CookCommandlet.h
+++ b/Engine/Source/Editor/UnrealEd/Classes/Commandlets/CookCommandlet.h
@@ -137,5 +137,12 @@ private:
 	/** Cooks all files */
 	bool Cook(const TArray<ITargetPlatform*>& Platforms, TArray<FString>& FilesInPath);
 
+	/** Cooks all files newly (in a new way) */
+	bool NewCook(const TArray<ITargetPlatform*>& Platforms, TArray<FString>& FilesInPath);
+
+
+	/**	Process deferred commands */
+	void ProcessDeferredCommands();
+
 
 };
\ No newline at end of file
diff --git a/Engine/Source/Editor/UnrealEd/Classes/CookOnTheSide/CookOnTheFlyServer.h b/Engine/Source/Editor/UnrealEd/Classes/CookOnTheSide/CookOnTheFlyServer.h
--- a/Engine/Source/Editor/UnrealEd/Classes/CookOnTheSide/CookOnTheFlyServer.h
+++ b/Engine/Source/Editor/UnrealEd/Classes/CookOnTheSide/CookOnTheFlyServer.h
@@ -9,31 +9,76 @@
 #include "CookOnTheFlyServer.generated.h"
 
 
+class FChunkManifestGenerator;
+
+enum class ECookInitializationFlags
+{
+	None = 0x0,						// No flags
+	Compressed = 0x1,				// will save compressed packages
+	Iterative = 0x2,				// use iterative cooking (previous cooks will not be cleaned unless detected out of date, experimental)
+	SkipEditorContent = 0x4,		// do not cook any content in the content\editor directory
+	Unversioned = 0x8,				// save the cooked packages without a version number
+	AutoTick = 0x10,				// enable ticking (only works in the editor)
+	AsyncSave = 0x20,				// save packages async
+	LeakTest = 0x40,				// test for uobject leaks after each level load
+	IncludeServerMaps = 0x80,		// should we include the server maps when cooking
+	GenerateStreamingInstallManifest = 0x100,  // should we generate streaming install manifest
+};
+ENUM_CLASS_FLAGS(ECookInitializationFlags);
+
+UENUM()
+namespace ECookMode
+{
+	enum Type
+	{
+		CookOnTheFly,				// default mode, handles requests from network
+		CookByTheBookFromTheEditor,	// precook all resources while in the editor
+		CookByTheBook,				// cooking by the book (not in the editor)
+	};
+}
+
 UCLASS()
 class UCookOnTheFlyServer : public UObject, public FTickableEditorObject
 {
 	GENERATED_UCLASS_BODY()
 
 private:
-	/** Cook ont he fly is almost always true, just leaving this option here if we want to switch the cook commandlet over to using this structure to cook */
-	bool bCookOnTheFly;
-	/** If true, iterative cooking is being done */
-	bool bIterativeCooking;
-	/** If true, packages are cooked compressed */
-	bool bCompressed;
-	/** Skip saving any packages in Engine/COntent/Editor* UNLESS TARGET HAS EDITORONLY DATA (in which case it will save those anyway) */
-	bool bSkipEditorContent;
-	/** Should we use Async loading and saving (good for the editor) */
-	bool bAsyncSaveLoad;
-	/** Should we auto tick (used in the editor) */
-	bool bIsAutoTick;
-	/** Holds the sandbox file wrapper to handle sandbox path conversion. */
-	TAutoPtr<class FSandboxPlatformFile> SandboxFile;
+	/** Current cook mode the cook on the fly server is running in */
+	ECookMode::Type CurrentCookMode;
+	
+	//////////////////////////////////////////////////////////////////////////
+	// Cook by the book options
+	struct FCookByTheBookOptions
+	{
+	public:
+		FCookByTheBookOptions() : bGenerateStreamingInstallManifests(false),
+			bRunning(false),
+			CookTime( 0.0f )
+		{ }
+
+		/** Should we generate streaming install manifests (only valid option in cook by the book) */
+		bool bGenerateStreamingInstallManifests;
+		/** Is cook by the book currently running */
+		bool bRunning;
+		/** Leak test: last gc items (only valid when running from commandlet requires gc between each cooked package) */
+		TSet<FWeakObjectPtr> LastGCItems;
+		/** Map of platform name to manifest generator, manifest is only used in cook by the book however it needs to be maintained across multiple cook by the books. */
+		TMap<FName, FChunkManifestGenerator*> ManifestGenerators;
+		double CookTime;
+	};
+	FCookByTheBookOptions* CookByTheBookOptions;
+	
+
+	//////////////////////////////////////////////////////////////////////////
+	// Cook on the fly options
 	/** Cook on the fly server uses the NetworkFileServer */
 	TArray<class INetworkFileServer*> NetworkFileServers;
 
-	/** We hook this up to a delegate to avoid reloading textures and whatnot */
-	TSet<FString> PackagesToNotReload;
+	//////////////////////////////////////////////////////////////////////////
+	// General cook options
+	ECookInitializationFlags CookFlags;
+	TAutoPtr<class FSandboxPlatformFile> SandboxFile;
+	bool bIsSavingPackage; // used to stop recursive mark package dirty functions
 
 private:
 
@@ -195,34 +240,35 @@ public:
 	{
 	private:
 		FName Filename;
-		FName Platformname;
-		uint32 CachedHash;
+		TArray<FName> Platformnames;
 	public:
-		FFilePlatformRequest() : CachedHash(0) { }
+		
+		// yes we have some friends
+		friend uint32 GetTypeHash(const UCookOnTheFlyServer::FFilePlatformRequest& Key);
+
+		FFilePlatformRequest() { }
 
-		FFilePlatformRequest( const FName &InFilename, const FName &InPlatformname ) : Filename( InFilename ), Platformname( InPlatformname ) 
+		FFilePlatformRequest( const FName& InFilename, const TArray<FName>& InPlatformname ) : Filename( InFilename )
 		{
-			UpdateHash();
+			Platformnames = InPlatformname;
 		}
 
-		FFilePlatformRequest( const FString &InFilename, const FName &InPlatformname ) : Filename( *InFilename ), Platformname( InPlatformname ) 
+		FFilePlatformRequest( const FName& InFilename, TArray<FName>&& InPlatformname ) : Filename( InFilename ), Platformnames(MoveTemp(InPlatformname))
 		{
-			UpdateHash();
 		}
 
-		FFilePlatformRequest( const FFilePlatformRequest &InFilePlatformRequest ) : Filename( InFilePlatformRequest.Filename ), Platformname( InFilePlatformRequest.Platformname ), CachedHash( InFilePlatformRequest.CachedHash )
+		FFilePlatformRequest( const FFilePlatformRequest& InFilePlatformRequest ) : Filename( InFilePlatformRequest.Filename ), Platformnames( InFilePlatformRequest.Platformnames )
+		{
+		}
+		
+		FFilePlatformRequest( FFilePlatformRequest&& InFilePlatformRequest ) : Filename( MoveTemp(InFilePlatformRequest.Filename) ), Platformnames( MoveTemp(InFilePlatformRequest.Platformnames) )
 		{
-#if DO_CHECK
-			uint32 OldHash = CachedHash;
-			UpdateHash();
-			check( OldHash == CachedHash );
-#endif
 		}
+		
 
 		void SetFilename( const FString &InFilename ) 
 		{
 			Filename = FName(*InFilename);
-			UpdateHash();
 		}
 
 		const FName &GetFilename() const
@@ -230,19 +276,25 @@ public:
 			return Filename;
 		}
 
-		const FName &GetPlatformname() const
+		const TArray<FName> &GetPlatformnames() const
 		{
-			return Platformname;
+			return Platformnames;
 		}
 
-		void UpdateHash()
+		void RemovePlatform( const FName &Platform )
 		{
-			CachedHash = GetTypeHash(Filename) ^ GetTypeHash(Platformname);
+			Platformnames.Remove(Platform);
 		}
 
-		uint32 GetCachedHash() const
+		void AddPlatform( const FName &Platform )
 		{
-			return CachedHash;
+			check( Platform != NAME_None );
+			Platformnames.Add(Platform );
+		}
+
+		bool HasPlatform( const FName &Platform ) const
+		{
+			return Platformnames.Find(Platform) != INDEX_NONE;
 		}
 
 		bool IsValid()  const
@@ -253,20 +305,23 @@ public:
 		void Clear()
 		{
 			Filename = TEXT("");
-			Platformname = TEXT("");
-			UpdateHash();
+			Platformnames.Empty();
+		}
+
+		FFilePlatformRequest &operator=( FFilePlatformRequest &&InFileRequest )
+		{
+			Filename = MoveTemp( InFileRequest.Filename );
+			Platformnames = MoveTemp( InFileRequest.Platformnames );
+			return *this;
 		}
 
 		bool operator ==( const FFilePlatformRequest &InFileRequest ) const
 		{
-			if ( CachedHash == InFileRequest.CachedHash )
+			if ( InFileRequest.Filename == Filename )
 			{
-				if ( InFileRequest.Filename == Filename )
+				if ( InFileRequest.Platformnames == Platformnames )
 				{
-					if ( InFileRequest.Platformname == Platformname )
-					{
-						return true;
-					}
+					return true;
 				}
 			}
 			return false;
@@ -274,24 +329,26 @@ public:
 
 		FORCEINLINE FString &&ToString() const
 		{
-			return MoveTemp( FString::Printf( TEXT("%s %s"), *Filename.ToString(), *Platformname.ToString() ) );
+			FString Result = FString::Printf(TEXT("%s;"), *Filename.ToString());
+
+			for ( const auto &Platform : Platformnames )
+			{
+				Result += FString::Printf(TEXT("%s,"), *Platform.ToString() );
+			}
+			return MoveTemp(Result);
 		}
 
 	};
 
-private:
-
-	FThreadSafeQueue<FFilePlatformRequest> FileRequests; // list of requested files
-	FLookupQueue<FFilePlatformRequest> UnsolicitedFileRequests; // list of files which haven't been requested but we think should cook based on previous requests
-	FThreadSafeQueue<FFilePlatformRequest> UnsolicitedCookedPackages; // list of files which weren't requested but were cooked (based on UnsolicitedFileRequests)
 
+private:
 
 	/** Helper list of all files which have been cooked */
 	struct FThreadSafeFilenameSet
 	{
 	private:
 		mutable FCriticalSection	SynchronizationObject;
-		TSet<FFilePlatformRequest> FilesProcessed;
+		TMap<FName, FFilePlatformRequest> FilesProcessed;
 	public:
 
 		void Lock()
@@ -303,38 +360,168 @@ private:
 			SynchronizationObject.Unlock();
 		}
 
-		void Add(const FFilePlatformRequest &Filename)
+		void Add(const FFilePlatformRequest& Request)
 		{
 			FScopeLock ScopeLock(&SynchronizationObject);
-			check(Filename.IsValid());
-			FilesProcessed.Add(Filename);
+			check(Request.IsValid());
+
+			// see if it's already in the requests list
+			FFilePlatformRequest *ExistingRequest = FilesProcessed.Find(Request.GetFilename() );
+			
+			if ( ExistingRequest )
+			{
+				check( ExistingRequest->GetFilename() == Request.GetFilename() );
+				for ( const auto &Platform : Request.GetPlatformnames() )
+				{
+					ExistingRequest->AddPlatform(Platform);
+				}
+			}
+			else
+				FilesProcessed.Add(Request.GetFilename(), Request);
 		}
-		bool Exists(const FFilePlatformRequest &Filename) const
+		bool Exists(const FFilePlatformRequest& Request) const
 		{
 			FScopeLock ScopeLock(&SynchronizationObject);
-			return FilesProcessed.Contains(Filename);
+
+			const FFilePlatformRequest* OurRequest = FilesProcessed.Find( Request.GetFilename() );
+			
+			if (!OurRequest)
+			{
+				return false;
+			}
+
+			// make sure all the platforms are completed
+			for ( const auto& Platform : Request.GetPlatformnames() )
+			{
+				if ( !OurRequest->GetPlatformnames().Contains( Platform ) )
+				{
+					return false;
+				}
+			}
+
+			return true;
+			// return FilesProcessed.Contains(Filename);
 		}
-		int RemoveAll( const FName &Filename )
+		bool GetCookedPlatforms( const FName& Filename, TArray<FName>& PlatformList )
 		{
 			FScopeLock ScopeLock( &SynchronizationObject );
-			int NumRemoved = 0;
-			// for ( TSet<FFilePlatformRequest>::TIterator It( FilesProcessed ); It; ++It )
-			for ( auto It = FilesProcessed.CreateIterator(); It; ++It )
+			const FFilePlatformRequest* Request = FilesProcessed.Find(Filename);
+			if ( Request )
 			{
-				if ( It->GetFilename() == Filename )
+				PlatformList = Request->GetPlatformnames();
+				return true;
+			}
+			return false;
+			
+		}
+		int RemoveAll( const FName& Filename )
+		{
+			FScopeLock ScopeLock( &SynchronizationObject );
+			return FilesProcessed.Remove( Filename );
+		}
+	};
+
+
+	struct FFilenameQueue
+	{
+	private:
+		TArray<FName>			Queue;
+		TMap<FName, TArray<FName>> PlatformList;
+		mutable FCriticalSection SynchronizationObject;
+	public:
+		void EnqueueUnique(const FFilePlatformRequest& Request, bool ForceEnqueFront = false)
+		{
+			FScopeLock ScopeLock( &SynchronizationObject );
+			TArray<FName>* Platforms = PlatformList.Find(Request.GetFilename());
+			if ( Platforms == NULL )
+			{
+				PlatformList.Add(Request.GetFilename(), Request.GetPlatformnames());
+				Queue.Add(Request.GetFilename());
+			}
+			else
+			{
+				// add the requested platforms to the platform list
+				for ( const auto& Platform : Request.GetPlatformnames() )
 				{
-					It.RemoveCurrent();
-					++NumRemoved;
+					Platforms->AddUnique(Platform);
 				}
 			}
-			return NumRemoved;
+
+			if ( ForceEnqueFront )
+			{
+				int32 Index = Queue.Find(Request.GetFilename());
+				if ( Index != 0 )
+				{
+					Queue.Swap(0, Index);
+				}
+			}
+		}
+		
+		bool Dequeue(FFilePlatformRequest* Result)
+		{
+			FScopeLock ScopeLock( &SynchronizationObject );
+			if (Queue.Num())
+			{
+				FName Filename = Queue[0];
+				Queue.RemoveAt(0);
+				TArray<FName> Platforms = PlatformList.FindChecked(Filename);
+				PlatformList.Remove(Filename);
+				*Result = MoveTemp(FFilePlatformRequest(MoveTemp(Filename), MoveTemp(Platforms)));
+				return true;
+			}
+			return false;
+		}
+		
+		bool HasItems() const
+		{
+			FScopeLock ScopeLock( &SynchronizationObject );
+			return Queue.Num() > 0;
+		}
+
+		int Num() const 
+		{
+			FScopeLock ScopeLock( &SynchronizationObject );
+			return Queue.Num();
 		}
 	};
-	FThreadSafeFilenameSet ThreadSafeFilenameSet; // set of files which have been cooked when needing to recook a file the entry will need to be removed from here
 
+	struct FThreadSafeUnsolicitedPackagesList
+	{
+	private:
+		FCriticalSection SyncObject;
+		TArray<FFilePlatformRequest> CookedPackages;
+	public:
+		void AddCookedPackage( const FFilePlatformRequest& PlatformRequest )
+		{
+			FScopeLock S( &SyncObject );
+			CookedPackages.Add( PlatformRequest );
+		}
+		void GetPackagesForPlatformAndRemove( const FName& Platform, TArray<FName> PackageNames )
+		{
+			FScopeLock S( &SyncObject );
+			
+			for ( int I = CookedPackages.Num()-1; I >= 0; --I )
+			{
+				FFilePlatformRequest &Request = CookedPackages[I];
 
-	FThreadSafeQueue<struct FRecompileRequest*> RecompileRequests;
+				if ( Request.GetPlatformnames().Contains( Platform ) )
+				{
+					// remove the platform
+					Request.RemovePlatform( Platform );
+
+					if ( Request.GetPlatformnames().Num() == 0 )
+					{
+						CookedPackages.RemoveAt(I);
+					}
+				}
+			}
+		}
+	};
 
+	FThreadSafeQueue<struct FRecompileRequest*> RecompileRequests;
+	FFilenameQueue CookRequests; // list of requested files
+	FThreadSafeUnsolicitedPackagesList UnsolicitedCookedPackages;
+	FThreadSafeFilenameSet CookedPackages; // set of files which have been cooked when needing to recook a file the entry will need to be removed from here
 
 public:
 
@@ -349,22 +536,20 @@ public:
 	};
 
 
+	virtual ~UCookOnTheFlyServer();
 
 	/**
 	 * FTickableEditorObject interface used by cook on the side
 	 */
 	TStatId GetStatId() const override;
-	 void Tick(float DeltaTime) override;
-	 bool IsTickable() const override;
-
-
-
-
+	void Tick(float DeltaTime) override;
+	bool IsTickable() const override;
+	ECookMode::Type GetCookMode() const { return CurrentCookMode; }
 
 	/**
 	 * Initialize the CookServer so that either CookOnTheFly can be called or Cook on the side can be started and ticked
 	 */
-	void Initialize( bool inCompressed, bool inIterativeCooking, bool inSkipEditorContent, bool inIsAutoTick = false, bool inAsyncLoadSave = false, const FString &OutputDirectoryOverride = FString(TEXT("")) );
+	void Initialize( ECookMode::Type DesiredCookMode, ECookInitializationFlags InCookInitializationFlags, const FString &OutputDirectoryOverride = FString(TEXT("")) );
 
 	/**
 	 * Cook on the side, cooks while also running the editor...
@@ -385,7 +570,15 @@ public:
 	 */
 	void EndNetworkFileServer();
 
-	
+
+	/**
+	 * Start a cook by the book session
+	 * Cook on the fly can't run at the same time as cook by the book
+	 */
+	void StartCookByTheBook(const TArray<ITargetPlatform*>& TargetPlatforms, const TArray<FString>& CookMaps, const TArray<FString>& CookDirectories, const TArray<FString>& CookCultures, const TArray<FString>& IniMapSections );
+
+	bool IsCookByTheBookRunning() const;
+
 	/**
 	 * Handles cook package requests until there are no more requests, then returns
 	 *
@@ -401,8 +594,7 @@ public:
 	 */
 	void TickRecompileShaderRequests();
 
-	bool HasCookRequests() const { return FileRequests.HasItems(); }
-	bool HasUnsolicitedCookRequests() const { return UnsolicitedFileRequests.HasItems(); }
+	bool HasCookRequests() const { return CookRequests.HasItems(); }
 
 	bool HasRecompileShaderRequests() const { return RecompileRequests.HasItems(); }
 
@@ -435,8 +627,81 @@ public:
 
 private:
 	
+	//////////////////////////////////////////////////////////////////////////
+	// cook by the book specific functions
+	/**
+	 * Collect all the files which need to be cooked for a cook by the book session
+	 */
+	void CollectFilesToCook(TArray<FString>& FilesInPath, 
+		const TArray<FString>& CookMaps, const TArray<FString>& CookDirectories, const TArray<FString>& CookCultures, 
+		const TArray<FString>& IniMapSections, bool bCookAll, bool bMapsOnly, bool bNoDev);
+	/**
+	 * Call back from the TickCookOnTheSide when a cook by the book finishes (when started form StartCookByTheBook)
+	 */
+	void CookByTheBookFinished();
+
+	/**
+	 * Helper function returns if we are in any cook by the book mode
+	 *
+	 * @return if the cook mode is a cook by the book mode
+	 */
+	inline bool IsCookByTheBookMode() const 
+	{ 
+		return CurrentCookMode == ECookMode::CookByTheBookFromTheEditor || CurrentCookMode == ECookMode::CookByTheBook; 
+	}
+
+	/**
+	 * GetDependencies
+	 * 
+	 * @param Packages List of packages to use as the root set for dependency checking
+	 * @param Found return value, all objects which package is dependent on
+	 */
+	void GetDependencies( const TSet<UPackage*>& Packages, TSet<UObject*>& Found);
+	/**
+	 * GenerateManifestInfo
+	 * generate the manifest information for a given package
+	 *
+	 * @param Package package to generate manifest information for
+	 */
+	void GenerateManifestInfo( UPackage* Package, const TArray<FName>& TargetPlatformNames );
 
-	bool ShouldCook(const FString& InFileName, const FName &InPlatformName);
+	//////////////////////////////////////////////////////////////////////////
+	// cook on the fly specific functions
+	/**
+	 * Cook requests for a package from network
+	 *  blocks until cook is complete
+	 * 
+	 * @param  Filename	requested file to cook
+	 * @param  Platformname platform to cook for
+	 * @param  out UnsolicitedFiles return a list of files which were cooked as a result of cooking the requested package
+	 */
+	void HandleNetworkFileServerFileRequest( const FString& Filename, const FString& Platformname, TArray<FString>& UnsolicitedFiles );
+
+	/**
+	 * Shader recompile request from network
+	 *  blocks until shader recompile complete
+	 *
+	 * @param  RecompileData input params for shader compile and compiled shader output
+	 */
+	void HandleNetworkFileServerRecompileShaders(const struct FShaderRecompileData& RecompileData);
+
+	//////////////////////////////////////////////////////////////////////////
+	// general functions
+
+	/**
+	 * Determines if a package should be cooked
+	 * 
+	 * @param InFileName		package file name
+	 * @param InPlatformName	desired platform to cook for
+	 * @return If the package should be cooked
+	 */
+	bool ShouldCook(const FString& InFileName, const FName& InPlatformName);
+
+
+	bool IsCookFlagSet( const ECookInitializationFlags& InCookFlags ) const 
+	{
+		return (CookFlags & InCookFlags) != ECookInitializationFlags::None;
+	}
 
 	/**
 	 * Returns cooker output directory.
@@ -479,17 +744,17 @@ private:
 	 */
 	bool SaveCookedPackage( UPackage* Package, uint32 SaveFlags, bool& bOutWasUpToDate, TArray<FName> &TargetPlatformNames );
 
-	// Callback for handling a network file request.
-	void HandleNetworkFileServerFileRequest( const FString& Filename, const FString &Platformname, TArray<FString>& UnsolicitedFiles );
-
-	// Callback for recompiling shaders
-	void HandleNetworkFileServerRecompileShaders(const struct FShaderRecompileData& RecompileData);
 
+	/**
+	 *  Save the global shader map
+	 *  
+	 *  @param	Platforms		List of platforms to make global shader maps for
+	 */
+	void SaveGlobalShaderMapFiles(const TArray<ITargetPlatform*>& Platforms);
 
-	void MaybeMarkPackageAsAlreadyLoaded(UPackage *Package);
 
 	/** Gets the output directory respecting any command line overrides */
-	FString GetOutputDirectoryOverride( const FString &OutputDirectoryOverride ) const;
+	FString GetOutputDirectoryOverride( const FString& OutputDirectoryOverride ) const;
 
 	/** Cleans sandbox folders for all target platforms */
 	void CleanSandbox(const TArray<ITargetPlatform*>& Platforms);
@@ -497,23 +762,21 @@ private:
 	/** Generates asset registry */
 	void GenerateAssetRegistry(const TArray<ITargetPlatform*>& Platforms);
 
-	/** Saves global shader map files */
-	// void SaveGlobalShaderMapFiles(const TArray<ITargetPlatform*>& Platforms);
-
-	/** Collects all files to be cooked. This includes all commandline specified maps */
-	//void CollectFilesToCook(TArray<FString>& FilesInPath);
-
 	/** Generates long package names for all files to be cooked */
 	void GenerateLongPackageNames(TArray<FString>& FilesInPath);
 
-	/** Cooks all files */
-	//bool Cook(const TArray<ITargetPlatform*>& Platforms, TArray<FString>& FilesInPath);
+
+	void GetDependencies( UPackage* Package, TArray<UPackage*> Dependencies );
 
 };
 
 FORCEINLINE uint32 GetTypeHash(const UCookOnTheFlyServer::FFilePlatformRequest &Key)
 {
-	//return GetTypeHash( Key.Filename ) ^ GetTypeHash( Key.Platformname );
-	return Key.GetCachedHash();
+	uint32 Hash = GetTypeHash( Key.Filename );
+	for ( const auto &PlatformName : Key.Platformnames )
+	{
+		Hash += Hash << 2 ^ GetTypeHash( PlatformName );
+	}
+	return Hash;
 }
 
diff --git a/Engine/Source/Editor/UnrealEd/Classes/Editor/EditorEngine.h b/Engine/Source/Editor/UnrealEd/Classes/Editor/EditorEngine.h
--- a/Engine/Source/Editor/UnrealEd/Classes/Editor/EditorEngine.h
+++ b/Engine/Source/Editor/UnrealEd/Classes/Editor/EditorEngine.h
@@ -1463,6 +1463,24 @@ public:
 	 */
 	virtual void PlayMap( const FVector* StartLocation = NULL, const FRotator* StartRotation = NULL, int32 DestinationConsole = -1, int32 InPlayInViewportIndex = -1, bool bUseMobilePreview = false, bool bMovieCapture = false );
 
+
+
+	/**
+	 * Can the editor do cook by the book in the editor process space
+	 */
+	virtual bool CanCookByTheBookInEditor() const { return false; }
+
+	/**
+	 * Start cook by the book in the editor process space
+	 */
+	virtual void StartCookByTheBookInEditor( const TArray<ITargetPlatform*> &TargetPlatforms, const TArray<FString> &CookMaps, const TArray<FString> &CookDirectories, const TArray<FString> &CookCultures, const TArray<FString> &IniMapSections ) { }
+
+	/**
+	 * Checks if the cook by the book is finished
+	 */
+	virtual bool IsCookByTheBookInEditorFinished() const { return true; }
+
+
 	/**
 	 * Makes a request to start a play from a Slate editor session
 	 * @param	bAtPlayerStart			Whether or not we would really like to use the game or level's PlayerStart vs the StartLocation
diff --git a/Engine/Source/Editor/UnrealEd/Classes/Editor/UnrealEdEngine.h b/Engine/Source/Editor/UnrealEd/Classes/Editor/UnrealEdEngine.h
--- a/Engine/Source/Editor/UnrealEd/Classes/Editor/UnrealEdEngine.h
+++ b/Engine/Source/Editor/UnrealEd/Classes/Editor/UnrealEdEngine.h
@@ -536,6 +536,24 @@ class UNREALED_API UUnrealEdEngine : public UEditorEngine, public FNotifyHook
 	void AttachSelectedActors();
 
 
+	
+	/**
+	 * Can the editor do cook by the book in the editor process space
+	 */
+	virtual bool CanCookByTheBookInEditor() const;
+
+	/**
+	 * Start cook by the book in the editor process space
+	 */
+	virtual void StartCookByTheBookInEditor( const TArray<ITargetPlatform*> &TargetPlatforms, const TArray<FString> &CookMaps, const TArray<FString> &CookDirectories, const TArray<FString> &CookCultures, const TArray<FString> &IniMapSections );
+
+	/**
+	 * Checks if the cook by the book is finished
+	 */
+	virtual bool IsCookByTheBookInEditorFinished() const override;
+
+
+
 	// Hook replacements.
 	void ShowActorProperties();
 
diff --git a/Engine/Source/Editor/UnrealEd/Private/Commandlets/ChunkManifestGenerator.cpp b/Engine/Source/Editor/UnrealEd/Private/Commandlets/ChunkManifestGenerator.cpp
--- a/Engine/Source/Editor/UnrealEd/Private/Commandlets/ChunkManifestGenerator.cpp
+++ b/Engine/Source/Editor/UnrealEd/Private/Commandlets/ChunkManifestGenerator.cpp
@@ -49,6 +49,14 @@ void FChunkManifestGenerator::OnAssetLoaded(UObject* Asset)
 	}
 }
 
+void FChunkManifestGenerator::OnLastPackageLoaded( UPackage* Package )
+{
+	if ( !AssetsLoadedWithLastPackage.Contains(Package->GetFName()))
+	{
+		AssetsLoadedWithLastPackage.Add(Package->GetFName());
+	}
+}
+
 bool FChunkManifestGenerator::CleanTempPackagingDirectory(const FString& Platform) const
 {
 	FString TmpPackagingDir = GetTempPackagingDirectoryForPlatform(Platform);
@@ -187,7 +195,7 @@ void FChunkManifestGenerator::Initialize(bool InGenerateChunks)
 }
 
 void FChunkManifestGenerator::AddPackageToChunkManifest(UPackage* Package, const FString& SandboxFilename, const FString& LastLoadedMapName, FSandboxPlatformFile* SandboxFile)
-{		
+{
 	TArray<int32> TargetChunks;
 	TArray<int32> ExistingChunkIDs;
 	
@@ -467,6 +475,11 @@ void FChunkManifestGenerator::AddPackageToManifest(const FString& PackageSandbox
 	UnassignedPackageSet.Remove(PackageName);
 }
 
+void FChunkManifestGenerator::AddUnassignedPackageToManifest(UPackage* Package, const FString& PackageSandboxPath )
+{
+	NotifyPackageWasNotAssigned(PackageSandboxPath, Package->GetFName() );
+}
+
 void FChunkManifestGenerator::RemovePackageFromManifest(FName PackageName, int32 ChunkId)
 {
 	if (ChunkManifests[ChunkId])
diff --git a/Engine/Source/Editor/UnrealEd/Private/Commandlets/ChunkManifestGenerator.h b/Engine/Source/Editor/UnrealEd/Private/Commandlets/ChunkManifestGenerator.h
--- a/Engine/Source/Editor/UnrealEd/Private/Commandlets/ChunkManifestGenerator.h
+++ b/Engine/Source/Editor/UnrealEd/Private/Commandlets/ChunkManifestGenerator.h
@@ -12,7 +12,7 @@ class FChunkManifestGenerator
 	/** Holds a reference to asset registry */
 	IAssetRegistry& AssetRegistry;
 	/** Platforms to generate the manifests for. */
-	const TArray<ITargetPlatform*>& Platforms;
+	const TArray<ITargetPlatform*> Platforms;
 	/** ChunkIDs associated with a package in the asset registry */
 	TMap<FName, TArray<int32> > RegistryChunkIDsMap;
 	/** List of all asset packages that were created while loading the last package in the cooker. */
@@ -52,6 +52,7 @@ class FChunkManifestGenerator
 		uint32		ParentNodeIndex;
 	};
 
+
 	/**
 	 * Adds a package to chunk manifest
 	 * 
@@ -212,7 +213,7 @@ public:
 	void Initialize(bool InGenerateChunks);
 
 	/**
-	 * Adds a package to chunk manifest.
+	 * Adds a package to chunk manifest (just calls the other AddPackageToChunkManifestFunction with more parameters)
 	 *
 	 * @param Package Package to add to one of the manifests
 	 * @param SandboxFilename Cooked sandbox path of the package to add to a manifest
@@ -220,6 +221,35 @@ public:
 	 * @param the SandboxPlatformFile used during cook
 	 */
 	void AddPackageToChunkManifest(UPackage* Package, const FString& SandboxFilename, const FString& LastLoadedMapName, FSandboxPlatformFile* SandboxFile);
+	
+
+	/**
+	 * Add a package to the manifest but don't assign it to any chunk yet, packages which are not assigned by the end of the cook will be put into chunk 0
+	 * 
+	 * @param Package which is unassigned
+	 * @param The sandbox file path of the package
+	 */
+	void AddUnassignedPackageToManifest(UPackage* Package, const FString& PackageSandboxPath );
+
+	/**
+	* Adds a package to chunk manifest
+	*
+	* @param PackageFName Package name (Package->GetFName())
+	* @param PackagePathName Package path name (Package->GetPathName())
+	* @param SandboxFilename Cooked sandbox path of the package to add to a manifest
+	* @param LastLoadedMapName Name of the last loaded map (can be empty)
+	* @param the SandboxPlatformFile used during cook
+	*/
+	void FChunkManifestGenerator::AddPackageToChunkManifest(const FName& PackageFName, const FString& PackagePathName, const FString& SandboxFilename, const FString& LastLoadedMapName, FSandboxPlatformFile* SandboxFile);
+
+
+
+	/**
+	 * Collects all the packages loaded 
+	 *
+	 * @param Package Package which was loaded 
+	 */
+	void OnLastPackageLoaded( UPackage* Package );
 
 	/**
 	 * The cooker is about to load a new package from the list, reset AssetsLoadedWithLastPackage
diff --git a/Engine/Source/Editor/UnrealEd/Private/Commandlets/CookCommandlet.cpp b/Engine/Source/Editor/UnrealEd/Private/Commandlets/CookCommandlet.cpp
--- a/Engine/Source/Editor/UnrealEd/Private/Commandlets/CookCommandlet.cpp
+++ b/Engine/Source/Editor/UnrealEd/Private/Commandlets/CookCommandlet.cpp
@@ -83,7 +83,13 @@ bool UCookCommandlet::CookOnTheFly( FGuid InstanceId, int32 Timeout, bool bForce
 	// make sure that the cookonthefly server doesn't get cleaned up while we are garbage collecting below :)
 	FScopeRootObject S(CookOnTheFlyServer);
 
-	CookOnTheFlyServer->Initialize( bCompressed, bIterativeCooking, bSkipEditorContent );
+	ECookInitializationFlags CookFlags = ECookInitializationFlags::None;
+	CookFlags |= bCompressed ? ECookInitializationFlags::Compressed : ECookInitializationFlags::None;
+	CookFlags |= bIterativeCooking ? ECookInitializationFlags::Iterative : ECookInitializationFlags::None;
+	CookFlags |= bSkipEditorContent ? ECookInitializationFlags::SkipEditorContent : ECookInitializationFlags::None;
+
+
+	CookOnTheFlyServer->Initialize( ECookMode::CookOnTheFly, CookFlags );
 
 	bool BindAnyPort = InstanceId.IsValid();
 
@@ -144,7 +150,7 @@ bool UCookCommandlet::CookOnTheFly( FGuid InstanceId, int32 Timeout, bool bForce
 				}
 
 				// delay the gc until we process some unsolicited packages
-				if ( bCookedAMapSinceLastGC && (CookOnTheFlyServer->HasUnsolicitedCookRequests() == false) )
+				if ( bCookedAMapSinceLastGC )
 				{
 					UE_LOG( LogCookCommandlet, Display, TEXT("Delaying map gc because we have unsolicited cook requests") );
 					bShouldGC |= bCookedAMapSinceLastGC;
@@ -168,21 +174,7 @@ bool UCookCommandlet::CookOnTheFly( FGuid InstanceId, int32 Timeout, bool bForce
 				}
 			}
 
-#if PLATFORM_MAC
-			// On Mac we need to process Cocoa events so that the console window for CookOnTheFlyServer is interactive
-			FPlatformMisc::PumpMessages(true);
-#endif
-
-			// update task graph
-			FTaskGraphInterface::Get().ProcessThreadUntilIdle(ENamedThreads::GameThread);
-
-			// execute deferred commands
-			for (int32 DeferredCommandsIndex = 0; DeferredCommandsIndex<GEngine->DeferredCommands.Num(); ++DeferredCommandsIndex)
-			{
-				GEngine->Exec( GWorld, *GEngine->DeferredCommands[DeferredCommandsIndex], *GLog);
-			}
-
-			GEngine->DeferredCommands.Empty();
+			ProcessDeferredCommands();
 
 			// handle server timeout
 			if (InstanceId.IsValid() || bForceClose)
@@ -569,8 +561,15 @@ int32 UCookCommandlet::Main(const FString& CmdLineParams)
 		// always generate the asset registry before starting to cook, for either method
 		GenerateAssetRegistry(Platforms);
 
-		Cook(Platforms, FilesInPath);
-
+		// new cook is better 
+		if ( Switches.Contains(TEXT("NEWCOOK")))
+		{
+			NewCook(Platforms, FilesInPath );
+		}
+		else
+		{
+			Cook(Platforms, FilesInPath);
+		}
 	}
 	
 	return 0;
@@ -992,6 +991,214 @@ void UCookCommandlet::GenerateLongPackageNames(TArray<FString>& FilesInPath)
 	Exchange(FilesInPathReverse, FilesInPath);
 }
 
+bool UCookCommandlet::NewCook( const TArray<ITargetPlatform*>& Platforms, TArray<FString>& FilesInPath )
+{
+
+	UCookOnTheFlyServer *CookOnTheFlyServer = ConstructObject<UCookOnTheFlyServer>( UCookOnTheFlyServer::StaticClass() );
+
+	struct FScopeRootObject
+	{
+		UObject *Object;
+		FScopeRootObject( UObject *InObject ) : Object( InObject )
+		{
+			Object->AddToRoot();
+		}
+
+		~FScopeRootObject()
+		{
+			Object->RemoveFromRoot();
+		}
+	};
+
+	// make sure that the cookonthefly server doesn't get cleaned up while we are garbage collecting below :)
+	FScopeRootObject S(CookOnTheFlyServer);
+
+	ECookInitializationFlags CookFlags = ECookInitializationFlags::IncludeServerMaps;
+	CookFlags |= bCompressed ? ECookInitializationFlags::Compressed : ECookInitializationFlags::None;
+	CookFlags |= bIterativeCooking ? ECookInitializationFlags::Iterative : ECookInitializationFlags::None;
+	CookFlags |= bSkipEditorContent ? ECookInitializationFlags::SkipEditorContent : ECookInitializationFlags::None;
+	CookFlags |= bGenerateStreamingInstallManifests ? ECookInitializationFlags::GenerateStreamingInstallManifest : ECookInitializationFlags::None;
+
+
+	CookOnTheFlyServer->Initialize( ECookMode::CookByTheBook, CookFlags );
+
+	//////////////////////////////////////////////////////////////////////////
+	// parse commandline options 
+
+	TArray<FString> CmdLineIniSections;
+	FString SectionStr;
+	if (FParse::Value(*Params, TEXT("MAPINISECTION="), SectionStr))
+	{
+		if (SectionStr.Contains(TEXT("+")))
+		{
+			SectionStr.ParseIntoArray(&CmdLineIniSections,TEXT("+"),true);
+		}
+		else
+		{
+			CmdLineIniSections.Add(SectionStr);
+		}
+	}
+
+	// Add any map sections specified on command line
+	TArray<FString> MapList;
+	GEditor->ParseMapSectionIni(*Params, MapList);
+
+	TArray<FString> CmdLineMapEntries;
+	TArray<FString> CmdLineDirEntries;
+	TArray<FString> CmdLineCultEntries;
+	for (int32 SwitchIdx = 0; SwitchIdx < Switches.Num(); SwitchIdx++)
+	{
+		const FString& Switch = Switches[SwitchIdx];
+
+		auto GetSwitchValueElements = [&Switch](const FString SwitchKey) -> TArray<FString>
+		{
+			TArray<FString> ValueElements;
+			if (Switch.StartsWith(SwitchKey + TEXT("=")) == true)
+			{
+				FString ValuesList = Switch.Right(Switch.Len() - (SwitchKey + TEXT("=")).Len());
+
+				// Allow support for -KEY=Value1+Value2+Value3 as well as -KEY=Value1 -KEY=Value2
+				for (int32 PlusIdx = ValuesList.Find(TEXT("+")); PlusIdx != INDEX_NONE; PlusIdx = ValuesList.Find(TEXT("+")))
+				{
+					const FString ValueElement = ValuesList.Left(PlusIdx);
+					ValueElements.Add(ValueElement);
+
+					ValuesList = ValuesList.Right(ValuesList.Len() - (PlusIdx + 1));
+				}
+				ValueElements.Add(ValuesList);
+			}
+			return ValueElements;
+		};
+
+		// Check for -MAP=<name of map> entries
+		CmdLineMapEntries += GetSwitchValueElements(TEXT("MAP"));
+
+		// Check for -COOKDIR=<path to directory> entries
+		CmdLineDirEntries += GetSwitchValueElements(TEXT("COOKDIR"));
+		for(FString& Entry : CmdLineDirEntries)
+		{
+			Entry = Entry.TrimQuotes();
+			FPaths::NormalizeDirectoryName(Entry);
+		}
+
+		// Check for -COOKCULTURES=<culture name> entries
+		CmdLineCultEntries += GetSwitchValueElements(TEXT("COOKCULTURES"));
+	}
+
+	// Also append any cookdirs from the project ini files; these dirs are relative to the game content directory
+	{
+		const FString AbsoluteGameContentDir = FPaths::ConvertRelativePathToFull(FPaths::GameContentDir());
+		const UProjectPackagingSettings* const PackagingSettings = GetDefault<UProjectPackagingSettings>();
+		for(const auto& DirToCook : PackagingSettings->DirectoriesToAlwaysCook)
+		{
+			CmdLineDirEntries.Add(AbsoluteGameContentDir / DirToCook.Path);
+		}
+	}
+
+	//////////////////////////////////////////////////////////////////////////
+	// start cook by the book 
+
+	for ( const auto& MapName : CmdLineMapEntries )
+	{
+		MapList.Add( MapName );
+	}
+
+	CookOnTheFlyServer->StartCookByTheBook(Platforms, MapList, CmdLineDirEntries, CmdLineCultEntries, CmdLineIniSections );
+
+	// Garbage collection should happen when either
+	//	1. We have cooked a map
+	//	2. We have cooked non-map packages and...
+	//		a. we have accumulated 50 of these since the last GC.
+	//		b. we have been idle for 20 seconds.
+	bool bShouldGC = true;
+
+	// megamoth
+	uint32 NonMapPackageCountSinceLastGC = 0;
+
+	const int32 PackagesPerGC = 50;
+
+	const double IdleTimeToGC = 20.0;
+	double LastCookActionTime = FPlatformTime::Seconds();
+
+	FDateTime LastConnectionTime = FDateTime::UtcNow();
+	bool bHadConnection = false;
+
+	bool bCookedAMapSinceLastGC = false;
+	while ( CookOnTheFlyServer->IsCookByTheBookRunning() )
+	{
+		uint32 TickResults = 0;
+		static const float CookOnTheSideTimeSlice = 10.0f;
+		TickResults = CookOnTheFlyServer->TickCookOnTheSide(CookOnTheSideTimeSlice, NonMapPackageCountSinceLastGC);
+
+		bCookedAMapSinceLastGC |= TickResults & UCookOnTheFlyServer::COSR_CookedMap;
+		if ( TickResults & (UCookOnTheFlyServer::COSR_CookedMap | UCookOnTheFlyServer::COSR_CookedPackage))
+		{
+			LastCookActionTime = FPlatformTime::Seconds();
+		}
+
+
+		if (NonMapPackageCountSinceLastGC > 0)
+		{
+			// We should GC if we have packages to collect and we've been idle for some time.
+			bShouldGC = (NonMapPackageCountSinceLastGC > PackagesPerGC) || 
+				((FPlatformTime::Seconds() - LastCookActionTime) >= IdleTimeToGC);
+		}
+
+		// delay the gc until we process some unsolicited packages
+		if ( bCookedAMapSinceLastGC )
+		{
+			UE_LOG( LogCookCommandlet, Display, TEXT("Delaying map gc because we have unsolicited cook requests") );
+			bShouldGC |= bCookedAMapSinceLastGC;
+		}
+
+		if (bShouldGC)
+		{
+			bShouldGC = false;
+			bCookedAMapSinceLastGC = false;
+			NonMapPackageCountSinceLastGC = 0;
+
+			UE_LOG(LogCookCommandlet, Display, TEXT("GC..."));
+
+			CollectGarbage( RF_Native );
+		}
+		else
+		{
+			CookOnTheFlyServer->TickRecompileShaderRequests();
+
+			FPlatformProcess::Sleep(0.0f);
+		}
+
+
+
+		ProcessDeferredCommands();
+	}
+
+	return true;
+
+
+
+}
+
+
+void UCookCommandlet::ProcessDeferredCommands()
+{
+#if PLATFORM_MAC
+	// On Mac we need to process Cocoa events so that the console window for CookOnTheFlyServer is interactive
+	FPlatformMisc::PumpMessages(true);
+#endif
+
+	// update task graph
+	FTaskGraphInterface::Get().ProcessThreadUntilIdle(ENamedThreads::GameThread);
+
+	// execute deferred commands
+	for (int32 DeferredCommandsIndex = 0; DeferredCommandsIndex<GEngine->DeferredCommands.Num(); ++DeferredCommandsIndex)
+	{
+		GEngine->Exec( GWorld, *GEngine->DeferredCommands[DeferredCommandsIndex], *GLog);
+	}
+
+	GEngine->DeferredCommands.Empty();
+}
+
 bool UCookCommandlet::Cook(const TArray<ITargetPlatform*>& Platforms, TArray<FString>& FilesInPath)
 {
 	// Subsets for parallel processing
@@ -1062,6 +1269,7 @@ bool UCookCommandlet::Cook(const TArray<ITargetPlatform*>& Platforms, TArray<FSt
 				{
 					// Populate streaming install manifests
 					FString SandboxFilename = SandboxFile->ConvertToAbsolutePathForExternalAppForWrite(*Filename);
+					UE_LOG(LogCookCommandlet, Display, TEXT("Adding package to manifest %s, %s, %s"), *Pkg->GetName(), *SandboxFilename, *LastLoadedMapName);
 					ManifestGenerator.AddPackageToChunkManifest(Pkg, SandboxFilename, LastLoadedMapName, SandboxFile.GetOwnedPointer());
 				}
 					
@@ -1113,7 +1321,6 @@ bool UCookCommandlet::Cook(const TArray<ITargetPlatform*>& Platforms, TArray<FSt
 		{
 			break;
 		}
-			
 		// Attempt to find file for package name. THis is to make sure no short package
 		// names are passed to LoadPackage.
 		FString Filename;
@@ -1123,6 +1330,8 @@ bool UCookCommandlet::Cook(const TArray<ITargetPlatform*>& Platforms, TArray<FSt
 			
 			continue;
 		}
+		
+		UE_LOG(LogCookCommandlet, Display, TEXT("Processing package %s"), *Filename);
 		Filename = FPaths::ConvertRelativePathToFull(Filename);
 
 		if (bDoSubset)
@@ -1151,6 +1360,7 @@ bool UCookCommandlet::Cook(const TArray<ITargetPlatform*>& Platforms, TArray<FSt
 
 		if (bGenerateStreamingInstallManifests)
 		{
+			UE_LOG(LogCookCommandlet, Display, TEXT("PrepareToLoadNewPackage %s"), *Filename );
 			ManifestGenerator.PrepareToLoadNewPackage(Filename);
 		}
 
diff --git a/Engine/Source/Editor/UnrealEd/Private/CookOnTheFlyServer.cpp b/Engine/Source/Editor/UnrealEd/Private/CookOnTheFlyServer.cpp
--- a/Engine/Source/Editor/UnrealEd/Private/CookOnTheFlyServer.cpp
+++ b/Engine/Source/Editor/UnrealEd/Private/CookOnTheFlyServer.cpp
@@ -21,13 +21,17 @@
 #include "UnrealEdMessages.h"
 #include "GameDelegates.h"
 
+// cook by the book requirements
+#include "Commandlets/ChunkManifestGenerator.h"
+#include "Engine/WorldComposition.h"
+
 
 
 
 DEFINE_LOG_CATEGORY_STATIC(LogCookOnTheFly, Log, All);
 
 #define DEBUG_COOKONTHEFLY 0
-#define OUTPUT_TIMING 0
+#define OUTPUT_TIMING 1
 
 #if OUTPUT_TIMING
 
@@ -123,6 +127,71 @@ int FScopeTimer::GScopeDepth = 0;
 
 void OutputTimers()
 {
+	
+	if ( GTimerInfo.Num() <= 0 )
+		return;
+	
+	static FArchive *OutputDevice = NULL;
+
+	static TMap<FString, int> TimerIndexMap;
+
+
+	if ( OutputDevice == NULL )
+	{
+		OutputDevice = IFileManager::Get().CreateFileWriter(TEXT("CookOnTheFlyServerTiming.csv") );
+	}
+
+	TArray<FString> OutputValues;
+	OutputValues.AddZeroed(TimerIndexMap.Num());
+
+	bool OutputTimerIndexMap = false;
+	for ( auto TimerInfo : GTimerInfo )
+	{
+		int *IndexPtr = TimerIndexMap.Find(TimerInfo.Name);
+		int Index = 0;
+		if (IndexPtr == NULL)
+		{
+			Index = TimerIndexMap.Num();
+			TimerIndexMap.Add( TimerInfo.Name, Index );
+			OutputValues.AddZeroed();
+			OutputTimerIndexMap = true;
+		}
+		else
+			Index = *IndexPtr;
+
+		OutputValues[Index] = FString::Printf(TEXT("%f"),TimerInfo.Length);
+	}
+	static FString NewLine = FString(TEXT("\n"));
+
+	if (OutputTimerIndexMap)
+	{
+		TArray<FString> OutputHeader;
+		OutputHeader.AddZeroed( TimerIndexMap.Num() );
+		for ( auto TimerIndex : TimerIndexMap )
+		{
+			int LocalIndex = TimerIndex.Value;
+			OutputHeader[LocalIndex] = TimerIndex.Key;
+		}
+
+		for ( auto OutputString : OutputHeader)
+		{
+			OutputString.Append(TEXT(", "));
+			OutputDevice->Serialize( (void*)(*OutputString), OutputString.Len() * sizeof(TCHAR));
+		}
+
+		
+		OutputDevice->Serialize( (void*)*NewLine, NewLine.Len() * sizeof(TCHAR) );
+	}
+
+	for ( auto OutputString : OutputValues)
+	{
+		OutputString.Append(TEXT(", "));
+		OutputDevice->Serialize( (void*)(*OutputString), OutputString.Len() * sizeof(TCHAR));
+	}
+	OutputDevice->Serialize( (void*)*NewLine, NewLine.Len() * sizeof(TCHAR) );
+
+	OutputDevice->Flush();
+
 	UE_LOG( LogCookOnTheFly, Display, TEXT("Timing information for cook") );
 	UE_LOG( LogCookOnTheFly, Display, TEXT("Name\tlength(ms)") );
 	for ( auto TimerInfo : GTimerInfo )
@@ -200,10 +269,9 @@ public:
 
 static TMap<FString, FCachedPackageFilename> PackageFilenameCache;
 
-
-static const FCachedPackageFilename &Cache(UPackage *Package)
+static const FCachedPackageFilename &Cache(const FString &PackageName)
 {
-	FCachedPackageFilename *Cached = PackageFilenameCache.Find( Package->GetName() );
+	FCachedPackageFilename *Cached = PackageFilenameCache.Find( PackageName );
 	if ( Cached != NULL )
 	{
 		return *Cached;
@@ -211,17 +279,19 @@ static const FCachedPackageFilename &Cache(UPackage *Package)
 
 
 	FString Filename;
+	FString PackageFilename;
 	FString StandardFilename;
 	FName StandardFileFName = NAME_None;
-	if (FPackageName::DoesPackageExist(Package->GetName(), NULL, &Filename))
+	if (FPackageName::DoesPackageExist(PackageName, NULL, &Filename))
 	{
-		StandardFilename = FPaths::ConvertRelativePathToFull(Filename);
+		StandardFilename = PackageFilename = FPaths::ConvertRelativePathToFull(Filename);
+		
 
 		FPaths::MakeStandardFilename(StandardFilename);
 		StandardFileFName = FName(*StandardFilename);
 	}
 
-	return PackageFilenameCache.Emplace( Package->GetName(), MoveTemp(FCachedPackageFilename(MoveTemp(Filename),MoveTemp(StandardFilename), StandardFileFName)) );
+	return PackageFilenameCache.Emplace( PackageName, MoveTemp(FCachedPackageFilename(MoveTemp(PackageFilename),MoveTemp(StandardFilename), StandardFileFName)) );
 }
 
 
@@ -236,20 +306,35 @@ static const FCachedPackageFilename &Cache(UPackage *Package)
 	return Filename;
 }*/
 
+static FString GetCachedPackageFilename( const FString &PackageName )
+{
+	return Cache( PackageName ).PackageFilename;
+}
+
+static FString GetCachedStandardPackageFilename( const FString &PackageName )
+{
+	return Cache( PackageName ).StandardFilename;
+}
+
+static FName GetCachedStandardPackageFileFName( const FString &PackageName )
+{
+	return Cache( PackageName ).StandardFileFName;
+}
+
 
 static FString GetCachedPackageFilename( UPackage *Package )
 {
-	return Cache( Package ).PackageFilename;
+	return Cache( Package->GetName() ).PackageFilename;
 }
 
 static FString GetCachedStandardPackageFilename( UPackage* Package )
 {
-	return Cache(Package).StandardFilename;
+	return Cache( Package->GetName() ).StandardFilename;
 }
 
 static FName GetCachedStandardPackageFileFName( UPackage *Package )
 {
-	return Cache(Package).StandardFileFName;
+	return Cache( Package->GetName() ).StandardFileFName;
 }
 
 /* helper structs
@@ -267,24 +352,34 @@ struct FRecompileRequest
  *****************************************************************************/
 
 UCookOnTheFlyServer::UCookOnTheFlyServer(const class FPostConstructInitializeProperties& PCIP)
-	: Super(PCIP)
+	: Super(PCIP),
+	CurrentCookMode(ECookMode::CookOnTheFly),
+	CookByTheBookOptions(NULL),
+	CookFlags(ECookInitializationFlags::None),
+	bIsSavingPackage( false )
 {
-	bCookOnTheFly = true;
 }
 
-
+UCookOnTheFlyServer::~UCookOnTheFlyServer()
+{
+	if ( CookByTheBookOptions )
+	{		
+		delete CookByTheBookOptions;
+		CookByTheBookOptions = NULL;
+	}
+}
 
 void UCookOnTheFlyServer::Tick(float DeltaTime)
 {
 	uint32 CookedPackagesCount = 0;
-	const static float CookOnTheSideTimeSlice = 0.01f;
+	const static float CookOnTheSideTimeSlice = 0.1f; // seconds
 	TickCookOnTheSide( CookOnTheSideTimeSlice, CookedPackagesCount);
 	TickRecompileShaderRequests();
 }
 
 bool UCookOnTheFlyServer::IsTickable() const 
 { 
-	return bIsAutoTick; 
+	return IsCookFlagSet(ECookInitializationFlags::AutoTick); 
 }
 
 TStatId UCookOnTheFlyServer::GetStatId() const
@@ -296,6 +391,7 @@ TStatId UCookOnTheFlyServer::GetStatId() const
 
 bool UCookOnTheFlyServer::StartNetworkFileServer( const bool BindAnyPort )
 {
+	check( CurrentCookMode == ECookMode::CookOnTheFly );
 	//GetDerivedDataCacheRef().WaitForQuiescence(false);
 
 	// start the listening thread
@@ -360,6 +456,290 @@ bool UCookOnTheFlyServer::BroadcastFileserverPresence( const FGuid &InstanceId )
 	return true;
 }
 
+/*----------------------------------------------------------------------------
+	FArchiveFindReferences.
+----------------------------------------------------------------------------*/
+/**
+ * Archive for gathering all the object references to other objects
+ */
+class FArchiveFindReferences : public FArchiveUObject
+{
+private:
+	/**
+	 * I/O function.  Called when an object reference is encountered.
+	 *
+	 * @param	Obj		a pointer to the object that was encountered
+	 */
+	FArchive& FArchiveFindReferences::operator<<( UObject*& Obj )
+	{
+		if( Obj )
+		{
+			FoundObject( Obj );
+		}
+		return *this;
+	}
+
+	virtual FArchive& operator<< (class FAssetPtr& Value) override
+	{
+		if ( Value.Get() )
+		{
+			Value.Get()->Serialize( *this );
+		}
+		return *this;
+	}
+	virtual FArchive& operator<< (struct FStringAssetReference& Value) override
+	{
+		if ( Value.ResolveObject() )
+		{
+			Value.ResolveObject()->Serialize( *this );
+		}
+		return *this;
+	}
+
+
+	void FoundObject( UObject* Object )
+	{
+		if ( RootSet.Find(Object) == NULL )
+		{
+			if ( Exclude.Find(Object) == INDEX_NONE )
+			{
+				// remove this check later because don't want this happening in development builds
+				check(RootSetArray.Find(Object)==INDEX_NONE);
+
+				RootSetArray.Add( Object );
+				RootSet.Add(Object);
+				Found.Add(Object);
+			}
+		}
+	}
+
+
+	/**
+	 * list of Outers to ignore;  any objects encountered that have one of
+	 * these objects as an Outer will also be ignored
+	 */
+	TArray<UObject*> &Exclude;
+
+	/** list of objects that have been found */
+	TSet<UObject*> &Found;
+	
+	/** the objects to display references to */
+	TArray<UObject*> RootSetArray;
+	/** Reflection of the rootsetarray */
+	TSet<UObject*> RootSet;
+
+public:
+
+	/**
+	 * Constructor
+	 * 
+	 * @param	inOutputAr		archive to use for logging results
+	 * @param	inOuter			only consider objects that do not have this object as its Outer
+	 * @param	inSource		object to show references for
+	 * @param	inExclude		list of objects that should be ignored if encountered while serializing SourceObject
+	 */
+	FArchiveFindReferences::FArchiveFindReferences( TSet<UObject*> InRootSet, TSet<UObject*> &inFound, TArray<UObject*> &inExclude )
+		: Exclude(inExclude)
+		, Found(inFound)
+		, RootSet(InRootSet)
+	{
+		ArIsObjectReferenceCollector = true;
+
+		for ( const auto& Object : RootSet )
+		{
+			RootSetArray.Add( Object );
+		}
+		
+		// loop through all the objects in the root set and serialize them
+		for ( int RootIndex = 0; RootIndex < RootSetArray.Num(); ++RootIndex )
+		{
+			UObject* SourceObject = RootSetArray[RootIndex];
+
+			// quick sanity check
+			check(SourceObject);
+			check(SourceObject->IsValidLowLevel());
+
+			SourceObject->Serialize( *this );
+		}
+
+	}
+
+	/**
+  	 * Returns the name of the Archive.  Useful for getting the name of the package a struct or object
+	 * is in when a loading error occurs.
+	 *
+	 * This is overridden for the specific Archive Types
+	 **/
+	virtual FString GetArchiveName() const { return TEXT("FArchiveFindReferences"); }
+};
+
+void UCookOnTheFlyServer::GetDependencies( const TSet<UPackage*>& Packages, TSet<UObject*>& Found)
+{
+
+	TSet<UObject*> RootSet;
+
+
+	// Iterate through the object list
+	for( FObjectIterator It; It; ++It )
+	{
+		UPackage* Package = It->GetOutermost();
+		if ( Packages.Find( Package ) )
+		{
+			RootSet.Add(*It);
+			Found.Add(*It);
+		}
+
+		/*if (It->IsIn(Package))
+		{
+		}*/
+	}
+
+	TArray<UObject*> Exclude;
+	FArchiveFindReferences ArFindReferences( RootSet, Found, Exclude );
+
+
+
+
+	// Iterate through the object list
+	/*for( FObjectIterator It; It; ++It )
+	{
+		// if this object is within the package specified, serialize the object
+		// into a specialized archive which logs object names encountered during
+		// serialization -- rjp
+		if ( It->IsIn(Pkg) )
+		{
+			if ( It->GetOuter() == Pkg )
+			{
+				FArchiveFindReferences ArFindReferences( Pkg, *It, Found, Exclude );
+			}
+			else if ( bRecurse )
+			{
+				// Two options -
+				// a) this object is a function or something (which we don't care about)
+				// b) this object is inside a group inside the specified package (which we do care about)
+				UObject* CurrentObject = *It;
+				UObject* CurrentOuter = It->GetOuter();
+				while ( CurrentObject && CurrentOuter )
+				{
+					// this object is a UPackage (a group inside a package)
+					// abort
+					if ( CurrentObject->GetClass() == UPackage::StaticClass() )
+						break;
+
+					// see if this object's outer is a UPackage
+					if ( CurrentOuter->GetClass() == UPackage::StaticClass() )
+					{
+						// if this object's outer is our original package, the original object (It)
+						// wasn't inside a group, it just wasn't at the base level of the package
+						// (its Outer wasn't the Pkg, it was something else e.g. a function, state, etc.)
+						/// ....just skip it
+						if ( CurrentOuter == Pkg )
+							break;
+
+						// otherwise, we've successfully found an object that was in the package we
+						// were searching, but would have been hidden within a group - let's log it
+						FArchiveFindReferences ArFindReferences( CurrentOuter, CurrentObject, Found, Exclude );
+						break;
+					}
+
+					CurrentObject = CurrentOuter;
+					CurrentOuter = CurrentObject->GetOuter();
+				}
+			}
+		}
+	}*/
+}
+
+
+void UCookOnTheFlyServer::GenerateManifestInfo( UPackage* Package, const TArray<FName>& TargetPlatformNames )
+{
+	if ( !CookByTheBookOptions )
+		return;
+	
+	// generate dependency information for this package
+	
+
+	TSet<UPackage*> RootPackages;
+	RootPackages.Add(Package);
+
+	FString LastLoadedMapName;
+
+	// load sublevels
+	UWorld* World = UWorld::FindWorldInPackage(Package);
+
+	if ( World )
+	{
+		for (const auto& StreamingLevel : World->StreamingLevels)
+		{
+			// Dependencies.Add( StreamingLevel->GetLoadedLevel() );
+			RootPackages.Add(StreamingLevel->GetLoadedLevel()->GetOutermost() );
+			//GetDependencies( StreamingLevel->GetLoadedLevel()->GetOutermost(), Dependencies );
+		}
+
+		TArray<FString> NewPackagesToCook;
+
+		// Collect world composition tile packages to cook
+		if (World->WorldComposition)
+		{
+			World->WorldComposition->CollectTilesToCook(NewPackagesToCook);
+		}
+
+		for ( const auto& PackageName : NewPackagesToCook)
+		{
+			UPackage* Package = LoadPackage( NULL, *PackageName, LOAD_None );
+
+			RootPackages.Add( Package );
+			//GetDependencies( Package, Dependencies );
+
+			// Dependencies.Add(Package);
+		}
+
+		LastLoadedMapName = Package->GetName();
+	}
+
+	TSet<UObject*> Dependencies; 
+	GetDependencies( RootPackages, Dependencies );
+
+
+	FName StandardFilename = GetCachedStandardPackageFileFName(Package);
+
+	TSet<UPackage*> Packages;
+	for ( const auto& Object : Dependencies )
+	{
+		Packages.Add( Object->GetOutermost() );
+	}
+
+	// update the manifests with generated dependencies
+	for ( const auto& PlatformName : TargetPlatformNames )
+	{
+		FChunkManifestGenerator* ManifestGenerator = CookByTheBookOptions->ManifestGenerators.FindChecked(PlatformName);
+
+		if ( CookByTheBookOptions->bGenerateStreamingInstallManifests )
+		{
+			UE_LOG(LogCookOnTheFly, Display, TEXT("PrepareToLoadNewPackage %s"), *StandardFilename.ToString() );
+			ManifestGenerator->PrepareToLoadNewPackage( StandardFilename.ToString() );
+		}
+
+		for ( const auto& DependentPackage : Packages )
+		{
+			ManifestGenerator->OnLastPackageLoaded( DependentPackage  );
+		}
+	
+		for ( const auto& DependentPackage  : Packages )
+		{
+			FString Filename = GetCachedPackageFilename( DependentPackage );
+			if (!Filename.IsEmpty())
+			{
+				// Populate streaming install manifests
+				FString SandboxFilename = SandboxFile->ConvertToAbsolutePathForExternalAppForWrite(*Filename);
+				UE_LOG(LogCookOnTheFly, Display, TEXT("Adding package to manifest %s, %s, %s"), *DependentPackage->GetName(), *SandboxFilename, *LastLoadedMapName);
+				ManifestGenerator->AddPackageToChunkManifest(DependentPackage, SandboxFilename, LastLoadedMapName, SandboxFile.GetOwnedPointer());
+			}
+		}
+	}
+}
+
+
 uint32 UCookOnTheFlyServer::TickCookOnTheSide( const float TimeSlice, uint32 &CookedPackageCount )
 {
 	double StartTime = FPlatformTime::Seconds();
@@ -369,22 +749,21 @@ uint32 UCookOnTheFlyServer::TickCookOnTheSide( const float TimeSlice, uint32 &Co
 	// we use this in the unsolicited packages processing below
 	TArray<FName> AllTargetPlatformNames;
 
-	while (!GIsRequestingExit)
+	while (!GIsRequestingExit || CurrentCookMode == ECookMode::CookByTheBook)
 	{
-		SCOPE_TIMER(TickCookOnTheSide);
+		// if we just cooked a map then don't process anything the rest of this tick
+		if ( Result & COSR_CookedMap )
+		{
+			break;
+		}
+
 
-		bool bCookedPackage = false;
 		FFilePlatformRequest ToBuild;
 
-		bool bIsUnsolicitedRequest = false;
-		if ( FileRequests.HasItems() )
-		{
-			FileRequests.Dequeue( &ToBuild );
-		}
-		else if ( UnsolicitedFileRequests.HasItems() )
+		//bool bIsUnsolicitedRequest = false;
+		if ( CookRequests.HasItems() )
 		{
-			bIsUnsolicitedRequest = true;
-			UnsolicitedFileRequests.Dequeue( &ToBuild );
+			CookRequests.Dequeue( &ToBuild );
 		}
 		else
 		{
@@ -392,130 +771,187 @@ uint32 UCookOnTheFlyServer::TickCookOnTheSide( const float TimeSlice, uint32 &Co
 			break;
 		}
 
+		if (CookedPackages.Exists(ToBuild))
+		{
+			UE_LOG(LogCookOnTheFly, Display, TEXT("Package for platform already cooked %s, discarding request"), *ToBuild.GetFilename().ToString());
+			continue;
+		}
+
+		UE_LOG(LogCookOnTheFly, Display, TEXT("Processing package %s"), *ToBuild.GetFilename().ToString());
+
+		SCOPE_TIMER(TickCookOnTheSide);
+
 		check( ToBuild.IsValid() );
-		TArray<FName> TargetPlatformNames;
+		const TArray<FName> &TargetPlatformNames = ToBuild.GetPlatformnames();
+
+		TArray<UPackage*> PackagesToSave;
+
+#if OUTPUT_TIMING
+		//FScopeTimer PackageManualTimer( ToBuild.GetFilename().ToString(), false );
+		UE_LOG(LogCookOnTheFly, Display,  TEXT("ProcessingPackage %s"), *ToBuild.GetFilename().ToString() );
+#endif
 
-		if( ToBuild.GetPlatformname() != NAME_None )
+		for ( const auto &PlatformName : TargetPlatformNames )
 		{
-			TargetPlatformNames.AddUnique( ToBuild.GetPlatformname() );
+			AllTargetPlatformNames.AddUnique(PlatformName);
 		}
 
 		bool bWasUpToDate = false;
 
-		if (!ThreadSafeFilenameSet.Exists(ToBuild))
+		bool bLastLoadWasMap = false;
+		FString LastLoadedMapName;
+
+		const FString BuildFilename = ToBuild.GetFilename().ToString();
+
+		// if we have no target platforms then we want to cook because this will cook for all target platforms in that case
+		bool bShouldCook = TargetPlatformNames.Num() > 0 ? false : ShouldCook( BuildFilename, NAME_None );
+		{
+			SCOPE_TIMER(ShouldCook);
+			for ( int Index = 0; Index < TargetPlatformNames.Num(); ++Index )
+			{
+				bShouldCook |= ShouldCook( ToBuild.GetFilename().ToString(), TargetPlatformNames[Index] );
+			}
+		}
+		
+		if ( bShouldCook ) // if we should cook the package then cook it otherwise add it to the list of already cooked packages below
 		{
+			FString PackageName;	
+			UPackage *Package = NULL;
+			if ( FPackageName::TryConvertFilenameToLongPackageName(BuildFilename, PackageName) )
 			{
-				const FString BuildFilename = ToBuild.GetFilename().ToString();
+				Package = FindObject<UPackage>( ANY_PACKAGE, *PackageName );
+			}
+
+			UE_LOG( LogCookOnTheFly, Display, TEXT("Processing request %s"), *BuildFilename)
 
-				// if we have no target platforms then we want to cook because this will cook for all target platforms in that case
-				bool bShouldCook = TargetPlatformNames.Num() > 0 ? false : ShouldCook( BuildFilename, NAME_None );
-				for ( int Index = 0; Index < TargetPlatformNames.Num(); ++Index )
+			//  if the package is already loaded then try to avoid reloading it :)
+			if ( ( Package == NULL ) || ( Package->IsFullyLoaded() == false ) )
+			{
+				// moved to GenerateDependencies 
+				/*if (CookByTheBookOptions && CookByTheBookOptions->bGenerateStreamingInstallManifests)
 				{
-					SCOPE_TIMER(ShouldCook);
-					bShouldCook |= ShouldCook( ToBuild.GetFilename().ToString(), TargetPlatformNames[Index] );
-				}
+					CookByTheBookOptions->ManifestGenerator->PrepareToLoadNewPackage(BuildFilename);
+				}*/
 
+				SCOPE_TIMER(LoadPackage);
+				Package = LoadPackage( NULL, *BuildFilename, LOAD_None );
+			}
+			else
+			{
+				UE_LOG(LogCookOnTheFly, Display, TEXT("Package already loaded %s avoiding reload"), *BuildFilename );
+			}
 
-				if ( bShouldCook ) // if we should cook the package then cook it otherwise add it to the list of already cooked packages below
-				{
-					FString PackageName;	
-					UPackage *Package = NULL;
-					if ( FPackageName::TryConvertFilenameToLongPackageName(BuildFilename, PackageName) )
-					{
-						Package = FindObject<UPackage>( ANY_PACKAGE, *PackageName );
-					}
+			if( Package == NULL )
+			{
+				UE_LOG(LogCookOnTheFly, Error, TEXT("Error loading %s!"), *BuildFilename );
+				Result |= COSR_ErrorLoadingPackage;
+			}
+			else
+			{
+				check(Package);
 
-					UE_LOG( LogCookOnTheFly, Display, TEXT("Processing request %s"), *BuildFilename)
 
-					//  if the package is already loaded then try to avoid reloading it :)
-					if ( ( Package == NULL ) || ( Package->IsFullyLoaded() == false ) )
-					{
-						SCOPE_TIMER(LoadPackage);
-						Package = LoadPackage( NULL, *BuildFilename, LOAD_None );
-					}
-					else
+				if (Package->ContainsMap())
+				{
+					if ( IsCookByTheBookMode() )
 					{
-						UE_LOG(LogCookOnTheFly, Display, TEXT("Package already loaded %s avoiding reload"), *BuildFilename );
-					}
+						// load sublevels
+						UWorld* World = UWorld::FindWorldInPackage(Package);
 
-					if( Package == NULL )
-					{
-						UE_LOG(LogCookOnTheFly, Error, TEXT("Error loading %s!"), *BuildFilename );
-						Result |= COSR_ErrorLoadingPackage;
-					}
-					else
-					{
-						SCOPE_TIMER(SaveCookedPackage);
-						if( SaveCookedPackage(Package, SAVE_KeepGUID | SAVE_Async, bWasUpToDate, TargetPlatformNames) )
+						// TArray<FString> PreviouslyCookedPackages;
+						if (World->StreamingLevels.Num())
 						{
-							FString Name = Package->GetPathName();
-							/*FString PackageFilename(GetPackageFilename(Package));
-							FPaths::MakeStandardFilename(PackageFilename);*/
-							FString PackageFilename( GetCachedStandardPackageFilename( Package ) );
-							if ( PackageFilename != BuildFilename )
-							{
-								// we have saved something which we didn't mean to load 
-								//  sounds unpossible.... but it is due to searching for files and such
-								//  mark the original request as processed (if this isn't actually the file they were requesting then it will fail)
-								//	and then also save our new request as processed so we don't do it again
-								ThreadSafeFilenameSet.Add( ToBuild );
-								UE_LOG( LogCookOnTheFly, Display, TEXT("Request for %s received saved %s"), *BuildFilename, *PackageFilename );
-								ToBuild.SetFilename( PackageFilename );
-							}
-							// ToBuild.Filename = PackageFilename;
+							//World->LoadSecondaryLevels(true, &PreviouslyCookedPackages);
+							World->LoadSecondaryLevels(true, NULL);
+						}
 
-							// Update flags used to determine garbage collection.
-							if (Package->ContainsMap())
-							{
-								Result |= COSR_CookedMap;
-								//bCookedAMapSinceLastGC = true;
-							}
-							else
-							{
-								++CookedPackageCount;
-								Result |= COSR_CookedPackage;
-								//NonMapPackageCountSinceLastGC++;
-								// LastCookActionTime = FPlatformTime::Seconds();
-							}
-							bCookedPackage = true;
+						TArray<FString> NewPackagesToCook;
+
+						// Collect world composition tile packages to cook
+						if (World->WorldComposition)
+						{
+							World->WorldComposition->CollectTilesToCook(NewPackagesToCook);
 						}
 
-						//@todo ResetLoaders outside of this (ie when Package is NULL) causes problems w/ default materials
-						if (Package->HasAnyFlags(RF_RootSet) == false)
+						for ( auto PackageName : NewPackagesToCook )
 						{
-							SCOPE_TIMER(ResetLoaders);
-							ResetLoaders(Package);
+							FString Filename;
+							if ( FPackageName::DoesPackageExist(PackageName, NULL, &Filename) )
+							{
+								FString StandardFilename = FPaths::ConvertRelativePathToFull(Filename);
+								FPaths::MakeStandardFilename(StandardFilename);
+
+								CookRequests.EnqueueUnique(FFilePlatformRequest(MoveTemp(FName(*StandardFilename)), MoveTemp(TargetPlatformNames)));
+							}
 						}
-					}
-				}
-			}
 
+						// maps don't compile level script actors correctly unless we do FULL GC's, they may also hold weak pointer refs that need to be reset
+						// NumProcessedSinceLastGC = GCInterval; 
+						// ForceGarbageCollect();
 
-			for ( int Index = 0; Index < TargetPlatformNames.Num(); ++Index )
-			{
-				AllTargetPlatformNames.AddUnique(TargetPlatformNames[Index]);
+						LastLoadedMapName = Package->GetName();
+						bLastLoadWasMap = true;
+					}
 
-				FFilePlatformRequest FileRequest( ToBuild.GetFilename(), TargetPlatformNames[Index]);
+				}
 
-				if ( bIsUnsolicitedRequest )
+				FString Name = Package->GetPathName();
+				FString PackageFilename( GetCachedStandardPackageFilename( Package ) );
+				if ( PackageFilename != BuildFilename )
 				{
-					if ((FPaths::FileExists(FileRequest.GetFilename().ToString()) == true) &&
-						(bWasUpToDate == false))
-					{
-						UnsolicitedCookedPackages.EnqueueUnique( FileRequest );
-						UE_LOG(LogCookOnTheFly, Display, TEXT("UnsolicitedCookedPackages: %s"), *FileRequest.GetFilename().ToString());
-					}
+					// we have saved something which we didn't mean to load 
+					//  sounds unpossible.... but it is due to searching for files and such
+					//  mark the original request as processed (if this isn't actually the file they were requesting then it will fail)
+					//	and then also save our new request as processed so we don't do it again
+					PackagesToSave.AddUnique( Package );
+					UE_LOG( LogCookOnTheFly, Display, TEXT("Request for %s received going to save %s"), *BuildFilename, *PackageFilename );
+					// CookedPackages.Add( ToBuild );
+					
+					// ToBuild.SetFilename( PackageFilename );
+				}
+				else
+				{
+					PackagesToSave.AddUnique( Package );
 				}
-				ThreadSafeFilenameSet.Add( FileRequest );
 			}
 		}
-		else
+
+		
+		if ( PackagesToSave.Num() == 0 )
 		{
-			UE_LOG(LogCookOnTheFly, Display, TEXT("Package for platform already cooked %s, discarding request"), *ToBuild.GetFilename().ToString());
+			// an error occured or the file doesn't exist
+			// this happens often as we request files with different extensions when we are searching for files
+			// just return that we processed the cook request
+			// the network file manager will then handle the missing file and search somewhere else
+			UE_LOG(LogCookOnTheFly, Display, TEXT("Not cooking package %s"), *ToBuild.GetFilename().ToString());
+			CookedPackages.Add( FFilePlatformRequest( ToBuild.GetFilename(), TargetPlatformNames ) );
+			continue;
+		}
+
+
+		
+		// if we are doing cook by the book we need to resolve string asset references
+		// this will load more packages :)
+		if ( IsCookByTheBookMode() )
+		{
+
+			// if we are doing cook by the book then we are not waiting on this single package to be cooked
+			// so we could just load some more packages and then save them all at once
+			/*if ( bLastLoadWasMap && !bNeedsGC )
+			{
+			// early out and load some more packages to save
+			continue;
+			}*/
+
+
+
+			SCOPE_TIMER(ResolveRedirectors);
+			GRedirectCollector.ResolveStringAssetReference();
 		}
 
+		int32 FirstUnsolicitedPackage = PackagesToSave.Num();
 
-		if ( bCookedPackage && !bIsUnsolicitedRequest)
+		// generate a list of other packages which were loaded with this one
 		{
 			SCOPE_TIMER(UnsolicitedMarkup);
 
@@ -529,92 +965,256 @@ uint32 UCookOnTheFlyServer::TickCookOnTheSide( const float TimeSlice, uint32 &Co
 			PackageNames.Empty(ObjectsInOuter.Num());
 			{
 				SCOPE_TIMER(GeneratePackageNames);
-				ACCUMULATE_TIMER(GetPathName);
-				ACCUMULATE_TIMER(GetPackageFilename);
 				ACCUMULATE_TIMER(PackageCast);
-				ACCUMULATE_TIMER(ConvertRelativePathToFull);
-				ACCUMULATE_TIMER(PackageNamesAdd);
-				ACCUMULATE_TIMER(PackageFNameConvert);
-				for( int32 Index = 0; Index < ObjectsInOuter.Num() && !GIsRequestingExit; Index++ )
+				ACCUMULATE_TIMER(AddUnassignedPackageToManifest);
+				ACCUMULATE_TIMER(ConvertToAbsolutePath);
+				ACCUMULATE_TIMER(ManifestAddUnassignedPackage);
+				ACCUMULATE_TIMER(CookedPackagesExists);
+				ACCUMULATE_TIMER(PackagesToSaveAdd);
+				for( int32 Index = 0; Index < ObjectsInOuter.Num(); Index++ )
 				{
 					ACCUMULATE_TIMER_START(PackageCast);
 					UPackage* Package = Cast<UPackage>(ObjectsInOuter[Index]);
 					ACCUMULATE_TIMER_STOP(PackageCast);
 
-				/*for ( TObjectIterator<UPackage> It; It; ++It ) // ~700ms
+					if (Package)
+					{
+						// this has been replaced to use the GenerateDependencies call
+						/*if ( CookByTheBookOptions && CookByTheBookOptions->ManifestGenerator != NULL )
+						{
+							check( CurrentCookMode == ECookMode::CookByTheBookFromTheEditor || CurrentCookMode == ECookMode::CookByTheBook );
+							// if we are generating a manifest add the package to the manifest :)
+							FString Filename = GetCachedPackageFilename(Package);
+							if (!Filename.IsEmpty())
+							{
+								// Populate streaming install manifests
+								FString SandboxFilename = SandboxFile->ConvertToAbsolutePathForExternalAppForWrite(*Filename);
+								CookByTheBookOptions->ManifestGenerator->AddPackageToChunkManifest(Package, SandboxFilename, LastLoadedMapName, SandboxFile.GetOwnedPointer());
+							}
+						}*/
+
+
+						if ( CookByTheBookOptions )
+						{
+							ACCUMULATE_TIMER_START(AddUnassignedPackageToManifest);
+							for ( const auto &TargetPlatform : AllTargetPlatformNames )
+							{
+								FChunkManifestGenerator*& Manifest = CookByTheBookOptions->ManifestGenerators.FindChecked(TargetPlatform);
+
+								FString Filename = GetCachedPackageFilename(Package);
+								if (!Filename.IsEmpty())
+								{
+									// Populate streaming install manifests
+									ACCUMULATE_TIMER_START(ConvertToAbsolutePath);
+									FString SandboxFilename = SandboxFile->ConvertToAbsolutePathForExternalAppForWrite(*Filename);
+									ACCUMULATE_TIMER_STOP(ConvertToAbsolutePath);
+									ACCUMULATE_TIMER_START(ManifestAddUnassignedPackage);
+									Manifest->AddUnassignedPackageToManifest(Package, SandboxFilename);
+									ACCUMULATE_TIMER_STOP(ManifestAddUnassignedPackage);
+								}
+							}
+							ACCUMULATE_TIMER_STOP(AddUnassignedPackageToManifest);
+						}
+
+
+						FName StandardFilename = GetCachedStandardPackageFileFName(Package);
+						if ( StandardFilename != NAME_None ) // if we have name none that means we are in core packages or something...
+						{
+							// check if the package has already been saved
+							bool bShouldSave = false;
+
+							ACCUMULATE_TIMER_START(CookedPackagesExists);
+							if ( !CookedPackages.Exists(FFilePlatformRequest(StandardFilename, TargetPlatformNames)))
+							{
+								bShouldSave = true;
+							}
+							ACCUMULATE_TIMER_STOP(CookedPackagesExists);
+							ACCUMULATE_TIMER_START(PackagesToSaveAdd);
+							if ( bShouldSave )
+							{
+								PackagesToSave.AddUnique( Package );
+							}
+							ACCUMULATE_TIMER_STOP(PackagesToSaveAdd);
+						}
+					}
+				}
+			}
+		}
+
+
+		{
+			SCOPE_TIMER(GenerateManifestInfo);
+			// update manifest with cooked package info
+			GenerateManifestInfo( PackagesToSave[0], AllTargetPlatformNames );
+		}
+
+
+
+		if ( PackagesToSave.Num() )
+		{
+			SCOPE_TIMER(SavingPackages);
+			for ( int32 I = 0; I < PackagesToSave.Num(); ++I )
+			{
+				// if we are processing unsolicited packages we can optionally not save these right now
+				// the unsolicited packages which we missed now will be picked up on next run
+				if ( (CurrentCookMode == ECookMode::CookOnTheFly || CurrentCookMode == ECookMode::CookByTheBookFromTheEditor) && (I >= FirstUnsolicitedPackage) )
+				{
+					bool bShouldFinishTick = false;
+
+					if ( CookRequests.HasItems() )
+					{
+						bShouldFinishTick = true;
+					}
+
+					if ( (FPlatformTime::Seconds() - StartTime) > TimeSlice)
+					{
+						bShouldFinishTick = true;
+						// our timeslice is up
+					}
+
+
+					if ( bShouldFinishTick )
+					{
+						// don't save these packages because we have a real request
+						// enqueue all the packages which we were about to save
+						if (CurrentCookMode==ECookMode::CookByTheBookFromTheEditor)
+						{
+							for ( int32 RemainingIndex = I; RemainingIndex < PackagesToSave.Num(); ++RemainingIndex )
+							{
+								FName StandardFilename = GetCachedStandardPackageFileFName(PackagesToSave[RemainingIndex]);
+								CookRequests.EnqueueUnique(FFilePlatformRequest(StandardFilename, AllTargetPlatformNames));
+							}
+						}
+
+						// break out of the loop
+						break;
+					}
+				}
+
+				UPackage *Package = PackagesToSave[I];
+				FName PackageFName = GetCachedStandardPackageFileFName(Package);
+				TArray<FName> SaveTargetPlatformNames = AllTargetPlatformNames;
+				TArray<FName> CookedTargetPlatforms;
+				if ( CookedPackages.GetCookedPlatforms(PackageFName, CookedTargetPlatforms) )
+				{
+					for ( auto const &CookedPlatform : CookedTargetPlatforms )
+					{
+						SaveTargetPlatformNames.Remove( CookedPlatform );
+					}
+				}
+
+				if ( SaveTargetPlatformNames.Num() == 0 )
 				{
-					UPackage *Package = *It;*/
-
+					// already saved this package
+					continue;
+				}
 
-					if (Package)
+				SCOPE_TIMER(SaveCookedPackage);
+				if( SaveCookedPackage(Package, SAVE_KeepGUID | SAVE_Async | (IsCookFlagSet(ECookInitializationFlags::Unversioned) ? SAVE_Unversioned : 0), bWasUpToDate, AllTargetPlatformNames ) )
+				{
+					// Update flags used to determine garbage collection.
+					if (Package->ContainsMap())
 					{
-						ACCUMULATE_TIMER_START(GetPackageFilename);
-						FName PackageFilename(GetCachedStandardPackageFileFName(Package));
-						ACCUMULATE_TIMER_STOP(GetPackageFilename);
-		
-						if ( PackageFilename != NAME_None )
-						{
-							ACCUMULATE_TIMER_SCOPE(PackageNamesAdd);
-							PackageNames.Add(PackageFilename);
-						}
-						
+						Result |= COSR_CookedMap;
+					}
+					else
+					{
+						++CookedPackageCount;
+						Result |= COSR_CookedPackage;
 					}
 				}
 
-			}
 
-			{
-				SCOPE_TIMER(EnqueueUnsolicitedPackages);
+				/*{
+					SCOPE_TIMER(GenerateManifestInfo);
+					// update manifest with cooked package info
+					GenerateManifestInfo( Package, AllTargetPlatformNames );
+				}*/
+				
+
+
+				//@todo ResetLoaders outside of this (ie when Package is NULL) causes problems w/ default materials
+				if (Package->HasAnyFlags(RF_RootSet) == false && ((CurrentCookMode==ECookMode::CookOnTheFly)) )
 				{
-					SCOPE_TIMER(ThreadSafeFileNameSetLock);
-					ThreadSafeFilenameSet.Lock(); // manually lock and unlock here so that we don't thrash the 
+					SCOPE_TIMER(ResetLoaders);
+					ResetLoaders(Package);
 				}
-				
 
-				ACCUMULATE_TIMER(CreateRequest);
-				ACCUMULATE_TIMER(ThreadSafeFilenameSetExists);
-				ACCUMULATE_TIMER(ReenqueueUnsolicitedRequests);
-				for ( const FName &PackageFilename : PackageNames )
+				FName StandardFilename = GetCachedStandardPackageFileFName(Package);
+
+				if ( StandardFilename != NAME_None )
 				{
-					bool AlreadyCooked = true;
-					for ( const FName TargetPlatformName : AllTargetPlatformNames )
-					{
-						ACCUMULATE_TIMER_START(CreateRequest);
-						auto Request = FFilePlatformRequest(PackageFilename, TargetPlatformName);
-						ACCUMULATE_TIMER_STOP(CreateRequest);
+					// mark the package as cooked
+					FFilePlatformRequest FileRequest( StandardFilename, AllTargetPlatformNames);
+					CookedPackages.Add( FileRequest );
 
-						ACCUMULATE_TIMER_SCOPE(ThreadSafeFilenameSetExists);
-						if ( !ThreadSafeFilenameSet.Exists(Request) )
-						{
-							ACCUMULATE_TIMER_SCOPE(ReenqueueUnsolicitedRequests);
-							UnsolicitedFileRequests.EnqueueUnique(Request);
-						}
-					}
-					/*if ( !AlreadyCooked )
+					if ( (CurrentCookMode == ECookMode::CookOnTheFly) && (I >= FirstUnsolicitedPackage) ) 
 					{
-						ACCUMULATE_TIMER_SCOPE(ReenqueueUnsolicitedRequests);
-						for (const FName TargetPlatformName : AllTargetPlatformNames)
+						// this is an unsolicited package
+						if ((FPaths::FileExists(FileRequest.GetFilename().ToString()) == true) &&
+							(bWasUpToDate == false))
 						{
-							const FFilePlatformRequest Request( PackageFilename, TargetPlatformName );
-							UnsolicitedFileRequests.EnqueueUnique(Request);
+							UnsolicitedCookedPackages.AddCookedPackage( FileRequest );
+							UE_LOG(LogCookOnTheFly, Display, TEXT("UnsolicitedCookedPackages: %s"), *FileRequest.GetFilename().ToString());
 						}
-					}*/
+					}
 				}
-				ThreadSafeFilenameSet.Unlock();
 			}
 		}
 
+		// if after all this our requested file didn't get saved then we need to mark it as saved (this can happen if the file loaded didn't match the one we saved)
+		// for example say we request A.uasset which doesn't exist however A.umap exists, our load call will succeed as it will find A.umap, then the save will save A.umap
+		// our original package request will fail however
+		if ( !CookedPackages.Exists(ToBuild) )
+		{
+			CookedPackages.Add( ToBuild );
+		}
+
+		// TODO: Daniel: this is reference code needs to be reimplemented on the callee side.
+		//  cooker can't depend on callee being able to garbage collect
+		// collect garbage
+		/*if (NumProcessedSinceLastGC >= GCInterval)
+		{
+		UE_LOG(LogCookCommandlet, Display, TEXT("Full GC..."));
+
+		CollectGarbage( RF_Native );
+		NumProcessedSinceLastGC = 0;
 
-		STOP_TIMER(TickCookOnTheSide);
-		// kill the timer 
-		OUTPUT_TIMERS();
+		if (bLeakTest)
+		{
+		for (FObjectIterator It; It; ++It)
+		{
+		if (!LastGCItems.Contains(FWeakObjectPtr(*It)))
+		{
+		UE_LOG(LogCookCommandlet, Warning, TEXT("\tLeaked %s"), *(It->GetFullName()));
+		LastGCItems.Add(FWeakObjectPtr(*It));
+		}
+		}
+		}
+		}*/
 
 		if ( (FPlatformTime::Seconds() - StartTime) > TimeSlice)
 		{
 			break;
 		}
+	}
+
+	OUTPUT_TIMERS();
+
+	if ( CookByTheBookOptions )
+	{
+		CookByTheBookOptions->CookTime += FPlatformTime::Seconds() - StartTime;
+	}
 
+
+	if ( IsCookByTheBookRunning() && 
+		( CookRequests.HasItems() == false ) )
+	{
+		check(IsCookByTheBookMode());
+		// if we are out of stuff and we are in cook by the book from the editor mode then we finish up
+		CookByTheBookFinished();
 	}
+
 	return Result;
 }
 
@@ -638,15 +1238,21 @@ void UCookOnTheFlyServer::OnObjectUpdated( UObject *Object )
 
 void UCookOnTheFlyServer::MarkPackageDirtyForCooker( UPackage *Package )
 {
-	// force that package to be recooked
-	const FString Name = Package->GetPathName();
-	/*FString PackageFilename(GetPackageFilename(Package));
-	FPaths::MakeStandardFilename(PackageFilename);*/
-	FString PackageFilename = GetCachedStandardPackageFilename(Package);
-	UE_LOG(LogCookOnTheFly, Display, TEXT("Modification detected to package %s"), *PackageFilename);
-	ThreadSafeFilenameSet.RemoveAll( FName(*PackageFilename) );
-
-	// TODO: go through package dependencies of this package, if a found dependency is not loaded then mark it as need recooking.
+	if ( !bIsSavingPackage )
+	{
+		// could have just cooked a file which we might need to write
+		UPackage::WaitForAsyncFileWrites();
+
+		// force that package to be recooked
+		const FString Name = Package->GetPathName();
+		/*FString PackageFilename(GetPackageFilename(Package));
+		FPaths::MakeStandardFilename(PackageFilename);*/
+		FString PackageFilename = GetCachedStandardPackageFilename(Package);
+		UE_LOG(LogCookOnTheFly, Display, TEXT("Modification detected to package %s"), *PackageFilename);
+		CookedPackages.RemoveAll( FName(*PackageFilename) );
+
+		// TODO: go through package dependencies of this package, if a found dependency is not loaded then mark it as need recooking.
+	}
 }
 
 
@@ -682,6 +1288,7 @@ void UCookOnTheFlyServer::TickRecompileShaderRequests()
 	{
 		HandleNetworkFileServerRecompileShaders(Request->RecompileData);
 
+
 		// all done! other thread can unblock now
 		Request->bComplete = true;
 	}
@@ -729,7 +1336,7 @@ bool UCookOnTheFlyServer::ShouldCook(const FString& InFileName, const FName &InP
 	FString PkgFilename;
 	FDateTime DependentTimeStamp = FDateTime::MinValue();
 
-	if (bIterativeCooking && FPackageName::DoesPackageExist(InFileName, NULL, &PkgFile))
+	if (IsCookFlagSet(ECookInitializationFlags::Iterative) && FPackageName::DoesPackageExist(InFileName, NULL, &PkgFile))
 	{
 		PkgFilename = PkgFile;
 
@@ -744,7 +1351,7 @@ bool UCookOnTheFlyServer::ShouldCook(const FString& InFileName, const FName &InP
 
 	ITargetPlatformManagerModule& TPM = GetTargetPlatformManagerRef();
 
-	static const TArray<ITargetPlatform*> &ActiveTargetPlatforms = TPM.GetActiveTargetPlatforms();
+	static const TArray<ITargetPlatform*> &ActiveTargetPlatforms = TPM.GetCookingTargetPlatforms();
 
 	TArray<ITargetPlatform*> Platforms;
 
@@ -763,7 +1370,7 @@ bool UCookOnTheFlyServer::ShouldCook(const FString& InFileName, const FName &InP
 		FString PlatFilename = PkgFilename.Replace(TEXT("[Platform]"), *Target->PlatformName());
 
 		// If we are not iterative cooking, then cook the package
-		bool bCookPackage = (bIterativeCooking == false);
+		bool bCookPackage = (IsCookFlagSet(ECookInitializationFlags::Iterative) == false);
 
 		if (bCookPackage == false)
 		{
@@ -780,10 +1387,9 @@ bool UCookOnTheFlyServer::ShouldCook(const FString& InFileName, const FName &InP
 
 bool UCookOnTheFlyServer::SaveCookedPackage( UPackage* Package, uint32 SaveFlags, bool& bOutWasUpToDate, TArray<FName> &TargetPlatformNames )
 {
-
-
 	bool bSavedCorrectly = true;
-
+	check( bIsSavingPackage == false );
+	bIsSavingPackage = true;
 	FString Filename(GetCachedPackageFilename(Package));
 
 	if (Filename.Len())
@@ -796,7 +1402,7 @@ bool UCookOnTheFlyServer::SaveCookedPackage( UPackage* Package, uint32 SaveFlags
 		FString PkgFile;
 		FString Name = Package->GetPathName();
 
-		if (bIterativeCooking && FPackageName::DoesPackageExist(Name, NULL, &PkgFile))
+		if (IsCookFlagSet(ECookInitializationFlags::Iterative) && FPackageName::DoesPackageExist(Name, NULL, &PkgFile))
 		{
 			PkgFilename = PkgFile;
 
@@ -814,14 +1420,14 @@ bool UCookOnTheFlyServer::SaveCookedPackage( UPackage* Package, uint32 SaveFlags
  		EObjectFlags Flags = RF_NoFlags;
 		bool bPackageFullyLoaded = false;
 
-		if (bCompressed)
+		if (IsCookFlagSet(ECookInitializationFlags::Compressed) )
 		{
 			Package->PackageFlags |= PKG_StoreCompressed;
 		}
 
 		ITargetPlatformManagerModule& TPM = GetTargetPlatformManagerRef();
 
-		static TArray<ITargetPlatform*> ActiveStartupPlatforms = TPM.GetActiveTargetPlatforms();
+		static TArray<ITargetPlatform*> ActiveStartupPlatforms = TPM.GetCookingTargetPlatforms();
 
 		TArray<ITargetPlatform*> Platforms;
 
@@ -858,7 +1464,7 @@ bool UCookOnTheFlyServer::SaveCookedPackage( UPackage* Package, uint32 SaveFlags
 			FString PlatFilename = Filename.Replace(TEXT("[Platform]"), *Target->PlatformName());
 
 			// If we are not iterative cooking, then cook the package
-			bool bCookPackage = (bIterativeCooking == false);
+			bool bCookPackage = (IsCookFlagSet(ECookInitializationFlags::Iterative) == false);
 
 			if (bCookPackage == false)
 			{
@@ -869,7 +1475,7 @@ bool UCookOnTheFlyServer::SaveCookedPackage( UPackage* Package, uint32 SaveFlags
 			}
 
 			// don't save Editor resources from the Engine if the target doesn't have editoronly data
-			if (bSkipEditorContent && Name.StartsWith(TEXT("/Engine/Editor")) && !Target->HasEditorOnlyData())
+			if (IsCookFlagSet(ECookInitializationFlags::SkipEditorContent) && Name.StartsWith(TEXT("/Engine/Editor")) && !Target->HasEditorOnlyData())
 			{
 				bCookPackage = false;
 			}
@@ -878,6 +1484,8 @@ bool UCookOnTheFlyServer::SaveCookedPackage( UPackage* Package, uint32 SaveFlags
 			{
 				if (bPackageFullyLoaded == false)
 				{
+					SCOPE_TIMER(LoadPackage);
+
 					Package->FullyLoad();
 					if (!Package->IsFullyLoaded())
 					{
@@ -922,6 +1530,7 @@ bool UCookOnTheFlyServer::SaveCookedPackage( UPackage* Package, uint32 SaveFlags
 				}
 				else
 				{
+					SCOPE_TIMER(GEditorSavePackage);
 					bSavedCorrectly &= GEditor->SavePackage( Package, World, Flags, *PlatFilename, GError, NULL, bSwap, false, SaveFlags, Target, FDateTime::MinValue() );
 				}
 
@@ -939,33 +1548,28 @@ bool UCookOnTheFlyServer::SaveCookedPackage( UPackage* Package, uint32 SaveFlags
 		Package->PackageFlags = OriginalPackageFlags;
 	}
 
+	check( bIsSavingPackage == true );
+	bIsSavingPackage = false;
+
 	// return success
 	return bSavedCorrectly;
 }
 
-
-void UCookOnTheFlyServer::MaybeMarkPackageAsAlreadyLoaded(UPackage *Package)
+void UCookOnTheFlyServer::Initialize( ECookMode::Type DesiredCookMode, ECookInitializationFlags InCookFlags, const FString &OutputDirectoryOverride )
 {
-	FString Name = Package->GetName();
-	if (PackagesToNotReload.Contains(Name))
+
+	CurrentCookMode = DesiredCookMode;
+	CookFlags = InCookFlags;
+	
+	if ( IsCookByTheBookMode() )
 	{
-		UE_LOG(LogCookOnTheFly, Verbose, TEXT("Marking %s already loaded."), *Name);
-		Package->PackageFlags |= PKG_ReloadingForCooker;
+		CookByTheBookOptions = new FCookByTheBookOptions();
+		CookByTheBookOptions->bGenerateStreamingInstallManifests = IsCookFlagSet(ECookInitializationFlags::GenerateStreamingInstallManifest);
 	}
-}
 
 
-void UCookOnTheFlyServer::Initialize( bool inCompressed, bool inIterativeCooking, bool inSkipEditorContent, bool inIsAutoTick, bool inAsyncSaveLoad, const FString &OutputDirectoryOverride )
-{
-	bCompressed = inCompressed;
-	bIterativeCooking = inIterativeCooking;
-	bSkipEditorContent = inSkipEditorContent;
-	bAsyncSaveLoad = inAsyncSaveLoad;
-	bIsAutoTick = inIsAutoTick;
-
-	
 	ITargetPlatformManagerModule& TPM = GetTargetPlatformManagerRef();
-	const TArray<ITargetPlatform*>& Platforms = TPM.GetActiveTargetPlatforms();
+	const TArray<ITargetPlatform*>& Platforms = TPM.GetCookingTargetPlatforms();
 
 	// Local sandbox file wrapper. This will be used to handle path conversions,
 	// but will not be used to actually write/read files so we can safely
@@ -981,10 +1585,6 @@ void UCookOnTheFlyServer::Initialize( bool inCompressed, bool inIterativeCooking
 
 	CleanSandbox(Platforms);
 
-	// allow the game to fill out the asset registry, as well as get a list of objects to always cook
-	TArray<FString> FilesInPath;
-	FGameDelegates::Get().GetCookModificationDelegate().ExecuteIfBound(FilesInPath);
-
 	// always generate the asset registry before starting to cook, for either method
 	GenerateAssetRegistry(Platforms);
 
@@ -1030,7 +1630,7 @@ void UCookOnTheFlyServer::CleanSandbox(const TArray<ITargetPlatform*>& Platforms
 	{
 		SCOPE_SECONDS_COUNTER(SandboxCleanTime);
 
-		if (bIterativeCooking == false)
+		if (IsCookFlagSet(ECookInitializationFlags::Iterative) == false)
 		{
 			// for now we are going to wipe the cooked directory
 			for (int32 Index = 0; Index < Platforms.Num(); Index++)
@@ -1114,7 +1714,7 @@ void UCookOnTheFlyServer::GenerateAssetRegistry(const TArray<ITargetPlatform*>&
 		}
 		
 		// When not cooking on the fly the registry will be saved after the cooker has finished
-		if (bCookOnTheFly)
+		if (CurrentCookMode == ECookMode::CookOnTheFly)
 		{
 			// write it out to a memory archive
 			FArrayWriter SerializedAssetRegistry;
@@ -1164,6 +1764,374 @@ void UCookOnTheFlyServer::GenerateLongPackageNames(TArray<FString>& FilesInPath)
 	Exchange(FilesInPathReverse, FilesInPath);
 }
 
+void UCookOnTheFlyServer::CollectFilesToCook(TArray<FString>& FilesInPath, const TArray<FString> &CookMaps, const TArray<FString> &InCookDirectories, const TArray<FString> &CookCultures, const TArray<FString> &IniMapSections, bool bCookAll, bool bMapsOnly, bool bNoDev)
+{
+	TArray<FString> MapList;
+	TArray<FString> CookDirectories = InCookDirectories;
+	// Add the default map section
+	GEditor->LoadMapListFromIni(TEXT("AlwaysCookMaps"), MapList);
+
+	for ( const auto &IniMapSection : IniMapSections )
+	{
+		GEditor->LoadMapListFromIni(*IniMapSection, MapList);
+	}
+
+	// Add any map sections specified on command line
+	/*GEditor->ParseMapSectionIni(*Params, MapList);
+	for (int32 MapIdx = 0; MapIdx < MapList.Num(); MapIdx++)
+	{
+		FilesInPath.AddUnique(MapList[MapIdx]);
+	}*/
+
+
+	// Also append any cookdirs from the project ini files; these dirs are relative to the game content directory
+	{
+		const FString AbsoluteGameContentDir = FPaths::ConvertRelativePathToFull(FPaths::GameContentDir());
+		const UProjectPackagingSettings* const PackagingSettings = GetDefault<UProjectPackagingSettings>();
+		for(const auto& DirToCook : PackagingSettings->DirectoriesToAlwaysCook)
+		{
+			CookDirectories.Add(AbsoluteGameContentDir / DirToCook.Path);
+		}
+	}
+
+	for ( const auto &CurrEntry : CookMaps )
+	{
+		if (FPackageName::IsShortPackageName(CurrEntry))
+		{
+			FString OutFilename;
+			if (FPackageName::SearchForPackageOnDisk(CurrEntry, NULL, &OutFilename) == false)
+			{
+				UE_LOG(LogCookOnTheFly, Warning, TEXT("Unable to find package for map %s."), *CurrEntry);
+			}
+			else
+			{
+				FilesInPath.AddUnique(OutFilename);
+			}
+		}
+		else
+		{
+			FilesInPath.AddUnique(CurrEntry);
+		}
+	}
+
+	const FString ExternalMountPointName(TEXT("/Game/"));
+	for ( const auto &CurrEntry : CookDirectories )
+	{
+		TArray<FString> Files;
+		IFileManager::Get().FindFilesRecursive(Files, *CurrEntry, *(FString(TEXT("*")) + FPackageName::GetAssetPackageExtension()), true, false);
+		for (int32 Index = 0; Index < Files.Num(); Index++)
+		{
+			FString StdFile = Files[Index];
+			FPaths::MakeStandardFilename(StdFile);
+			FilesInPath.AddUnique(StdFile);
+
+			// this asset may not be in our currently mounted content directories, so try to mount a new one now
+			FString LongPackageName;
+			if(!FPackageName::IsValidLongPackageName(StdFile) && !FPackageName::TryConvertFilenameToLongPackageName(StdFile, LongPackageName))
+			{
+				FPackageName::RegisterMountPoint(ExternalMountPointName, CurrEntry);
+			}
+		}
+	}
+
+	if ((FilesInPath.Num() == 0) || bCookAll)
+	{
+		TArray<FString> Tokens;
+		Tokens.Empty(2);
+		Tokens.Add(FString("*") + FPackageName::GetAssetPackageExtension());
+		Tokens.Add(FString("*") + FPackageName::GetMapPackageExtension());
+
+		uint8 PackageFilter = NORMALIZE_DefaultFlags | NORMALIZE_ExcludeEnginePackages;
+		if ( bMapsOnly )
+		{
+			PackageFilter |= NORMALIZE_ExcludeContentPackages;
+		}
+
+		if ( bNoDev )
+		{
+			PackageFilter |= NORMALIZE_ExcludeDeveloperPackages;
+		}
+
+		// assume the first token is the map wildcard/pathname
+		TArray<FString> Unused;
+		for ( int32 TokenIndex = 0; TokenIndex < Tokens.Num(); TokenIndex++ )
+		{
+			TArray<FString> TokenFiles;
+			if ( !NormalizePackageNames( Unused, TokenFiles, Tokens[TokenIndex], PackageFilter) )
+			{
+				UE_LOG(LogCookOnTheFly, Display, TEXT("No packages found for parameter %i: '%s'"), TokenIndex, *Tokens[TokenIndex]);
+				continue;
+			}
+
+			for (int32 TokenFileIndex = 0; TokenFileIndex < TokenFiles.Num(); ++TokenFileIndex)
+			{
+				FilesInPath.AddUnique(TokenFiles[TokenFileIndex]);
+			}
+		}
+	}
+
+	// make sure we cook the default maps
+	ITargetPlatformManagerModule& TPM = GetTargetPlatformManagerRef();
+	static const TArray<ITargetPlatform*>& Platforms =  TPM.GetTargetPlatforms();
+	for (int32 Index = 0; Index < Platforms.Num(); Index++)
+	{
+		// load the platform specific ini to get its DefaultMap
+		FConfigFile PlatformEngineIni;
+		FConfigCacheIni::LoadLocalIniFile(PlatformEngineIni, TEXT("Engine"), true, *Platforms[Index]->IniPlatformName());
+
+		// get the server and game default maps and cook them
+		FString Obj;
+		if (PlatformEngineIni.GetString(TEXT("/Script/EngineSettings.GameMapsSettings"), TEXT("GameDefaultMap"), Obj))
+		{
+			if (Obj != FName(NAME_None).ToString())
+			{
+				FilesInPath.AddUnique(Obj);
+			}
+		}
+		if ( IsCookFlagSet(ECookInitializationFlags::IncludeServerMaps) )
+		{
+			if (PlatformEngineIni.GetString(TEXT("/Script/EngineSettings.GameMapsSettings"), TEXT("ServerDefaultMap"), Obj))
+			{
+				if (Obj != FName(NAME_None).ToString())
+				{
+					FilesInPath.AddUnique(Obj);
+				}
+			}
+		}
+		if (PlatformEngineIni.GetString(TEXT("/Script/EngineSettings.GameMapsSettings"), TEXT("GlobalDefaultGameMode"), Obj))
+		{
+			if (Obj != FName(NAME_None).ToString())
+			{
+				FilesInPath.AddUnique(Obj);
+			}
+		}
+		if (PlatformEngineIni.GetString(TEXT("/Script/EngineSettings.GameMapsSettings"), TEXT("GlobalDefaultServerGameMode"), Obj))
+		{
+			if (Obj != FName(NAME_None).ToString())
+			{
+				FilesInPath.AddUnique(Obj);
+			}
+		}
+	}
+
+	// make sure we cook any extra assets for the default touch interface
+	// @todo need a better approach to cooking assets which are dynamically loaded by engine code based on settings
+	FConfigFile InputIni;
+	FString InterfaceFile;
+	FConfigCacheIni::LoadLocalIniFile(InputIni, TEXT("Input"), true);
+	if (InputIni.GetString(TEXT("/Script/Engine.InputSettings"), TEXT("DefaultTouchInterface"), InterfaceFile))
+	{
+		if (InterfaceFile != TEXT("None") && InterfaceFile != TEXT(""))
+		{
+			FilesInPath.AddUnique(InterfaceFile);
+		}
+	}
+
+	//@todo SLATE: This is a hack to ensure all slate referenced assets get cooked.
+	// Slate needs to be refactored to properly identify required assets at cook time.
+	// Simply jamming everything in a given directory into the cook list is error-prone
+	// on many levels - assets not required getting cooked/shipped; assets not put under 
+	// the correct folder; etc.
+	{
+		TArray<FString> UIContentPaths;
+		if (GConfig->GetArray(TEXT("UI"), TEXT("ContentDirectories"), UIContentPaths, GEditorIni) > 0)
+		{
+			for (int32 DirIdx = 0; DirIdx < UIContentPaths.Num(); DirIdx++)
+			{
+				FString ContentPath = FPackageName::LongPackageNameToFilename(UIContentPaths[DirIdx]);
+
+				TArray<FString> Files;
+				IFileManager::Get().FindFilesRecursive(Files, *ContentPath, *(FString(TEXT("*")) + FPackageName::GetAssetPackageExtension()), true, false);
+				for (int32 Index = 0; Index < Files.Num(); Index++)
+				{
+					FString StdFile = Files[Index];
+					FPaths::MakeStandardFilename(StdFile);
+					FilesInPath.AddUnique(StdFile);
+				}
+			}
+		}
+	}
+}
+
+
+bool UCookOnTheFlyServer::IsCookByTheBookRunning() const
+{
+	return CookByTheBookOptions && CookByTheBookOptions->bRunning;
+}
+
+
+void UCookOnTheFlyServer::SaveGlobalShaderMapFiles(const TArray<ITargetPlatform*>& Platforms)
+{
+	for (int32 Index = 0; Index < Platforms.Num(); Index++)
+	{
+		// make sure global shaders are up to date!
+		TArray<FString> Files;
+		FShaderRecompileData RecompileData;
+		RecompileData.PlatformName = Platforms[Index]->PlatformName();
+		// Compile for all platforms
+		RecompileData.ShaderPlatform = -1;
+		RecompileData.ModifiedFiles = &Files;
+		RecompileData.MeshMaterialMaps = NULL;
+
+		check( IsInGameThread() );
+
+		FString OutputDir = GetOutputDirectory(RecompileData.PlatformName);
+
+		RecompileShadersForRemote
+			(RecompileData.PlatformName, 
+			RecompileData.ShaderPlatform == -1 ? SP_NumPlatforms : (EShaderPlatform)RecompileData.ShaderPlatform,
+			OutputDir, 
+			RecompileData.MaterialsToLoad, 
+			RecompileData.SerializedShaderResources, 
+			RecompileData.MeshMaterialMaps, 
+			RecompileData.ModifiedFiles);
+	}
+}
+
+
+void UCookOnTheFlyServer::CookByTheBookFinished()
+{
+	check( IsCookByTheBookMode() );
+	check( CookByTheBookOptions->bRunning == true );
+
+	UPackage::WaitForAsyncFileWrites();
+
+	GetDerivedDataCacheRef().WaitForQuiescence(true);
+
+	{
+		// Save modified asset registry with all streaming chunk info generated during cook
+		FString RegistryFilename = FPaths::GameDir() / TEXT("AssetRegistry.bin");
+		FString SandboxRegistryFilename = SandboxFile->ConvertToAbsolutePathForExternalAppForWrite(*RegistryFilename);
+		// the registry filename will be modified when we call save asset registry
+
+		for ( auto& Manifest : CookByTheBookOptions->ManifestGenerators )
+		{
+			// Always try to save the manifests, this is required to make the asset registry work, but doesn't necessarily write a file
+			Manifest.Value->SaveManifests(SandboxFile.GetOwnedPointer());
+			Manifest.Value->SaveAssetRegistry(SandboxRegistryFilename);
+		}
+	}
+
+	CookByTheBookOptions->LastGCItems.Empty();
+
+	UE_LOG(LogCookOnTheFly, Display, TEXT("Cook by the book total time %fs"), CookByTheBookOptions->CookTime);
+
+
+
+	CookByTheBookOptions->bRunning = false;
+}
+
+void UCookOnTheFlyServer::StartCookByTheBook(const TArray<ITargetPlatform*>& TargetPlatforms, const TArray<FString> &CookMaps, const TArray<FString> &CookDirectories, const TArray<FString> &CookCultures, const TArray<FString> &IniMapSections )
+{
+	// can transition from cook on the fly to cook by the book from the editor
+	check( IsCookByTheBookMode() );
+
+	/// reference ini parsing code
+	/*FString SectionStr;
+	if (FParse::Value(InCmdParams, TEXT("MAPINISECTION="), SectionStr))
+	{
+		if (SectionStr.Contains(TEXT("+")))
+		{
+			TArray<FString> Sections;
+			SectionStr.ParseIntoArray(&Sections,TEXT("+"),true);
+			for (int32 Index = 0; Index < Sections.Num(); Index++)
+			{
+				LoadMapListFromIni(Sections[Index], OutMapList);
+			}
+		}
+		else
+		{
+			LoadMapListFromIni(SectionStr, OutMapList);
+		}
+	}*/
+	
+
+	/*
+	// copied from original cook commandlet
+	// need to fill these in
+	bCookAll = Switches.Contains(TEXT("COOKALL"));   // Cook everything
+	bLeakTest = Switches.Contains(TEXT("LEAKTEST"));   // Test for UObject leaks
+
+	*/
+
+	CookByTheBookOptions->bRunning = true;
+	CookByTheBookOptions->CookTime = 0.0f;
+
+	bool bCookAll = false;
+	bool bMapsOnly = false;
+	bool bNoDev = false;
+
+	bool bLeakTest = IsCookFlagSet(ECookInitializationFlags::LeakTest); // this won't work from the editor this needs to be standalone
+	check( !bLeakTest || CurrentCookMode == ECookMode::CookByTheBook );
+
+	CookByTheBookOptions->LastGCItems.Empty();
+	if (bLeakTest)
+	{
+		for (FObjectIterator It; It; ++It)
+		{
+			CookByTheBookOptions->LastGCItems.Add(FWeakObjectPtr(*It));
+		}
+	}
+
+
+	
+	// allow the game to fill out the asset registry, as well as get a list of objects to always cook
+	TArray<FString> FilesInPath;
+	FGameDelegates::Get().GetCookModificationDelegate().ExecuteIfBound(FilesInPath);
+
+	SaveGlobalShaderMapFiles(TargetPlatforms);
+
+	CollectFilesToCook(FilesInPath, CookMaps, CookDirectories, CookCultures, IniMapSections, bCookAll, bMapsOnly, bNoDev );
+	if (FilesInPath.Num() == 0)
+	{
+		UE_LOG(LogCookOnTheFly, Warning, TEXT("No files found."));
+	}
+
+
+	TArray<FName> TargetPlatformNames;
+	for( const auto &Platform : TargetPlatforms )
+	{
+		FName PlatformName = FName(*Platform->PlatformName());
+		TargetPlatformNames.Add( PlatformName ); // build list of all target platform names
+
+		// make sure we have a manifest for all the platforms 
+		// we want a seperate manifest for each platform because they can all be in different states of cooked content
+		FChunkManifestGenerator* ManifestGenerator = CookByTheBookOptions->ManifestGenerators.FindRef( PlatformName );
+		if ( ManifestGenerator == NULL )
+		{
+			TArray<ITargetPlatform*> Platforms;
+			Platforms.Add( Platform );
+			ManifestGenerator = new FChunkManifestGenerator(Platforms);
+			ManifestGenerator->CleanManifestDirectories();
+			ManifestGenerator->Initialize( CookByTheBookOptions->bGenerateStreamingInstallManifests);
+
+			CookByTheBookOptions->ManifestGenerators.Add(PlatformName, ManifestGenerator);
+		}
+	}
+
+	GenerateLongPackageNames(FilesInPath);
+
+	// add all the files for the requested platform to the cook list
+	for ( auto FileName : FilesInPath )
+	{
+		FName FileFName = GetCachedStandardPackageFileFName(FileName);
+		
+		if (FileFName != NAME_None)
+		{
+			CookRequests.EnqueueUnique( MoveTemp( FFilePlatformRequest( FileFName, TargetPlatformNames ) ) );
+		}
+		else
+		{
+			UE_LOG(LogCookOnTheFly, Warning, TEXT("Unable to find package for cooking %s"), *FileName)
+		}
+		
+	}
+
+}
+
+
+
+
 
 /* UCookOnTheFlyServer callbacks
  *****************************************************************************/
@@ -1171,6 +2139,8 @@ void UCookOnTheFlyServer::GenerateLongPackageNames(TArray<FString>& FilesInPath)
 
 void UCookOnTheFlyServer::HandleNetworkFileServerFileRequest( const FString& Filename, const FString &Platformname, TArray<FString>& UnsolicitedFiles )
 {
+	check( CurrentCookMode == ECookMode::CookOnTheFly );	
+	
 	bool bIsCookable = FPackageName::IsPackageExtension(*FPaths::GetExtension(Filename, true));
 
 
@@ -1178,55 +2148,45 @@ void UCookOnTheFlyServer::HandleNetworkFileServerFileRequest( const FString& Fil
 
 	if (!bIsCookable)
 	{
-		TArray<FFilePlatformRequest> UnsolicitedRequests;
-		UnsolicitedCookedPackages.DequeueAll(UnsolicitedRequests);
-		for (const FFilePlatformRequest &Request : UnsolicitedRequests)
-		{
-			if ( Request.GetPlatformname() == PlatformFname )
-			{
-				FString StandardFilename = Request.GetFilename().ToString();
-				FPaths::MakeStandardFilename( StandardFilename );
-				UnsolicitedFiles.Add( StandardFilename );
-			}
-			else
-			{
-				UnsolicitedCookedPackages.Enqueue( Request );
-			}
+		TArray<FName> UnsolicitedFilenames;
+		UnsolicitedCookedPackages.GetPackagesForPlatformAndRemove(PlatformFname, UnsolicitedFilenames);
+
+		for ( const auto &UnsolicitedFile : UnsolicitedFilenames )
+		{	
+			FString StandardFilename = UnsolicitedFile.ToString();
+			FPaths::MakeStandardFilename( StandardFilename );
+			UnsolicitedFiles.Add( StandardFilename );
 		}
 		UPackage::WaitForAsyncFileWrites();
 		return;
 	}
+
 	FString StandardFileName = Filename;
 	FPaths::MakeStandardFilename( StandardFileName );
-	FFilePlatformRequest FileRequest( StandardFileName, PlatformFname );
-	FileRequests.Enqueue(FileRequest);
+
+	FName StandardFileFname = FName( *StandardFileName );
+	TArray<FName> Platforms;
+	Platforms.Add( PlatformFname );
+	FFilePlatformRequest FileRequest( StandardFileFname, Platforms);
+	CookRequests.EnqueueUnique(FileRequest, true);
 
 	do
 	{
 		FPlatformProcess::Sleep(0.0f);
 	}
-	while (!ThreadSafeFilenameSet.Exists(FileRequest));
+	while (!CookedPackages.Exists(FileRequest));
 
 	UE_LOG( LogCookOnTheFly, Display, TEXT("Cook complete %s"), *FileRequest.GetFilename().ToString())
 
-	TArray<FFilePlatformRequest> UnsolicitedRequests;
-	UnsolicitedCookedPackages.DequeueAll(UnsolicitedRequests);
-	for (const FFilePlatformRequest &Request : UnsolicitedRequests)
-	{
-		if ( Request.GetFilename() == FileRequest.GetFilename() )
-		{
-			// don't do anything we don't want this guy
-		}
-		else if ( Request.GetPlatformname() == PlatformFname )
-		{
-			const FString RequestFilename = Request.GetFilename().ToString();
-			UE_LOG( LogCookOnTheFly, Display, TEXT("Returning unsolicited file %s"), *RequestFilename );
-			UnsolicitedFiles.Add( RequestFilename );
-		}
-		else
-		{
-			UnsolicitedCookedPackages.Enqueue( Request );
-		}
+	TArray<FName> UnsolicitedFilenames;
+	UnsolicitedCookedPackages.GetPackagesForPlatformAndRemove(PlatformFname, UnsolicitedFilenames);
+	UnsolicitedFilenames.Remove(FileRequest.GetFilename());
+
+	for ( const auto &UnsolicitedFile : UnsolicitedFilenames )
+	{	
+		FString StandardFilename = UnsolicitedFile.ToString();
+		FPaths::MakeStandardFilename( StandardFilename );
+		UnsolicitedFiles.Add( StandardFilename );
 	}
 
 	UPackage::WaitForAsyncFileWrites();
@@ -1241,6 +2201,8 @@ void UCookOnTheFlyServer::HandleNetworkFileServerFileRequest( const FString& Fil
 
 void UCookOnTheFlyServer::HandleNetworkFileServerRecompileShaders(const FShaderRecompileData& RecompileData)
 {
+	// shouldn't receive network requests unless we are in cook on the fly mode
+	check( CurrentCookMode == ECookMode::CookOnTheFly );
 	// if we aren't in the game thread, we need to push this over to the game thread and wait for it to finish
 	if (!IsInGameThread())
 	{
diff --git a/Engine/Source/Editor/UnrealEd/Private/PlayLevel.cpp b/Engine/Source/Editor/UnrealEd/Private/PlayLevel.cpp
--- a/Engine/Source/Editor/UnrealEd/Private/PlayLevel.cpp
+++ b/Engine/Source/Editor/UnrealEd/Private/PlayLevel.cpp
@@ -1256,11 +1256,17 @@ void UEditorEngine::PlayUsingLauncher()
 		// Setup launch profile, keep the setting here to a minimum.
 		ILauncherProfileRef LauncherProfile = LauncherServicesModule.CreateProfile(TEXT("Play On Device"));
 		LauncherProfile->SetBuildGame(bHasCode && FSourceCodeNavigation::IsCompilerAvailable());
-		LauncherProfile->SetCookMode(ELauncherProfileCookModes::ByTheBook);
+		LauncherProfile->SetCookMode(CanCookByTheBookInEditor() ? ELauncherProfileCookModes::ByTheBookInEditor : ELauncherProfileCookModes::ByTheBook);
 		LauncherProfile->AddCookedPlatform(PlayUsingLauncherDeviceId.Left(PlayUsingLauncherDeviceId.Find(TEXT("@"))));
 		LauncherProfile->SetDeployedDeviceGroup(DeviceGroup);
 		LauncherProfile->SetEditorExe(FUnrealEdMisc::Get().GetExecutableForCommandlets());
 
+
+		// hack to make launch on use pak files and generate manifests
+		LauncherProfile->SetDeployWithUnrealPak( true );
+
+
+
 		const FString DummyDeviceName(FString::Printf(TEXT("All_iOS_On_%s"), FPlatformProcess::ComputerName()));
 		if (PlayUsingLauncherDeviceId.Left(PlayUsingLauncherDeviceId.Find(TEXT("@"))) != TEXT("IOS") || !PlayUsingLauncherDeviceName.Contains(DummyDeviceName))
 		{
@@ -1269,9 +1275,10 @@ void UEditorEngine::PlayUsingLauncher()
 
 		TArray<FString> MapNames;
 		FWorldContext & EditorContext = GetEditorWorldContext();
-		if (EditorContext.World()->WorldComposition)
+		if (EditorContext.World()->WorldComposition || (LauncherProfile->GetCookMode() == ELauncherProfileCookModes::ByTheBookInEditor) )
 		{
 			// Open world composition from original folder
+			// Or if using by book in editor don't need to resave the package just cook it by the book 
 			FString MapName = EditorContext.World()->GetOutermost()->GetName();
 			MapNames.Add(MapName);
 		}
@@ -1293,6 +1300,35 @@ void UEditorEngine::PlayUsingLauncher()
 			LauncherProfile->AddCookedMap(MapName);
 		}
 
+
+		if ( LauncherProfile->GetCookMode() == ELauncherProfileCookModes::ByTheBookInEditor )
+		{
+			// check( GEditor != NULL );
+
+			TArray<ITargetPlatform*> TargetPlatforms;
+			for ( const auto &PlatformName : LauncherProfile->GetCookedPlatforms() )
+			{
+				ITargetPlatform* TargetPlatform = GetTargetPlatformManager()->FindTargetPlatform(PlatformName);
+				// todo pass in all the target platforms instead of just the single platform
+				// crashes if two requests are inflight but we can support having multiple platforms cooking at once
+				TargetPlatforms.Add( TargetPlatform ); 
+			}
+
+
+			const TArray<FString> &CookedMaps = LauncherProfile->GetCookedMaps();
+
+			// const TArray<FString>& CookedMaps = ChainState.Profile->GetCookedMaps();
+			TArray<FString> CookDirectories;
+			TArray<FString> CookCultures;
+			TArray<FString> IniMapSections;
+
+			StartCookByTheBookInEditor(TargetPlatforms, CookedMaps, CookDirectories, CookCultures, IniMapSections );
+
+			FIsCookFinishedDelegate &CookerFinishedDelegate = LauncherProfile->OnIsCookFinished();
+
+			CookerFinishedDelegate.BindUObject(this, &UEditorEngine::IsCookByTheBookInEditorFinished);   // AddUObject(this, &UEditorEngine::IsCookByTheBookInEditorFinished);
+		}
+
 		ILauncherPtr Launcher = LauncherServicesModule.CreateLauncher();
 		ILauncherWorkerPtr LauncherWorker = Launcher->Launch(TargetDeviceServicesModule.GetDeviceProxyManager(), LauncherProfile);
 
diff --git a/Engine/Source/Editor/UnrealEd/Private/UnrealEdEngine.cpp b/Engine/Source/Editor/UnrealEd/Private/UnrealEdEngine.cpp
--- a/Engine/Source/Editor/UnrealEd/Private/UnrealEdEngine.cpp
+++ b/Engine/Source/Editor/UnrealEd/Private/UnrealEdEngine.cpp
@@ -110,16 +110,42 @@ void UUnrealEdEngine::Init(IEngineLoop* InEngineLoop)
 		PropertyModule.RegisterCustomClassLayout("ProjectPackagingSettings", FOnGetDetailCustomizationInstance::CreateStatic(&FProjectPackagingSettingsCustomization::MakeInstance));
 	}
 
-	bool bCookOnTheSide = FParse::Param( FCommandLine::Get(),TEXT("COOKONTHESIDE"));
-	if ( bCookOnTheSide )
+	if ( FParse::Param( FCommandLine::Get(),TEXT("COOKONTHESIDE")) )
 	{
 		CookServer = ConstructObject<UCookOnTheFlyServer>( UCookOnTheFlyServer::StaticClass() );
-		CookServer->Initialize( false, false, false, true, true );
+		CookServer->Initialize( ECookMode::CookOnTheFly, ECookInitializationFlags::AutoTick | ECookInitializationFlags::AsyncSave );
 		CookServer->StartNetworkFileServer( false );
 
 		FCoreDelegates::OnObjectPropertyChanged.AddUObject(CookServer, &UCookOnTheFlyServer::OnObjectPropertyChanged);
 		FCoreDelegates::OnObjectModified.AddUObject(CookServer, &UCookOnTheFlyServer::OnObjectModified);
 	}
+	else if ( FParse::Param( FCommandLine::Get(), TEXT("COOKINTHEEDITOR")))
+	{
+		CookServer = ConstructObject<UCookOnTheFlyServer>( UCookOnTheFlyServer::StaticClass() );
+		CookServer->Initialize( ECookMode::CookByTheBookFromTheEditor, ECookInitializationFlags::AutoTick | ECookInitializationFlags::AsyncSave );
+
+		FCoreDelegates::OnObjectPropertyChanged.AddUObject(CookServer, &UCookOnTheFlyServer::OnObjectPropertyChanged);
+		FCoreDelegates::OnObjectModified.AddUObject(CookServer, &UCookOnTheFlyServer::OnObjectModified);
+	}
+}
+
+bool UUnrealEdEngine::CanCookByTheBookInEditor() const 
+{ 
+	if ( CookServer )
+	{
+		return CookServer->GetCookMode() == ECookMode::CookByTheBookFromTheEditor; 
+	}
+	return false;
+}
+
+void UUnrealEdEngine::StartCookByTheBookInEditor( const TArray<ITargetPlatform*> &TargetPlatforms, const TArray<FString> &CookMaps, const TArray<FString> &CookDirectories, const TArray<FString> &CookCultures, const TArray<FString> &IniMapSections )
+{
+	CookServer->StartCookByTheBook( TargetPlatforms, CookMaps, CookDirectories, CookCultures, IniMapSections );
+}
+
+bool UUnrealEdEngine::IsCookByTheBookInEditorFinished() const 
+{ 
+	return !CookServer->IsCookByTheBookRunning();
 }
 
 void UUnrealEdEngine::MakeSortedSpriteInfo(TArray<FSpriteCategoryInfo>& OutSortedSpriteInfo) const
diff --git a/Engine/Source/Runtime/CoreUObject/Private/UObject/SavePackage.cpp b/Engine/Source/Runtime/CoreUObject/Private/UObject/SavePackage.cpp
--- a/Engine/Source/Runtime/CoreUObject/Private/UObject/SavePackage.cpp
+++ b/Engine/Source/Runtime/CoreUObject/Private/UObject/SavePackage.cpp
@@ -2348,13 +2348,27 @@ bool UPackage::SavePackage( UPackage* InOuter, UObject* Base, EObjectFlags TopLe
 			ensureMsgf(false, TEXT("Recursive SavePackage() is not supported"));
 			return false;
 		}
+		UE_LOG_COOK_TIME(TEXT("First"));
 
+
+		
+		bool bIsCooking = TargetPlatform != NULL;
+
+		// if we are cooking we should be doing it in the editor
+		// otherwise some other assumptions are bad
+		check( !bIsCooking || WITH_EDITOR );
 	#if WITH_EDITOR
-		// Attempt to create a backup of this package before it is saved, if applicable
-		if (FCoreUObjectDelegates::AutoPackageBackupDelegate.IsBound())
+		if ( !bIsCooking )
 		{
-			FCoreUObjectDelegates::AutoPackageBackupDelegate.Execute(*InOuter);
+			// Attempt to create a backup of this package before it is saved, if applicable
+			if (FCoreUObjectDelegates::AutoPackageBackupDelegate.IsBound())
+			{
+				FCoreUObjectDelegates::AutoPackageBackupDelegate.Execute(*InOuter);
+			}
+
+			UE_LOG_COOK_TIME(TEXT("AutoPackageBackupDelegate"));
 		}
+		
 	#endif	// #if WITH_EDITOR
 
 		// do any path replacements on the source DestFile
@@ -2367,11 +2381,15 @@ bool UPackage::SavePackage( UPackage* InOuter, UObject* Base, EObjectFlags TopLe
 		// The latter implies flushing all file handles which is a pre-requisite of saving a package. The code basically needs 
 		// to be sure that we are not reading from a file that is about to be overwritten and that there is no way we might 
 		// start reading from the file till we are done overwriting it.
+
 		FlushAsyncLoading();
+
+		UE_LOG_COOK_TIME(TEXT("Flush Async loading"));
 		(*GFlushStreamingFunc)();
+		UE_LOG_COOK_TIME(TEXT("Flush streaming func"));
 		FIOSystem::Get().BlockTillAllRequestsFinishedAndFlushHandles();
 
-		UE_LOG_COOK_TIME(TEXT("Flush Async loading"));
+		UE_LOG_COOK_TIME(TEXT("Block till all requests finished and flush handles"));
 
 		check(InOuter);
 		check(Filename);
@@ -2510,13 +2528,15 @@ bool UPackage::SavePackage( UPackage* InOuter, UObject* Base, EObjectFlags TopLe
 			FArchiveSaveTagExports ExportTaggerArchive( InOuter );
 			ExportTaggerArchive.SetPortFlags( ComparisonFlags );
 			ExportTaggerArchive.SetCookingTarget(TargetPlatform);
-	
+			
+			check( ExportTaggerArchive.IsCooking() == !!TargetPlatform );
+			check( ExportTaggerArchive.IsCooking() == bIsCooking );
+
 			// Tag exports and route presave.
 			FPackageExportTagger PackageExportTagger( Base, TopLevelFlags, InOuter );
 			PackageExportTagger.TagPackageExports( ExportTaggerArchive, true );
 			ExportTaggerArchive.SetFilterEditorOnly(FilterEditorOnly);
 
-
 			UE_LOG_COOK_TIME(TEXT("TagPackageExports"));
 		
 			{
@@ -2544,7 +2564,9 @@ bool UPackage::SavePackage( UPackage* InOuter, UObject* Base, EObjectFlags TopLe
 				// only need to do this if we are cooking a different platform then the one which is currently running
 				// TODO: if save package is cancelled then call ClearCache on each object
 				TArray<UObject*> CachedObjects;
-				if ( !!TargetPlatform )
+
+#if WITH_EDITOR
+				if ( bIsCooking )
 				{
 					TArray<UObject*> TagExpObjects;
 					GetObjectsWithAnyMarks(TagExpObjects, OBJECTMARK_TagExp);
@@ -2558,7 +2580,7 @@ bool UPackage::SavePackage( UPackage* InOuter, UObject* Base, EObjectFlags TopLe
 						}
 					}
 				}
-
+#endif
 
 
 
diff --git a/Engine/Source/Runtime/Engine/Private/LevelActor.cpp b/Engine/Source/Runtime/Engine/Private/LevelActor.cpp
--- a/Engine/Source/Runtime/Engine/Private/LevelActor.cpp
+++ b/Engine/Source/Runtime/Engine/Private/LevelActor.cpp
@@ -811,9 +811,13 @@ void UWorld::LoadSecondaryLevels(bool bForce, TSet<FString>* CookedPackages)
 							}
 							check(LoadedWorld);
 
-							// LoadedWorld won't be serialized as there's a BeginLoad on the stack so we manually serialize it here.
-							check( LoadedWorld->GetLinker() );
-							LoadedWorld->GetLinker()->Preload( LoadedWorld );
+							if ( !LevelPackage->IsFullyLoaded() )
+							{
+								// LoadedWorld won't be serialized as there's a BeginLoad on the stack so we manually serialize it here.
+								check( LoadedWorld->GetLinker() );
+								LoadedWorld->GetLinker()->Preload( LoadedWorld );
+							}
+							
 
 							// Keep reference to prevent garbage collection.
 							check( LoadedWorld->PersistentLevel );
