343741	PERSONA: Animation Composite: Unable to add to world from content browser

diff --git a/Engine/Source/Editor/UnrealEd/Private/Factories/ActorFactory.cpp b/Engine/Source/Editor/UnrealEd/Private/Factories/ActorFactory.cpp
--- a/Engine/Source/Editor/UnrealEd/Private/Factories/ActorFactory.cpp
+++ b/Engine/Source/Editor/UnrealEd/Private/Factories/ActorFactory.cpp
@@ -645,7 +645,7 @@ UActorFactoryAnimationAsset::UActorFactoryAnimationAsset(const class FPostConstr
 bool UActorFactoryAnimationAsset::CanCreateActorFrom( const FAssetData& AssetData, FText& OutErrorMsg )
 { 
 	if ( !AssetData.IsValid() || 
-		( !AssetData.GetClass()->IsChildOf( UAnimSequence::StaticClass() ) && 
+		( !AssetData.GetClass()->IsChildOf( UAnimSequenceBase::StaticClass() ) && 
 		  !AssetData.GetClass()->IsChildOf( UVertexAnimation::StaticClass() )) )
 	{
 		OutErrorMsg = NSLOCTEXT("CanCreateActor", "NoAnimData", "A valid anim data must be specified.");
@@ -655,7 +655,7 @@ bool UActorFactoryAnimationAsset::CanCreateActorFrom( const FAssetData& AssetDat
 	FAssetRegistryModule& AssetRegistryModule = FModuleManager::LoadModuleChecked<FAssetRegistryModule>(TEXT("AssetRegistry"));
 	IAssetRegistry& AssetRegistry = AssetRegistryModule.Get();
 
-	if ( AssetData.GetClass()->IsChildOf( UAnimSequence::StaticClass() ) )
+	if ( AssetData.GetClass()->IsChildOf( UAnimSequenceBase::StaticClass() ) )
 	{
 		const FString* SkeletonPath = AssetData.TagsAndValues.Find("Skeleton");
 		if ( SkeletonPath == NULL || SkeletonPath->IsEmpty() ) 
@@ -728,7 +728,7 @@ bool UActorFactoryAnimationAsset::CanCreateActorFrom( const FAssetData& AssetDat
 USkeletalMesh* UActorFactoryAnimationAsset::GetSkeletalMeshFromAsset( UObject* Asset ) const
 {
 	USkeletalMesh* SkeletalMesh = NULL;
-	UAnimSequence* AnimationAsset = Cast<UAnimSequence>( Asset );
+	UAnimSequenceBase* AnimationAsset = Cast<UAnimSequenceBase>( Asset );
 	UVertexAnimation* VertexAnimation = Cast<UVertexAnimation>( Asset );
 
 	if( AnimationAsset != NULL )
diff --git a/Engine/Source/Editor/UnrealEd/Private/LevelEditorViewport.cpp b/Engine/Source/Editor/UnrealEd/Private/LevelEditorViewport.cpp
--- a/Engine/Source/Editor/UnrealEd/Private/LevelEditorViewport.cpp
+++ b/Engine/Source/Editor/UnrealEd/Private/LevelEditorViewport.cpp
@@ -608,7 +608,7 @@ static bool AttemptApplyObjToActor( UObject* ObjToUse, AActor* ActorToApplyTo, i
 		// block anything else than just anim sequence
 		if( DroppedObjAsAnimationAsset != NULL )
 		{
-			if( ! DroppedObjAsAnimationAsset->IsA(UAnimSequence::StaticClass()) )
+			if( ! DroppedObjAsAnimationAsset->IsA(UAnimSequenceBase::StaticClass()) )
 			{
 				DroppedObjAsAnimationAsset = NULL;
 			}
