clang fixes (extra qualifications).

diff --git a/Engine/Source/Editor/UnrealEd/Private/Commandlets/ChunkManifestGenerator.h b/Engine/Source/Editor/UnrealEd/Private/Commandlets/ChunkManifestGenerator.h
--- a/Engine/Source/Editor/UnrealEd/Private/Commandlets/ChunkManifestGenerator.h
+++ b/Engine/Source/Editor/UnrealEd/Private/Commandlets/ChunkManifestGenerator.h
@@ -240,7 +240,7 @@ public:
 	* @param LastLoadedMapName Name of the last loaded map (can be empty)
 	* @param the SandboxPlatformFile used during cook
 	*/
-	void FChunkManifestGenerator::AddPackageToChunkManifest(const FName& PackageFName, const FString& PackagePathName, const FString& SandboxFilename, const FString& LastLoadedMapName, FSandboxPlatformFile* SandboxFile);
+	void AddPackageToChunkManifest(const FName& PackageFName, const FString& PackagePathName, const FString& SandboxFilename, const FString& LastLoadedMapName, FSandboxPlatformFile* SandboxFile);
 
 
 
diff --git a/Engine/Source/Editor/UnrealEd/Private/CookOnTheFlyServer.cpp b/Engine/Source/Editor/UnrealEd/Private/CookOnTheFlyServer.cpp
--- a/Engine/Source/Editor/UnrealEd/Private/CookOnTheFlyServer.cpp
+++ b/Engine/Source/Editor/UnrealEd/Private/CookOnTheFlyServer.cpp
@@ -470,7 +470,7 @@ private:
 	 *
 	 * @param	Obj		a pointer to the object that was encountered
 	 */
-	FArchive& FArchiveFindReferences::operator<<( UObject*& Obj )
+	FArchive& operator<<( UObject*& Obj )
 	{
 		if( Obj )
 		{
