Fixed compilation error in clang.

diff --git a/Engine/Source/Editor/UnrealEd/Private/CookOnTheFlyServer.cpp b/Engine/Source/Editor/UnrealEd/Private/CookOnTheFlyServer.cpp
--- a/Engine/Source/Editor/UnrealEd/Private/CookOnTheFlyServer.cpp
+++ b/Engine/Source/Editor/UnrealEd/Private/CookOnTheFlyServer.cpp
@@ -538,7 +538,7 @@ public:
 	 * @param	inSource		object to show references for
 	 * @param	inExclude		list of objects that should be ignored if encountered while serializing SourceObject
 	 */
-	FArchiveFindReferences::FArchiveFindReferences( TSet<UObject*> InRootSet, TSet<UObject*> &inFound, TArray<UObject*> &inExclude )
+	FArchiveFindReferences( TSet<UObject*> InRootSet, TSet<UObject*> &inFound, TArray<UObject*> &inExclude )
 		: Exclude(inExclude)
 		, Found(inFound)
 		, RootSet(InRootSet)
