Convert uses of FActorIterator to TActorIterator where appropriate

diff --git a/Engine/Source/Developer/GameplayDebugger/Private/GameplayDebugger.cpp b/Engine/Source/Developer/GameplayDebugger/Private/GameplayDebugger.cpp
--- a/Engine/Source/Developer/GameplayDebugger/Private/GameplayDebugger.cpp
+++ b/Engine/Source/Developer/GameplayDebugger/Private/GameplayDebugger.cpp
@@ -341,13 +341,13 @@ bool FGameplayDebugger::Exec(UWorld* Inworld, const TCHAR* Cmd, FOutputDevice& A
 		if (MyPC && EQS)
 		{
 			AGameplayDebuggingReplicator* DebuggingReplicator = NULL;
-			for (FActorIterator It(Inworld); It; ++It)
+			for (TActorIterator<AGameplayDebuggingReplicator> It(Inworld); It; ++It)
 			{
-				AActor* A = *It;
-				if (A && A->IsA(AGameplayDebuggingReplicator::StaticClass()) && !A->IsPendingKill())
+				AGameplayDebuggingReplicator* A = *It;
+				if (!A->IsPendingKill())
 				{
-					DebuggingReplicator = Cast<AGameplayDebuggingReplicator>(A);
-					if (DebuggingReplicator && !DebuggingReplicator->IsGlobalInWorld() && DebuggingReplicator->GetLocalPlayerOwner() == MyPC)
+					DebuggingReplicator = A;
+					if (!DebuggingReplicator->IsGlobalInWorld() && DebuggingReplicator->GetLocalPlayerOwner() == MyPC)
 					{
 						break;
 					}
diff --git a/Engine/Source/Developer/GameplayDebugger/Private/GameplayDebuggingHUDComponent.cpp b/Engine/Source/Developer/GameplayDebugger/Private/GameplayDebuggingHUDComponent.cpp
--- a/Engine/Source/Developer/GameplayDebugger/Private/GameplayDebuggingHUDComponent.cpp
+++ b/Engine/Source/Developer/GameplayDebugger/Private/GameplayDebuggingHUDComponent.cpp
@@ -50,17 +50,13 @@ AGameplayDebuggingReplicator* AGameplayDebuggingHUDComponent::GetDebuggingReplic
 		return CachedDebuggingReplicator.Get();
 	}
 	
-	for (FActorIterator It(GetWorld()); It; ++It)
+	for (TActorIterator<AGameplayDebuggingReplicator> It(GetWorld()); It; ++It)
 	{
-		AActor* A = *It;
-		if (A && A->IsA(AGameplayDebuggingReplicator::StaticClass()) && !A->IsPendingKill())
+		AGameplayDebuggingReplicator* Replicator = *It;
+		if (!Replicator->IsPendingKill() && Replicator->GetLocalPlayerOwner() == PlayerOwner)
 		{
-			AGameplayDebuggingReplicator* Replicator = Cast<AGameplayDebuggingReplicator>(A);
-			if (Replicator->GetLocalPlayerOwner() == PlayerOwner)
-			{
-				CachedDebuggingReplicator = Replicator;
-				return Replicator;
-			}
+			CachedDebuggingReplicator = Replicator;
+			return Replicator;
 		}
 	}
 
diff --git a/Engine/Source/Developer/LogVisualizer/Private/LogVisualizer.cpp b/Engine/Source/Developer/LogVisualizer/Private/LogVisualizer.cpp
--- a/Engine/Source/Developer/LogVisualizer/Private/LogVisualizer.cpp
+++ b/Engine/Source/Developer/LogVisualizer/Private/LogVisualizer.cpp
@@ -88,14 +88,12 @@ class AActor* FLogVisualizer::GetHelperActor(class UWorld* InWorld)
 		return DebugActor.Get();
 	}
 
-	for (FActorIterator It(InWorld); It; ++It)
+	for (TActorIterator<ALogVisualizerDebugActor> It(InWorld); It; ++It)
 	{
-		ALogVisualizerDebugActor* LogVisualizerDebugActor = Cast<ALogVisualizerDebugActor>(*It);
-		if (LogVisualizerDebugActor != NULL)
-		{
-			DebugActor = LogVisualizerDebugActor;
-			return LogVisualizerDebugActor;
-		}
+		ALogVisualizerDebugActor* LogVisualizerDebugActor = *It;
+
+		DebugActor = LogVisualizerDebugActor;
+		return LogVisualizerDebugActor;
 	}
 
 	FActorSpawnParameters SpawnInfo;
diff --git a/Engine/Source/Developer/LogVisualizer/Private/SLogVisualizer.cpp b/Engine/Source/Developer/LogVisualizer/Private/SLogVisualizer.cpp
--- a/Engine/Source/Developer/LogVisualizer/Private/SLogVisualizer.cpp
+++ b/Engine/Source/Developer/LogVisualizer/Private/SLogVisualizer.cpp
@@ -586,7 +586,7 @@ void SLogVisualizer::OnListDoubleClick(TSharedPtr<FLogsListItem> LogListItem)
 		TSharedPtr<FActorsVisLog>& Log = LogVisualizer->Logs[LogListItem->LogIndex];
 		for (FActorIterator It(GetWorld()); It; ++It)
 		{
-			AActor* Actor = Cast<AActor>(*It);
+			AActor* Actor = *It;
 			if (Actor->GetFName() == Log->Name)
 			{
 				Actor->GetActorBounds(false, Orgin, Extent);
diff --git a/Engine/Source/Editor/BehaviorTreeEditor/Private/BehaviorTreeDebugger.cpp b/Engine/Source/Editor/BehaviorTreeEditor/Private/BehaviorTreeDebugger.cpp
--- a/Engine/Source/Editor/BehaviorTreeEditor/Private/BehaviorTreeDebugger.cpp
+++ b/Engine/Source/Editor/BehaviorTreeEditor/Private/BehaviorTreeDebugger.cpp
@@ -656,12 +656,12 @@ void FBehaviorTreeDebugger::FindLockedDebugActor(UWorld* World)
 	if (LocalPC && LocalPC->GetHUD() && LocalPC->GetPawnOrSpectator())
 	{
 		AGameplayDebuggingReplicator* DebuggingReplicator = NULL;
-		for (FActorIterator It(World); It; ++It)
+		for (TActorIterator<AGameplayDebuggingReplicator> It(World); It; ++It)
 		{
-			AActor* A = *It;
-			if (A && A->IsA(AGameplayDebuggingReplicator::StaticClass()) && !A->IsPendingKill())
+			AGameplayDebuggingReplicator* A = *It;
+			if (!A->IsPendingKill())
 			{
-				DebuggingReplicator = Cast<AGameplayDebuggingReplicator>(A);
+				DebuggingReplicator = A;
 				break;
 			}
 		}
diff --git a/Engine/Source/Editor/LandscapeEditor/Private/LandscapeEdMode.cpp b/Engine/Source/Editor/LandscapeEditor/Private/LandscapeEdMode.cpp
--- a/Engine/Source/Editor/LandscapeEditor/Private/LandscapeEdMode.cpp
+++ b/Engine/Source/Editor/LandscapeEditor/Private/LandscapeEdMode.cpp
@@ -300,14 +300,10 @@ void FEdModeLandscape::Enter()
 		GEditor->SelectNone(false, true);
 	}
 
-	for (FActorIterator It(GetWorld()); It; ++It)
+	for (TActorIterator<ALandscapeGizmoActiveActor> It(GetWorld()); It; ++It)
 	{
-		ALandscapeGizmoActiveActor* Gizmo = Cast<ALandscapeGizmoActiveActor>(*It);
-		if (Gizmo)
-		{
-			CurrentGizmoActor = Gizmo;
-			break;
-		}
+		CurrentGizmoActor = *It;
+		break;
 	}
 
 	if (!CurrentGizmoActor.IsValid())
@@ -501,25 +497,17 @@ void FEdModeLandscape::Exit()
 
 	// Clear all GizmoActors if there is no Landscape in World
 	bool bIsLandscapeExist = false;
-	for (FActorIterator It(GetWorld()); It; ++It)
+	for (TActorIterator<ALandscapeProxy> It(GetWorld()); It; ++It)
 	{
-		ALandscapeProxy* Proxy = Cast<ALandscapeProxy>(*It);
-		if (Proxy)
-		{
-			bIsLandscapeExist = true;
-			break;
-		}
+		bIsLandscapeExist = true;
+		break;
 	}
 
 	if (!bIsLandscapeExist)
 	{
-		for (FActorIterator It(GetWorld()); It; ++It)
+		for (TActorIterator<ALandscapeGizmoActor> It(GetWorld()); It; ++It)
 		{
-			ALandscapeGizmoActor* Gizmo = Cast<ALandscapeGizmoActor>(*It);
-			if (Gizmo)
-			{
-				GetWorld()->DestroyActor(Gizmo, false, false);
-			}
+			GetWorld()->DestroyActor(*It, false, false);
 		}
 	}
 
@@ -1630,14 +1618,10 @@ int32 FEdModeLandscape::UpdateLandscapeList()
 	if (!CurrentGizmoActor.IsValid())
 	{
 		ALandscapeGizmoActiveActor* GizmoActor = NULL;
-		for (FActorIterator It(GetWorld()); It; ++It)
+		for (TActorIterator<ALandscapeGizmoActiveActor> It(GetWorld()); It; ++It)
 		{
-			ALandscapeGizmoActiveActor* Gizmo = Cast<ALandscapeGizmoActiveActor>(*It);
-			if (Gizmo)
-			{
-				GizmoActor = Gizmo;
-				break;
-			}
+			GizmoActor = *It;
+			break;
 		}
 	}
 
diff --git a/Engine/Source/Editor/LandscapeEditor/Private/LandscapeEditorObject.cpp b/Engine/Source/Editor/LandscapeEditor/Private/LandscapeEditorObject.cpp
--- a/Engine/Source/Editor/LandscapeEditor/Private/LandscapeEditorObject.cpp
+++ b/Engine/Source/Editor/LandscapeEditor/Private/LandscapeEditorObject.cpp
@@ -255,10 +255,10 @@ void ULandscapeEditorObject::Load()
 
 	// Gizmo History (not saved!)
 	GizmoHistories.Empty();
-	for (FActorIterator It(ParentMode->GetWorld()); It; ++It)
+	for (TActorIterator<ALandscapeGizmoActor> It(ParentMode->GetWorld()); It; ++It)
 	{
-		ALandscapeGizmoActor* Gizmo = Cast<ALandscapeGizmoActor>(*It);
-		if (Gizmo && !Gizmo->IsEditable())
+		ALandscapeGizmoActor* Gizmo = *It;
+		if (!Gizmo->IsEditable())
 		{
 			new(GizmoHistories) FGizmoHistory(Gizmo);
 		}
diff --git a/Engine/Source/Editor/LevelEditor/Private/LevelEditorActions.cpp b/Engine/Source/Editor/LevelEditor/Private/LevelEditorActions.cpp
--- a/Engine/Source/Editor/LevelEditor/Private/LevelEditorActions.cpp
+++ b/Engine/Source/Editor/LevelEditor/Private/LevelEditorActions.cpp
@@ -1382,13 +1382,9 @@ void FLevelEditorActionCallbacks::OnSelectAllLights()
 {
 	GEditor->GetSelectedActors()->BeginBatchSelectOperation();
 	// Select all light actors.
-	for( FActorIterator It(GetWorld()); It; ++It )
+	for( TActorIterator<ALight> It(GetWorld()); It; ++It )
 	{
-		ALight* Light = Cast<ALight>(*It);
-		if( Light )
-		{
-			GUnrealEd->SelectActor( Light, true, false, false );
-		}
+		GUnrealEd->SelectActor( *It, true, false, false );
 	}
 
 	GEditor->GetSelectedActors()->EndBatchSelectOperation();
diff --git a/Engine/Source/Editor/UnrealEd/Classes/Editor/UnrealEdEngine.h b/Engine/Source/Editor/UnrealEd/Classes/Editor/UnrealEdEngine.h
--- a/Engine/Source/Editor/UnrealEd/Classes/Editor/UnrealEdEngine.h
+++ b/Engine/Source/Editor/UnrealEd/Classes/Editor/UnrealEdEngine.h
@@ -325,7 +325,7 @@ class UNREALED_API UUnrealEdEngine : public UEditorEngine, public FNotifyHook
 	 * @param	InClass		Class of actor to select
 	 * @param	InArchetype	Archetype of actor to select
 	 */
-	virtual void edactSelectOfClassAndArchetype( UWorld* InWorld, const UClass* InClass, const UObject* InArchetype );
+	virtual void edactSelectOfClassAndArchetype( UWorld* InWorld, const TSubclassOf<AActor> InClass, const UObject* InArchetype );
 
 	/**
 	 * Select all actors in a particular class and its subclasses.
@@ -654,7 +654,7 @@ class UNREALED_API UUnrealEdEngine : public UEditorEngine, public FNotifyHook
 	 * @param InVolumeActorClass	The type of volume actors to update.  If NULL is passed in all volume actor types are updated.
 	 * @param InViewport			The viewport where actor visibility should apply.  Pass NULL for all editor viewports.
 	 */
-	void UpdateVolumeActorVisibility( const UClass* InVolumeActorClass = NULL , FLevelEditorViewportClient* InViewport = NULL);
+	void UpdateVolumeActorVisibility( UClass* InVolumeActorClass = NULL , FLevelEditorViewportClient* InViewport = NULL);
 
 	/**
 	 * Get the index of the provided sprite category
diff --git a/Engine/Source/Editor/UnrealEd/Private/Commandlets/PackageUtilities.cpp b/Engine/Source/Editor/UnrealEd/Private/Commandlets/PackageUtilities.cpp
--- a/Engine/Source/Editor/UnrealEd/Private/Commandlets/PackageUtilities.cpp
+++ b/Engine/Source/Editor/UnrealEd/Private/Commandlets/PackageUtilities.cpp
@@ -2374,7 +2374,7 @@ int32 UReplaceActorCommandlet::Main(const FString& Params)
 				World->UpdateWorldComponents( true, false );
 
 				// iterate through all the actors in the world, looking for matches with the class to replace (must have exact match, not subclass)
-				for (FActorIterator It(World); It; ++It)
+				for (TActorIterator<AActor> It(World, ClassToReplace); It; ++It)
 				{
 					AActor* OldActor = *It;
 					if (OldActor->GetClass() == ClassToReplace)
diff --git a/Engine/Source/Editor/UnrealEd/Private/DebugToolExec.cpp b/Engine/Source/Editor/UnrealEd/Private/DebugToolExec.cpp
--- a/Engine/Source/Editor/UnrealEd/Private/DebugToolExec.cpp
+++ b/Engine/Source/Editor/UnrealEd/Private/DebugToolExec.cpp
@@ -215,9 +215,9 @@ bool FDebugToolExec::Exec( UWorld* InWorld, const TCHAR* Cmd, FOutputDevice& Ar
 			}
 
 			float   MinDist = FLT_MAX;
-			for( FActorIterator It(InWorld); It; ++It )
+			for( TActorIterator<AActor> It(InWorld, Class); It; ++It )
 			{
-				if ( !It->IsPendingKill() && It->IsA(Class) )
+				if ( !It->IsPendingKill() )
 				{
 					float const Dist = (PlayerController && It->GetRootComponent()) ? FVector::Dist(It->GetActorLocation(), PlayerLocation) : 0.f;
 					if (Dist < MinDist)
diff --git a/Engine/Source/Editor/UnrealEd/Private/Editor.cpp b/Engine/Source/Editor/UnrealEd/Private/Editor.cpp
--- a/Engine/Source/Editor/UnrealEd/Private/Editor.cpp
+++ b/Engine/Source/Editor/UnrealEd/Private/Editor.cpp
@@ -3287,7 +3287,7 @@ void UEditorEngine::SelectAllActorsWithClass( bool bArchetype )
 		// For this function to have been called in the first place, all of the selected actors should be of the same type
 		// and with the same archetype; however, it's safest to confirm the assumption first
 		bool bAllSameClassAndArchetype = false;
-		UClass* FirstClass = NULL;
+		TSubclassOf<AActor> FirstClass;
 		UObject* FirstArchetype = NULL;
 
 		// Find the class and archetype of the first selected actor; they will be used to check that all selected actors
@@ -4519,14 +4519,10 @@ void UEditorEngine::OnPostSaveWorld(uint32 SaveFlags, UWorld* World, uint32 Orig
 APlayerStart* UEditorEngine::CheckForPlayerStart()
 {
 	UWorld* IteratorWorld = GWorld;
-	for( FActorIterator It(IteratorWorld); It; ++It )
+	for( TActorIterator<APlayerStart> It(IteratorWorld); It; ++It )
 	{
-		APlayerStart* Start = Cast<APlayerStart>(*It);
-		if (Start)
-		{
-			// Return the found start.
-			return Start;
-		}
+		// Return the found start.
+		return *It;
 	}
 
 	// No player start was found, return NULL.
diff --git a/Engine/Source/Editor/UnrealEd/Private/EditorActor.cpp b/Engine/Source/Editor/UnrealEd/Private/EditorActor.cpp
--- a/Engine/Source/Editor/UnrealEd/Private/EditorActor.cpp
+++ b/Engine/Source/Editor/UnrealEd/Private/EditorActor.cpp
@@ -895,10 +895,10 @@ void UUnrealEdEngine::edactReplaceClassWithClass(UWorld* InWorld, UClass* SrcCla
 {
 	// Make a list of actors to replace.
 	TArray<AActor*> ActorsToReplace;
-	for( FActorIterator It(InWorld); It; ++It )
+	for( TActorIterator<AActor> It(InWorld, SrcClass); It; ++It )
 	{
 		AActor* Actor = *It;
-		if ( Actor->IsA( SrcClass ) && Actor->HasAnyFlags(RF_Transactional) )
+		if ( Actor->HasAnyFlags(RF_Transactional) )
 		{
 			ActorsToReplace.Add( Actor );
 		}
@@ -1537,7 +1537,7 @@ void UUnrealEdEngine::edactSelectOfClass( UWorld* InWorld, UClass* Class )
 
 	SelectedActors->Modify();
 
-	for( FActorIterator It(InWorld); It; ++It )
+	for( TActorIterator<AActor> It(InWorld, Class); It; ++It )
 	{
 		AActor* Actor = *It;
 		if( Actor->GetClass()==Class && !Actor->IsSelected() && !Actor->IsHiddenEd() )
@@ -1556,7 +1556,7 @@ void UUnrealEdEngine::edactSelectOfClass( UWorld* InWorld, UClass* Class )
 }
 
 
-void UUnrealEdEngine::edactSelectOfClassAndArchetype( UWorld* InWorld, const UClass* InClass, const UObject* InArchetype )
+void UUnrealEdEngine::edactSelectOfClassAndArchetype( UWorld* InWorld, const TSubclassOf<AActor> InClass, const UObject* InArchetype )
 {
 	USelection* SelectedActors = GetSelectedActors();
 	SelectedActors->BeginBatchSelectOperation();
@@ -1565,7 +1565,7 @@ void UUnrealEdEngine::edactSelectOfClassAndArchetype( UWorld* InWorld, const UCl
 
 	// Select all actors with of the provided class and archetype, assuming they aren't already selected, 
 	// aren't hidden in the editor, aren't a member of a prefab, and aren't builder brushes
-	for( FActorIterator ActorIter(InWorld); ActorIter; ++ActorIter )
+	for( TActorIterator<AActor> ActorIter(InWorld, InClass); ActorIter; ++ActorIter )
 	{
 		AActor* CurActor = *ActorIter;
 		if ( CurActor->GetClass() == InClass && CurActor->GetArchetype() == InArchetype && !CurActor->IsSelected() 
@@ -1587,10 +1587,10 @@ void UUnrealEdEngine::edactSelectSubclassOf( UWorld* InWorld, UClass* Class )
 
 	SelectedActors->Modify();
 
-	for( FActorIterator It(InWorld); It; ++It )
+	for( TActorIterator<AActor> It(InWorld, Class); It; ++It )
 	{
 		AActor* Actor = *It;
-		if( !Actor->IsSelected() && !Actor->IsHiddenEd() && Actor->GetClass()->IsChildOf(Class) )
+		if( !Actor->IsSelected() && !Actor->IsHiddenEd() )
 		{
 			// Selection by class not permitted for actors belonging to prefabs.
 			// Selection by class not permitted for builder brushes.
@@ -1969,15 +1969,14 @@ void UUnrealEdEngine::edactSelectMatchingEmitter()
 	SelectedActors->BeginBatchSelectOperation();
 	SelectedActors->Modify();
 	// Iterate over all of the non-hidden actors, selecting those who have a particle system template that matches one from the previously-found list
-	for( FActorIterator ActorIterator(WorldList[0]); ActorIterator; ++ActorIterator )
+	for( TActorIterator<AEmitter> ActorIterator(WorldList[0]); ActorIterator; ++ActorIterator )
 	{
-		AActor* Actor = *ActorIterator;
-		if ( !Actor->IsHiddenEd() )
+		AEmitter* ActorAsEmitter = *ActorIterator;
+		if ( !ActorAsEmitter->IsHiddenEd() )
 		{
-			AEmitter* ActorAsEmitter = Cast<AEmitter>( Actor );
-			if ( ActorAsEmitter && ActorAsEmitter->ParticleSystemComponent && SelectedParticleSystemTemplates.Contains( ActorAsEmitter->ParticleSystemComponent->Template ) )
+			if ( ActorAsEmitter->ParticleSystemComponent && SelectedParticleSystemTemplates.Contains( ActorAsEmitter->ParticleSystemComponent->Template ) )
 			{
-				SelectActor( Actor, true, false );
+				SelectActor( ActorAsEmitter, true, false );
 			}
 		}
 	}
diff --git a/Engine/Source/Editor/UnrealEd/Private/EditorCsg.cpp b/Engine/Source/Editor/UnrealEd/Private/EditorCsg.cpp
--- a/Engine/Source/Editor/UnrealEd/Private/EditorCsg.cpp
+++ b/Engine/Source/Editor/UnrealEd/Private/EditorCsg.cpp
@@ -266,10 +266,10 @@ void UEditorEngine::csgRebuild( UWorld* InWorld )
 
 	TArray<ABrush*> DynamicBrushes;
 	DynamicBrushes.Empty();
-	for( FActorIterator It(InWorld); It; ++It )
+	for( TActorIterator<ABrush> It(InWorld); It; ++It )
 	{
-		ABrush* B=Cast<ABrush>(*It);
-		if ( B && B->Brush && !B->IsStaticBrush() )
+		ABrush* B=*It;
+		if ( B->Brush && !B->IsStaticBrush() )
 		{
 			DynamicBrushes.Add(B);
 		}
diff --git a/Engine/Source/Editor/UnrealEd/Private/EditorModes.cpp b/Engine/Source/Editor/UnrealEd/Private/EditorModes.cpp
--- a/Engine/Source/Editor/UnrealEd/Private/EditorModes.cpp
+++ b/Engine/Source/Editor/UnrealEd/Private/EditorModes.cpp
@@ -608,12 +608,12 @@ void FEdMode::Render(const FSceneView* View,FViewport* Viewport,FPrimitiveDrawIn
 	{
 		// Draw translucent polygons on brushes and volumes
 
-		for( FActorIterator It(GetWorld()); It; ++ It )
+		for( TActorIterator<ABrush> It(GetWorld()); It; ++ It )
 		{
-			ABrush* Brush = Cast<ABrush>( *It );
+			ABrush* Brush = *It;
 
 			// Brush->Brush is checked to safe from brushes that were created without having their brush members attached.
-			if( Brush && Brush->Brush && (FActorEditorUtils::IsABuilderBrush(Brush) || Brush->IsVolumeBrush()) && GEditor->GetSelectedActors()->IsSelected(Brush) )
+			if( Brush->Brush && (FActorEditorUtils::IsABuilderBrush(Brush) || Brush->IsVolumeBrush()) && GEditor->GetSelectedActors()->IsSelected(Brush) )
 			{
 				// Build a mesh by basically drawing the triangles of each 
 				FDynamicMeshBuilder MeshBuilder;
diff --git a/Engine/Source/Editor/UnrealEd/Private/EditorServer.cpp b/Engine/Source/Editor/UnrealEd/Private/EditorServer.cpp
--- a/Engine/Source/Editor/UnrealEd/Private/EditorServer.cpp
+++ b/Engine/Source/Editor/UnrealEd/Private/EditorServer.cpp
@@ -108,10 +108,10 @@ static int32 CleanBSPMaterials(UWorld* InWorld, bool bPreviewOnly, bool bLogBrus
 	// Clear the mark flag the polys of all non-volume, non-builder brushes.
 	// Make a list of all brushes that were encountered.
 	TArray<ABrush*> Brushes;
-	for ( FActorIterator It(InWorld) ; It ; ++It )
+	for ( TActorIterator<ABrush> It(InWorld) ; It ; ++It )
 	{
-		ABrush* ItBrush = Cast<ABrush>(*It);
-		if ( ItBrush && !ItBrush->IsVolumeBrush() && !FActorEditorUtils::IsABuilderBrush(ItBrush) && !ItBrush->IsBrushShape() )
+		ABrush* ItBrush = *It;
+		if ( !ItBrush->IsVolumeBrush() && !FActorEditorUtils::IsABuilderBrush(ItBrush) && !ItBrush->IsBrushShape() )
 		{
 			if( ItBrush->Brush && ItBrush->Brush->Polys )
 			{
@@ -2987,13 +2987,9 @@ void UEditorEngine::CopySelectedActorsToClipboard( UWorld* InWorld, bool bShould
 				}
 
 				// Clean-up flag for Landscape Proxy cases...
-				for( FActorIterator CurActorIt(InWorld); CurActorIt; ++CurActorIt )
+				for( TActorIterator<ALandscapeProxy> ProxyIt(InWorld); ProxyIt; ++ProxyIt )
 				{
-					ALandscapeProxy* Proxy = Cast<ALandscapeProxy>(*CurActorIt);
-					if (Proxy)
-					{
-						Proxy->bIsMovingToLevel = false;
-					}
+					ProxyIt->bIsMovingToLevel = false;
 				}
 
 				BufferLevel->ClearLevelComponents();
@@ -6074,7 +6070,7 @@ bool UEditorEngine::HandleSetDetailModeViewCommand( const TCHAR* Str, FOutputDev
 			Actor->GetComponents(Components);
 
 			for(int32 ComponentIndex = 0;ComponentIndex < Components.Num();ComponentIndex++)
-				{
+			{
 				Components[ComponentIndex]->MarkRenderStateDirty();
 			}
 		}
@@ -6177,12 +6173,11 @@ bool UEditorEngine::HandleListMapPackageDependenciesCommand( const TCHAR* Str, F
 
 bool UEditorEngine::HandleRebuildVolumesCommand( const TCHAR* Str, FOutputDevice& Ar, UWorld* InWorld )
 {
-	for( FActorIterator It(InWorld); It; ++It )
+	for( TActorIterator<AVolume> It(InWorld); It; ++It )
 	{
-		AActor* Actor = *It;
-		AVolume* Volume = Cast<AVolume>(Actor);
+		AVolume* Volume = *It;
 
-		if(Volume != NULL && !Volume->IsTemplate() && Volume->BrushComponent.IsValid())
+		if(!Volume->IsTemplate() && Volume->BrushComponent.IsValid())
 		{
 			UE_LOG(LogEditorServer, Log, TEXT("BSBC: %s"), *Volume->GetPathName() );
 			Volume->BrushComponent->BuildSimpleBrushCollision();
diff --git a/Engine/Source/Editor/UnrealEd/Private/Lightmass/Lightmass.cpp b/Engine/Source/Editor/UnrealEd/Private/Lightmass/Lightmass.cpp
--- a/Engine/Source/Editor/UnrealEd/Private/Lightmass/Lightmass.cpp
+++ b/Engine/Source/Editor/UnrealEd/Private/Lightmass/Lightmass.cpp
@@ -823,57 +823,55 @@ void FLightmassExporter::WriteVisibilityData( int32 Channel )
 			ACameraActor* Camera = *CameraIt;
 			if (World->ContainsActor(Camera) && !Camera->IsPendingKill())
 			{
-				for( FActorIterator It(World); It; ++It )
+				for( TActorIterator<AMatineeActor> It(World); It; ++It )
 				{
-					AMatineeActor* MatineeActor = Cast<AMatineeActor>(*It);
-					if( MatineeActor )
+					AMatineeActor* MatineeActor = *It;
+
+					bool bNeedsTermInterp = false;
+					if (MatineeActor->GroupInst.Num() == 0)
 					{
-						bool bNeedsTermInterp = false;
-						if (MatineeActor->GroupInst.Num() == 0)
-						{
-							// If matinee is closed, GroupInst will be empty, so we need to populate it
-							bNeedsTermInterp = true;
-							MatineeActor->InitInterp();
-						}
-						UInterpGroupInst* GroupInstance = MatineeActor->FindGroupInst(Camera);
-						if (GroupInstance && GroupInstance->Group)
+						// If matinee is closed, GroupInst will be empty, so we need to populate it
+						bNeedsTermInterp = true;
+						MatineeActor->InitInterp();
+					}
+					UInterpGroupInst* GroupInstance = MatineeActor->FindGroupInst(Camera);
+					if (GroupInstance && GroupInstance->Group)
+					{
+						for (int32 TrackIndex = 0; TrackIndex < GroupInstance->Group->InterpTracks.Num(); TrackIndex++)
 						{
-							for (int32 TrackIndex = 0; TrackIndex < GroupInstance->Group->InterpTracks.Num(); TrackIndex++)
+							UInterpTrackMove* MoveTrack = Cast<UInterpTrackMove>(GroupInstance->Group->InterpTracks[TrackIndex]);
+							if (MoveTrack)
 							{
-								UInterpTrackMove* MoveTrack = Cast<UInterpTrackMove>(GroupInstance->Group->InterpTracks[TrackIndex]);
-								if (MoveTrack)
+								float StartTime;
+								float EndTime;
+								//@todo - look at the camera cuts to only process time ranges where the camera is actually active
+								MoveTrack->GetTimeRange(StartTime, EndTime);
+								for (int32 TrackInstanceIndex = 0; TrackInstanceIndex < GroupInstance->TrackInst.Num(); TrackInstanceIndex++)
 								{
-									float StartTime;
-									float EndTime;
-									//@todo - look at the camera cuts to only process time ranges where the camera is actually active
-									MoveTrack->GetTimeRange(StartTime, EndTime);
-									for (int32 TrackInstanceIndex = 0; TrackInstanceIndex < GroupInstance->TrackInst.Num(); TrackInstanceIndex++)
+									UInterpTrackInst* TrackInstance = GroupInstance->TrackInst[TrackInstanceIndex];
+									UInterpTrackInstMove* MoveTrackInstance = Cast<UInterpTrackInstMove>(TrackInstance);
+									if (MoveTrackInstance)
 									{
-										UInterpTrackInst* TrackInstance = GroupInstance->TrackInst[TrackInstanceIndex];
-										UInterpTrackInstMove* MoveTrackInstance = Cast<UInterpTrackInstMove>(TrackInstance);
-										if (MoveTrackInstance)
+										//@todo - handle long camera paths
+										for (float Time = StartTime; Time < EndTime; Time += FMath::Max((EndTime - StartTime) * .001f, .001f))
 										{
-											//@todo - handle long camera paths
-											for (float Time = StartTime; Time < EndTime; Time += FMath::Max((EndTime - StartTime) * .001f, .001f))
+											const FVector RelativePosition = MoveTrack->EvalPositionAtTime(TrackInstance, Time);
+											FVector CurrentPosition;
+											FRotator CurrentRotation;
+											MoveTrack->ComputeWorldSpaceKeyTransform(MoveTrackInstance, RelativePosition, FRotator::ZeroRotator, CurrentPosition, CurrentRotation);
+											if (CameraTrackPositions.Num() == 0 || !CurrentPosition.Equals(CameraTrackPositions.Last(), CellSize * .1f))
 											{
-												const FVector RelativePosition = MoveTrack->EvalPositionAtTime(TrackInstance, Time);
-												FVector CurrentPosition;
-												FRotator CurrentRotation;
-												MoveTrack->ComputeWorldSpaceKeyTransform(MoveTrackInstance, RelativePosition, FRotator::ZeroRotator, CurrentPosition, CurrentRotation);
-												if (CameraTrackPositions.Num() == 0 || !CurrentPosition.Equals(CameraTrackPositions.Last(), CellSize * .1f))
-												{
-													CameraTrackPositions.Add(CurrentPosition);
-												}
+												CameraTrackPositions.Add(CurrentPosition);
 											}
 										}
 									}
 								}
 							}
 						}
-						if (bNeedsTermInterp)
-						{
-							MatineeActor->TermInterp();
-						}
+					}
+					if (bNeedsTermInterp)
+					{
+						MatineeActor->TermInterp();
 					}
 				}
 			}
diff --git a/Engine/Source/Editor/UnrealEd/Private/Tests/BlueprintAutomationTests.cpp b/Engine/Source/Editor/UnrealEd/Private/Tests/BlueprintAutomationTests.cpp
--- a/Engine/Source/Editor/UnrealEd/Private/Tests/BlueprintAutomationTests.cpp
+++ b/Engine/Source/Editor/UnrealEd/Private/Tests/BlueprintAutomationTests.cpp
@@ -1097,14 +1097,13 @@ bool FBlueprintInstancesTest::RunTest(const FString& InParameters)
 	TSet<AActor*> BlueprintInstances;
 	for ( FActorIterator It(GWorld); It; ++It )
 	{
-		if (AActor* Actor = Cast<AActor>(*It))
-		{
-			UClass* ActorClass = Actor->GetClass();
+		AActor* Actor = *It;
 
-			if (ActorClass->ClassGeneratedBy && ActorClass->ClassGeneratedBy->IsA( UBlueprint::StaticClass() ) )
-			{
-				BlueprintInstances.Add(Actor);
-			}
+		UClass* ActorClass = Actor->GetClass();
+
+		if (ActorClass->ClassGeneratedBy && ActorClass->ClassGeneratedBy->IsA( UBlueprint::StaticClass() ) )
+		{
+			BlueprintInstances.Add(Actor);
 		}
 	}
 
diff --git a/Engine/Source/Editor/UnrealEd/Private/Tests/EditorAutomationTests.cpp b/Engine/Source/Editor/UnrealEd/Private/Tests/EditorAutomationTests.cpp
--- a/Engine/Source/Editor/UnrealEd/Private/Tests/EditorAutomationTests.cpp
+++ b/Engine/Source/Editor/UnrealEd/Private/Tests/EditorAutomationTests.cpp
@@ -1316,15 +1316,13 @@ bool FCleanupConvertToValidation::Update()
 			// If we found the mesh, find and delete the static mesh actor we added to the level to clear the reference to it.
 			if (GeneratedMesh)
 			{
-				for (FActorIterator ActorIt(TestWorld); ActorIt; ++ActorIt)
+				for (TActorIterator<AStaticMeshActor> ActorIt(TestWorld); ActorIt; ++ActorIt)
 				{
-					AStaticMeshActor* StaticMeshActor = Cast<AStaticMeshActor>(*ActorIt);
-					if (StaticMeshActor)
+					AStaticMeshActor* StaticMeshActor = *ActorIt;
+
+					if ( StaticMeshActor->StaticMeshComponent->StaticMesh == GeneratedMesh )
 					{
-						if ( StaticMeshActor->StaticMeshComponent->StaticMesh == GeneratedMesh )
-						{
-							TestWorld->DestroyActor(StaticMeshActor);
-						}
+						TestWorld->DestroyActor(StaticMeshActor);
 					}
 				}
 			}
diff --git a/Engine/Source/Editor/UnrealEd/Private/UnrealEdEngine.cpp b/Engine/Source/Editor/UnrealEd/Private/UnrealEdEngine.cpp
--- a/Engine/Source/Editor/UnrealEd/Private/UnrealEdEngine.cpp
+++ b/Engine/Source/Editor/UnrealEd/Private/UnrealEdEngine.cpp
@@ -990,21 +990,16 @@ static void InternalUpdateVolumeActorVisibility( TArray<AActor*>& ActorsToUpdate
 }
 
 
-void UUnrealEdEngine::UpdateVolumeActorVisibility( const UClass* InVolumeActorClass, FLevelEditorViewportClient* InViewport )
+void UUnrealEdEngine::UpdateVolumeActorVisibility( UClass* InVolumeActorClass, FLevelEditorViewportClient* InViewport )
 {
-	const UClass* VolumeClassToCheck = InVolumeActorClass ? InVolumeActorClass : AVolume::StaticClass();
+	TSubclassOf<AActor> VolumeClassToCheck = InVolumeActorClass ? InVolumeActorClass : AVolume::StaticClass();
 	
 	// Build a list of actors that need to be updated.  Only take actors of the passed in volume class.  
 	UWorld* World = InViewport ? InViewport->GetWorld() : GWorld;
 	TArray< AActor *> ActorsToUpdate;
-	for( FActorIterator It( World ); It; ++It)
+	for( TActorIterator<AActor> It( World, VolumeClassToCheck ); It; ++It)
 	{
-		AActor* Actor = *It;
-
-		if( Actor->IsA( VolumeClassToCheck ) )
-		{
-			ActorsToUpdate.Add(Actor);
-		}
+		ActorsToUpdate.Add(*It);
 	}
 
 	if( ActorsToUpdate.Num() > 0 )
diff --git a/Engine/Source/Editor/UnrealEd/Private/UnrealEdSrv.cpp b/Engine/Source/Editor/UnrealEd/Private/UnrealEdSrv.cpp
--- a/Engine/Source/Editor/UnrealEd/Private/UnrealEdSrv.cpp
+++ b/Engine/Source/Editor/UnrealEd/Private/UnrealEdSrv.cpp
@@ -502,10 +502,10 @@ bool UUnrealEdEngine::HandleUpdateLandscapeEditorDataCommand( const TCHAR* Str,
 				
 		// for removing 
 		TMap<ULandscapeInfo*, ALandscapeGizmoActiveActor*> GizmoMap;
-		for (FActorIterator It(InWorld); It; ++It)
+		for (TActorIterator<ALandscapeGizmoActiveActor> It(InWorld); It; ++It)
 		{
-			ALandscapeGizmoActiveActor* Gizmo = Cast<ALandscapeGizmoActiveActor>(*It);
-			if (Gizmo && Gizmo->TargetLandscapeInfo)
+			ALandscapeGizmoActiveActor* Gizmo = *It;
+			if (Gizmo->TargetLandscapeInfo)
 			{
 				if (!GizmoMap.FindRef(Gizmo->TargetLandscapeInfo))
 				{
diff --git a/Engine/Source/Runtime/Engine/Private/CheatManager.cpp b/Engine/Source/Runtime/Engine/Private/CheatManager.cpp
--- a/Engine/Source/Runtime/Engine/Private/CheatManager.cpp
+++ b/Engine/Source/Runtime/Engine/Private/CheatManager.cpp
@@ -237,10 +237,10 @@ void UCheatManager::DestroyTarget()
 
 void UCheatManager::DestroyAll(TSubclassOf<AActor> aClass)
 {
-	for (FActorIterator It(GetWorld()); It; ++It)
+	for (TActorIterator<AActor> It(GetWorld(),aClass); It; ++It)
 	{
 		AActor* A = *It;
-		if (A && A->IsA(aClass) && !A->IsPendingKill())
+		if (!A->IsPendingKill())
 		{
 			APawn* Pawn = Cast<APawn>(A);
 			if (Pawn != NULL)
@@ -438,10 +438,10 @@ void UCheatManager::ViewClass( TSubclassOf<AActor> DesiredClass )
 {
 	bool bFound = false;
 	AActor* First = NULL;
-	for (FActorIterator It(GetWorld()); It; ++It)
+	for (TActorIterator<AActor> It(GetWorld(), DesiredClass); It; ++It)
 	{
 		AActor* TestActor = *It;
-		if (TestActor && !TestActor->IsPendingKill() && (*TestActor->GetClass()).IsChildOf(*DesiredClass) )
+		if (!TestActor->IsPendingKill())
 		{
 			AActor* Other = TestActor;
 			if (bFound || (First == NULL))
diff --git a/Engine/Source/Runtime/Engine/Private/GameState.cpp b/Engine/Source/Runtime/Engine/Private/GameState.cpp
--- a/Engine/Source/Runtime/Engine/Private/GameState.cpp
+++ b/Engine/Source/Runtime/Engine/Private/GameState.cpp
@@ -57,14 +57,9 @@ void AGameState::PostInitializeComponents()
 	UWorld* World = GetWorld();
 	World->GameState = this;
 
-	for (FActorIterator It(World); It; ++It)
+	for (TActorIterator<APlayerState> It(World); It; ++It)
 	{
-		AActor* TestActor = *It;
-		if (TestActor && !TestActor->IsPendingKill() && TestActor->IsA<APlayerState>())
-		{
-			APlayerState* PlayerState = Cast<APlayerState>(TestActor);
-			AddPlayerState(PlayerState);
-		}
+		AddPlayerState(*It);
 	}
 }
 
diff --git a/Engine/Source/Runtime/Engine/Private/Landscape/Landscape.cpp b/Engine/Source/Runtime/Engine/Private/Landscape/Landscape.cpp
--- a/Engine/Source/Runtime/Engine/Private/Landscape/Landscape.cpp
+++ b/Engine/Source/Runtime/Engine/Private/Landscape/Landscape.cpp
@@ -1752,9 +1752,9 @@ void ULandscapeInfo::RecreateLandscapeInfo(UWorld* InWorld, bool bMapCheck)
 
 	TMap<FGuid, TArray<ALandscapeProxy*>> ValidLandscapesMap;
 	// Gather all valid landscapes in the world
-	for (FActorIterator It(InWorld); It; ++It)
+	for (TActorIterator<ALandscapeProxy> It(InWorld); It; ++It)
 	{
-		ALandscapeProxy* Proxy = Cast<ALandscapeProxy>(*It);
+		ALandscapeProxy* Proxy = *It;
 		if (Proxy
 			&& Proxy->HasAnyFlags(RF_BeginDestroyed) == false
 			&& Proxy->IsPendingKill() == false
diff --git a/Engine/Source/Runtime/Engine/Private/Landscape/LandscapeEdit.cpp b/Engine/Source/Runtime/Engine/Private/Landscape/LandscapeEdit.cpp
--- a/Engine/Source/Runtime/Engine/Private/Landscape/LandscapeEdit.cpp
+++ b/Engine/Source/Runtime/Engine/Private/Landscape/LandscapeEdit.cpp
@@ -3333,10 +3333,10 @@ void ALandscapeProxy::PostEditImport()
 	Super::PostEditImport();
 	if (!bIsProxy && GetWorld()) // For Landscape
 	{
-		for (FActorIterator It(GetWorld()); It; ++It)
+		for (TActorIterator<ALandscape> It(GetWorld()); It; ++It)
 		{
-			ALandscape* Landscape = Cast<ALandscape>(*It);
-			if (Landscape && Landscape != this && !Landscape->HasAnyFlags(RF_BeginDestroyed) && Landscape->LandscapeGuid == LandscapeGuid)
+			ALandscape* Landscape = *It;
+			if (Landscape != this && !Landscape->HasAnyFlags(RF_BeginDestroyed) && Landscape->LandscapeGuid == LandscapeGuid)
 			{
 				// Copy/Paste case, need to generate new GUID
 				LandscapeGuid = FGuid::NewGuid();
@@ -3794,13 +3794,9 @@ void ALandscape::PostEditChangeProperty(FPropertyChangedEvent& PropertyChangedEv
 		// Need to update Gizmo scene proxy
 		if (bNeedsRecalcBoundingBox && GetWorld())
 		{
-			for (FActorIterator It(GetWorld()); It; ++It)
+			for (TActorIterator<ALandscapeGizmoActiveActor> It(GetWorld()); It; ++It)
 			{
-				ALandscapeGizmoActiveActor* Gizmo = Cast<ALandscapeGizmoActiveActor>(*It);
-				if (Gizmo)
-				{
-					Gizmo->ReregisterAllComponents();
-				}
+				It->ReregisterAllComponents();
 			}
 		}
 
diff --git a/Engine/Source/Runtime/Engine/Private/LocalPlayer.cpp b/Engine/Source/Runtime/Engine/Private/LocalPlayer.cpp
--- a/Engine/Source/Runtime/Engine/Private/LocalPlayer.cpp
+++ b/Engine/Source/Runtime/Engine/Private/LocalPlayer.cpp
@@ -1337,29 +1337,26 @@ bool ULocalPlayer::HandleCancelMatineeCommand( const TCHAR* Cmd, FOutputDevice&
 	{
 		TArray<UWorld*> MatineeActorWorldsThatSkipped;
 		// if so, look for all active matinees that has this Player in a director group
-		for (FActorIterator It(GetWorld()); It; ++It)
+		for (TActorIterator<AMatineeActor> It(GetWorld()); It; ++It)
 		{
-			AMatineeActor* MatineeActor = Cast<AMatineeActor>(*It);
+			AMatineeActor* MatineeActor = *It;
 
-			if( MatineeActor )
+			// is it currently playing (and skippable)?
+			if (MatineeActor->bIsPlaying && MatineeActor->bIsSkippable && (MatineeActor->bClientSideOnly || MatineeActor->GetWorld()->IsServer()))
 			{
-				// is it currently playing (and skippable)?
-				if (MatineeActor->bIsPlaying && MatineeActor->bIsSkippable && (MatineeActor->bClientSideOnly || MatineeActor->GetWorld()->IsServer()))
+				for (int32 GroupIndex = 0; GroupIndex < MatineeActor->GroupInst.Num(); GroupIndex++)
 				{
-					for (int32 GroupIndex = 0; GroupIndex < MatineeActor->GroupInst.Num(); GroupIndex++)
+					// is the PC the group actor?
+					if (MatineeActor->GroupInst[GroupIndex]->GetGroupActor() == PlayerController)
 					{
-						// is the PC the group actor?
-						if (MatineeActor->GroupInst[GroupIndex]->GetGroupActor() == PlayerController)
+						const float RightBeforeEndTime = 0.1f;
+						// make sure we aren';t already at the end (or before the allowed skip time)
+						if ((MatineeActor->InterpPosition < MatineeActor->MatineeData->InterpLength - RightBeforeEndTime) && 
+							(MatineeActor->InterpPosition >= InitialNoSkipTime))
 						{
-							const float RightBeforeEndTime = 0.1f;
-							// make sure we aren';t already at the end (or before the allowed skip time)
-							if ((MatineeActor->InterpPosition < MatineeActor->MatineeData->InterpLength - RightBeforeEndTime) && 
-								(MatineeActor->InterpPosition >= InitialNoSkipTime))
-							{
-								// skip to end
-								MatineeActor->SetPosition(MatineeActor->MatineeData->InterpLength - RightBeforeEndTime, true);
-								MatineeActorWorldsThatSkipped.AddUnique( MatineeActor->GetWorld() );
-							}
+							// skip to end
+							MatineeActor->SetPosition(MatineeActor->MatineeData->InterpLength - RightBeforeEndTime, true);
+							MatineeActorWorldsThatSkipped.AddUnique( MatineeActor->GetWorld() );
 						}
 					}
 				}
diff --git a/Engine/Source/Runtime/Engine/Private/NetConnection.cpp b/Engine/Source/Runtime/Engine/Private/NetConnection.cpp
--- a/Engine/Source/Runtime/Engine/Private/NetConnection.cpp
+++ b/Engine/Source/Runtime/Engine/Private/NetConnection.cpp
@@ -1502,9 +1502,9 @@ void UChildConnection::HandleClientPlayer(APlayerController* PC, UNetConnection*
 #if !(UE_BUILD_SHIPPING || UE_BUILD_TEST)
 		UE_LOG(LogNet, Error, TEXT("Failed to find LocalPlayer for received PlayerController '%s' with index %d. PlayerControllers:"), *PC->GetName(), int32(PC->NetPlayerIndex));
 		check( PC->GetWorld() );
-		for (FActorIterator It(PC->GetWorld()); It; ++It)
+		for (TActorIterator<APlayerController> It(PC->GetWorld()); It; ++It)
 		{
-			if (Cast<const APlayerController>(*It) && It->Role < ROLE_Authority)
+			if (It->Role < ROLE_Authority)
 			{
 				UE_LOG(LogNet, Log, TEXT(" - %s"), *It->GetFullName());
 			}
diff --git a/Engine/Source/Runtime/Engine/Private/PlayerController.cpp b/Engine/Source/Runtime/Engine/Private/PlayerController.cpp
--- a/Engine/Source/Runtime/Engine/Private/PlayerController.cpp
+++ b/Engine/Source/Runtime/Engine/Private/PlayerController.cpp
@@ -1668,16 +1668,14 @@ void APlayerController::NotifyLoadedWorld(FName WorldPackageName, bool bFinalDes
 	// place the camera at the first playerstart we can find
 	SetViewTarget(this);
 	
-	for( FActorIterator It(GetWorld()); It; ++It )
+	for( TActorIterator<APlayerStart> It(GetWorld()); It; ++It )
 	{
-		APlayerStart* P = Cast<APlayerStart>(*It);
-		if (P)
-		{
-			FRotator SpawnRotation(ForceInit);
-			SpawnRotation.Yaw = P->GetActorRotation().Yaw;
-			SetInitialLocationAndRotation(P->GetActorLocation(), SpawnRotation);
-			break;
-		}
+		APlayerStart* P = *It;
+
+		FRotator SpawnRotation(ForceInit);
+		SpawnRotation.Yaw = P->GetActorRotation().Yaw;
+		SetInitialLocationAndRotation(P->GetActorLocation(), SpawnRotation);
+		break;
 	}
 }
 
diff --git a/Engine/Source/Runtime/Engine/Private/World.cpp b/Engine/Source/Runtime/Engine/Private/World.cpp
--- a/Engine/Source/Runtime/Engine/Private/World.cpp
+++ b/Engine/Source/Runtime/Engine/Private/World.cpp
@@ -4672,10 +4672,10 @@ void UWorld::UpdateConstraintActors()
 {
 	if( bAreConstraintsDirty )
 	{
-		for( FActorIterator It(this); It; ++It )
+		for( TActorIterator<APhysicsConstraintActor> It(this); It; ++It )
 		{
-			APhysicsConstraintActor* ConstraintActor = Cast<APhysicsConstraintActor>( *It );
-			if( ConstraintActor && ConstraintActor->ConstraintComp)
+			APhysicsConstraintActor* ConstraintActor = *It;
+			if( ConstraintActor->ConstraintComp)
 			{
 				ConstraintActor->ConstraintComp->UpdateConstraintFrames();
 			}
diff --git a/Engine/Source/Runtime/GameplayAbilities/Private/AbilitySystemDebugHUD.cpp b/Engine/Source/Runtime/GameplayAbilities/Private/AbilitySystemDebugHUD.cpp
--- a/Engine/Source/Runtime/GameplayAbilities/Private/AbilitySystemDebugHUD.cpp
+++ b/Engine/Source/Runtime/GameplayAbilities/Private/AbilitySystemDebugHUD.cpp
@@ -130,12 +130,10 @@ static void	ToggleDebugHUD(const TArray<FString>& Args, UWorld* InWorld)
 
 	AAbilitySystemDebugHUD *HUD = NULL;
 
-	for (FActorIterator It(InWorld); It; ++It)
+	for (TActorIterator<AAbilitySystemDebugHUD> It(InWorld); It; ++It)
 	{
-		AActor *Actor = *It;
-		HUD = Cast<AAbilitySystemDebugHUD>(Actor);
-		if (HUD)
-			break;
+		HUD = *It;
+		break;
 	}
 
 	if (!HUD)
