Gathering from assets and metadata nows loads required modules before trying to load packages. Specify modules using ModulesToPreload in the commandlet config files.

diff --git a/Engine/Source/Editor/UnrealEd/Private/Commandlets/GatherTextFromAssetsCommandlet.cpp b/Engine/Source/Editor/UnrealEd/Private/Commandlets/GatherTextFromAssetsCommandlet.cpp
--- a/Engine/Source/Editor/UnrealEd/Private/Commandlets/GatherTextFromAssetsCommandlet.cpp
+++ b/Engine/Source/Editor/UnrealEd/Private/Commandlets/GatherTextFromAssetsCommandlet.cpp
@@ -682,6 +682,15 @@ int32 UGatherTextFromAssetsCommandlet::Main(const FString& Params)
 		return -1;
 	}
 
+	//Modules to Preload
+	TArray<FString> ModulesToPreload;
+	GetConfigArray(*SectionName, TEXT("ModulesToPreload"), ModulesToPreload, GatherTextConfigPath);
+
+	for (const FString& ModuleName : ModulesToPreload)
+	{
+		FModuleManager::Get().LoadModule(*ModuleName);
+	}
+
 	//Include paths
 	TArray<FString> IncludePaths;
 	GetConfigArray(*SectionName, TEXT("IncludePaths"), IncludePaths, GatherTextConfigPath);
@@ -912,6 +921,7 @@ int32 UGatherTextFromAssetsCommandlet::Main(const FString& Params)
 	for( int32 BatchIndex = 0; BatchIndex < BatchCount; ++BatchIndex )
 	{
 		int32 PackagesInThisBatch = 0;
+		int32 FailuresInThisBatch = 0;
 		for( ; PackageIndex < PackageCount && PackagesInThisBatch < PackagesPerBatchCount; ++PackageIndex )
 		{
 			FString PackageFileName = PackageFileNamesToLoad[PackageIndex];
@@ -932,13 +942,14 @@ int32 UGatherTextFromAssetsCommandlet::Main(const FString& Params)
 			else
 			{
 				FailedPackageFileNames.Add( PackageFileName );
+				++FailuresInThisBatch;
 				continue;
 			}
 
 			++PackagesInThisBatch;
 		}
 
-		UE_LOG(LogGatherTextFromAssetsCommandlet, Log, TEXT("Loaded %i packages in batch %i of %i."), PackagesInThisBatch, BatchIndex + 1, BatchCount);
+		UE_LOG(LogGatherTextFromAssetsCommandlet, Log, TEXT("Loaded %i packages in batch %i of %i. %i failed."), PackagesInThisBatch, BatchIndex + 1, BatchCount, FailuresInThisBatch);
 
 		ProcessPackages(PackagesToProcess);
 		PackagesToProcess.Empty(PackagesPerBatchCount);
diff --git a/Engine/Source/Editor/UnrealEd/Private/Commandlets/GatherTextFromMetadataCommandlet.cpp b/Engine/Source/Editor/UnrealEd/Private/Commandlets/GatherTextFromMetadataCommandlet.cpp
--- a/Engine/Source/Editor/UnrealEd/Private/Commandlets/GatherTextFromMetadataCommandlet.cpp
+++ b/Engine/Source/Editor/UnrealEd/Private/Commandlets/GatherTextFromMetadataCommandlet.cpp
@@ -49,6 +49,15 @@ int32 UGatherTextFromMetaDataCommandlet::Main( const FString& Params )
 		return -1;
 	}
 
+	//Modules to Preload
+	TArray<FString> ModulesToPreload;
+	GetConfigArray(*SectionName, TEXT("ModulesToPreload"), ModulesToPreload, GatherTextConfigPath);
+
+	for (const FString& ModuleName : ModulesToPreload)
+	{
+		FModuleManager::Get().LoadModule(*ModuleName);
+	}
+
 	//Include paths
 	TArray<FString> IncludePaths;
 	GetConfigArray(*SectionName, TEXT("IncludePaths"), IncludePaths, GatherTextConfigPath);
