UE4: The default value for bGenerateLightmapUVs can now be configured and is stored in the EditorUserSettings.

diff --git a/Engine/Source/Editor/UnrealEd/Classes/Factories/FbxStaticMeshImportData.h b/Engine/Source/Editor/UnrealEd/Classes/Factories/FbxStaticMeshImportData.h
--- a/Engine/Source/Editor/UnrealEd/Classes/Factories/FbxStaticMeshImportData.h
+++ b/Engine/Source/Editor/UnrealEd/Classes/Factories/FbxStaticMeshImportData.h
@@ -24,7 +24,7 @@ class UFbxStaticMeshImportData : public UFbxMeshImportData
 	UPROPERTY(EditAnywhere, Category=ImportSettings)
 	uint32 bRemoveDegenerates:1;
 	
-	UPROPERTY(EditAnywhere, Category=ImportSettings)
+	UPROPERTY(EditAnywhere, config, Category=ImportSettings)
 	uint32 bGenerateLightmapUVs:1;
 
 	/** If checked, one convex hull per UCX_ prefixed collision mesh will be generated instead of decomposing into multiple hulls */
