Add cooker log messages.

diff --git a/Engine/Source/Editor/UnrealEd/Private/CookOnTheFlyServer.cpp b/Engine/Source/Editor/UnrealEd/Private/CookOnTheFlyServer.cpp
--- a/Engine/Source/Editor/UnrealEd/Private/CookOnTheFlyServer.cpp
+++ b/Engine/Source/Editor/UnrealEd/Private/CookOnTheFlyServer.cpp
@@ -25,13 +25,16 @@
 #include "Commandlets/ChunkManifestGenerator.h"
 #include "Engine/WorldComposition.h"
 
+// error message log
+#include "TokenizedMessage.h"
+#include "MessageLog.h"
 
 
 
 DEFINE_LOG_CATEGORY_STATIC(LogCookOnTheFly, Log, All);
 
 #define DEBUG_COOKONTHEFLY 0
-#define OUTPUT_TIMING 1
+#define OUTPUT_TIMING 0
 
 #if OUTPUT_TIMING
 
@@ -348,6 +351,28 @@ struct FRecompileRequest
 };
 
 
+
+/**
+ * Uses the FMessageLog to log a message
+ * 
+ * @param Message to log
+ * @param Severity of the message
+ */
+void LogCookerMessage( const FString& MessageText, EMessageSeverity::Type Severity)
+{
+	FMessageLog MessageLog("CookResults");
+
+	TSharedRef<FTokenizedMessage> Message = FTokenizedMessage::Create(Severity);
+
+	Message->AddToken( FTextToken::Create( FText::FromString(MessageText) ) );
+	// Message->AddToken(FTextToken::Create(MessageLogTextDetail)); 
+	// Message->AddToken(FDocumentationToken::Create(TEXT("https://docs.unrealengine.com/latest/INT/Platforms/iOS/QuickStart/6/index.html"))); 
+	MessageLog.AddMessage(Message);
+
+	MessageLog.Notify();
+}
+
+
 /* UCookOnTheFlyServer structors
  *****************************************************************************/
 
@@ -428,8 +453,8 @@ bool UCookOnTheFlyServer::BroadcastFileserverPresence( const FGuid &InstanceId )
 		INetworkFileServer *NetworkFileServer = NetworkFileServers[i];
 		if ((NetworkFileServer == NULL || !NetworkFileServer->IsItReadyToAcceptConnections() || !NetworkFileServer->GetAddressList(AddressList)))
 		{
+			LogCookerMessage( FString(TEXT("Failed to create network file server")), EMessageSeverity::Error );
 			UE_LOG(LogCookOnTheFly, Error, TEXT("Failed to create network file server"));
-
 			continue;
 		}
 
@@ -716,7 +741,6 @@ void UCookOnTheFlyServer::GenerateManifestInfo( UPackage* Package, const TArray<
 
 		if ( CookByTheBookOptions->bGenerateStreamingInstallManifests )
 		{
-			UE_LOG(LogCookOnTheFly, Display, TEXT("PrepareToLoadNewPackage %s"), *StandardFilename.ToString() );
 			ManifestGenerator->PrepareToLoadNewPackage( StandardFilename.ToString() );
 		}
 
@@ -732,7 +756,7 @@ void UCookOnTheFlyServer::GenerateManifestInfo( UPackage* Package, const TArray<
 			{
 				// Populate streaming install manifests
 				FString SandboxFilename = SandboxFile->ConvertToAbsolutePathForExternalAppForWrite(*Filename);
-				UE_LOG(LogCookOnTheFly, Display, TEXT("Adding package to manifest %s, %s, %s"), *DependentPackage->GetName(), *SandboxFilename, *LastLoadedMapName);
+				//UE_LOG(LogCookOnTheFly, Display, TEXT("Adding package to manifest %s, %s, %s"), *DependentPackage->GetName(), *SandboxFilename, *LastLoadedMapName);
 				ManifestGenerator->AddPackageToChunkManifest(DependentPackage, SandboxFilename, LastLoadedMapName, SandboxFile.GetOwnedPointer());
 			}
 		}
@@ -822,7 +846,7 @@ uint32 UCookOnTheFlyServer::TickCookOnTheSide( const float TimeSlice, uint32 &Co
 				Package = FindObject<UPackage>( ANY_PACKAGE, *PackageName );
 			}
 
-			UE_LOG( LogCookOnTheFly, Display, TEXT("Processing request %s"), *BuildFilename)
+			UE_LOG( LogCookOnTheFly, Display, TEXT("Processing request %s"), *BuildFilename);
 
 			//  if the package is already loaded then try to avoid reloading it :)
 			if ( ( Package == NULL ) || ( Package->IsFullyLoaded() == false ) )
@@ -843,6 +867,7 @@ uint32 UCookOnTheFlyServer::TickCookOnTheSide( const float TimeSlice, uint32 &Co
 
 			if( Package == NULL )
 			{
+				LogCookerMessage( FString::Printf(TEXT("Error loading %s!"), *BuildFilename), EMessageSeverity::Error );
 				UE_LOG(LogCookOnTheFly, Error, TEXT("Error loading %s!"), *BuildFilename );
 				Result |= COSR_ErrorLoadingPackage;
 			}
@@ -1489,6 +1514,8 @@ bool UCookOnTheFlyServer::SaveCookedPackage( UPackage* Package, uint32 SaveFlags
 					Package->FullyLoad();
 					if (!Package->IsFullyLoaded())
 					{
+						LogCookerMessage( FString::Printf(TEXT("Package %s supposed to be fully loaded but isn't. RF_WasLoaded is %s"), 
+							*Package->GetName(), Package->HasAnyFlags(RF_WasLoaded) ? TEXT("set") : TEXT("not set")), EMessageSeverity::Warning);
 						UE_LOG(LogCookOnTheFly, Warning, TEXT("Package %s supposed to be fully loaded but isn't. RF_WasLoaded is %s"), 
 							*Package->GetName(), Package->HasAnyFlags(RF_WasLoaded) ? TEXT("set") : TEXT("not set"));
 					}
@@ -1525,6 +1552,7 @@ bool UCookOnTheFlyServer::SaveCookedPackage( UPackage* Package, uint32 SaveFlags
 				const FString FullFilename = FPaths::ConvertRelativePathToFull( PlatFilename );
 				if( FullFilename.Len() >= PLATFORM_MAX_FILEPATH_LENGTH )
 				{
+					LogCookerMessage( FString::Printf(TEXT("Couldn't save package, filename is too long: %s"), *PlatFilename), EMessageSeverity::Error );
 					UE_LOG( LogCookOnTheFly, Error, TEXT( "Couldn't save package, filename is too long :%s" ), *PlatFilename );
 					bSavedCorrectly = false;
 				}
@@ -1757,6 +1785,7 @@ void UCookOnTheFlyServer::GenerateLongPackageNames(TArray<FString>& FilesInPath)
 			}
 			else
 			{
+				LogCookerMessage( FString::Printf(TEXT("Unable to generate long package name for %s"), *FileInPath), EMessageSeverity::Warning);
 				UE_LOG(LogCookOnTheFly, Warning, TEXT("Unable to generate long package name for %s"), *FileInPath);
 			}
 		}
@@ -1801,6 +1830,7 @@ void UCookOnTheFlyServer::CollectFilesToCook(TArray<FString>& FilesInPath, const
 			FString OutFilename;
 			if (FPackageName::SearchForPackageOnDisk(CurrEntry, NULL, &OutFilename) == false)
 			{
+				LogCookerMessage( FString::Printf(TEXT("Unable to find package for map %s."), *CurrEntry), EMessageSeverity::Warning);
 				UE_LOG(LogCookOnTheFly, Warning, TEXT("Unable to find package for map %s."), *CurrEntry);
 			}
 			else
@@ -2084,6 +2114,7 @@ void UCookOnTheFlyServer::StartCookByTheBook(const TArray<ITargetPlatform*>& Tar
 	CollectFilesToCook(FilesInPath, CookMaps, CookDirectories, CookCultures, IniMapSections, bCookAll, bMapsOnly, bNoDev );
 	if (FilesInPath.Num() == 0)
 	{
+		LogCookerMessage( FString::Printf(TEXT("No files found to cook.")), EMessageSeverity::Warning );
 		UE_LOG(LogCookOnTheFly, Warning, TEXT("No files found."));
 	}
 
@@ -2122,6 +2153,7 @@ void UCookOnTheFlyServer::StartCookByTheBook(const TArray<ITargetPlatform*>& Tar
 		}
 		else
 		{
+			LogCookerMessage( FString::Printf(TEXT("Unable to find package for cooking %s"), *FileName), EMessageSeverity::Warning );
 			UE_LOG(LogCookOnTheFly, Warning, TEXT("Unable to find package for cooking %s"), *FileName)
 		}
 		
