Cleaned up the editor performance capturing code  and added more data that will be captured.
AutomationEditorCommon.cpp

diff --git a/Engine/Source/Editor/UnrealEd/Private/Tests/AutomationEditorCommon.cpp b/Engine/Source/Editor/UnrealEd/Private/Tests/AutomationEditorCommon.cpp
--- a/Engine/Source/Editor/UnrealEd/Private/Tests/AutomationEditorCommon.cpp
+++ b/Engine/Source/Editor/UnrealEd/Private/Tests/AutomationEditorCommon.cpp
@@ -315,13 +315,13 @@ float ValueFromStringArrayofNumbers(const TArray<FString>& NumberArray, bool bAv
 	int32 NumberCounter = 0;
 
 	//Loop through the array.
-	for (int32 Index = 0; Index < NumberArray.Num(); ++Index)
+	for (int32 I = 0; I < NumberArray.Num(); ++I)
 	{
 		//If the character in the array is a number then we'll want to add it to our total.
-		if (NumberArray[Index].IsNumeric())
+		if (NumberArray[I].IsNumeric())
 		{
 			NumberCounter++;
-			TotalValue += FCString::Atof(NumberArray[Index].GetCharArray().GetData());
+			TotalValue += FCString::Atof(NumberArray[I].GetCharArray().GetData());
 		}
 	}	
 
@@ -338,6 +338,28 @@ float ValueFromStringArrayofNumbers(const TArray<FString>& NumberArray, bool bAv
 
 }
 
+float LargestValueFromStringArrayofNumbers(const TArray<FString>& NumberArray)
+{
+	//Total Value holds the sum of all the numbers available in the array.
+	float LargestValue = FCString::Atof(NumberArray[0].GetCharArray().GetData());
+
+	//Loop through the array.
+	for (int32 I = 0; I < NumberArray.Num(); ++I)
+	{
+		//If the character in the array is a number then we'll want to add it to our total.
+		if (NumberArray[I].IsNumeric())
+		{
+			float CurrentValue = FCString::Atof(NumberArray[I].GetCharArray().GetData());
+			if ( LargestValue < CurrentValue)
+			{
+				LargestValue = CurrentValue;
+			}
+		}
+	}
+	UE_LOG(LogEditorAutomationTests, Log, TEXT("The Largest value of the array is %f"), LargestValue);
+	return LargestValue;
+}
+
 TArray<FString> CreateArrayFromFile(const FString& FullFileLocation)
 {
 
@@ -450,12 +472,12 @@ bool FEditorPerformanceCommand::Update()
 	FString SavedMemory;
 	FString SavedDuration;
 	static SIZE_T StaticLastTotalAllocated = 0;
+	static SIZE_T AvailablePhysicalMemory = 0;
+	static SIZE_T AvailableVirtualMemory = 0;
+	static SIZE_T UsedVirtualMemory = 0;
+	static SIZE_T UsedPeakPhysical = 0;
+	static SIZE_T UsedPeakVirtual = 0;
 
-	//Variables needed for the WriteToTextFile function.
-	FString TestName = TEXT("Performance");
-	FString Delimiter = TEXT(",");
-	FString FileName;
-	
 	//Grab the FPS and working memory number.
 	float CurrentTime = FPlatformTime::Seconds();
 	if ((CurrentTime - StartTime) <= Duration)
@@ -474,14 +496,40 @@ bool FEditorPerformanceCommand::Update()
 		FPlatformMemoryStats MemoryStats = FPlatformMemory::GetStats();
 		StaticLastTotalAllocated = MemoryStats.UsedPhysical;
 
+		//Query OS for available physical memory
+		AvailablePhysicalMemory = MemoryStats.AvailablePhysical;
+
+		//Query OS for available virtual memory
+		AvailableVirtualMemory = MemoryStats.AvailableVirtual;
+
+		//Query OS for used virtual memory
+		UsedVirtualMemory = MemoryStats.UsedVirtual;
+
+		//Query OS for used Peak Used physical memory
+		UsedPeakPhysical = MemoryStats.PeakUsedPhysical;
+
+		//Query OS for used Peak Used virtual memory
+		UsedPeakVirtual = MemoryStats.PeakUsedVirtual;
+
 		//Log out to a text file the Duration.
-		WriteToTextFile(TestName, ShortMapName, FileName = TEXT("RAWDuration"), CurrentTime, Delimiter);
+		WriteToTextFile(TEXT("Performance"), ShortMapName, TEXT("RAWDuration"), CurrentTime, TEXT(","));
 		//Log out to a text file the AverageFPS
-		WriteToTextFile(TestName, ShortMapName, FileName = TEXT("RAWAverageFPS"), ClampedFPS, Delimiter);
+		WriteToTextFile(TEXT("Performance"), ShortMapName, TEXT("RAWAverageFPS"), ClampedFPS, TEXT(","));
 		//Log out to a text file the AverageFrameTime
-		WriteToTextFile(TestName, ShortMapName, FileName = TEXT("RAWFrameTime"), ClampedMS, Delimiter);
-		//Log out to a text file the Memory in kilobytes.
-		WriteToTextFile(TestName, ShortMapName, FileName = TEXT("RAWMemoryUsedPhysical"), (float)StaticLastTotalAllocated / 1024, Delimiter);
+		WriteToTextFile(TEXT("Performance"), ShortMapName, TEXT("RAWFrameTime"), ClampedMS, TEXT(","));
+		//Log out to a text file the Used Physical Memory in kilobytes.
+		WriteToTextFile(TEXT("Performance"), ShortMapName, TEXT("RAWMemoryUsedPhysical"), (float)StaticLastTotalAllocated / 1024, TEXT(","));
+		//Log out to a text file the Used Peaked Memory in kilobytes.
+		WriteToTextFile(TEXT("Performance"), ShortMapName, TEXT("RAWMemoryUsedPeak"), (float)UsedPeakPhysical / 1024, TEXT(","));
+		//Log out to a text file the Used Peaked virtual Memory in kilobytes.
+		WriteToTextFile(TEXT("Performance"), ShortMapName, TEXT("RAWMemoryUsedPeakVirtual"), (float)UsedPeakVirtual / 1024, TEXT(","));
+		//Log out to a text file the Available Physical Memory in kilobytes.
+		WriteToTextFile(TEXT("Performance"), ShortMapName, TEXT("RAWMemoryAvailablePhysical"), (float)AvailablePhysicalMemory / 1024, TEXT(","));
+		//Log out to a text file the Available Virtual Memory in kilobytes.
+		WriteToTextFile(TEXT("Performance"), ShortMapName, TEXT("RAWMemoryAvailableVirtual"), (float)AvailableVirtualMemory / 1024, TEXT(","));
+		//Log out to a text file the Used Virtual Memory in kilobytes.
+		WriteToTextFile(TEXT("Performance"), ShortMapName, TEXT("RAWMemoryUsedVirtual"), (float)UsedVirtualMemory / 1024, TEXT(","));
+
 
 		return false;
 	}
@@ -502,18 +550,28 @@ bool FGenerateEditorPerformanceCharts::Update()
 	FString LoadedMapName = MainFrameModule.GetLoadedLevelName();
 
 	//The locations of the raw data.
-	FString FileSaveLocation = FPaths::AutomationLogDir() + TEXT("Performance/") + LoadedMapName + TEXT("/");
-	FString FPSFileSaveLocation = FileSaveLocation + TEXT("RAWAverageFPS.txt");
-	FString FrameTimeSaveLocation = FileSaveLocation + TEXT("RAWFrameTime.txt");
-	FString MemoryFileSaveLocation = FileSaveLocation + TEXT("RAWMemoryUsedPhysical.txt");
-	FString MapLoadTimeFileSaveLocation = FileSaveLocation + TEXT("RAWMapLoadTime.txt");
-	FString DurationFileSaveLocation = FileSaveLocation + TEXT("RAWDuration.txt");
+	FString FileLocation = FPaths::AutomationLogDir() + TEXT("Performance/") + LoadedMapName + TEXT("/");
+	FString FPSFileLocation = FileLocation + TEXT("RAWAverageFPS.txt");
+	FString FrameTimeLocation = FileLocation + TEXT("RAWFrameTime.txt");
+	FString MemoryFileLocation = FileLocation + TEXT("RAWMemoryUsedPhysical.txt");
+	FString MemoryAvailPhysFileLocation = FileLocation + TEXT("RAWMemoryAvailablePhysical.txt");
+	FString MemoryAvailVirtualFileLocation = FileLocation + TEXT("RAWMemoryAvailableVirtual.txt");
+	FString MemoryUsedVirtualFileLocation = FileLocation + TEXT("RAWMemoryUsedVirtual.txt");
+	FString MapLoadTimeFileLocation = FileLocation + TEXT("RAWMapLoadTime.txt");
+	FString DurationFileLocation = FileLocation + TEXT("RAWDuration.txt");
+	FString MemoryUsedPeakFileLocation = FileLocation + TEXT("RAWMemoryUsedPeak.txt");
+	FString MemoryUsedPeakVirtualFileLocation = FileLocation + TEXT("RAWMemoryUsedPeakVirtual.txt");
 	//Get the data from the raw files and put them into an array.
-	TArray<FString> FPSRawArray = CreateArrayFromFile(FPSFileSaveLocation);
-	TArray<FString> FrameTimeRawArray = CreateArrayFromFile(FrameTimeSaveLocation);
-	TArray<FString> MemoryRawArray = CreateArrayFromFile(MemoryFileSaveLocation);
-	TArray<FString> MapLoadTimeRawArray = CreateArrayFromFile(MapLoadTimeFileSaveLocation);
-	TArray<FString> TestRunTimeRawArray = CreateArrayFromFile(DurationFileSaveLocation);
+	TArray<FString> FPSRawArray = CreateArrayFromFile(FPSFileLocation);
+	TArray<FString> FrameTimeRawArray = CreateArrayFromFile(FrameTimeLocation);
+	TArray<FString> MemoryRawArray = CreateArrayFromFile(MemoryFileLocation);
+	TArray<FString> MemoryAvailPhysRawArray = CreateArrayFromFile(MemoryAvailPhysFileLocation);
+	TArray<FString> MemoryAvailVirtualRawArray = CreateArrayFromFile(MemoryAvailVirtualFileLocation);
+	TArray<FString> MemoryUsedVirtualRawArray = CreateArrayFromFile(MemoryUsedVirtualFileLocation);
+	TArray<FString> MapLoadTimeRawArray = CreateArrayFromFile(MapLoadTimeFileLocation);
+	TArray<FString> TestRunTimeRawArray = CreateArrayFromFile(DurationFileLocation);
+	TArray<FString> MemoryUsedPeakArray = CreateArrayFromFile(MemoryUsedPeakFileLocation);
+	TArray<FString> MemoryUsedPeakVirtualArray = CreateArrayFromFile(MemoryUsedPeakVirtualFileLocation);
 
 	//Variable that will hold the old performance csv file data.
 	FString OldPerformanceCSVFile;
@@ -521,8 +579,8 @@ bool FGenerateEditorPerformanceCharts::Update()
 	//The CSV files.
 	//RAW csv holds all of the information from the text files
 	//Final csv only holds the averaged information.
-	FString FinalCSVFilename = FString::Printf(TEXT("%s%s_Performance.csv"), *FileSaveLocation, *LoadedMapName);
-	FString RAWCSVFilename = FString::Printf(TEXT("%sRAW_%s_%s.csv"), *FileSaveLocation, *LoadedMapName, *FDateTime::Now().ToString());
+	FString FinalCSVFilename = FString::Printf(TEXT("%s%s_Performance.csv"), *FileLocation, *LoadedMapName);
+	FString RAWCSVFilename = FString::Printf(TEXT("%sRAW_%s_%s.csv"), *FileLocation, *LoadedMapName, *FDateTime::Now().ToString());
 
 	//Create the .csv file that will hold the raw data from the text files.
 	FArchive* RAWCSVFile = IFileManager::Get().CreateFileWriter(*RAWCSVFilename);
@@ -533,11 +591,11 @@ bool FGenerateEditorPerformanceCharts::Update()
 	}
 
 	//Add the top title row to the raw csv file
-	FString RAWCSVLine = (TEXT("Map Name, Changelist, Test Run Time, Map Load Time, Average FPS, Frame Time, Used Physical Memory\n"));
+	FString RAWCSVLine = (TEXT("Map Name, Changelist, Test Run Time, Map Load Time, Average FPS, Frame Time, Used Physical Memory, Used Virtual Memory, Used Peak Physical, Used Peak Virtual, Available Physical Memory, Available Virtual Memory\n"));
 	RAWCSVFile->Serialize(TCHAR_TO_ANSI(*RAWCSVLine), RAWCSVLine.Len());
 	
 	//Loop through the text files and add them to the raw csv file.
-	for (int32 IterLoop = 0; IterLoop < FPSRawArray.Num(); IterLoop++)
+	for (int32 I = 0; I < FPSRawArray.Num(); I++)
 	{
 		//If the raw file isn't available to write to then we'll fail back this test.
 		if (!RAWCSVFile)
@@ -547,7 +605,7 @@ bool FGenerateEditorPerformanceCharts::Update()
 		}
 
 		//Create the line that will be written and then write it to the csv file.
-		RAWCSVLine = LoadedMapName + TEXT(",") + Changelist + TEXT(",") + TestRunTimeRawArray[IterLoop] + TEXT(",") + MapLoadTimeRawArray[0] + TEXT(",") + FPSRawArray[IterLoop] + TEXT(",") + FrameTimeRawArray[IterLoop] + TEXT(",") + MemoryRawArray[IterLoop] + LINE_TERMINATOR;
+		RAWCSVLine = LoadedMapName + TEXT(",") + Changelist + TEXT(",") + TestRunTimeRawArray[I] + TEXT(",") + MapLoadTimeRawArray[0] + TEXT(",") + FPSRawArray[I] + TEXT(",") + FrameTimeRawArray[I] + TEXT(",") + MemoryRawArray[I] + TEXT(",") + MemoryUsedVirtualRawArray[I] + TEXT(",") + MemoryUsedPeakArray[I] + TEXT(",") + MemoryUsedPeakVirtualArray[I] + TEXT(",") + MemoryAvailPhysRawArray[I] + TEXT(",") + MemoryAvailVirtualRawArray[I] + LINE_TERMINATOR;
 		RAWCSVFile->Serialize(TCHAR_TO_ANSI(*RAWCSVLine), RAWCSVLine.Len());
 	}
 	
@@ -556,9 +614,16 @@ bool FGenerateEditorPerformanceCharts::Update()
 	
 	//Get the final data for the Performance csv file.
 	//This gets the averaged numbers using the raw data arrays.
-	float FPSAverage = ValueFromStringArrayofNumbers(FPSRawArray, true);
-	float FrameTimeAverage = ValueFromStringArrayofNumbers(FrameTimeRawArray, true);
-	float MemoryAverage = ValueFromStringArrayofNumbers(MemoryRawArray, true);
+	float FPSAvg = ValueFromStringArrayofNumbers(FPSRawArray, true);
+	float FrameTimeAvg = ValueFromStringArrayofNumbers(FrameTimeRawArray, true);
+	float MemoryAvg = ValueFromStringArrayofNumbers(MemoryRawArray, true);
+	float MemoryAvailPhysAvg = ValueFromStringArrayofNumbers(MemoryAvailPhysRawArray, true);
+	float MemoryAvailVirtualAvg = ValueFromStringArrayofNumbers(MemoryAvailVirtualRawArray, true);
+	float MemoryUsedVirtualAvg = ValueFromStringArrayofNumbers(MemoryUsedVirtualRawArray, true);
+	
+	//Find the largest value in the array.
+	float MemoryUsedPeak = LargestValueFromStringArrayofNumbers(MemoryUsedPeakArray);
+	float MemoryUsedPeakVirtual = LargestValueFromStringArrayofNumbers(MemoryUsedPeakVirtualArray);
 	
 	//This gets the actual test run time by subtracting the last test run entry against the first test run entry.
 	FString TestRunStartTime = TestRunTimeRawArray[0];
@@ -588,7 +653,7 @@ bool FGenerateEditorPerformanceCharts::Update()
 		}
 		
 		//Create the .csv file that will hold the raw data from the text files.
-		FString FinalCSVLine = (TEXT("Date, Map Name, Changelist, Test Run Time, Map Load Time, Average FPS, Average MS, Used Physical KB\n"));
+		FString FinalCSVLine = (TEXT("Date, Map Name, Changelist, Test Run Time , Map Load Time, Average FPS, Average MS, Used Physical KB, Used Virtual KB, Used Peak Physcial KB, Used Peak Virtual KB, Available Physical KB, Available Virtual KB\n"));
 		FinalCSVFile->Serialize(TCHAR_TO_ANSI(*FinalCSVLine), FinalCSVLine.Len());
 		FinalCSVFile->Close();
 	}
@@ -607,26 +672,32 @@ bool FGenerateEditorPerformanceCharts::Update()
 	FinalCSVFile->Serialize(TCHAR_TO_ANSI(*OldPerformanceCSVFile), OldPerformanceCSVFile.Len());
 
 	//Write the averaged numbers to the Performance CSV file.
-	FString FinalCSVLine = FString::Printf(TEXT("%s,%s,%s,%.0f,%s,%.3f,%.3f,%.0f%s"), *FDateTime::Today().ToString(), *LoadedMapName, *Changelist, TestRunTime, *MapLoadTime, FPSAverage, FrameTimeAverage, MemoryAverage, LINE_TERMINATOR);
+	FString FinalCSVLine = FString::Printf(TEXT("%s,%s,%s,%.0f,%s,%.3f,%.3f,%.0f,%.0f,%.0f,%.0f,%.0f,%.0f%s"), *FDateTime::Today().ToString(), *LoadedMapName, *Changelist, TestRunTime, *MapLoadTime, FPSAvg, FrameTimeAvg, MemoryAvg, MemoryUsedVirtualAvg, MemoryUsedPeak, MemoryUsedPeakVirtual, MemoryAvailPhysAvg, MemoryAvailVirtualAvg, LINE_TERMINATOR);
 	FinalCSVFile->Serialize(TCHAR_TO_ANSI(*FinalCSVLine), FinalCSVLine.Len());
 
 	//Close the csv files so we can use them while the editor is still open.
 	FinalCSVFile->Close();
 
 	//Display the averaged numbers.
-	UE_LOG(LogEditorAutomationTests, Display, TEXT("AVG FPS: '%.1f'"), FPSAverage);
-	UE_LOG(LogEditorAutomationTests, Display, TEXT("AVG Frame Time: '%.1f' ms"), FrameTimeAverage);
-	UE_LOG(LogEditorAutomationTests, Display, TEXT("AVG Used Physical Memory: '%.0f' kb"), MemoryAverage);
+	UE_LOG(LogEditorAutomationTests, Display, TEXT("AVG FPS: '%.1f'"), FPSAvg);
+	UE_LOG(LogEditorAutomationTests, Display, TEXT("AVG Frame Time: '%.1f' ms"), FrameTimeAvg);
+	UE_LOG(LogEditorAutomationTests, Display, TEXT("AVG Used Physical Memory: '%.0f' kb"), MemoryAvg);
+	UE_LOG(LogEditorAutomationTests, Display, TEXT("AVG Used Virtual Memory: '%.0f' kb"), MemoryUsedVirtualAvg);
 	UE_LOG(LogEditorAutomationTests, Display, TEXT("Performance csv file is located here: %s"), *FPaths::ConvertRelativePathToFull(FinalCSVFilename));
 	UE_LOG(LogEditorAutomationTests, Log, TEXT("Performance csv file is located here: %s"), *FPaths::ConvertRelativePathToFull(FinalCSVFilename));
 
 	//Delete the text files.
 	UE_LOG(LogEditorAutomationTests, Log, TEXT("Deleting temporary data text files."));
-	IFileManager::Get().Delete(*FPSFileSaveLocation);
-	IFileManager::Get().Delete(*FrameTimeSaveLocation);
-	IFileManager::Get().Delete(*MemoryFileSaveLocation);
-	IFileManager::Get().Delete(*MapLoadTimeFileSaveLocation);
-	IFileManager::Get().Delete(*DurationFileSaveLocation);
+	IFileManager::Get().Delete(*FPSFileLocation);
+	IFileManager::Get().Delete(*FrameTimeLocation);
+	IFileManager::Get().Delete(*MemoryFileLocation);
+	IFileManager::Get().Delete(*MemoryAvailPhysFileLocation);
+	IFileManager::Get().Delete(*MapLoadTimeFileLocation);
+	IFileManager::Get().Delete(*DurationFileLocation);
+	IFileManager::Get().Delete(*MemoryAvailVirtualFileLocation);
+	IFileManager::Get().Delete(*MemoryUsedVirtualFileLocation);
+	IFileManager::Get().Delete(*MemoryUsedPeakFileLocation);
+	IFileManager::Get().Delete(*MemoryUsedPeakVirtualFileLocation);
 
 	return true;
 }
@@ -642,11 +713,6 @@ bool FEditorLoadMap::Update()
 	double MapLoadStartTime = 0.0f;
 	double MapLoadTime = 0.0f;
 
-	//Variables needed for the WriteToTextFile function.
-	FString TestName = TEXT("Performance");
-	FString Delimiter = TEXT(",");
-	FString FileName;
-
 	//Get the current number of seconds.  This will be used to track how long it took to load the map.
 	MapLoadStartTime = FPlatformTime::Seconds();
 
@@ -663,7 +729,7 @@ bool FEditorLoadMap::Update()
 	UE_LOG(LogEditorAutomationTests, Log, TEXT("%s has been loaded."), *ShortMapName);
 
 	//Log out to a text file the time it takes to load the map.
-	WriteToTextFile(TestName, LoadedMapName, FileName = TEXT("RAWMapLoadTime"), MapLoadTime, Delimiter);
+	WriteToTextFile(TEXT("Performance"), LoadedMapName, TEXT("RAWMapLoadTime"), MapLoadTime, TEXT(","));
 	
 	UE_LOG(LogEditorAutomationTests, Display, TEXT("%s took %.3f to load."), *LoadedMapName, MapLoadTime);
 	
