FBX UI refactor for NUX. Removed old UI and replaced with a more flexible and extensible version based on details panels.

diff --git a/Engine/Source/Editor/DetailCustomizations/Private/DetailCustomizations.cpp b/Engine/Source/Editor/DetailCustomizations/Private/DetailCustomizations.cpp
--- a/Engine/Source/Editor/DetailCustomizations/Private/DetailCustomizations.cpp
+++ b/Engine/Source/Editor/DetailCustomizations/Private/DetailCustomizations.cpp
@@ -84,6 +84,7 @@
 #include "AudioSettingsDetails.h"
 #include "DateTimeStructCustomization.h"
 #include "TimespanStructCustomization.h"
+#include "FbxImportUIDetails.h"
 #include "RigDetails.h"
 
 IMPLEMENT_MODULE( FDetailCustomizationsModule, DetailCustomizations );
@@ -257,6 +258,8 @@ void FDetailCustomizationsModule::RegisterObjectCustomizations()
 	RegisterCustomClassLayout("MediaPlayer", FOnGetDetailCustomizationInstance::CreateStatic(&FMediaPlayerCustomization::MakeInstance));
 	RegisterCustomClassLayout("MediaSoundWave", FOnGetDetailCustomizationInstance::CreateStatic(&FMediaSoundWaveCustomization::MakeInstance));
 	RegisterCustomClassLayout("MediaTexture", FOnGetDetailCustomizationInstance::CreateStatic(&FMediaTextureCustomization::MakeInstance));
+
+	RegisterCustomClassLayout("FbxImportUI", FOnGetDetailCustomizationInstance::CreateStatic(&FFbxImportUIDetails::MakeInstance));
 }
 
 
@@ -281,4 +284,4 @@ void FDetailCustomizationsModule::RegisterCustomPropertyTypeLayout(FName Propert
 	static FName PropertyEditor("PropertyEditor");
 	FPropertyEditorModule& PropertyModule = FModuleManager::GetModuleChecked<FPropertyEditorModule>(PropertyEditor);
 	PropertyModule.RegisterCustomPropertyTypeLayout(PropertyTypeName, PropertyTypeLayoutDelegate);
-}
+}
\ No newline at end of file
diff --git a/Engine/Source/Editor/DetailCustomizations/Private/FbxImportUIDetails.cpp b/Engine/Source/Editor/DetailCustomizations/Private/FbxImportUIDetails.cpp
new file mode 100644
--- /dev/null
+++ b/Engine/Source/Editor/DetailCustomizations/Private/FbxImportUIDetails.cpp
@@ -0,0 +1,213 @@
+// Copyright 1998-2014 Epic Games, Inc. All Rights Reserved.
+
+#include "DetailCustomizationsPrivatePCH.h"
+#include "FbxImportUIDetails.h"
+#include "Factories/FbxAnimSequenceImportData.h"
+
+#define LOCTEXT_NAMESPACE "FbxImportUIDetails"
+
+
+FFbxImportUIDetails::FFbxImportUIDetails()
+{
+	CachedDetailBuilder = nullptr;
+}
+
+TSharedRef<IDetailCustomization> FFbxImportUIDetails::MakeInstance()
+{
+	return MakeShareable( new FFbxImportUIDetails );
+}
+
+void FFbxImportUIDetails::CustomizeDetails( IDetailLayoutBuilder& DetailBuilder )
+{
+	CachedDetailBuilder = &DetailBuilder;
+	TArray<TWeakObjectPtr<UObject>> EditingObjects;
+	DetailBuilder.GetObjectsBeingCustomized(EditingObjects);
+	check(EditingObjects.Num() == 1);
+
+	ImportUI = Cast<UFbxImportUI>(EditingObjects[0].Get());
+
+	// Handle mesh category
+	IDetailCategoryBuilder& MeshCategory = DetailBuilder.EditCategory("Mesh", TEXT(""), ECategoryPriority::Important);
+	IDetailCategoryBuilder& TransformCategory = DetailBuilder.EditCategory("Transform");
+	TArray<TSharedRef<IPropertyHandle>> CategoryDefaultProperties;
+	TArray<TSharedPtr<IPropertyHandle>> ExtraProperties;
+
+	// Grab and hide per-type import options
+	TSharedRef<IPropertyHandle> StaticMeshDataProp = DetailBuilder.GetProperty(GET_MEMBER_NAME_CHECKED(UFbxImportUI, StaticMeshImportData));
+	TSharedRef<IPropertyHandle> SkeletalMeshDataProp = DetailBuilder.GetProperty(GET_MEMBER_NAME_CHECKED(UFbxImportUI, SkeletalMeshImportData));
+	TSharedRef<IPropertyHandle> AnimSequenceDataProp = DetailBuilder.GetProperty(GET_MEMBER_NAME_CHECKED(UFbxImportUI, AnimSequenceImportData));
+	DetailBuilder.HideProperty(StaticMeshDataProp);
+	DetailBuilder.HideProperty(SkeletalMeshDataProp);
+	DetailBuilder.HideProperty(AnimSequenceDataProp);
+
+	MeshCategory.GetDefaultProperties(CategoryDefaultProperties);
+
+	EFBXImportType ImportType = ImportUI->MeshTypeToImport;
+	switch(ImportType)
+	{
+		case FBXIT_StaticMesh:
+			CollectChildPropertiesRecursive(StaticMeshDataProp, ExtraProperties);
+			break;
+		case FBXIT_SkeletalMesh:
+			CollectChildPropertiesRecursive(SkeletalMeshDataProp, ExtraProperties);
+			break;
+		default:
+			break;
+	}
+
+	if(ImportType != FBXIT_Animation)
+	{
+		TSharedRef<IPropertyHandle> ImportSkeletalProp = DetailBuilder.GetProperty(GET_MEMBER_NAME_CHECKED(UFbxImportUI, bImportAsSkeletal));
+		ImportSkeletalProp->SetOnPropertyValueChanged(FSimpleDelegate::CreateSP(this, &FFbxImportUIDetails::MeshImportModeChanged));
+		MeshCategory.AddProperty(ImportSkeletalProp);
+	}
+
+	for(TSharedRef<IPropertyHandle> Handle : CategoryDefaultProperties)
+	{
+		FString MetaData = Handle->GetMetaData(TEXT("ImportType"));
+		if(!IsImportTypeMetaDataValid(ImportType, MetaData))
+		{
+			DetailBuilder.HideProperty(Handle);
+		}
+	}
+
+	for(TSharedPtr<IPropertyHandle> Handle : ExtraProperties)
+	{
+		FString ImportTypeMetaData = Handle->GetMetaData(TEXT("ImportType"));
+		FString CategoryMetaData = Handle->GetMetaData(TEXT("ImportCategory"));
+		if(IsImportTypeMetaDataValid(ImportType, ImportTypeMetaData))
+		{
+			// Decide on category
+			if(CategoryMetaData == TEXT("Transform"))
+			{
+				// Transform data can be held inside the asset uobjects (used for reimport later).
+				// Move it out into its own UI section.
+				TransformCategory.AddProperty(Handle);
+			}
+			else
+			{
+				MeshCategory.AddProperty(Handle);
+			}
+		}
+	}
+
+	// Animation Category
+	IDetailCategoryBuilder& AnimCategory = DetailBuilder.EditCategory("Animation", TEXT(""), ECategoryPriority::Important);
+	if(ImportType == FBXIT_Animation || ImportType == FBXIT_SkeletalMesh)
+	{
+		ExtraProperties.Empty();
+		CollectChildPropertiesRecursive(AnimSequenceDataProp, ExtraProperties);
+
+		// Before we add the import data properties we need to re-add any properties we want to appear above them in the UI
+		TSharedRef<IPropertyHandle> ImportAnimProp = DetailBuilder.GetProperty(GET_MEMBER_NAME_CHECKED(UFbxImportUI, bImportAnimations));
+		// If we're importing an animation file we really don't need to ask this
+		DetailBuilder.HideProperty(ImportAnimProp);
+		if(ImportType == FBXIT_Animation)
+		{
+			ImportUI->bImportAnimations = true;
+		}
+		else
+		{
+			AnimCategory.AddProperty(ImportAnimProp);
+		}
+
+		for(TSharedPtr<IPropertyHandle> Handle : ExtraProperties)
+		{
+			if(Handle->GetProperty()->GetOuter() == UFbxAnimSequenceImportData::StaticClass())
+			{
+				IDetailPropertyRow& PropertyRow = AnimCategory.AddProperty(Handle);
+			}
+		}
+
+		// Start and End frame for animations need a custom edit condition
+		TSharedRef<IPropertyHandle> StartFrameProp = DetailBuilder.GetProperty(GET_MEMBER_NAME_CHECKED(UFbxAnimSequenceImportData, StartFrame));
+		TSharedRef<IPropertyHandle> EndFrameProp = DetailBuilder.GetProperty(GET_MEMBER_NAME_CHECKED(UFbxAnimSequenceImportData, EndFrame));	
+	}
+	else
+	{
+		// Hide animation options
+		CategoryDefaultProperties.Empty();
+		AnimCategory.GetDefaultProperties(CategoryDefaultProperties);
+
+		for(TSharedRef<IPropertyHandle> Handle : CategoryDefaultProperties)
+		{
+			DetailBuilder.HideProperty(Handle);
+		}
+	}
+
+	// Material Category
+	IDetailCategoryBuilder& MaterialCategory = DetailBuilder.EditCategory("Material");
+	if(ImportType == FBXIT_Animation)
+	{
+		// In animation-only mode, hide the material display
+		CategoryDefaultProperties.Empty();
+		MaterialCategory.GetDefaultProperties(CategoryDefaultProperties);
+
+		for(TSharedRef<IPropertyHandle> Handle : CategoryDefaultProperties)
+		{
+			DetailBuilder.HideProperty(Handle);
+		}
+	}
+	else
+	{
+		TSharedRef<IPropertyHandle> TextureDataProp = DetailBuilder.GetProperty(GET_MEMBER_NAME_CHECKED(UFbxImportUI, TextureImportData));
+		DetailBuilder.HideProperty(TextureDataProp);
+
+		ExtraProperties.Empty();
+		CollectChildPropertiesRecursive(TextureDataProp, ExtraProperties);
+
+		for(TSharedPtr<IPropertyHandle> Handle : ExtraProperties)
+		{
+			// We ignore base import data for this window.
+			if(Handle->GetProperty()->GetOuter() == UFbxTextureImportData::StaticClass())
+			{
+				MaterialCategory.AddProperty(Handle);
+			}
+		}
+	}
+}
+
+void FFbxImportUIDetails::CollectChildPropertiesRecursive(TSharedPtr<IPropertyHandle> Node, TArray<TSharedPtr<IPropertyHandle>>& OutProperties)
+{
+	uint32 NodeNumChildren = 0;
+	Node->GetNumChildren(NodeNumChildren);
+
+	for(uint32 ChildIdx = 0 ; ChildIdx < NodeNumChildren ; ++ChildIdx)
+	{
+		TSharedPtr<IPropertyHandle> ChildHandle = Node->GetChildHandle(ChildIdx);
+		CollectChildPropertiesRecursive(ChildHandle, OutProperties);
+
+		if(ChildHandle->GetProperty())
+		{
+			OutProperties.AddUnique(ChildHandle);
+		}
+	}
+}
+
+bool FFbxImportUIDetails::IsImportTypeMetaDataValid(EFBXImportType& ImportType, FString& MetaData)
+{
+	TArray<FString> Types;
+	MetaData.ParseIntoArray(&Types, TEXT("|"), 1);
+	switch(ImportType)
+	{
+		case FBXIT_StaticMesh:
+			return Types.Contains(TEXT("StaticMesh")) || Types.Contains(TEXT("Mesh"));
+		case FBXIT_SkeletalMesh:
+			return Types.Contains(TEXT("SkeletalMesh")) || Types.Contains(TEXT("Mesh"));
+		case FBXIT_Animation:
+			return Types.Contains(TEXT("Animation"));
+		default:
+			return false;
+	}
+}
+
+void FFbxImportUIDetails::MeshImportModeChanged()
+{
+	if(CachedDetailBuilder)
+	{
+		ImportUI->MeshTypeToImport = ImportUI->MeshTypeToImport == FBXIT_SkeletalMesh ? FBXIT_StaticMesh : FBXIT_SkeletalMesh;
+		CachedDetailBuilder->ForceRefreshDetails();
+	}
+}
+
+#undef LOCTEXT_NAMESPACE
diff --git a/Engine/Source/Editor/DetailCustomizations/Private/FbxImportUIDetails.h b/Engine/Source/Editor/DetailCustomizations/Private/FbxImportUIDetails.h
new file mode 100644
--- /dev/null
+++ b/Engine/Source/Editor/DetailCustomizations/Private/FbxImportUIDetails.h
@@ -0,0 +1,33 @@
+// Copyright 1998-2014 Epic Games, Inc. All Rights Reserved.
+
+#pragma once
+
+class FFbxImportUIDetails : public IDetailCustomization
+{
+public:
+	/** Makes a new instance of this detail layout class for a specific detail view requesting it */
+	static TSharedRef<IDetailCustomization> MakeInstance();
+
+	/** IDetailCustomization interface */
+	virtual void CustomizeDetails( IDetailLayoutBuilder& DetailBuilder ) override;
+	/** End IDetailCustomization interface */
+
+	void CollectChildPropertiesRecursive(TSharedPtr<IPropertyHandle> Node, TArray<TSharedPtr<IPropertyHandle>>& OutProperties);
+	
+	/** Checks whether a metadata string is valid for a given import type 
+	* @param ImportType the type of mesh being imported
+	* @param MetaData the metadata string to validate
+	*/
+	bool IsImportTypeMetaDataValid(EFBXImportType& ImportType, FString& MetaData);
+
+	/** Called if the mesh mode (static / skeletal) changes */
+	void MeshImportModeChanged();
+
+	TWeakObjectPtr<UFbxImportUI> ImportUI;		// The UI data object being customised
+	IDetailLayoutBuilder* CachedDetailBuilder;	// The detail builder for this cusomtomisation
+
+private:
+
+	/** Use MakeInstance to create an instance of this class */
+	FFbxImportUIDetails();
+};
diff --git a/Engine/Source/Editor/UnrealEd/Classes/Factories/FbxAnimSequenceImportData.h b/Engine/Source/Editor/UnrealEd/Classes/Factories/FbxAnimSequenceImportData.h
--- a/Engine/Source/Editor/UnrealEd/Classes/Factories/FbxAnimSequenceImportData.h
+++ b/Engine/Source/Editor/UnrealEd/Classes/Factories/FbxAnimSequenceImportData.h
@@ -24,7 +24,7 @@ enum EFBXAnimationLengthImportType
  * Import data and options used when importing any mesh from FBX
  */
 UCLASS()
-class UFbxAnimSequenceImportData : public UFbxAssetImportData
+class UNREALED_API UFbxAnimSequenceImportData : public UFbxAssetImportData
 {
 	GENERATED_UCLASS_BODY()
 
@@ -33,20 +33,24 @@ class UFbxAnimSequenceImportData : public UFbxAssetImportData
 	TEnumAsByte<enum EFBXAnimationLengthImportType> AnimationLength;
 
 	/** Type of asset to import from the FBX file */
-	UPROPERTY(EditAnywhere, Category=ImportSettings, meta=(DisplayName = "Start Frame"))
+	UPROPERTY(EditAnywhere, AdvancedDisplay, Category=ImportSettings, meta=(DisplayName = "Start Frame"))
 	int32	StartFrame;
+
 	/** Type of asset to import from the FBX file */
-	UPROPERTY(EditAnywhere, Category=ImportSettings, meta=(DisplayName = "End Frame"))
+	UPROPERTY(EditAnywhere, AdvancedDisplay, Category=ImportSettings, meta=(DisplayName = "End Frame"))
 	int32	EndFrame;
 
 	/** Name of source animation that was imported, used to reimport correct animation from the FBX file*/
-	UPROPERTY(EditAnywhere, Category = ImportSettings, meta = (DisplayName = "Source Animation Name"))
+	UPROPERTY()
 	FString SourceAnimationName;
 
 	/** Import if custom attribute as a curve within the animation **/
-	UPROPERTY(EditAnywhere, config, Category = ImportSettings, meta = (DisplayName = "Import Custom Attribute"))
+	UPROPERTY(EditAnywhere, AdvancedDisplay, config, Category = ImportSettings, meta = (DisplayName = "Import Custom Attribute"))
 	bool bImportCustomAttribute;
 
 	/** Gets or creates fbx import data for the specified anim sequence */
 	static UFbxAnimSequenceImportData* GetImportDataForAnimSequence(UAnimSequence* AnimSequence, UFbxAnimSequenceImportData* TemplateForCreation);
+
+	virtual bool CanEditChange(const UProperty* InProperty) const;
+
 };
\ No newline at end of file
diff --git a/Engine/Source/Editor/UnrealEd/Classes/Factories/FbxAssetImportData.h b/Engine/Source/Editor/UnrealEd/Classes/Factories/FbxAssetImportData.h
--- a/Engine/Source/Editor/UnrealEd/Classes/Factories/FbxAssetImportData.h
+++ b/Engine/Source/Editor/UnrealEd/Classes/Factories/FbxAssetImportData.h
@@ -8,7 +8,16 @@
  * Base class for import data and options used when importing any asset from FBX
  */
 UCLASS(config=EditorUserSettings, HideCategories=Object, abstract)
-class UFbxAssetImportData : public UAssetImportData
+class UNREALED_API UFbxAssetImportData : public UAssetImportData
 {
 	GENERATED_UCLASS_BODY()
+
+	UPROPERTY()
+	FVector ImportTranslation;
+
+	UPROPERTY()
+	FRotator ImportRotation;
+
+	UPROPERTY()
+	float ImportUniformScale;
 };
\ No newline at end of file
diff --git a/Engine/Source/Editor/UnrealEd/Classes/Factories/FbxImportUI.h b/Engine/Source/Editor/UnrealEd/Classes/Factories/FbxImportUI.h
--- a/Engine/Source/Editor/UnrealEd/Classes/Factories/FbxImportUI.h
+++ b/Engine/Source/Editor/UnrealEd/Classes/Factories/FbxImportUI.h
@@ -22,81 +22,93 @@ enum EFBXImportType
 	FBXIT_MAX,
 };
 
-UCLASS(config=EditorUserSettings, AutoExpandCategories=(General, USkeletalMesh, UStaticMesh, Materials), HideCategories=Object, MinimalAPI)
+UCLASS(config=EditorUserSettings, AutoExpandCategories=(FTransform), HideCategories=Object, MinimalAPI)
 class UFbxImportUI : public UObject
 {
 	GENERATED_UCLASS_BODY()
 
+	/** Whether or not the imported file is in OBJ format */
+	UPROPERTY()
+	bool bIsObjImport;
+
+	/** The original detected type of this import */
+	UPROPERTY()
+	TEnumAsByte<enum EFBXImportType> OriginalImportType;
+
 	/** Type of asset to import from the FBX file */
-	UPROPERTY(EditAnywhere, Category=General, meta=(DisplayName = "Import Type", ToolTip="The type of mesh being imported"))
+	UPROPERTY()
 	TEnumAsByte<enum EFBXImportType> MeshTypeToImport;
 
 	/** Use the string in "Name" field as full name of mesh. The option only works when the scene contains one mesh. */
-	UPROPERTY(EditAnywhere, config, Category=General)
+	UPROPERTY(EditAnywhere, AdvancedDisplay, config, Category=Miscellaneous, meta=(OBJRestrict="true"))
 	uint32 bOverrideFullName:1;
 
 	/** Whether to convert scene from FBX scene. */
-	UPROPERTY(EditAnywhere, config, Category=General)
+	UPROPERTY(EditAnywhere, AdvancedDisplay, config, Category=Miscellaneous, meta=(OBJRestrict="true", ToolTip="Convert the scene from FBX coordinate system to UE4 coordinate system"))
 	uint32 bConvertScene:1;
 
+	/** Whether to import the incoming FBX as a skeletal object */
+	UPROPERTY(EditAnywhere, Category = Mesh, meta = (ImportType = "StaticMesh|SkeletalMesh"))
+	bool bImportAsSkeletal;
+
 	/** For static meshes, enabling this option will combine all meshes in the FBX into a single monolithic mesh in Unreal */
-	UPROPERTY(EditAnywhere, config, Category=StaticMesh, meta=(ToolTip="If enabled, combines all meshes into a single mesh"))
+	UPROPERTY(EditAnywhere, AdvancedDisplay, config, Category=Mesh, meta=(ToolTip="If enabled, combines all meshes into a single mesh", ImportType="StaticMesh"))
 	uint32 bCombineMeshes:1;
 
-	/** Create Skeleton from this SkeletalMesh with postfix _Skeleton. Don't want this flag to be saved as config, so removed config. */
-	UPROPERTY(EditAnywhere, Category=SkeletalMesh)
+	/** Skeleton to use for imported asset. When importing a mesh, leaving this as "None" will create a new skeleton. When importing and animation this MUST be specified to import the asset. */
+	UPROPERTY(EditAnywhere, Category=Mesh, meta=(ImportType="SkeletalMesh|Animation"))
 	class USkeleton* Skeleton;
 
 	/** If checked, create new PhysicsAsset if it doesn't have it */
-	UPROPERTY(EditAnywhere, config, Category=SkeletalMesh)
+	UPROPERTY(EditAnywhere, AdvancedDisplay, config, Category=Mesh, meta=(ImportType="SkeletalMesh"))
 	uint32 bCreatePhysicsAsset:1;
 
 	/** If this is set, use this PhysicsAsset. It is possible bCreatePhysicsAsset == false, and PhysicsAsset == NULL. It is possible they do not like to create anything. */
-	UPROPERTY(EditAnywhere, Category=SkeletalMesh)
+	UPROPERTY(EditAnywhere, AdvancedDisplay, Category=Mesh, meta=(ImportType="SkeletalMesh", editcondition="!bCreatePhysicsAsset"))
 	class UPhysicsAsset* PhysicsAsset;
 
 	/** True to import animations from the FBX File */
-	UPROPERTY(EditAnywhere, config, Category=SkeletalMesh)
+	UPROPERTY(EditAnywhere, config, Category=Animation, meta=(ImportType="SkeletalMesh|Animation"))
 	uint32 bImportAnimations:1;
 
-	/** If you import animation and if you'd like to chnage name, please type here **/
-	UPROPERTY(EditAnywhere, Category=SkeletalMesh, meta=(DisplayName = "Animation Name", editcondition = "bImportAnimations")) 
+	/** Override for the name of the animation to import **/
+	UPROPERTY(EditAnywhere, AdvancedDisplay, Category=Animation, meta=(editcondition="bImportAnimations", ImportType="SkeletalMesh")) 
 	FString AnimationName;
 
-	/** Enables importing of 'rigid skeletalmesh' (unskinned, hierarchy-based animation) from this FBX file */
-	UPROPERTY(EditAnywhere, config, Category=SkeletalMesh, meta=(editcondition = "bImportAnimations"))
+	/** Enables importing of 'rigid skeletalmesh' (unskinned, hierarchy-based animation) from this FBX file, no longer shown, used behind the scenes */
+	UPROPERTY()
 	uint32 bImportRigidMesh:1;
 
 	/** Enable this option to use default sample rate for the imported animation at 30 frames per second */
-	UPROPERTY(EditAnywhere, config, Category=SkeletalMesh, meta=(editcondition = "bImportAnimations", ToolTip="If enabled, samples all animation curves to 30 FPS"))
+	UPROPERTY(EditAnywhere, AdvancedDisplay, config, Category=Animation, meta=(editcondition = "bImportAnimations", ToolTip="If enabled, samples all animation curves to 30 FPS", ImportType="SkeletalMesh"))
 	uint32 bUseDefaultSampleRate:1;
 
 	/** Whether to automatically create Unreal materials for materials found in the FBX scene */
-	UPROPERTY(EditAnywhere, config, Category=Materials)
+	UPROPERTY(EditAnywhere, config, Category = Material, meta = (OBJRestrict="true"))
 	uint32 bImportMaterials:1;
 
 	/** The option works only when option "Import UMaterial" is OFF. If "Import UMaterial" is ON, textures are always imported. */
-	UPROPERTY(EditAnywhere, config, Category=Materials)
+	UPROPERTY(EditAnywhere, config, Category=Material, meta=(OBJRestrict="true"))
 	uint32 bImportTextures:1;
 
 	/** Import data used when importing static meshes */
-	UPROPERTY(EditAnywhere, Instanced, Category=StaticMesh)
+	UPROPERTY(EditAnywhere, Instanced, Category = Mesh, meta=(ImportType = "StaticMesh"))
 	class UFbxStaticMeshImportData* StaticMeshImportData;
 
 	/** Import data used when importing skeletal meshes */
-	UPROPERTY(EditAnywhere, Instanced, Category=SkeletalMesh)
+	UPROPERTY(EditAnywhere, Instanced, Category=Mesh, meta=(ImportType = "SkeletalMesh"))
 	class UFbxSkeletalMeshImportData* SkeletalMeshImportData;
 
 	/** Import data used when importing animations */
-	UPROPERTY(EditAnywhere, Instanced, Category=Animation)
+	UPROPERTY(EditAnywhere, Instanced, Category=Animation, meta=(editcondition="bImportAnimations", ImportType = "Animation"))
 	class UFbxAnimSequenceImportData* AnimSequenceImportData;
 
 	/** Type of asset to import from the FBX file */
-	UPROPERTY(EditAnywhere, config, Category=ImportSettings)
+	UPROPERTY(EditAnywhere, AdvancedDisplay, config, Category=Animation, meta=(editcondition="bImportAnimations", ImportType = "Animation"))
 	bool	bPreserveLocalTransform;
 
 	/** Import data used when importing textures */
-	UPROPERTY(EditAnywhere, Instanced, Category=Textures)
+	UPROPERTY(EditAnywhere, Instanced, Category=Material)
 	class UFbxTextureImportData* TextureImportData;
 
 	// Begin UObject Interface
diff --git a/Engine/Source/Editor/UnrealEd/Classes/Factories/FbxMeshImportData.h b/Engine/Source/Editor/UnrealEd/Classes/Factories/FbxMeshImportData.h
--- a/Engine/Source/Editor/UnrealEd/Classes/Factories/FbxMeshImportData.h
+++ b/Engine/Source/Editor/UnrealEd/Classes/Factories/FbxMeshImportData.h
@@ -21,10 +21,12 @@ class UFbxMeshImportData : public UFbxAssetImportData
 	GENERATED_UCLASS_BODY()
 
 	/** Enables importing of mesh LODs from FBX LOD groups, if present in the FBX file */
-	UPROPERTY(EditAnywhere, config, Category=ImportSettings, meta=(ToolTip="If enabled, creates LOD models for Unreal meshes from LODs in the import file; If not enabled, only the base mesh from the LOD group is imported"))
+	UPROPERTY(EditAnywhere, AdvancedDisplay, config, Category=ImportSettings, meta=(OBJRestrict="true", ImportType="Mesh", ToolTip="If enabled, creates LOD models for Unreal meshes from LODs in the import file; If not enabled, only the base mesh from the LOD group is imported"))
 	uint32 bImportMeshLODs:1;
 
 	/** Enabling this option will read the tangents(tangent,binormal,normal) from FBX file instead of generating them automatically. */
-	UPROPERTY(EditAnywhere, config, Category=ImportSettings)
+	UPROPERTY(EditAnywhere, AdvancedDisplay, config, Category=ImportSettings, meta=(ImportType="Mesh"))
 	TEnumAsByte<enum EFBXNormalImportMethod> NormalImportMethod;
+
+	bool CanEditChange( const UProperty* InProperty ) const;
 };
\ No newline at end of file
diff --git a/Engine/Source/Editor/UnrealEd/Classes/Factories/FbxSkeletalMeshImportData.h b/Engine/Source/Editor/UnrealEd/Classes/Factories/FbxSkeletalMeshImportData.h
--- a/Engine/Source/Editor/UnrealEd/Classes/Factories/FbxSkeletalMeshImportData.h
+++ b/Engine/Source/Editor/UnrealEd/Classes/Factories/FbxSkeletalMeshImportData.h
@@ -12,29 +12,31 @@ class UFbxSkeletalMeshImportData : public UFbxMeshImportData
 	GENERATED_UCLASS_BODY()
 
 	/** Enable this option to update Skeleton (of the mesh)'s reference pose. Mesh's reference pose is always updated.  */
-	UPROPERTY(EditAnywhere, Category=ImportSettings, meta=(ToolTip="If enabled, update the Skeleton (of the mesh being imported)'s reference pose."))
+	UPROPERTY(EditAnywhere, AdvancedDisplay, Category=ImportSettings, meta=(ImportType="SkeletalMesh", ToolTip="If enabled, update the Skeleton (of the mesh being imported)'s reference pose."))
 	uint32 bUpdateSkeletonReferencePose:1;
 
 	/** Enable this option to use frame 0 as reference pose */
-	UPROPERTY(EditAnywhere, config, Category=ImportSettings)
+	UPROPERTY(EditAnywhere, AdvancedDisplay, config, Category=ImportSettings, meta=(ImportType="SkeletalMesh"))
 	uint32 bUseT0AsRefPose:1;
 
 	/** If checked, triangles with non-matching smoothing groups will be physically split. */
-	UPROPERTY(EditAnywhere, config, Category=ImportSettings)
+	UPROPERTY(EditAnywhere, AdvancedDisplay, config, Category=ImportSettings, meta=(ImportType="SkeletalMesh"))
 	uint32 bPreserveSmoothingGroups:1;
 
 	/** If checked, meshes nested in bone hierarchies will be imported instead of being converted to bones. */
-	UPROPERTY(EditAnywhere, config, Category=ImportSettings)
+	UPROPERTY(EditAnywhere, AdvancedDisplay, config, Category=ImportSettings, meta=(ImportType="SkeletalMesh"))
 	uint32 bImportMeshesInBoneHierarchy:1;
 
 	/** True to import morph target meshes from the FBX file */
-	UPROPERTY(EditAnywhere, config, Category=ImportSettings, meta=(ToolTip="If enabled, creates Unreal morph objects for the imported meshes"))
+	UPROPERTY(EditAnywhere, AdvancedDisplay, config, Category=ImportSettings, meta=(ImportType="SkeletalMesh", ToolTip="If enabled, creates Unreal morph objects for the imported meshes"))
 	uint32 bImportMorphTargets:1;
 
 	/** If checked, do not filter same vertices. Keep all vertices even if they have exact same properties*/
-	UPROPERTY()
+	UPROPERTY(EditAnywhere, AdvancedDisplay, config, Category=ImportSettings, meta=(ImportType="SkeletalMesh"))
 	uint32 bKeepOverlappingVertices:1;
 
 	/** Gets or creates fbx import data for the specified skeletal mesh */
 	static UFbxSkeletalMeshImportData* GetImportDataForSkeletalMesh(USkeletalMesh* SkeletalMesh, UFbxSkeletalMeshImportData* TemplateForCreation);
+
+	bool CanEditChange( const UProperty* InProperty ) const;
 };
\ No newline at end of file
diff --git a/Engine/Source/Editor/UnrealEd/Classes/Factories/FbxStaticMeshImportData.h b/Engine/Source/Editor/UnrealEd/Classes/Factories/FbxStaticMeshImportData.h
--- a/Engine/Source/Editor/UnrealEd/Classes/Factories/FbxStaticMeshImportData.h
+++ b/Engine/Source/Editor/UnrealEd/Classes/Factories/FbxStaticMeshImportData.h
@@ -13,30 +13,32 @@ class UFbxStaticMeshImportData : public UFbxMeshImportData
 	GENERATED_UCLASS_BODY()
 
 	/** For static meshes, enabling this option will combine all meshes in the FBX into a single monolithic mesh in Unreal */
-	UPROPERTY(EditAnywhere, Category=ImportSettings)
+	UPROPERTY(EditAnywhere, AdvancedDisplay, Category=ImportSettings, meta=(ImportType="StaticMesh"))
 	FName StaticMeshLODGroup;
 
 	/** If true, will replace the vertex colors on an existing mesh with the vertex colors from the FBX file */
-	UPROPERTY(EditAnywhere, config, Category=ImportSettings, meta=(ToolTip="If enabled, existing vertex colors will be replaced with ones from the imported file"))
+	UPROPERTY(EditAnywhere, AdvancedDisplay, config, Category=ImportSettings, meta=(OBJRestrict="true", ImportType="StaticMesh", ToolTip="If enabled, existing vertex colors will be replaced with ones from the imported file"))
 	uint32 bReplaceVertexColors:1;
 
 	/** Disabling this option will keep degenerate triangles found.  In general you should leave this option on. */
-	UPROPERTY(EditAnywhere, Category=ImportSettings)
+	UPROPERTY(EditAnywhere, AdvancedDisplay, Category = ImportSettings, meta = (ImportType = "StaticMesh"))
 	uint32 bRemoveDegenerates:1;
 	
-	UPROPERTY(EditAnywhere, config, Category=ImportSettings)
+	UPROPERTY(EditAnywhere, config, AdvancedDisplay, Category=ImportSettings, meta=(ImportType="StaticMesh"))
 	uint32 bGenerateLightmapUVs:1;
 
 	/** If checked, one convex hull per UCX_ prefixed collision mesh will be generated instead of decomposing into multiple hulls */
-	UPROPERTY(EditAnywhere, config, Category=ImportSettings)
+	UPROPERTY(EditAnywhere, AdvancedDisplay, config, Category=ImportSettings, meta=(OBJRestrict="true", ImportType="StaticMesh"))
 	uint32 bOneConvexHullPerUCX:1;
 
 	/** If checked, collision will automatically be generated (ignored if custom collision is imported or used). */
-	UPROPERTY(EditAnywhere, config, Category = ImportSettings)
+	UPROPERTY(EditAnywhere, config, Category = ImportSettings, meta=(OBJRestrict="true", ImportType="StaticMesh"))
 	uint32 bAutoGenerateCollision : 1;
 
 	/** Gets or creates fbx import data for the specified static mesh */
 	static UFbxStaticMeshImportData* GetImportDataForStaticMesh(UStaticMesh* StaticMesh, UFbxStaticMeshImportData* TemplateForCreation);
+
+	bool CanEditChange( const UProperty* InProperty ) const;
 };
 
 
diff --git a/Engine/Source/Editor/UnrealEd/Classes/Factories/FbxTextureImportData.h b/Engine/Source/Editor/UnrealEd/Classes/Factories/FbxTextureImportData.h
--- a/Engine/Source/Editor/UnrealEd/Classes/Factories/FbxTextureImportData.h
+++ b/Engine/Source/Editor/UnrealEd/Classes/Factories/FbxTextureImportData.h
@@ -7,11 +7,13 @@
  * Import data and options used when importing any mesh from FBX
  */
 UCLASS(AutoExpandCategories=(Texture))
-class UFbxTextureImportData : public UFbxAssetImportData
+class UNREALED_API UFbxTextureImportData : public UFbxAssetImportData
 {
 	GENERATED_UCLASS_BODY()
 
 	/** If either importing of textures (or materials) is enabled, this option will cause normal map values to be inverted */
-	UPROPERTY(EditAnywhere, config, Category=ImportSettings)
+	UPROPERTY(EditAnywhere, AdvancedDisplay, config, Category=ImportSettings, meta=(OBJRestrict="true"))
 	uint32 bInvertNormalMaps:1;
+
+	bool CanEditChange( const UProperty* InProperty ) const;
 };
\ No newline at end of file
diff --git a/Engine/Source/Editor/UnrealEd/Private/Fbx/FbxAnimSequenceImportData.cpp b/Engine/Source/Editor/UnrealEd/Private/Fbx/FbxAnimSequenceImportData.cpp
--- a/Engine/Source/Editor/UnrealEd/Private/Fbx/FbxAnimSequenceImportData.cpp
+++ b/Engine/Source/Editor/UnrealEd/Private/Fbx/FbxAnimSequenceImportData.cpp
@@ -29,4 +29,16 @@ UFbxAnimSequenceImportData* UFbxAnimSequenceImportData::GetImportDataForAnimSequ
 	}
 
 	return ImportData;
-}
\ No newline at end of file
+}
+
+bool UFbxAnimSequenceImportData::CanEditChange(const UProperty* InProperty) const
+{
+	bool bMutable = Super::CanEditChange(InProperty);
+	UObject* Outer = GetOuter();
+	if(Outer && bMutable)
+	{
+		// Let the FbxImportUi object handle the editability of our properties
+		bMutable = Outer->CanEditChange(InProperty);
+	}
+	return bMutable;
+}
diff --git a/Engine/Source/Editor/UnrealEd/Private/Fbx/FbxAssetImportData.cpp b/Engine/Source/Editor/UnrealEd/Private/Fbx/FbxAssetImportData.cpp
--- a/Engine/Source/Editor/UnrealEd/Private/Fbx/FbxAssetImportData.cpp
+++ b/Engine/Source/Editor/UnrealEd/Private/Fbx/FbxAssetImportData.cpp
@@ -4,6 +4,9 @@
 
 UFbxAssetImportData::UFbxAssetImportData(const class FPostConstructInitializeProperties& PCIP)
 	: Super(PCIP)
+	, ImportTranslation(0)
+	, ImportRotation(0)
+	, ImportUniformScale(1.0f)
 {
 	
 }
\ No newline at end of file
diff --git a/Engine/Source/Editor/UnrealEd/Private/Fbx/FbxFactory.cpp b/Engine/Source/Editor/UnrealEd/Private/Fbx/FbxFactory.cpp
--- a/Engine/Source/Editor/UnrealEd/Private/Fbx/FbxFactory.cpp
+++ b/Engine/Source/Editor/UnrealEd/Private/Fbx/FbxFactory.cpp
@@ -104,6 +104,7 @@ bool UFbxFactory::DetectImportType(const FString& InFilename)
 	else
 	{
 		ImportUI->MeshTypeToImport = EFBXImportType(ImportType);
+		ImportUI->OriginalImportType = ImportUI->MeshTypeToImport;
 	}
 	
 	return true;
@@ -587,9 +588,20 @@ bool UFbxImportUI::CanEditChange( const UProperty* InProperty ) const
 	bool bIsMutable = Super::CanEditChange( InProperty );
 	if( bIsMutable && InProperty != NULL )
 	{
-		if( InProperty->GetFName() == TEXT( "AnimationName" ) )
+		FName PropName = InProperty->GetFName();
+
+		if(PropName == TEXT("StartFrame") || PropName == TEXT("EndFrame"))
+		{
+			bIsMutable = AnimSequenceImportData->AnimationLength == FBXALIT_SetRange && bImportAnimations;
+		}
+		else if(PropName == TEXT("bImportCustomAttribute") || PropName == TEXT("AnimationLength"))
+		{
+			bIsMutable = bImportAnimations;
+		}
+
+		if(bIsObjImport && InProperty->GetBoolMetaData(TEXT("OBJRestrict")))
 		{
-			bIsMutable = (MeshTypeToImport == FBXIT_SkeletalMesh);
+			bIsMutable = false;
 		}
 	}
 
diff --git a/Engine/Source/Editor/UnrealEd/Private/Fbx/FbxMainImport.cpp b/Engine/Source/Editor/UnrealEd/Private/Fbx/FbxMainImport.cpp
--- a/Engine/Source/Editor/UnrealEd/Private/Fbx/FbxMainImport.cpp
+++ b/Engine/Source/Editor/UnrealEd/Private/Fbx/FbxMainImport.cpp
@@ -83,6 +83,15 @@ FBXImportOptions* GetImportOptions( UnFbx::FFbxImporter* FbxImporter, UFbxImport
 			ImportUI->PhysicsAsset = NULL;
 		}
 
+		if(bForceImportType)
+		{
+			ImportUI->MeshTypeToImport = ImportType;
+			ImportUI->OriginalImportType = ImportType;
+		}
+
+		ImportUI->bImportAsSkeletal = ImportUI->MeshTypeToImport == FBXIT_SkeletalMesh;
+		ImportUI->bIsObjImport = bIsObjFormat;
+
 		TSharedPtr<SWindow> ParentWindow;
 
 		if( FModuleManager::Get().IsModuleLoaded( "MainFrame" ) )
@@ -172,23 +181,28 @@ void ApplyImportUIToImportOptions(UFbxImportUI* ImportUI, FBXImportOptions& InOu
 	InOutImportOptions.bConvertScene = ImportUI->bConvertScene;
 	InOutImportOptions.bImportAnimations = ImportUI->bImportAnimations;
 	InOutImportOptions.SkeletonForAnimation = ImportUI->Skeleton;
+
 	if ( ImportUI->MeshTypeToImport == FBXIT_StaticMesh )
 	{
-		InOutImportOptions.NormalImportMethod = ImportUI->StaticMeshImportData->NormalImportMethod;
+		UFbxStaticMeshImportData* StaticMeshData	= ImportUI->StaticMeshImportData;
+		InOutImportOptions.NormalImportMethod		= StaticMeshData->NormalImportMethod;
 	}
 	else if ( ImportUI->MeshTypeToImport == FBXIT_SkeletalMesh )
 	{
-		InOutImportOptions.NormalImportMethod = ImportUI->SkeletalMeshImportData->NormalImportMethod;
+		UFbxSkeletalMeshImportData* SkeletalMeshData	= ImportUI->SkeletalMeshImportData;
+		InOutImportOptions.NormalImportMethod			= SkeletalMeshData->NormalImportMethod;
 	}
 	else
 	{
+		UFbxAnimSequenceImportData* AnimData	= ImportUI->AnimSequenceImportData;
 		InOutImportOptions.NormalImportMethod = FBXNIM_ComputeNormals;
 	}
+
 	// only re-sample if they don't want to use default sample rate
 	InOutImportOptions.bResample = ImportUI->bUseDefaultSampleRate==false;
 	InOutImportOptions.bImportMorph = ImportUI->SkeletalMeshImportData->bImportMorphTargets;
 	InOutImportOptions.bUpdateSkeletonReferencePose = ImportUI->SkeletalMeshImportData->bUpdateSkeletonReferencePose;
-	InOutImportOptions.bImportRigidMesh = ImportUI->bImportRigidMesh;
+	InOutImportOptions.bImportRigidMesh = ImportUI->OriginalImportType == FBXIT_StaticMesh && ImportUI->MeshTypeToImport == FBXIT_SkeletalMesh;
 	InOutImportOptions.bUseT0AsRefPose = ImportUI->SkeletalMeshImportData->bUseT0AsRefPose;
 	InOutImportOptions.bPreserveSmoothingGroups = ImportUI->SkeletalMeshImportData->bPreserveSmoothingGroups;
 	InOutImportOptions.bKeepOverlappingVertices = ImportUI->SkeletalMeshImportData->bKeepOverlappingVertices;
diff --git a/Engine/Source/Editor/UnrealEd/Private/Fbx/FbxMeshImportData.cpp b/Engine/Source/Editor/UnrealEd/Private/Fbx/FbxMeshImportData.cpp
--- a/Engine/Source/Editor/UnrealEd/Private/Fbx/FbxMeshImportData.cpp
+++ b/Engine/Source/Editor/UnrealEd/Private/Fbx/FbxMeshImportData.cpp
@@ -6,4 +6,16 @@ UFbxMeshImportData::UFbxMeshImportData(const class FPostConstructInitializePrope
 	: Super(PCIP)
 {
 	NormalImportMethod = FBXNIM_ComputeNormals;
+}
+
+bool UFbxMeshImportData::CanEditChange(const UProperty* InProperty) const
+{
+	bool bMutable = Super::CanEditChange(InProperty);
+	UObject* Outer = GetOuter();
+	if(Outer && bMutable)
+	{
+		// Let the parent object handle the editability of our properties
+		bMutable = Outer->CanEditChange(InProperty);
+	}
+	return bMutable;
 }
\ No newline at end of file
diff --git a/Engine/Source/Editor/UnrealEd/Private/Fbx/FbxOptionWindow.cpp b/Engine/Source/Editor/UnrealEd/Private/Fbx/FbxOptionWindow.cpp
--- a/Engine/Source/Editor/UnrealEd/Private/Fbx/FbxOptionWindow.cpp
+++ b/Engine/Source/Editor/UnrealEd/Private/Fbx/FbxOptionWindow.cpp
@@ -3,282 +3,10 @@
 #include "FbxOptionWindow.h"
 #include "Editor/ContentBrowser/Public/ContentBrowserModule.h"
 #include "IDocumentation.h"
+#include "SKismetInspector.h"
 
 #define LOCTEXT_NAMESPACE "FBXOption"
 
-DECLARE_DELEGATE_OneParam( FOnImportTypeChanged, EFBXImportType );
-
-class SImportTypeButton : public SCompoundWidget
-{
-public:
-	SLATE_BEGIN_ARGS( SImportTypeButton )
-	:	_DefaultImportType(FBXIT_StaticMesh)
-	,	_ForceImportType( false )
-	,	_OnSelectionChanged()
-	{}
-		SLATE_ARGUMENT( EFBXImportType, DefaultImportType )
-		SLATE_ARGUMENT( bool, ForceImportType )
-		SLATE_EVENT( FOnImportTypeChanged, OnSelectionChanged )
-	SLATE_END_ARGS()
-
-	void Construct( const FArguments& InArgs )
-	{
-		DefaultImportType = InArgs._DefaultImportType;
-		CurrentChoice = InArgs._DefaultImportType;
-		OnSelectionChanged = InArgs._OnSelectionChanged;
-		bForceImportType = InArgs._ForceImportType;
-
-		const UEnum* ImportTypeEnum = FindObject<UEnum>(ANY_PACKAGE, TEXT("EFBXImportType"), true);
-		check(ImportTypeEnum != nullptr);
-
-		FText AutoDetected = FText::Format(LOCTEXT("FbxOptionWindow_AutoDetectedPrefix", "Auto-detected: {0}"), ImportTypeEnum->GetEnumText(CurrentChoice));
-
-		this->ChildSlot
-		[
-			SNew(SBox)
-			// TODO: Ideally I would like to specify that the VerticalBox slot containing the radio buttons defines the width of the VerticalBox.
-			// While that's not an option, the width is fixed, so that the warning STextBlock slot doesn't expand the window to its length.
-			.WidthOverride(400.0f)
-			[
-				SNew(SVerticalBox)
-
-				+SVerticalBox::Slot().AutoHeight().Padding(0, 2)
-				[
-					SNew(SHorizontalBox)
-					
-					+SHorizontalBox::Slot()
-					.Padding(0, 10)
-					[
-						SNew(STextBlock)
-						.TextStyle(FEditorStyle::Get(), "FBXLargeFont")
-						.Text(LOCTEXT("FbxOptionWindow_SelectType", "Select Asset Type"))
-					]
-
-
-					+SHorizontalBox::Slot()
-					.HAlign(HAlign_Left)
-					.FillWidth(1)
-					[
-						IDocumentation::Get()->CreateAnchor(FString("Engine/Content/FBX/ImportOptions"))
-					]
-				]
-
-				+SVerticalBox::Slot().AutoHeight().Padding(0, 5)
-				[
-					SNew(STextBlock)
-					.Text(AutoDetected)
-				]
-
-				+SVerticalBox::Slot().AutoHeight()
-				[
-					SNew(SUniformGridPanel)
-					+ SUniformGridPanel::Slot(0, 0)
-					[
-						CreateRadioButton(*ImportTypeEnum, FBXIT_StaticMesh, FEditorStyle::GetBrush("FBXIcon.StaticMesh"), 0)
-					]
-					+ SUniformGridPanel::Slot(1, 0)
-					[
-						CreateRadioButton(*ImportTypeEnum, FBXIT_SkeletalMesh, FEditorStyle::GetBrush("FBXIcon.SkeletalMesh"), 1)
-					]
-					+ SUniformGridPanel::Slot(2, 0)
-					[
-						CreateRadioButton(*ImportTypeEnum, FBXIT_Animation, FEditorStyle::GetBrush("FBXIcon.Animation"), 2)
-					]
-				]
-
-				+SVerticalBox::Slot().AutoHeight().Padding(0, 5)
-				[
-					SNew(STextBlock)
-					.Text(LOCTEXT("FbxOptionWindow_TypeMismatch", "Warning: Asset does not appear to correspond to selected type, and may not import correctly."))
-					.ColorAndOpacity(FColor::Yellow)
-					.AutoWrapText(true)
-					.Visibility(this, &SImportTypeButton::IsWarningVisible)
-				]
-			]
-		];
-	}
-
-private:
-
-	TSharedRef<SWidget> CreateRadioButton( const UEnum& Enum, EFBXImportType RadioButtonChoice, const FSlateBrush* Brush, int32 Position)
-	{
-		return
-			SNew(SCheckBox)
-			.Style(FEditorStyle::Get(), (Position == 0) ? "Property.ToggleButton.Start" : ((Position==1)? "Property.ToggleButton.Middle" : "Property.ToggleButton.End") )
-			.IsChecked( this, &SImportTypeButton::IsRadioChecked, RadioButtonChoice )
-			.OnCheckStateChanged( this, &SImportTypeButton::OnRadioChanged, RadioButtonChoice )
-			.ToolTipText(Enum.GetToolTipText(RadioButtonChoice))
-			.IsEnabled( !bForceImportType || RadioButtonChoice == DefaultImportType )
-			[
-				SNew(SHorizontalBox)
-
-				+ SHorizontalBox::Slot()
-				.AutoWidth()
-				.VAlign(VAlign_Center)
-				.Padding(3, 2)
-				[
-					SNew(SImage)
-					.Image(Brush)
-				]
-
-				+ SHorizontalBox::Slot()
-				.FillWidth(1.0f)
-				.VAlign(VAlign_Center)
-				.HAlign(HAlign_Center)
-				.Padding(6, 2)
-				[
-					SNew(STextBlock)
-					.Text(Enum.GetEnumText(RadioButtonChoice))
-					.ColorAndOpacity(FLinearColor(0.72f, 0.72f, 0.72f, 1.f))
-				]
-			];
-	}
-
-	ESlateCheckBoxState::Type IsRadioChecked( EFBXImportType ButtonId ) const
-	{
-		return (CurrentChoice == ButtonId)
-			? ESlateCheckBoxState::Checked
-			: ESlateCheckBoxState::Unchecked;
-	}
-
-	EVisibility IsWarningVisible() const
-	{
-		return (CurrentChoice == DefaultImportType)
-			? EVisibility::Collapsed
-			: EVisibility::Visible;
-	}
-
-	void OnRadioChanged( ESlateCheckBoxState::Type NewRadioState, EFBXImportType RadioThatChanged )
-	{
-		if (NewRadioState == ESlateCheckBoxState::Checked)
-		{
-			CurrentChoice = RadioThatChanged;
-
-			if (OnSelectionChanged.IsBound())
-			{
-				OnSelectionChanged.Execute(CurrentChoice);
-			}
-		}
-	}
-
-	EFBXImportType DefaultImportType;
-	EFBXImportType CurrentChoice;
-	FOnImportTypeChanged OnSelectionChanged;
-	bool bForceImportType;
-};
-
-DECLARE_DELEGATE_OneParam( FOnAnimImportLengthOptionChanged, EFBXAnimationLengthImportType );
-
-class SAnimImportLengthOption : public SCompoundWidget
-{
-public:
-	SLATE_BEGIN_ARGS( SAnimImportLengthOption )
-		:	_AnimationLengthOption(FBXALIT_ExportedTime)
-		,	_OnSelectionChanged()
-		,	_OnTextCommittedRange1()
-		,	_OnTextCommittedRange2()
-	{}
-		SLATE_ARGUMENT( EFBXAnimationLengthImportType, AnimationLengthOption )
-		SLATE_EVENT( FOnAnimImportLengthOptionChanged, OnSelectionChanged )
-		SLATE_EVENT( FOnTextCommitted,	OnTextCommittedRange1)
-		SLATE_EVENT( FOnTextCommitted,	OnTextCommittedRange2)
-	SLATE_END_ARGS()
-
-	void Construct( const FArguments& InArgs )
-	{
-		CurrentChoice = InArgs._AnimationLengthOption;
-		OnSelectionChanged = InArgs._OnSelectionChanged;
-
-		this->ChildSlot
-		[
-			SNew(SVerticalBox)
-			+SVerticalBox::Slot().AutoHeight() .Padding(1)
-			[
-				CreateRadioButton( LOCTEXT("AnimImportLengthOption_ExportTime", "Exported Time"), FBXALIT_ExportedTime )
-			]
-			+SVerticalBox::Slot().AutoHeight() .Padding(1)
-			[
-				CreateRadioButton( LOCTEXT("AnimImportLengthOption_AnimTime", "Animated Time"), FBXALIT_AnimatedKey )
-			]
-			+SVerticalBox::Slot().AutoHeight() .Padding(1)
-			[
-				SNew(SHorizontalBox)
-
-				+SHorizontalBox::Slot().AutoWidth() .Padding(2)
-				[
-					CreateRadioButton( LOCTEXT("AnimImportLengthOption_SetRange", "Set Range"), FBXALIT_SetRange )
-				]
-
-				+SHorizontalBox::Slot().AutoWidth() .Padding(2) 
-				[
-					SNew(SEditableTextBox)
-					.MinDesiredWidth(20)
-					.OnTextCommitted(InArgs._OnTextCommittedRange1)
-					.IsEnabled(this, &SAnimImportLengthOption::CanEnterAnimationRange)
-				]
-
-				+SHorizontalBox::Slot().AutoWidth() .Padding(2)
-				[
-					SNew(STextBlock)
-					.Text(LOCTEXT("RangeSeparator", "-"))
-				]
-				+SHorizontalBox::Slot().AutoWidth() .Padding(2)
-				[
-					SNew(SEditableTextBox)
-					.MinDesiredWidth(20)
-					.OnTextCommitted(InArgs._OnTextCommittedRange2)
-					.IsEnabled(this, &SAnimImportLengthOption::CanEnterAnimationRange)
-				]
-
-				+SHorizontalBox::Slot().AutoWidth() .Padding(2)
-				[
-					SNew(STextBlock)
-					.Text(LOCTEXT("AnimImportLengthOption_InFrameNumber", " In Frame Number"))
-				]
-			]
-		];
-	}
-
-private:
-
-	TSharedRef<SWidget> CreateRadioButton( const FText& RadioText, EFBXAnimationLengthImportType RadioButtonChoice )
-	{
-		return
-			SNew(SCheckBox)
-			.Style(FEditorStyle::Get(), "RadioButton")
-			.IsChecked( this, &SAnimImportLengthOption::IsRadioChecked, RadioButtonChoice )
-			.OnCheckStateChanged( this, &SAnimImportLengthOption::OnRadioChanged, RadioButtonChoice )
-			[
-				SNew(STextBlock).Text(RadioText)
-			];
-	}
-
-	ESlateCheckBoxState::Type IsRadioChecked( EFBXAnimationLengthImportType ButtonId ) const
-	{
-		return (CurrentChoice == ButtonId)
-			? ESlateCheckBoxState::Checked
-			: ESlateCheckBoxState::Unchecked;
-	}
-
-	bool CanEnterAnimationRange() const { return CurrentChoice == FBXALIT_SetRange; }
-
-	void OnRadioChanged( ESlateCheckBoxState::Type NewRadioState, EFBXAnimationLengthImportType RadioThatChanged )
-	{
-		if (NewRadioState == ESlateCheckBoxState::Checked)
-		{
-			CurrentChoice = RadioThatChanged;
-
-			if (OnSelectionChanged.IsBound())
-			{
-				OnSelectionChanged.Execute(CurrentChoice);
-			}
-		}
-	}
-
-	EFBXAnimationLengthImportType CurrentChoice;
-	FOnAnimImportLengthOptionChanged OnSelectionChanged;
-};
-
 void SFbxOptionWindow::Construct(const FArguments& InArgs)
 {
 	ImportUI = InArgs._ImportUI;
@@ -287,1237 +15,87 @@ void SFbxOptionWindow::Construct(const FArguments& InArgs)
 
 	check (ImportUI);
 
-	UStaticMesh::GetLODGroups(StaticMeshLODGroupNames);
-	UStaticMesh::GetLODGroupsDisplayNames(StaticMeshLODGroupDisplayNames);
-
-	for (int32 GroupIndex = 0; GroupIndex < StaticMeshLODGroupNames.Num(); ++GroupIndex)
-	{
-		StaticMeshLODGroups.Add(MakeShareable(new FString(StaticMeshLODGroupDisplayNames[GroupIndex].ToString())));
-	}
-
-	bForceImportType = InArgs._ForcedImportType.IsSet();
-
-	// Force the import type
-	if( bForceImportType )
-	{
-		ImportUI->MeshTypeToImport = InArgs._ForcedImportType.GetValue();
-	}
-
-	// clear Import UI bug
-	ImportUI->AnimSequenceImportData->AnimationLength = FBXALIT_ExportedTime;
-	ImportUI->AnimSequenceImportData->StartFrame = 0;
-	ImportUI->AnimSequenceImportData->EndFrame = 0;
-
+	TSharedPtr<SBox> InspectorBox;
 	this->ChildSlot
 	[
-		SNew(SBorder)
-		.BorderImage(FEditorStyle::GetBrush("ToolPanel.GroupBorder"))
+		SNew(SVerticalBox)
+		+ SVerticalBox::Slot()
+		.AutoHeight()
+		.Padding(2)
 		[
-		
-			SNew(SVerticalBox)
-
-			// first 3 radio box for staticmesh/skeletalmesh/animation		
-			+SVerticalBox::Slot()
-			.AutoHeight()
-			.Padding(5)
-			[
-				SNew(STextBlock)
-				.Font(FEditorStyle::GetFontStyle("CurveEd.InfoFont"))
-				.Text(InArgs._FullPath)
-			]
-
-			// first 3 radio box for staticmesh/skeletalmesh/animation		
-			+SVerticalBox::Slot()
-			.AutoHeight()
-			.Padding(5)
-			[
-				SNew(SImportTypeButton)
-				.DefaultImportType(ImportUI->MeshTypeToImport)
-				.OnSelectionChanged(this, &SFbxOptionWindow::SetImportType)
-				.ForceImportType( bForceImportType )
-			]
-
-			+SVerticalBox::Slot()
-			.AutoHeight()
-			.Padding(5)
-			[
-				SNew(SSeparator)
-			]
-
-			+SVerticalBox::Slot()
-			.AutoHeight()
-			.Padding(5)
-			.Expose(CustomBox)
-
-			+SVerticalBox::Slot()
-			.AutoHeight() 
-			.Padding(5)
-			[
-				SNew(SSeparator)
-			]
-
-			+SVerticalBox::Slot()
+			SAssignNew(InspectorBox, SBox)
+			.MaxDesiredHeight(650.0f)
+			.WidthOverride(400.0f)
+		]
+		+ SVerticalBox::Slot()
 			.AutoHeight()
 			.HAlign(HAlign_Right)
 			.Padding(2)
 			[
 				SNew(SUniformGridPanel)
 				.SlotPadding(2)
-				+SUniformGridPanel::Slot(0,0)
+				+ SUniformGridPanel::Slot(0, 0)
 				[
-					SNew(SButton)
-					.HAlign(HAlign_Center)
-					.Text(LOCTEXT("FbxOptionWindow_ImportAll", "Import All"))
-					.ToolTipText( LOCTEXT("FbxOptionWindow_ImportAll_ToolTip", "Import all files with these same settings") )
-					.IsEnabled(this, &SFbxOptionWindow::CanImport)
-					.OnClicked( this, &SFbxOptionWindow::OnImportAll )
+					IDocumentation::Get()->CreateAnchor(FString("Engine/Content/FBX/ImportOptions"))
 				]
-				+SUniformGridPanel::Slot(1,0)
+				+ SUniformGridPanel::Slot(1, 0)
 				[
-					SAssignNew( ImportButton, SButton)
+					SNew(SButton)
 					.HAlign(HAlign_Center)
-					.Text(LOCTEXT("FbxOptionWindow_Import", "Import"))
+					.Text(LOCTEXT("FbxOptionWindow_ImportAll", "Import All"))
+					.ToolTipText(LOCTEXT("FbxOptionWindow_ImportAll_ToolTip", "Import all files with these same settings"))
 					.IsEnabled(this, &SFbxOptionWindow::CanImport)
-					.OnClicked( this, &SFbxOptionWindow::OnImport )
-				]
-				+SUniformGridPanel::Slot(2,0)
-				[
-					SNew(SButton) 
-					.HAlign(HAlign_Center)
-					.Text(LOCTEXT("FbxOptionWindow_Cancel", "Cancel"))
-					.ToolTipText(LOCTEXT("FbxOptionWindow_Cancel_ToolTip", "Cancels importing this FBX file"))
-					.OnClicked( this, &SFbxOptionWindow::OnCancel )
-				]
-			]
-		]
-	];
-
-	if(WidgetWindow.IsValid())
-	{
-		WidgetWindow.Pin()->SetWidgetToFocusOnActivate(ImportButton);
-	}
-
-	RefreshWindow();
-}
-
-void SFbxOptionWindow::RefreshWindow()
-{
-	TSharedPtr<SVerticalBox> VerticalBox;
-
-	(*CustomBox)
-	[
-		SAssignNew( VerticalBox, SVerticalBox )
-	];
-
-	if (ImportUI->MeshTypeToImport == FBXIT_StaticMesh)
-	{
-		VerticalBox->AddSlot()
-		.AutoHeight()
-		[
-			ConstructStaticMeshBasic()
-		];
-
-		VerticalBox->AddSlot()
-		.AutoHeight()
-		.Padding(0, 4)
-		[
-			SNew(SExpandableArea)
-			.InitiallyCollapsed(true)
-			.HeaderContent()
-			[
-				SNew(STextBlock)
-				.Text(LOCTEXT("SFbxOptionWindow_StaticAdvanced", "Advanced"))
-				.TextStyle(FEditorStyle::Get(), "FBXSmallFont")
-			]
-			.BodyContent()
-			[
-				SNew (SVerticalBox)
-
-				+SVerticalBox::Slot()
-				.AutoHeight()
-				[
-					ConstructStaticMeshAdvanced()
-				]
-
-				+SVerticalBox::Slot()
-				.AutoHeight()
-				[
-					ConstructMaterialOption()
-				]
-
-				+SVerticalBox::Slot()
-				.AutoHeight()
-				[
-					ConstructMiscOption()
-				]
-			]
-		];
-	}
-	else if (ImportUI->MeshTypeToImport == FBXIT_SkeletalMesh)
-	{
-		VerticalBox->AddSlot()
-		.AutoHeight()
-		[
-			ConstructSkeletonOptionForMesh()
-		];
-
-		VerticalBox->AddSlot()
-		.AutoHeight()
-		[
-			ConstructSkeletalMeshBasic()
-		];
-
-
-		VerticalBox->AddSlot()
-		.Padding(0, 4)
-		.AutoHeight()
-		[
-			SNew(SExpandableArea)
-			.InitiallyCollapsed(true)
-			.HeaderContent()
-			[
-				SNew(STextBlock)
-				.Text(LOCTEXT("SFbxOptionWindow_StaticAdvanced", "Advanced"))
-				.TextStyle( FEditorStyle::Get(), "FBXSmallFont" ) 
-			]
-			.BodyContent()
-			[
-				SNew (SVerticalBox)
-
-				+SVerticalBox::Slot().AutoHeight()
-				[
-					ConstructSkeletalMeshAdvanced()
-				]
-
-				+SVerticalBox::Slot().AutoHeight()
-				[
-					ConstructMaterialOption()
-				]
-
-				+SVerticalBox::Slot().AutoHeight()
-				[
-					ConstructMiscOption()
+					.OnClicked(this, &SFbxOptionWindow::OnImportAll)
 				]
+				+ SUniformGridPanel::Slot(2, 0)
+					[
+						SAssignNew(ImportButton, SButton)
+						.HAlign(HAlign_Center)
+						.Text(LOCTEXT("FbxOptionWindow_Import", "Import"))
+						.IsEnabled(this, &SFbxOptionWindow::CanImport)
+						.OnClicked(this, &SFbxOptionWindow::OnImport)
+					]
+				+ SUniformGridPanel::Slot(3, 0)
+					[
+						SNew(SButton)
+						.HAlign(HAlign_Center)
+						.Text(LOCTEXT("FbxOptionWindow_Cancel", "Cancel"))
+						.ToolTipText(LOCTEXT("FbxOptionWindow_Cancel_ToolTip", "Cancels importing this FBX file"))
+						.OnClicked(this, &SFbxOptionWindow::OnCancel)
+					]
 			]
-		];
-	}
-	else 	
-	{
-		VerticalBox->AddSlot()
-		.AutoHeight()
-		[
-			ConstructSkeletonOptionForAnim()
-		];
-
-		VerticalBox->AddSlot()
-		.Padding( 0.0f, 3.0f )
-		.AutoHeight()
-		[
-			ConstructAnimationOption()
-		];
-
-		VerticalBox->AddSlot()
-		.Padding(0, 4)
-		.AutoHeight()
-		[
-			SNew(SExpandableArea)
-			.InitiallyCollapsed(true)
-			.HeaderContent()
-			[
-				SNew(STextBlock)
-				.Text(LOCTEXT("SFbxOptionWindow_StaticAdvanced", "Advanced"))
-				.TextStyle(FEditorStyle::Get(), "FBXSmallFont")
-			]
-			.BodyContent()
-			[
-					ConstructMiscOption()
-				]
-		];
-	}
-}
-
-TSharedRef<SWidget> SFbxOptionWindow::ConstructMiscOption()
-{
-	TSharedRef<SVerticalBox> NewBox = SNew(SVerticalBox);
-
-	NewBox->AddSlot().AutoHeight() .Padding(4)
-	[
-		SNew(STextBlock)
-		.TextStyle( FEditorStyle::Get(), "FBXSmallFont" ) 
-		.Text( LOCTEXT("FbxOptionWindow_Misc", "Misc") )
 	];
 
-	NewBox->AddSlot().AutoHeight() .Padding(2)
-	[
-		SNew(SCheckBox)
-		.IsChecked(ImportUI->bOverrideFullName ? ESlateCheckBoxState::Checked : ESlateCheckBoxState::Unchecked)
-		.OnCheckStateChanged(this, &SFbxOptionWindow::SetGeneral_OverrideFullName)
-		.IsEnabled( !bIsObjFormat )
-		[
-			SNew(STextBlock)
-			.Text(LOCTEXT("FbxOptionWindow_NameOverride", "Override FullName"))
-		]
-	];
+	TSharedPtr<SKismetInspector> Inspector;
+	SAssignNew(Inspector, SKismetInspector)
+		.ShowTitleArea(false)
+		.HideNameArea(true);
 
-	NewBox->AddSlot().AutoHeight() .Padding(2)
-	[
-		SNew(SCheckBox)
-		.IsChecked(ImportUI->bConvertScene ? ESlateCheckBoxState::Checked : ESlateCheckBoxState::Unchecked)
-		.OnCheckStateChanged(this, &SFbxOptionWindow::SetGeneral_ConvertScene)
-		.IsEnabled( !bIsObjFormat )
-		[
-			SNew(STextBlock)
-			.Text(LOCTEXT("FbxOptionWindow_ConvertScene", "Convert Scene"))
-		]
-	];
+	InspectorBox->SetContent(Inspector->AsShared());
 
-	return NewBox;
+	Inspector->ShowDetailsForSingleObject(ImportUI);
 }
 
-TSharedRef<SWidget> SFbxOptionWindow::ConstructNormalImportOptions()
+bool SFbxOptionWindow::CanImport()  const
 {
-	NormalImportOptions.Empty();
-	for( int32 NormalMethod = 0; NormalMethod < FBXNIM_MAX; ++NormalMethod )
+	// do test to see if we are ready to import
+	if (ImportUI->MeshTypeToImport == FBXIT_Animation)
 	{
-		if( !bIsObjFormat || NormalMethod != FBXNIM_ImportNormalsAndTangents  )
+		if (ImportUI->Skeleton == NULL || !ImportUI->bImportAnimations)
 		{
-			NormalImportOptions.Add( TSharedPtr<EFBXNormalImportMethod>() );
+			return false;
 		}
 	}
 
-	NormalImportOptions[FBXNIM_ComputeNormals] = MakeShareable( new EFBXNormalImportMethod( FBXNIM_ComputeNormals ) );
-	NormalImportOptions[FBXNIM_ImportNormals] = MakeShareable( new EFBXNormalImportMethod( FBXNIM_ImportNormals ) );
-
-	if( !bIsObjFormat )
+	if (ImportUI->AnimSequenceImportData->AnimationLength == FBXALIT_SetRange)
 	{
-		NormalImportOptions[FBXNIM_ImportNormalsAndTangents] = MakeShareable( new EFBXNormalImportMethod( FBXNIM_ImportNormalsAndTangents ) );
+		if (ImportUI->AnimSequenceImportData->StartFrame > ImportUI->AnimSequenceImportData->EndFrame)
+		{
+			return false;
+		}
 	}
 
-	return 
-	SNew( SHorizontalBox )
-	+SHorizontalBox::Slot()
-	.AutoWidth()
-	.Padding(2)
-	[
-		SNew( SHorizontalBox )
-		+ SHorizontalBox::Slot()
-		.VAlign(VAlign_Center)
-		.AutoWidth()
-		[
-			SNew( STextBlock )
-			.Text(LOCTEXT("FBXOptionWindow_NormalImport", "Normals"))
-		]
-		+ SHorizontalBox::Slot()
-		.VAlign(VAlign_Center)
-		.AutoWidth()
-		.Padding(4)
-		[
-			SNew( SComboBox< TSharedPtr<EFBXNormalImportMethod> > )
-			.ContentPadding(1.f)
-			.ToolTipText(LOCTEXT("FBXOptionWindow_NormalInputMethod", "Options for importing normals and tangents"))
-			.OptionsSource( &NormalImportOptions )
-			.InitiallySelectedItem( NormalImportOptions[ GetCurrentNormalImportMethod() ] )
-			.OnSelectionChanged( this, &SFbxOptionWindow::OnNormalImportMethodChanged )
-			.OnGenerateWidget( this, &SFbxOptionWindow::OnGenerateWidgetForComboItem )
-			[
-				SNew( STextBlock )
-				.Text( this, &SFbxOptionWindow::OnGenerateTextForImportMethod, FBXNIM_MAX )
-			]
-		]
-	];
-}
-
-TSharedRef<SWidget> SFbxOptionWindow::ConstructStaticMeshBasic()
-{
-	TSharedRef<SVerticalBox> NewBox = SNew(SVerticalBox);
-
-	NewBox->AddSlot().AutoHeight() .Padding(2)
-	[
-		SNew(SHorizontalBox)
-
-		+SHorizontalBox::Slot().AutoWidth()
-			.Padding(2)
-			.VAlign(VAlign_Center)
-		[
-			SNew(STextBlock)
-			.Text(LOCTEXT("FbxOptionWindow_StaticMeshLODGroup", "LOD Group"))
-		]
-
-		+SHorizontalBox::Slot()
-			.Padding(2)
-			.FillWidth(1.0f)
-			.VAlign(VAlign_Center)
-		[
-			SNew(STextComboBox)
-			.OptionsSource(&StaticMeshLODGroups)
-			.InitiallySelectedItem(StaticMeshLODGroups[StaticMeshLODGroupNames.Find(NAME_None)])
-			.OnSelectionChanged(this, &SFbxOptionWindow::SetStaticMesh_LODGroup)
-		]
-	];
-
-	NewBox->AddSlot().AutoHeight() .Padding(2)
-	[
-		ConstructNormalImportOptions()
-	];
-
-	return NewBox;
-}
-TSharedRef<SWidget> SFbxOptionWindow::ConstructStaticMeshAdvanced()
-{
-	TSharedRef<SVerticalBox> NewBox = SNew(SVerticalBox);
-
-	NewBox->AddSlot().AutoHeight() .Padding(4)
-	[
-		SNew(STextBlock)
-		.TextStyle( FEditorStyle::Get(), "FBXSmallFont" ) 
-		.Text( LOCTEXT("FbxOptionWindow_StaticMesh", "Mesh"))
-	];
-
-
-	NewBox->AddSlot().AutoHeight().Padding(2)
-	[
-		SNew(SCheckBox)
-		.IsChecked(ImportUI->bCombineMeshes ? ESlateCheckBoxState::Checked : ESlateCheckBoxState::Unchecked)
-		.OnCheckStateChanged(this, &SFbxOptionWindow::SetStaticMesh_CombineMeshes)
-		.Content()
-		[
-			SNew(STextBlock)
-			.Text(LOCTEXT("FbxOptionWindow_CombineMeshes", "Combine Meshes"))
-		]
-	];
-
-
-	NewBox->AddSlot().AutoHeight() .Padding(2)
-	[
-		SNew(SCheckBox)
-		.IsChecked(ImportUI->StaticMeshImportData->bImportMeshLODs ? ESlateCheckBoxState::Checked : ESlateCheckBoxState::Unchecked)
-		.OnCheckStateChanged(this, &SFbxOptionWindow::SetStaticMesh_ImportMeshLODs)
-		.IsEnabled( !bIsObjFormat )
-		[
-			SNew(STextBlock)
-			.Text(LOCTEXT("FbxOptionWindow_ImportMeshLODs", "Import Mesh LODs"))
-		]
-	];
-
-	NewBox->AddSlot().AutoHeight() .Padding(2)
-	[
-		SNew(SCheckBox)
-		.IsChecked(ImportUI->StaticMeshImportData->bReplaceVertexColors ? ESlateCheckBoxState::Checked : ESlateCheckBoxState::Unchecked)
-		.OnCheckStateChanged(this, &SFbxOptionWindow::SetStaticMesh_ReplaceVertexColor)
-		.IsEnabled( !bIsObjFormat )
-		[
-			SNew(STextBlock)
-			.Text(LOCTEXT("FbxOptionWindow_ReplaceVertexColors", "Replace Vertex Colors"))
-		]
-	];
-
-	NewBox->AddSlot().AutoHeight() .Padding(2)
-	[
-		SNew(SCheckBox)
-		.IsChecked(ImportUI->StaticMeshImportData->bRemoveDegenerates ? ESlateCheckBoxState::Checked : ESlateCheckBoxState::Unchecked)
-		.OnCheckStateChanged(this, &SFbxOptionWindow::SetStaticMesh_RemoveDegenerates)
-		[
-			SNew(STextBlock)
-			.Text(LOCTEXT("FbxOptionWindow_RemoveDegeneates", "Remove Degenerates"))
-		]
-	];
-
-	NewBox->AddSlot().AutoHeight() .Padding(2)
-	[
-		SNew(SCheckBox)
-		.IsChecked(ImportUI->StaticMeshImportData->bGenerateLightmapUVs ? ESlateCheckBoxState::Checked : ESlateCheckBoxState::Unchecked)
-		.OnCheckStateChanged(this, &SFbxOptionWindow::SetStaticMesh_GenerateLightmapUVs)
-		[
-			SNew(STextBlock)
-			.Text(LOCTEXT("FbxOptionWindow_GenerateLightmapUVs", "Generate Lightmap UVs"))
-		]
-	];
-
-	NewBox->AddSlot().AutoHeight() .Padding(2)
-	[
-		SNew(SCheckBox)
-		.IsChecked(ImportUI->StaticMeshImportData->bOneConvexHullPerUCX ? ESlateCheckBoxState::Checked : ESlateCheckBoxState::Unchecked)
-		.OnCheckStateChanged(this, &SFbxOptionWindow::SetStaticMesh_OneConvexHullPerUCX)
-		.IsEnabled( !bIsObjFormat )
-		[
-			SNew(STextBlock)
-			.Text(LOCTEXT("FbxOptionWindow_OneConvexHullPerUCX", "One Convex Hull Per UCX"))
-		]
-	];
-
-	NewBox->AddSlot().AutoHeight().Padding(2)
-		[
-			SNew(SCheckBox)
-			.IsChecked(ImportUI->StaticMeshImportData->bAutoGenerateCollision ? ESlateCheckBoxState::Checked : ESlateCheckBoxState::Unchecked)
-			.OnCheckStateChanged(this, &SFbxOptionWindow::SetStaticMesh_AutoGenerateCollision)
-			.IsEnabled(!bIsObjFormat)
-			[
-				SNew(STextBlock)
-				.Text(LOCTEXT("FbxOptionWindow_AutoGenerateCollision", "Auto Generate Collision"))
-			]
-		];
-
-
-	return NewBox;
-}
-
-TSharedRef<SWidget> SFbxOptionWindow::ConstructSkeletalMeshBasic()
-{
-	TSharedRef<SVerticalBox> NewBox = SNew(SVerticalBox);
-
-	NewBox->AddSlot()
-	.AutoHeight()
-	.Padding(2)
-	[
-		ConstructNormalImportOptions()
-	];
-
-	NewBox->AddSlot().AutoHeight() .Padding(2)
-	[
-		SNew(SCheckBox)
-		.IsChecked(ImportUI->SkeletalMeshImportData->bImportMorphTargets ? ESlateCheckBoxState::Checked : ESlateCheckBoxState::Unchecked)
-		.OnCheckStateChanged(this, &SFbxOptionWindow::SetSkeletalMesh_ImportMorphTargets)
-		[
-			SNew(STextBlock)
-			.Text(LOCTEXT("FbxOptionWindow_ImportMorphTargets", "Import Morph Targets"))
-		]
-	];
-
-	NewBox->AddSlot().AutoHeight() .Padding(2)
-	[
-		SNew(SHorizontalBox)
-
-		+SHorizontalBox::Slot()
-		.AutoWidth()
-		[
-			SNew(SCheckBox)
-			.IsChecked(ImportUI->bImportAnimations ? ESlateCheckBoxState::Checked : ESlateCheckBoxState::Unchecked)
-			.OnCheckStateChanged(this, &SFbxOptionWindow::SetSkeletalMesh_ImportAnimation)
-			[
-				SNew(STextBlock)
-				.Text(LOCTEXT("FbxOptionWindow_ImportAnimation", "Import Animation"))
-			]
-		]
-		+SHorizontalBox::Slot()
-		.AutoWidth()
-		.Padding(2)
-		[
-			SNew(SEditableTextBox)
-			.IsReadOnly(false)
-			.ToolTipText(LOCTEXT("FbxOptionWindow_ImportAnimationToolTip", "Type animation name if you're importing animation"))
-			.Text( this, &SFbxOptionWindow::GetAnimationName)
-			.MinDesiredWidth(50)
-			.OnTextCommitted(this, &SFbxOptionWindow::SetSkeletalMesh_AnimationName)
-			.IsEnabled(this, &SFbxOptionWindow::CanEnterAnimationName)
-		]
-	];
-
-	NewBox->AddSlot().AutoHeight() .Padding(2)
-	[
-		SNew(SCheckBox)
-		.IsChecked(ImportUI->SkeletalMeshImportData->bUpdateSkeletonReferencePose ? ESlateCheckBoxState::Checked : ESlateCheckBoxState::Unchecked)
-		.OnCheckStateChanged(this, &SFbxOptionWindow::SetSkeletalMesh_UpdateSkeletonRefPose)
-		[
-			SNew(STextBlock)
-			.Text(LOCTEXT("FbxOptionWindow_UpdateSkeletonRefPose", "Update Skeleton Reference Pose"))
-		]
-	];
-
-	return NewBox;
-}
-
-BEGIN_SLATE_FUNCTION_BUILD_OPTIMIZATION
-TSharedRef<SWidget> SFbxOptionWindow::ConstructSkeletalMeshAdvanced()
-{
-	TSharedRef<SVerticalBox> NewBox = SNew(SVerticalBox);
-
-	NewBox->AddSlot().AutoHeight() .Padding(4)
-	[
-		SNew(STextBlock)
-		.TextStyle( FEditorStyle::Get(), "FBXSmallFont" ) 
-		.Text(LOCTEXT("FbxOptionWindow_SkeletalMesh", "Mesh"))
-	];
-
-	NewBox->AddSlot().AutoHeight() .Padding(2)
-	[
-		SNew(SCheckBox)
-		.IsChecked(ImportUI->SkeletalMeshImportData->bImportMeshLODs ? ESlateCheckBoxState::Checked : ESlateCheckBoxState::Unchecked)
-		.OnCheckStateChanged(this, &SFbxOptionWindow::SetSkeletalMesh_ImportMeshLODs)
-		[
-			SNew(STextBlock)
-			.Text(LOCTEXT("FbxOptionWindow_ImportSkeltalMeshLODs", "Import Mesh LODs"))
-		]
-	];
-
-	NewBox->AddSlot().AutoHeight() .Padding(2)
-	[
-		SNew(SCheckBox)
-		.IsChecked(ImportUI->bImportRigidMesh ? ESlateCheckBoxState::Checked : ESlateCheckBoxState::Unchecked)
-		.OnCheckStateChanged(this, &SFbxOptionWindow::SetSkeletalMesh_ImportRigidMesh)
-		[
-			SNew(STextBlock)
-			.Text(LOCTEXT("FbxOptionWindow_ImportRigidMesh", "Import Rigid Mesh"))
-		]
-	];
-
-	NewBox->AddSlot().AutoHeight() .Padding(2)
-	[
-		SNew(SCheckBox)
-		.IsChecked(ImportUI->SkeletalMeshImportData->bUseT0AsRefPose ? ESlateCheckBoxState::Checked : ESlateCheckBoxState::Unchecked)
-		.OnCheckStateChanged(this, &SFbxOptionWindow::SetSkeletalMesh_UseT0AsRefPose)
-		[
-			SNew(STextBlock)
-			.Text(LOCTEXT("FbxOptionWindow_UseRefPose", "Use Time 0 Pose for Ref Pose"))
-		]
-	];
-
-	NewBox->AddSlot().AutoHeight() .Padding(2)
-	[
-		SNew(SCheckBox)
-		.IsChecked(ImportUI->SkeletalMeshImportData->bPreserveSmoothingGroups ? ESlateCheckBoxState::Checked : ESlateCheckBoxState::Unchecked)
-		.OnCheckStateChanged(this, &SFbxOptionWindow::SetSkeletalMesh_ReserveSmoothingGroups)
-		[
-			SNew(STextBlock)
-			.Text(LOCTEXT("FbxOptionWindow_PreserveSmoothingGroup", "Preserve Smoothing Groups"))
-		]
-	];
-
-
-	NewBox->AddSlot().AutoHeight() .Padding(2)
-	[
-		SNew(SCheckBox)
-		.IsChecked(ImportUI->SkeletalMeshImportData->bKeepOverlappingVertices ? ESlateCheckBoxState::Checked : ESlateCheckBoxState::Unchecked)
-		.OnCheckStateChanged(this, &SFbxOptionWindow::SetSkeletalMesh_KeepOverlappingVertices)
-		[
-			SNew(STextBlock)
-			.Text(LOCTEXT("FbxOptionWindow_KeepOverlappingVertices", "Keep Overlapping Vertices"))
-		]
-	];
-	NewBox->AddSlot().AutoHeight() .Padding(2)
-	[
-		SNew(SCheckBox)
-		.IsChecked(ImportUI->SkeletalMeshImportData->bImportMeshesInBoneHierarchy ? ESlateCheckBoxState::Checked : ESlateCheckBoxState::Unchecked)
-		.OnCheckStateChanged(this, &SFbxOptionWindow::SetSkeletalMesh_ImportMeshesInBoneHierarchy)
-		[
-			SNew(STextBlock)
-			.Text(LOCTEXT("FbxOptionWindow_ImportMeshesInBoneHierarchy", "Import Meshes in Bone Hierarchy"))
-		]
-	];
-	NewBox->AddSlot().AutoHeight() .Padding(4)
-	[
-		SNew(STextBlock)
-		.TextStyle( FEditorStyle::Get(), "FBXSmallFont" ) 
-		.Text(LOCTEXT("FbxOptionWindow_PhysicsAsset", "Collision"))
-	];
-
-	NewBox->AddSlot().AutoHeight() .Padding(2)
-	[
-		SNew(SCheckBox)
-		.IsChecked(ImportUI->bCreatePhysicsAsset ? ESlateCheckBoxState::Checked : ESlateCheckBoxState::Unchecked)
-		.OnCheckStateChanged(this, &SFbxOptionWindow::SetSkeletalMesh_CreatePhysicsAsset)
-		.IsEnabled(!bForceImportType)
-		[
-			SNew(STextBlock)
-			.Text(LOCTEXT("FbxOptionWindow_CreatePhysicsAsset", "Create Physics Asset"))
-		]
-	];
-
-	NewBox->AddSlot().AutoHeight() .Padding(2)
-	[
-		SNew (SVerticalBox)
-
-		+SVerticalBox::Slot() 
-		.Padding(2) 
-		[
-			SNew(STextBlock)
-			.Text( LOCTEXT("FbxOptionWindow_SelectPhysicsAsset", "Select PhysicsAsset") )
-			.IsEnabled(this, &SFbxOptionWindow::ShouldShowPhysicsAssetPicker)
-		]
-		+SVerticalBox::Slot()
-		.Padding(2) 
-		[
-			SAssignNew(PhysicsAssetPickerComboButton, SComboButton)
-			.ContentPadding(1.f)
-			.OnGetMenuContent( this, &SFbxOptionWindow::MakePhysicsAssetPickerMenu )
-			.HasDownArrow(true)
-			.ToolTipText(LOCTEXT("FBXOption", "Pick an skeleton from a popup menu"))
-			.IsEnabled(this, &SFbxOptionWindow::ShouldShowPhysicsAssetPicker)
-			.ButtonContent()
-			[
-				SNew(STextBlock)
-				.Text(this, &SFbxOptionWindow::GetPhysicsAssetDisplay)
-			]
-		]
-	];
-
-	return NewBox;
-}
-END_SLATE_FUNCTION_BUILD_OPTIMIZATION
-
-TSharedRef<SWidget> SFbxOptionWindow::ConstructMaterialOption()
-{
-	TSharedRef<SVerticalBox> NewBox = SNew(SVerticalBox);
-
-	NewBox->AddSlot().AutoHeight() .Padding(4)
-	[
-		SNew(STextBlock)
-		.TextStyle( FEditorStyle::Get(), "FBXSmallFont" ) 
-		.Text( LOCTEXT("FbxOptionWindow_Material", "Material") )
-	];
-
-	NewBox->AddSlot().AutoHeight() .Padding(2)
-	[
-		SNew(SCheckBox)
-		.IsChecked(ImportUI->bImportMaterials ? ESlateCheckBoxState::Checked : ESlateCheckBoxState::Unchecked)
-		.OnCheckStateChanged(this, &SFbxOptionWindow::SetMaterial_ImportMaterials)
-		.IsEnabled( !bIsObjFormat )
-		[
-			SNew(STextBlock)
-			.Text(LOCTEXT("FbxOptionWindow_ImportMaterials", "Import Materials"))
-		]
-	];
-
-	NewBox->AddSlot().AutoHeight() .Padding(2)
-	[
-		SNew(SCheckBox)
-		.IsChecked(ImportUI->bImportTextures ? ESlateCheckBoxState::Checked : ESlateCheckBoxState::Unchecked)
-		.OnCheckStateChanged(this, &SFbxOptionWindow::SetMaterial_ImportTextures)
-		.IsEnabled( !bIsObjFormat )
-		[
-			SNew(STextBlock)
-			.Text(LOCTEXT("FbxOptionWindow_ImportTextures", "Import Textures"))
-		]
-	];
-
-	NewBox->AddSlot().AutoHeight() .Padding(2)
-	[
-		SNew(SCheckBox)
-		.IsChecked(ImportUI->TextureImportData->bInvertNormalMaps ? ESlateCheckBoxState::Checked : ESlateCheckBoxState::Unchecked)
-		.OnCheckStateChanged(this, &SFbxOptionWindow::SetMaterial_InvertNormalMaps)
-		.IsEnabled( !bIsObjFormat )
-		[
-			SNew(STextBlock)
-			.Text(LOCTEXT("FbxOptionWindow_InvertNormalMaps", "Invert Normal Maps"))
-		]
-	];
-
-	return NewBox;
-}
-
-TSharedRef<SWidget> SFbxOptionWindow::ConstructAnimationOption()
-{
-	TSharedRef<SVerticalBox> NewBox = SNew(SVerticalBox);
-
-	NewBox->AddSlot().AutoHeight() .Padding(2, 4)
-	[
-		SNew(SCheckBox)
-		.IsChecked(ImportUI->bUseDefaultSampleRate ? ESlateCheckBoxState::Checked : ESlateCheckBoxState::Unchecked)
-		.OnCheckStateChanged(this, &SFbxOptionWindow::SetSkeletalMesh_UseDefaultSampleRate)
-		[
-			SNew(STextBlock)
-			.Text(LOCTEXT("FbxOptionWindow_UseDefaultSampleRate", "Use Default Sample Rate"))
-		]
-	];
-
-	NewBox->AddSlot().AutoHeight().Padding(2)
-	[
-		SNew(SCheckBox)
-		.IsChecked(ImportUI->AnimSequenceImportData->bImportCustomAttribute ? ESlateCheckBoxState::Checked : ESlateCheckBoxState::Unchecked)
-		.OnCheckStateChanged(this, &SFbxOptionWindow::SetAnimation_ImportCustomAttribute)
-		[
-			SNew(STextBlock)
-			.Text(LOCTEXT("FbxOptionWindow_ImportCustomAttribute", "Import Custom Attribute as Curves"))
-		]
-	];
-	NewBox->AddSlot().AutoHeight() .Padding(2)
-	[
-		SNew(SCheckBox)
-		.IsChecked(ImportUI->bPreserveLocalTransform ? ESlateCheckBoxState::Checked : ESlateCheckBoxState::Unchecked)
-		.OnCheckStateChanged(this, &SFbxOptionWindow::SetAnimation_ReserveLocalTransform)
-		[
-			SNew(STextBlock)
-			.Text(LOCTEXT("FbxOptionWindow_PreserveLocalTransform", "Preserve Local Transform"))
-		]
-	];
-
-	NewBox->AddSlot().AutoHeight() .Padding(4)
-	[
-		SNew(STextBlock)
-		.Text( LOCTEXT("FbxOptionWindow_AnimationLength", "Animation Length") )
-		.TextStyle( FEditorStyle::Get(), "FBXSmallFont" ) 
-	];
-
-	NewBox->AddSlot() .Padding(1)
-	[
-		SNew(SAnimImportLengthOption)
-			.AnimationLengthOption(ImportUI->AnimSequenceImportData->AnimationLength)
-			.OnSelectionChanged(this, &SFbxOptionWindow::SetAnimLengthOption)
-			.OnTextCommittedRange1(this, &SFbxOptionWindow::SetAnimationRangeStart)
-			.OnTextCommittedRange2(this, &SFbxOptionWindow::SetAnimationRangeEnd)
-	];
-
-	return NewBox;
-}
-
-TSharedRef<SWidget> SFbxOptionWindow::ConstructSkeletonOptionForMesh()
-{
-	TSharedRef<SVerticalBox> NewBox = SNew(SVerticalBox)
-		.ToolTip(IDocumentation::Get()->CreateToolTip(FText::FromString("Pick a skeleton for this mesh"), NULL, FString("Shared/Editors/Persona"), FString("Skeleton")));
-
-	NewBox->AddSlot().AutoHeight() .Padding(2)
-	[
-		SNew(SHorizontalBox)
-		+SHorizontalBox::Slot()
-		.AutoWidth()
-		[
-			SNew(STextBlock)
-			.TextStyle( FEditorStyle::Get(), "FBXMediumFont" ) 
-			.Text( LOCTEXT("FbxOptionWindow_SelectSkeletonForMesh", "Choose Skeleton") )
-		]
-		+SHorizontalBox::Slot()
-		.HAlign(HAlign_Left)
-		[
-			IDocumentation::Get()->CreateAnchor(FString("Engine/Animation/Skeleton"))
-		]
-	];
-
-	NewBox->AddSlot().AutoHeight() .Padding(2)
-	[
-		SNew(STextBlock)
-		.Text( LOCTEXT("FbxOptionWindow_AutoCreateSkeleton", "If none is selected, a new skeleton will be created") )
-	];
-
-	NewBox->AddSlot().AutoHeight() .Padding(2)
-	[
-		SAssignNew(SkeletonPickerComboButton, SComboButton)
-		.ContentPadding(1.f)
-		.OnGetMenuContent( this, &SFbxOptionWindow::MakeSkeletonPickerMenu )
-		.HasDownArrow(true)
-		.IsEnabled( !bForceImportType )
-		.ButtonContent()
-		[
-			SNew(STextBlock)
-			.Text(this, &SFbxOptionWindow::GetSkeletonDisplay)
-		]
-	];
-
-	return NewBox;
-}
-
-TSharedRef<SWidget> SFbxOptionWindow::ConstructSkeletonOptionForAnim()
-{
-	TSharedRef<SVerticalBox> NewBox = SNew(SVerticalBox)
-		.ToolTip(IDocumentation::Get()->CreateToolTip(FText::FromString("Pick a skeleton for this mesh"), NULL, FString("Shared/Editors/Persona"), FString("Skeleton")));;
-	
-
-	NewBox->AddSlot().AutoHeight() .Padding(2) 
-	[
-		SNew(SHorizontalBox)
-		+SHorizontalBox::Slot()
-		.AutoWidth()
-		[
-			SNew(STextBlock)
-			.TextStyle( FEditorStyle::Get(), "FBXMediumFont" ) 
-			.Text( LOCTEXT("FbxOptionWindow_SelectSkeletonForAnim", "Choose Skeleton") )
-		]
-		
-		+SHorizontalBox::Slot()
-		.HAlign(HAlign_Left)
-		[
-			IDocumentation::Get()->CreateAnchor(FString("Engine/Animation/Skeleton"))
-		]
-	];
-
-	NewBox->AddSlot().AutoHeight() .Padding(2)
-	[
-		SNew(STextBlock)
-		.Text( LOCTEXT("FbxOptionWindow_NeedToSelectSkeletonForAnimation", "You need to select skeleton for the animation.") )
-	];
-
-	NewBox->AddSlot().AutoHeight() .Padding(2)
-	[
-		SAssignNew(SkeletonPickerComboButton, SComboButton)
-		.ContentPadding(1.f)
-		.OnGetMenuContent( this, &SFbxOptionWindow::MakeSkeletonPickerMenu )
-		.HasDownArrow(true)
-		.IsEnabled( !bForceImportType )
-		.ButtonContent()
-		[
-			SNew(STextBlock)
-			.Text(this, &SFbxOptionWindow::GetSkeletonDisplay)
-		]
-	];
-
-	return NewBox;
-}
-
-FText SFbxOptionWindow::GetPhysicsAssetDisplay() const
-{
-	if ( ImportUI->PhysicsAsset )
-	{
-		return FText::FromString(ImportUI->PhysicsAsset->GetName());
-	}
-
-	return LOCTEXT("NoPhysicsAssetToDisplay", "None");
-}
-
-FText SFbxOptionWindow::GetSkeletonDisplay() const
-{
-	if ( ImportUI->Skeleton )
-	{
-		return FText::FromString(ImportUI->Skeleton->GetName());
-	}
-
-	return LOCTEXT("NoSkeletonToDisplay", "None");
-}
-
-void SFbxOptionWindow::SetImportType(EFBXImportType ImportType)
-{
-	ImportUI->MeshTypeToImport = ImportType;
-
-	RefreshWindow();
-}
-
-// data set functions
-void SFbxOptionWindow::SetGeneral_OverrideFullName(ESlateCheckBoxState::Type NewType)
-{
-	ImportUI->bOverrideFullName = (NewType == ESlateCheckBoxState::Checked)? 1:0;
-}
-
-void SFbxOptionWindow::SetGeneral_ConvertScene(ESlateCheckBoxState::Type NewType)
-{
-	ImportUI->bConvertScene = (NewType == ESlateCheckBoxState::Checked)? 1:0;
-}
-
-void SFbxOptionWindow::SetSkeletalMesh_ImportMeshLODs(ESlateCheckBoxState::Type NewType)
-{
-	ImportUI->SkeletalMeshImportData->bImportMeshLODs = (NewType == ESlateCheckBoxState::Checked)? 1:0;
-}
-
-void SFbxOptionWindow::SetSkeletalMesh_ImportMorphTargets(ESlateCheckBoxState::Type NewType)
-{
-	ImportUI->SkeletalMeshImportData->bImportMorphTargets = (NewType == ESlateCheckBoxState::Checked)? 1:0;
-}
-
-void SFbxOptionWindow::SetSkeletalMesh_UpdateSkeletonRefPose(ESlateCheckBoxState::Type NewType)
-{
-	ImportUI->SkeletalMeshImportData->bUpdateSkeletonReferencePose = (NewType == ESlateCheckBoxState::Checked)? 1:0;
-}
-
-void SFbxOptionWindow::SetSkeletalMesh_ImportAnimation(ESlateCheckBoxState::Type NewType)
-{
-	ImportUI->bImportAnimations = (NewType == ESlateCheckBoxState::Checked)? 1:0;
-}
-
-void SFbxOptionWindow::SetSkeletalMesh_ImportRigidMesh(ESlateCheckBoxState::Type NewType)
-{
-	ImportUI->bImportRigidMesh = (NewType == ESlateCheckBoxState::Checked)? 1:0;
-}
-
-void SFbxOptionWindow::SetSkeletalMesh_UseDefaultSampleRate(ESlateCheckBoxState::Type NewType)
-{
-	ImportUI->bUseDefaultSampleRate = (NewType == ESlateCheckBoxState::Checked)? 1:0;
-}
-
-void SFbxOptionWindow::SetSkeletalMesh_UseT0AsRefPose(ESlateCheckBoxState::Type NewType)
-{
-	ImportUI->SkeletalMeshImportData->bUseT0AsRefPose = (NewType == ESlateCheckBoxState::Checked)? 1:0;
-}
-
-void SFbxOptionWindow::SetSkeletalMesh_ReserveSmoothingGroups(ESlateCheckBoxState::Type NewType)
-{
-	ImportUI->SkeletalMeshImportData->bPreserveSmoothingGroups = (NewType == ESlateCheckBoxState::Checked)? 1:0;
-}
-
-void SFbxOptionWindow::SetSkeletalMesh_KeepOverlappingVertices(ESlateCheckBoxState::Type NewType)
-{
-	ImportUI->SkeletalMeshImportData->bKeepOverlappingVertices = (NewType == ESlateCheckBoxState::Checked)? 1:0;
-}
-
-void SFbxOptionWindow::SetSkeletalMesh_ImportMeshesInBoneHierarchy(ESlateCheckBoxState::Type NewType)
-{
-	ImportUI->SkeletalMeshImportData->bImportMeshesInBoneHierarchy = (NewType == ESlateCheckBoxState::Checked)? 1:0;
-}
-
-void SFbxOptionWindow::SetSkeletalMesh_CreatePhysicsAsset(ESlateCheckBoxState::Type NewType)
-{
-	ImportUI->bCreatePhysicsAsset = (NewType == ESlateCheckBoxState::Checked)? 1:0;
-}
-
-void SFbxOptionWindow::SetStaticMesh_ImportMeshLODs(ESlateCheckBoxState::Type NewType)
-{
-	ImportUI->StaticMeshImportData->bImportMeshLODs = (NewType == ESlateCheckBoxState::Checked)? 1:0;
-}
-
-void SFbxOptionWindow::SetStaticMesh_LODGroup(TSharedPtr<FString> NewValue, ESelectInfo::Type SelectInfo)
-{
-	int32 GroupIndex = StaticMeshLODGroups.Find(NewValue);
-	ImportUI->StaticMeshImportData->StaticMeshLODGroup = StaticMeshLODGroupNames[GroupIndex];
-}
-
-void SFbxOptionWindow::SetStaticMesh_CombineMeshes(ESlateCheckBoxState::Type NewType)
-{
-	ImportUI->bCombineMeshes = (NewType == ESlateCheckBoxState::Checked)? 1:0;
-}
-
-void SFbxOptionWindow::SetStaticMesh_ReplaceVertexColor(ESlateCheckBoxState::Type NewType)
-{
-	ImportUI->StaticMeshImportData->bReplaceVertexColors = (NewType == ESlateCheckBoxState::Checked)? 1:0;
-}
-
-void SFbxOptionWindow::SetStaticMesh_RemoveDegenerates(ESlateCheckBoxState::Type NewType)
-{
-	ImportUI->StaticMeshImportData->bRemoveDegenerates = (NewType == ESlateCheckBoxState::Checked)? 1:0;
-}
-
-void SFbxOptionWindow::SetStaticMesh_GenerateLightmapUVs(ESlateCheckBoxState::Type NewType)
-{
-	ImportUI->StaticMeshImportData->bGenerateLightmapUVs = (NewType == ESlateCheckBoxState::Checked)? 1:0;
-}
-
-void SFbxOptionWindow::SetStaticMesh_OneConvexHullPerUCX(ESlateCheckBoxState::Type NewType)
-{
-	ImportUI->StaticMeshImportData->bOneConvexHullPerUCX = (NewType == ESlateCheckBoxState::Checked)? 1:0;
-}
-
-void SFbxOptionWindow::SetStaticMesh_AutoGenerateCollision(ESlateCheckBoxState::Type NewType)
-{
-	ImportUI->StaticMeshImportData->bAutoGenerateCollision = (NewType == ESlateCheckBoxState::Checked) ? 1 : 0;
-}
-
-void SFbxOptionWindow::SetMaterial_ImportMaterials(ESlateCheckBoxState::Type NewType)
-{
-	ImportUI->bImportMaterials = (NewType == ESlateCheckBoxState::Checked)? 1:0;
-}
-
-void SFbxOptionWindow::SetMaterial_ImportTextures(ESlateCheckBoxState::Type NewType)
-{
-	ImportUI->bImportTextures = (NewType == ESlateCheckBoxState::Checked)? 1:0;
-}
-
-void SFbxOptionWindow::SetMaterial_InvertNormalMaps(ESlateCheckBoxState::Type NewType)
-{
-	ImportUI->TextureImportData->bInvertNormalMaps = (NewType == ESlateCheckBoxState::Checked)? 1:0;
-}
-
-void SFbxOptionWindow::SetSkeletalMesh_AnimationName(const FText& Name, ETextCommit::Type /*CommitInfo*/ )
-{
-	ImportUI->AnimationName = Name.ToString();
-}
-
-void SFbxOptionWindow::SetAnimLengthOption(EFBXAnimationLengthImportType AnimLengthOption)
-{
-	ImportUI->AnimSequenceImportData->AnimationLength = AnimLengthOption;
-}
-
-void SFbxOptionWindow::SetAnimationRangeStart(const FText& Name, ETextCommit::Type /*CommitInfo*/ )
-{
-	ImportUI->AnimSequenceImportData->StartFrame = FCString::Atoi( *Name.ToString() );
-
-	if (ImportUI->AnimSequenceImportData->StartFrame >= ImportUI->AnimSequenceImportData->EndFrame)
-	{
-		// TODO: ErrorMessage no longer appears to be used?
-		ErrorMessage = LOCTEXT("FbxOptionWindow_InvalidStartFrame", "Invalid StartFrame").ToString();
-	}
-	else
-	{
-		// @fixme
-		// this error message doesn't work well globally, -i.e. if the message wasn't for this, but for something else
-		// we need better error feed back message thing here
-		ErrorMessage = TEXT("");
-	}
-}
-
-void SFbxOptionWindow::SetAnimationRangeEnd(const FText& Name, ETextCommit::Type /*CommitInfo*/)
-{
-	ImportUI->AnimSequenceImportData->EndFrame = FCString::Atoi(*Name.ToString());
-
-	if (ImportUI->AnimSequenceImportData->StartFrame >= ImportUI->AnimSequenceImportData->EndFrame)
-	{
-		ErrorMessage = LOCTEXT("FbxOptionWindow_InvalidEndFrame", "Invalid EndFrame").ToString();
-	}
-	else
-	{
-		// @fixme
-		// this error message doesn't work well globally, -i.e. if the message wasn't for this, but for something else
-		// we need better error feed back message thing here
-		ErrorMessage = TEXT("");
-	}
-}
-
-void SFbxOptionWindow::SetAnimation_ReserveLocalTransform(ESlateCheckBoxState::Type NewType)
-{
-	ImportUI->bPreserveLocalTransform = (NewType == ESlateCheckBoxState::Checked)? 1:0;
-}
-
-void SFbxOptionWindow::SetAnimation_ImportCustomAttribute(ESlateCheckBoxState::Type NewType)
-{
-	ImportUI->AnimSequenceImportData->bImportCustomAttribute = (NewType == ESlateCheckBoxState::Checked)? 1:0;
-}
-
-TSharedRef<SWidget> SFbxOptionWindow::MakeSkeletonPickerMenu()
-{
-	FContentBrowserModule& ContentBrowserModule = FModuleManager::Get().LoadModuleChecked<FContentBrowserModule>(TEXT("ContentBrowser"));
-
-	FAssetPickerConfig AssetPickerConfig;
-	AssetPickerConfig.Filter.ClassNames.Add(USkeleton::StaticClass()->GetFName());
-	AssetPickerConfig.OnAssetSelected = FOnAssetSelected::CreateSP(this, &SFbxOptionWindow::OnAssetSelectedFromSkeletonPicker);
-	AssetPickerConfig.bAllowNullSelection = true;
-	AssetPickerConfig.InitialAssetViewType = EAssetViewType::List;
-	AssetPickerConfig.ThumbnailScale = 0.0f;
-
-	return SNew(SBox)
-		.WidthOverride(384)
-		.HeightOverride(768)
-		[
-			ContentBrowserModule.Get().CreateAssetPicker(AssetPickerConfig)
-		];
-}
-
-void SFbxOptionWindow::OnAssetSelectedFromSkeletonPicker(const FAssetData& AssetData)
-{
-	// @todo Set the content reference
-	if (SkeletonPickerComboButton.IsValid())
-	{
-		SkeletonPickerComboButton->SetIsOpen(false);
-		ImportUI->Skeleton = Cast<USkeleton>(AssetData.GetAsset());
-	}
-}
-
-TSharedRef<SWidget> SFbxOptionWindow::MakePhysicsAssetPickerMenu()
-{
-	FContentBrowserModule& ContentBrowserModule = FModuleManager::Get().LoadModuleChecked<FContentBrowserModule>(TEXT("ContentBrowser"));
-
-	FAssetPickerConfig AssetPickerConfig;
-	AssetPickerConfig.Filter.ClassNames.Add(UPhysicsAsset::StaticClass()->GetFName());
-	AssetPickerConfig.OnAssetSelected = FOnAssetSelected::CreateSP(this, &SFbxOptionWindow::OnAssetSelectedFromPhysicsAssetPicker);
-	AssetPickerConfig.bAllowNullSelection = true;
-	AssetPickerConfig.InitialAssetViewType = EAssetViewType::List;
-	AssetPickerConfig.ThumbnailScale = 0.0f;
-
-	return SNew(SBox)
-		.WidthOverride(384)
-		.HeightOverride(768)
-		[
-			ContentBrowserModule.Get().CreateAssetPicker(AssetPickerConfig)
-		];
-}
-
-void SFbxOptionWindow::OnAssetSelectedFromPhysicsAssetPicker(const FAssetData& AssetData)
-{
-	// @todo Set the content reference
-	if (PhysicsAssetPickerComboButton.IsValid())
-	{
-		PhysicsAssetPickerComboButton->SetIsOpen(false);
-		ImportUI->PhysicsAsset = Cast<UPhysicsAsset>(AssetData.GetAsset());
-	}
-}
-
-TSharedRef<SWidget> SFbxOptionWindow::OnGenerateWidgetForComboItem( TSharedPtr<EFBXNormalImportMethod> ImportMethod )
-{
-	return SNew( STextBlock )
-	.Text( OnGenerateTextForImportMethod( *ImportMethod ) )
-	.ToolTipText( OnGenerateToolTipForImportMethod( *ImportMethod ) );
-}
-
-FText SFbxOptionWindow::OnGenerateTextForImportMethod( EFBXNormalImportMethod ImportMethod ) const
-{
-	// TODO: change this to read the display name from the enum directly
-	if( ImportMethod == FBXNIM_MAX )
-	{
-		ImportMethod = GetCurrentNormalImportMethod();
-	}
-
-	if( ImportMethod == FBXNIM_ComputeNormals ) 
-	{	
-		return LOCTEXT("FBXOptions_CalculateNormals", "Calculate Normals");
-	}
-	else if( ImportMethod == FBXNIM_ImportNormals )
-	{
-		return LOCTEXT("FBXOptions_ImportNormals", "Import Normals");
-	}
-	else
-	{
-		return LOCTEXT("FBXOptions_ImportNormalsAndTangents", "Import Normals and Tangents");
-	}
-}
-
-FText SFbxOptionWindow::OnGenerateToolTipForImportMethod( EFBXNormalImportMethod ImportMethod ) const
-{
-	// TODO: change this to read the tooltip text from the enum directly
-	if( ImportMethod == FBXNIM_ComputeNormals ) 
-	{	
-		return LOCTEXT("FBXOptions_CalculateNormalsToolTip", "Let Unreal calculate normals and tangents. Ignores normals in the fbx file");
-	}
-	else if( ImportMethod == FBXNIM_ImportNormals )
-	{
-		return LOCTEXT("FBXOptions_ImportNormalsToolTip", "Import normals found in the fbx file.  Tangents are calculated by Unreal");
-	}
-	else
-	{
-		return LOCTEXT("FBXOptions_ImportNormalsAndTangentsToolTip", "Import Normals and Tangents found in the fbx file instead of computing them");
-	}
-}
-
-void SFbxOptionWindow::OnNormalImportMethodChanged( TSharedPtr<EFBXNormalImportMethod> NewMethod, ESelectInfo::Type SelectionType )
-{
-	if ( ImportUI->MeshTypeToImport == FBXIT_StaticMesh )
-	{
-		ImportUI->StaticMeshImportData->NormalImportMethod = *NewMethod;
-	}
-	else if ( ImportUI->MeshTypeToImport == FBXIT_SkeletalMesh )
-	{
-		ImportUI->SkeletalMeshImportData->NormalImportMethod = *NewMethod;
-	}
-	else
-	{
-		// Invalid mode
-	}
-}
-
-bool SFbxOptionWindow::CanImport()  const
-{
-	// do test to see if we are ready to import
-
-	if (ImportUI->MeshTypeToImport == FBXIT_Animation)
-	{
-		if (ImportUI->Skeleton == NULL)
-		{
-			return false;
-		}
-	}
-
-	if (ImportUI->AnimSequenceImportData->AnimationLength == FBXALIT_SetRange)
-	{
-		if (ImportUI->AnimSequenceImportData->StartFrame > ImportUI->AnimSequenceImportData->EndFrame)
-		{
-			return false;
-		}
-	}
-
-	return true;
-}
-
-bool SFbxOptionWindow::ShouldShowPhysicsAssetPicker() const
-{
-	return (!bForceImportType && ImportUI->bCreatePhysicsAsset == false);
-}
-
-EFBXNormalImportMethod SFbxOptionWindow::GetCurrentNormalImportMethod() const
-{
-	if ( ImportUI->MeshTypeToImport == FBXIT_StaticMesh )
-	{
-		return ImportUI->StaticMeshImportData->NormalImportMethod;
-	}
-	else if ( ImportUI->MeshTypeToImport == FBXIT_SkeletalMesh )
-	{
-		return ImportUI->SkeletalMeshImportData->NormalImportMethod;
-	}
-	else
-	{
-		return FBXNIM_ComputeNormals;
-	}
+	return true;
 }
 
 #undef LOCTEXT_NAMESPACE
-
diff --git a/Engine/Source/Editor/UnrealEd/Private/Fbx/FbxSkeletalMeshImportData.cpp b/Engine/Source/Editor/UnrealEd/Private/Fbx/FbxSkeletalMeshImportData.cpp
--- a/Engine/Source/Editor/UnrealEd/Private/Fbx/FbxSkeletalMeshImportData.cpp
+++ b/Engine/Source/Editor/UnrealEd/Private/Fbx/FbxSkeletalMeshImportData.cpp
@@ -28,4 +28,16 @@ UFbxSkeletalMeshImportData* UFbxSkeletalMeshImportData::GetImportDataForSkeletal
 	}
 
 	return ImportData;
+}
+
+bool UFbxSkeletalMeshImportData::CanEditChange(const UProperty* InProperty) const
+{
+	bool bMutable = Super::CanEditChange(InProperty);
+	UObject* Outer = GetOuter();
+	if(Outer && bMutable)
+	{
+		// Let the FbxImportUi object handle the editability of our properties
+		bMutable = Outer->CanEditChange(InProperty);
+	}
+	return bMutable;
 }
\ No newline at end of file
diff --git a/Engine/Source/Editor/UnrealEd/Private/Fbx/FbxStaticMeshImportData.cpp b/Engine/Source/Editor/UnrealEd/Private/Fbx/FbxStaticMeshImportData.cpp
--- a/Engine/Source/Editor/UnrealEd/Private/Fbx/FbxStaticMeshImportData.cpp
+++ b/Engine/Source/Editor/UnrealEd/Private/Fbx/FbxStaticMeshImportData.cpp
@@ -32,4 +32,16 @@ UFbxStaticMeshImportData* UFbxStaticMeshImportData::GetImportDataForStaticMesh(U
 	}
 
 	return ImportData;
+}
+
+bool UFbxStaticMeshImportData::CanEditChange(const UProperty* InProperty) const
+{
+	bool bMutable = Super::CanEditChange(InProperty);
+	UObject* Outer = GetOuter();
+	if(Outer && bMutable)
+	{
+		// Let the FbxImportUi object handle the editability of our properties
+		bMutable = Outer->CanEditChange(InProperty);
+	}
+	return bMutable;
 }
\ No newline at end of file
diff --git a/Engine/Source/Editor/UnrealEd/Private/Fbx/FbxTextureImportData.cpp b/Engine/Source/Editor/UnrealEd/Private/Fbx/FbxTextureImportData.cpp
--- a/Engine/Source/Editor/UnrealEd/Private/Fbx/FbxTextureImportData.cpp
+++ b/Engine/Source/Editor/UnrealEd/Private/Fbx/FbxTextureImportData.cpp
@@ -6,4 +6,16 @@ UFbxTextureImportData::UFbxTextureImportData(const class FPostConstructInitializ
 	: Super(PCIP)
 {
 	
+}
+
+bool UFbxTextureImportData::CanEditChange(const UProperty* InProperty) const
+{
+	bool bMutable = Super::CanEditChange(InProperty);
+	UObject* Outer = GetOuter();
+	if(Outer && bMutable)
+	{
+		// Let the parent object handle the editability of our properties
+		bMutable = Outer->CanEditChange(InProperty);
+	}
+	return bMutable;
 }
\ No newline at end of file
diff --git a/Engine/Source/Editor/UnrealEd/Private/FbxOptionWindow.h b/Engine/Source/Editor/UnrealEd/Private/FbxOptionWindow.h
--- a/Engine/Source/Editor/UnrealEd/Private/FbxOptionWindow.h
+++ b/Engine/Source/Editor/UnrealEd/Private/FbxOptionWindow.h
@@ -74,8 +74,6 @@ public:
 		return bShouldImportAll;
 	}
 
-	void RefreshWindow();
-	
 	SFbxOptionWindow() 
 		: ImportUI(NULL)
 		, bShouldImport(false)
@@ -83,91 +81,14 @@ public:
 	{}
 		
 private:
-	// sub options
-	TSharedRef<SWidget> ConstructStaticMeshBasic();
-	TSharedRef<SWidget> ConstructStaticMeshAdvanced();
-	TSharedRef<SWidget> ConstructSkeletalMeshBasic();
-	TSharedRef<SWidget> ConstructSkeletalMeshAdvanced();
-	TSharedRef<SWidget> ConstructMaterialOption();
-	TSharedRef<SWidget> ConstructAnimationOption();
-	TSharedRef<SWidget> ConstructSkeletonOptionForMesh();
-	TSharedRef<SWidget> ConstructSkeletonOptionForAnim();
-	TSharedRef<SWidget> ConstructMiscOption();
-	TSharedRef<SWidget> ConstructNormalImportOptions();
 
-	void SetImportType(EFBXImportType ImportType);
-
-	void SetGeneral_OverrideFullName(ESlateCheckBoxState::Type NewType);
-	void SetGeneral_ConvertScene(ESlateCheckBoxState::Type NewType);
-	void SetMesh_ImportTangents(ESlateCheckBoxState::Type NewType);
-	void SetSkeletalMesh_ImportMeshLODs(ESlateCheckBoxState::Type NewType);
-	void SetSkeletalMesh_ImportMorphTargets(ESlateCheckBoxState::Type NewType);
-	void SetSkeletalMesh_UpdateSkeletonRefPose(ESlateCheckBoxState::Type NewType);
-	void SetSkeletalMesh_ImportAnimation(ESlateCheckBoxState::Type NewType);
-	void SetSkeletalMesh_ImportRigidMesh(ESlateCheckBoxState::Type NewType);
-	void SetSkeletalMesh_UseDefaultSampleRate(ESlateCheckBoxState::Type NewType);
-	void SetSkeletalMesh_UseT0AsRefPose(ESlateCheckBoxState::Type NewType);
-	void SetSkeletalMesh_ReserveSmoothingGroups(ESlateCheckBoxState::Type NewType);
-	void SetSkeletalMesh_KeepOverlappingVertices(ESlateCheckBoxState::Type NewType);
-	void SetSkeletalMesh_ImportMeshesInBoneHierarchy(ESlateCheckBoxState::Type NewType);
-	void SetSkeletalMesh_CreatePhysicsAsset(ESlateCheckBoxState::Type NewType);
-	void SetStaticMesh_ImportMeshLODs(ESlateCheckBoxState::Type NewType);
-	void SetStaticMesh_LODGroup(TSharedPtr<FString> NewValue, ESelectInfo::Type SelectInfo);
-	void SetStaticMesh_CombineMeshes(ESlateCheckBoxState::Type NewType);
-	void SetStaticMesh_ReplaceVertexColor(ESlateCheckBoxState::Type NewType);
-	void SetStaticMesh_RemoveDegenerates(ESlateCheckBoxState::Type NewType);
-	void SetStaticMesh_GenerateLightmapUVs(ESlateCheckBoxState::Type NewType);
-	void SetStaticMesh_OneConvexHullPerUCX(ESlateCheckBoxState::Type NewType);
-	void SetStaticMesh_AutoGenerateCollision(ESlateCheckBoxState::Type NewType);
-	void SetMaterial_ImportMaterials(ESlateCheckBoxState::Type NewType);
-	void SetMaterial_ImportTextures(ESlateCheckBoxState::Type NewType);
-	void SetMaterial_InvertNormalMaps(ESlateCheckBoxState::Type NewType);
-
-	bool CanEnterAnimationName() const { return ImportUI->bImportAnimations; }
-
-	void SetSkeletalMesh_AnimationName(const FText& Name, ETextCommit::Type CommitInfo);
-
-	void SetAnimLengthOption(EFBXAnimationLengthImportType AnimLengthOption);
-	void SetAnimationRangeStart(const FText& Name, ETextCommit::Type CommitInfo);
-	void SetAnimationRangeEnd(const FText& Name, ETextCommit::Type CommitInfo);
-	void SetAnimation_ReserveLocalTransform(ESlateCheckBoxState::Type NewType);
-	void SetAnimation_ImportCustomAttribute(ESlateCheckBoxState::Type NewType);
-
-	TSharedRef<SWidget> MakeSkeletonPickerMenu();
-	void OnAssetSelectedFromSkeletonPicker(const FAssetData& AssetData);
-	TSharedPtr<SComboButton> SkeletonPickerComboButton;
-	FText GetSkeletonDisplay() const;
-
-	TSharedRef<SWidget> MakePhysicsAssetPickerMenu();
-	void OnAssetSelectedFromPhysicsAssetPicker(const FAssetData& AssetData);
-	TSharedPtr<SComboButton> PhysicsAssetPickerComboButton;
-	FText GetPhysicsAssetDisplay() const;
-
-	TSharedRef<SWidget> OnGenerateWidgetForComboItem( TSharedPtr<EFBXNormalImportMethod> ImportMethod );
-	void OnNormalImportMethodChanged( TSharedPtr<EFBXNormalImportMethod> NewMethod, ESelectInfo::Type SelectionType );
-	FText OnGenerateTextForImportMethod( EFBXNormalImportMethod ImportMethod ) const;
-	FText OnGenerateToolTipForImportMethod( EFBXNormalImportMethod ImportMethod ) const;
-
-	FText GetAnimationName() const { return FText::FromString( ImportUI->AnimationName ); }
-	FString GetAnimationRangeStart() const { return FString::FromInt(ImportUI->AnimSequenceImportData->StartFrame); }
-	FString GetAnimationRangeEnd() const { return FString::FromInt(ImportUI->AnimSequenceImportData->StartFrame); }
-	
 	bool CanImport() const;
-	bool ShouldShowPhysicsAssetPicker() const;
-	EFBXNormalImportMethod GetCurrentNormalImportMethod() const;
 
 private:
 	UFbxImportUI*	ImportUI;
-	FString			ErrorMessage;
-	SVerticalBox::FSlot* CustomBox;
 	TWeakPtr< SWindow > WidgetWindow;
 	TSharedPtr< SButton > ImportButton;
-	TArray<FName> StaticMeshLODGroupNames;
-	TArray<FText> StaticMeshLODGroupDisplayNames;
-	TArray< TSharedPtr<FString> > StaticMeshLODGroups;
-	TArray< TSharedPtr<EFBXNormalImportMethod> > NormalImportOptions;
 	bool			bShouldImport;
 	bool			bShouldImportAll;
-	bool			bForceImportType;
 	bool			bIsObjFormat;
 };
