InternationalizationExportCommandlet: Added bUseCultureDirectory flag.  Defaults to true, but for translation editor, we need to not add that because we need to pick the exact file path with the file open dialog, and not have a culture path added to that later.

diff --git a/Engine/Source/Editor/TranslationEditor/Private/InternationalizationExportSettings.h b/Engine/Source/Editor/TranslationEditor/Private/InternationalizationExportSettings.h
--- a/Engine/Source/Editor/TranslationEditor/Private/InternationalizationExportSettings.h
+++ b/Engine/Source/Editor/TranslationEditor/Private/InternationalizationExportSettings.h
@@ -45,4 +45,8 @@ public:
 	/** Whether or not to import localization data */
 	UPROPERTY(Category = GatherTextStep, EditAnywhere, config)
 	bool bImportLoc;
+
+	/** Whether or not to use culture path */
+	UPROPERTY(Category = GatherTextStep, EditAnywhere, config)
+	bool bUseCultureDirectory;
 };
\ No newline at end of file
diff --git a/Engine/Source/Editor/TranslationEditor/Private/TranslationEditor.cpp b/Engine/Source/Editor/TranslationEditor/Private/TranslationEditor.cpp
--- a/Engine/Source/Editor/TranslationEditor/Private/TranslationEditor.cpp
+++ b/Engine/Source/Editor/TranslationEditor/Private/TranslationEditor.cpp
@@ -1144,6 +1144,7 @@ void FTranslationEditor::ExportToPortableObjectFormat_Execute()
 		ExportSettings->ArchiveName = FPaths::GetBaseFilename(ManifestFilePath) + ".archive";
 		ExportSettings->bExportLoc = true;
 		ExportSettings->bImportLoc = false;
+		ExportSettings->bUseCultureDirectory = false;
 
 		ExportSettings->DestinationPath = DefaultPath / DefaultFilename;
 
@@ -1292,6 +1293,7 @@ void FTranslationEditor::ImportFromPortableObjectFormat_Execute()
 		ImportSettings->ArchiveName = FPaths::GetBaseFilename(ManifestFilePath) + ".archive";
 		ImportSettings->bExportLoc = false;
 		ImportSettings->bImportLoc = true;
+		ImportSettings->bUseCultureDirectory = false;
 
 		ImportSettings->SourcePath = DefaultPath / FPaths::GetBaseFilename(ManifestFilePath);
 
diff --git a/Engine/Source/Editor/UnrealEd/Private/Commandlets/InternationalizationExportCommandlet.cpp b/Engine/Source/Editor/UnrealEd/Private/Commandlets/InternationalizationExportCommandlet.cpp
--- a/Engine/Source/Editor/UnrealEd/Private/Commandlets/InternationalizationExportCommandlet.cpp
+++ b/Engine/Source/Editor/UnrealEd/Private/Commandlets/InternationalizationExportCommandlet.cpp
@@ -1053,6 +1053,13 @@ bool UInternationalizationExportCommandlet::DoExport( const FString& SourcePath,
 		return false;
 	}
 
+	// Get culture directory setting, default to true if not specified (used to allow picking of export directory with windows file dialog from Translation Editor)
+	bool bUseCultureDirectory = true;
+	if (!(GetConfigBool(*SectionName, TEXT("bUseCultureDirectory"), bUseCultureDirectory, ConfigPath)))
+	{
+		bUseCultureDirectory = true;
+	}
+
 
 	TSharedRef< FInternationalizationManifest > InternationalizationManifest = MakeShareable( new FInternationalizationManifest );
 	// Load the manifest info
@@ -1143,7 +1150,15 @@ bool UInternationalizationExportCommandlet::DoExport( const FString& SourcePath,
 				// Write out the Portable Object to .po file.
 				{
 					FString OutputString = PortableObj.ToString();
-					FString OutputFileName = DestinationPath / CultureName / Filename;
+					FString OutputFileName = "";
+					if (bUseCultureDirectory)
+					{
+						OutputFileName = DestinationPath / CultureName / Filename;
+					}
+					else
+					{
+						OutputFileName = DestinationPath / Filename;
+					}
 
 					if( SourceControlInfo.IsValid() )
 					{
@@ -1186,12 +1201,27 @@ bool UInternationalizationExportCommandlet::DoImport(const FString& SourcePath,
 		return false;
 	}
 
+	// Get culture directory setting, default to true if not specified (used to allow picking of import directory with file open dialog from Translation Editor)
+	bool bUseCultureDirectory = true;
+	if (!(GetConfigBool(*SectionName, TEXT("bUseCultureDirectory"), bUseCultureDirectory, ConfigPath)))
+	{
+		bUseCultureDirectory = true;
+	}
+
 	// Process the desired cultures
 	for(int32 Culture = 0; Culture < CulturesToGenerate.Num(); Culture++)
 	{
 		// Load the Portable Object file if found
 		const FString CultureName = CulturesToGenerate[Culture];
-		const FString POFilePath = SourcePath / CultureName / Filename;
+		FString POFilePath = "";
+		if (bUseCultureDirectory)
+		{
+			POFilePath = SourcePath / CultureName / Filename;
+		}
+		else
+		{
+			POFilePath = SourcePath / Filename;
+		}
 
 		if( !FPaths::FileExists(POFilePath) )
 		{
