Don't change widget mode (W, E, R) while using trackpad in editor viewports

diff --git a/Engine/Source/Editor/UnrealEd/Private/LevelEditorViewport.cpp b/Engine/Source/Editor/UnrealEd/Private/LevelEditorViewport.cpp
--- a/Engine/Source/Editor/UnrealEd/Private/LevelEditorViewport.cpp
+++ b/Engine/Source/Editor/UnrealEd/Private/LevelEditorViewport.cpp
@@ -3058,7 +3058,7 @@ bool FLevelEditorViewportClient::IsVolumeVisibleInViewport( const AActor& Volume
 
 void FLevelEditorViewportClient::SetWidgetMode( FWidget::EWidgetMode ActivatedMode )
 {
-	if( !GLevelEditorModeTools().IsTracking() )
+	if( !GLevelEditorModeTools().IsTracking() && !FSlateApplication::Get().IsUsingTrackpad() )
 	{
 		GLevelEditorModeTools().SetWidgetMode( ActivatedMode );
 
