update fbx import warning for scale

diff --git a/Engine/Source/Editor/UnrealEd/Private/Fbx/FbxStaticMeshImport.cpp b/Engine/Source/Editor/UnrealEd/Private/Fbx/FbxStaticMeshImport.cpp
--- a/Engine/Source/Editor/UnrealEd/Private/Fbx/FbxStaticMeshImport.cpp
+++ b/Engine/Source/Editor/UnrealEd/Private/Fbx/FbxStaticMeshImport.cpp
@@ -815,7 +815,7 @@ void UnFbx::FFbxImporter::VerifyGeometry(UStaticMesh* StaticMesh)
 
 		if (Extents.GetAbsMax() < 5.f)
 		{
-			AddTokenizedErrorMessage(FTokenizedMessage::Create(EMessageSeverity::Warning, LOCTEXT("Prompt_MeshVerySmall", "Warning: The imported mesh is very small. This is most likely an issue with the units used when exporting to FBX. Please import again and increase the Scale to make the mesh bigger.")), FFbxErrors::Generic_Mesh_SmallGeometry);
+			AddTokenizedErrorMessage(FTokenizedMessage::Create(EMessageSeverity::Warning, LOCTEXT("Prompt_MeshVerySmall", "Warning: The imported mesh is very small. This is most likely an issue with the units used when exporting to FBX.")), FFbxErrors::Generic_Mesh_SmallGeometry);
 		}
 	}
 }
