Build - Fixing a non-unity build issue.

diff --git a/Engine/Source/Editor/UnrealEd/Private/SViewportToolBarComboMenu.cpp b/Engine/Source/Editor/UnrealEd/Private/SViewportToolBarComboMenu.cpp
--- a/Engine/Source/Editor/UnrealEd/Private/SViewportToolBarComboMenu.cpp
+++ b/Engine/Source/Editor/UnrealEd/Private/SViewportToolBarComboMenu.cpp
@@ -3,7 +3,7 @@
 #pragma once
 
 #include "UnrealEd.h"
-#include "SViewportToolBarIconMenu.h"
+#include "SViewportToolBarComboMenu.h"
 
 void SViewportToolBarComboMenu::Construct( const FArguments& InArgs )
 {
