Editor: Fixed no default device selected in LaunchOn toolbar button

diff --git a/Engine/Source/Editor/UnrealEd/Private/Kismet2/DebuggerCommands.cpp b/Engine/Source/Editor/UnrealEd/Private/Kismet2/DebuggerCommands.cpp
--- a/Engine/Source/Editor/UnrealEd/Private/Kismet2/DebuggerCommands.cpp
+++ b/Engine/Source/Editor/UnrealEd/Private/Kismet2/DebuggerCommands.cpp
@@ -187,7 +187,8 @@ FPlayWorldCommands::FPlayWorldCommands()
 
 				if (PlayDevice.IsValid())
 				{
-					PlaySettings->LastExecutedLaunchDevice = PlayDevice->GetName();
+					PlaySettings->LastExecutedLaunchDevice = PlayDevice->GetId().ToString();
+					PlaySettings->LastExecutedLaunchName = PlayDevice->GetName();
 					PlaySettings->SaveConfig();
 				}
 			}
