Editor: Fixed wrong default device selected in LaunchOn on Mac

diff --git a/Engine/Source/Editor/UnrealEd/Private/Kismet2/DebuggerCommands.cpp b/Engine/Source/Editor/UnrealEd/Private/Kismet2/DebuggerCommands.cpp
--- a/Engine/Source/Editor/UnrealEd/Private/Kismet2/DebuggerCommands.cpp
+++ b/Engine/Source/Editor/UnrealEd/Private/Kismet2/DebuggerCommands.cpp
@@ -172,7 +172,7 @@ FPlayWorldCommands::FPlayWorldCommands()
 		{
 			PlayPlatformName = TEXT("WindowsNoEditor");
 		}
-		else if (RunningPlatformName == TEXT("Windows"))
+		else if (RunningPlatformName == TEXT("Mac"))
 		{
 			PlayPlatformName = TEXT("MacNoEditor");
 		}
