Newly spawned actors in the editor now have more correct unique names

This addresses TTP#345442 - NUX:Editor: Dropping new actors into totally empty maps adds number suffixes to my actor names

diff --git a/Engine/Source/Editor/UnrealEd/Private/AssetSelection.cpp b/Engine/Source/Editor/UnrealEd/Private/AssetSelection.cpp
--- a/Engine/Source/Editor/UnrealEd/Private/AssetSelection.cpp
+++ b/Engine/Source/Editor/UnrealEd/Private/AssetSelection.cpp
@@ -623,7 +623,7 @@ AActor* FActorFactoryAssetProxy::AddActorFromSelection( UClass* ActorClass, cons
 	{
 		UObject* TargetObject = GEditor->GetSelectedObjects()->GetTop<UObject>();
 
-		if( ActorFactory->CanCreateActorFrom( FAssetData(TargetObject), ErrorMessage ) )
+		if( TargetObject && ActorFactory->CanCreateActorFrom( FAssetData(TargetObject), ErrorMessage ) )
 		{
 			// Attempt to add the actor
 			Result = PrivateAddActor( TargetObject, ActorFactory, SelectActor, ObjectFlags );
diff --git a/Engine/Source/Editor/UnrealEd/Private/Factories/ActorFactory.cpp b/Engine/Source/Editor/UnrealEd/Private/Factories/ActorFactory.cpp
--- a/Engine/Source/Editor/UnrealEd/Private/Factories/ActorFactory.cpp
+++ b/Engine/Source/Editor/UnrealEd/Private/Factories/ActorFactory.cpp
@@ -208,6 +208,10 @@ AActor* UActorFactory::SpawnActor( UObject* Asset, ULevel* InLevel, const FVecto
 void UActorFactory::PostSpawnActor( UObject* Asset, AActor* NewActor)
 {
 	// Subclasses may implement this to modify the actor after it has been spawned
+	if (Asset)
+	{
+		GEditor->SetActorLabelUnique(NewActor, Asset->GetName());
+	}
 }
 
 void UActorFactory::PostCreateBlueprint( UObject* Asset, AActor* CDO )
@@ -385,6 +389,10 @@ void UActorFactoryDeferredDecal::PostSpawnActor(UObject* Asset, AActor* NewActor
 		// Init Component
 		DecalComponent->RegisterComponent();
 	}
+	else if (Asset)
+	{
+		GEditor->SetActorLabelUnique(NewActor, Asset->GetName());
+	}
 }
 
 void UActorFactoryDeferredDecal::PostCreateBlueprint( UObject* Asset, AActor* CDO )
@@ -769,6 +777,8 @@ void UActorFactoryAnimationAsset::PostSpawnActor( UObject* Asset, AActor* NewAct
 			NewSASComponent->AnimationData.AnimToPlay = AnimationAsset;
 			// set runtime data
 			NewSASComponent->SetAnimation(AnimationAsset);
+
+			GEditor->SetActorLabelUnique(NewActor, AnimationAsset->GetName());
 		}
 		else if( VertexAnimation )
 		{
@@ -777,6 +787,8 @@ void UActorFactoryAnimationAsset::PostSpawnActor( UObject* Asset, AActor* NewAct
 
 			// set runtime data
 			NewSASComponent->SetVertexAnimation(VertexAnimation);
+
+			GEditor->SetActorLabelUnique(NewActor, VertexAnimation->GetName());
 		}
 	}
 }
@@ -1521,6 +1533,8 @@ void UActorFactoryVectorFieldVolume::PostSpawnActor( UObject* Asset, AActor* New
 	UVectorField* VectorField = CastChecked<UVectorField>(Asset);
 	AVectorFieldVolume* VectorFieldVolumeActor = CastChecked<AVectorFieldVolume>(NewActor);
 
+	GEditor->SetActorLabelUnique(NewActor, VectorField->GetName());
+
 	if ( VectorFieldVolumeActor && VectorFieldVolumeActor->VectorFieldComponent )
 	{
 		VectorFieldVolumeActor->VectorFieldComponent->VectorField = VectorField;
@@ -1599,6 +1613,11 @@ void UActorFactoryBoxVolume::PostSpawnActor( UObject* Asset, AActor* NewActor )
 	{
 		UCubeBuilder* Builder = ConstructObject<UCubeBuilder>( UCubeBuilder::StaticClass() );
 		CreateBrushForVolumeActor( VolumeActor, Builder );
+
+		if (Asset)
+		{
+			GEditor->SetActorLabelUnique(NewActor, Asset->GetName());
+		}
 	}
 }
 
@@ -1636,6 +1655,11 @@ void UActorFactorySphereVolume::PostSpawnActor( UObject* Asset, AActor* NewActor
 		Builder->SphereExtrapolation = 2;
 		Builder->Radius = 192.0f;
 		CreateBrushForVolumeActor( VolumeActor, Builder );
+
+		if (Asset)
+		{
+			GEditor->SetActorLabelUnique(NewActor, Asset->GetName());
+		}
 	}
 }
 
@@ -1671,6 +1695,11 @@ void UActorFactoryCylinderVolume::PostSpawnActor( UObject* Asset, AActor* NewAct
 		UCylinderBuilder* Builder = ConstructObject<UCylinderBuilder>( UCylinderBuilder::StaticClass() );
 		Builder->OuterRadius = 128.0f;
 		CreateBrushForVolumeActor( VolumeActor, Builder );
+
+		if (Asset)
+		{
+			GEditor->SetActorLabelUnique(NewActor, Asset->GetName());
+		}
 	}
 }
 
