Added analytics for average frame times

TTP# 346675 - Add an analytics event for average FPS and frame times

The following extra attributes are now sent with the "Editor.Usage.Heartbeat" event (they may not always be present, eg, if the current performance samples aren't a reliable average):
- AverageFPS
- AverageFrameTime
- AverageGameThreadTime
- AverageRenderThreadTime
- AverageGPUFrameTime

These are sampled once per-second when the editor is in a normal state (eg, not being throttled).

diff --git a/Engine/Source/Editor/UnrealEd/Private/Editor.cpp b/Engine/Source/Editor/UnrealEd/Private/Editor.cpp
--- a/Engine/Source/Editor/UnrealEd/Private/Editor.cpp
+++ b/Engine/Source/Editor/UnrealEd/Private/Editor.cpp
@@ -1463,7 +1463,7 @@ void UEditorEngine::Tick( float DeltaSeconds, bool bIdleMode )
 
 	FUnrealEdMisc::Get().TickAssetAnalytics();
 
-	FUnrealEdMisc::Get().TickPerformanceSurvey();
+	FUnrealEdMisc::Get().TickPerformanceAnalytics();
 
 	// If the fadeout animation has completed for the undo/redo notification item, allow it to be deleted
 	if(UndoRedoNotificationItem.IsValid() && UndoRedoNotificationItem->GetCompletionState() == SNotificationItem::CS_None)
diff --git a/Engine/Source/Editor/UnrealEd/Private/PerformanceMonitor.h b/Engine/Source/Editor/UnrealEd/Private/PerformanceMonitor.h
--- a/Engine/Source/Editor/UnrealEd/Private/PerformanceMonitor.h
+++ b/Engine/Source/Editor/UnrealEd/Private/PerformanceMonitor.h
@@ -1,5 +1,7 @@
 // Copyright 1998-2014 Epic Games, Inc. All Rights Reserved.
 
+#pragma once
+
 /**
  * Helper class for a calculating a moving average. This works by maintaining an accumulator which is then sampled at regular intervals into the "Samples" array.
  */
@@ -19,7 +21,7 @@ struct FMovingAverage
 	}
 	
 	/** Check if this data is reliable or not. Returns false if the sampler is not populated with enough data */
-	inline float IsReliable() const
+	inline bool IsReliable() const
 	{
 		return SampleSize != 0 && Samples.Num() == SampleSize;
 	}
@@ -36,6 +38,12 @@ struct FMovingAverage
 		return SampleAverage;
 	}
 
+	/** Get the current number of available samples */
+	inline int32 GetSampleCount() const
+	{
+		return Samples.Num();
+	}
+
 	/** Allow this sampler to tick the frame time, and potentially save a new sample */
 	void Tick(double CurrentTime, float Value)
 	{
diff --git a/Engine/Source/Editor/UnrealEd/Private/UnrealEdMisc.cpp b/Engine/Source/Editor/UnrealEd/Private/UnrealEdMisc.cpp
--- a/Engine/Source/Editor/UnrealEd/Private/UnrealEdMisc.cpp
+++ b/Engine/Source/Editor/UnrealEd/Private/UnrealEdMisc.cpp
@@ -35,14 +35,108 @@
 #include "ShaderCompiler.h"
 #include "NavigationBuildingNotification.h"
 #include "HotReloadInterface.h"
+#include "PerformanceMonitor.h"
 
 #define LOCTEXT_NAMESPACE "UnrealEd"
 
 DEFINE_LOG_CATEGORY_STATIC(LogUnrealEdMisc, Log, All);
 
+namespace
+{
+	static const FName LevelEditorName("LevelEditor");
+	static const FName AssetRegistryName("AssetRegistry");
+}
+
+/**
+ * Manages the stats needed by the analytics heartbeat
+ * This is very similar to FStatUnitData, however it's not tied to a single viewport, 
+ * nor does it rely on the stats being active to be updated
+ */
+class FPerformanceAnalyticsStats
+{
+public:
+	FPerformanceAnalyticsStats()
+		: AverageFrameTime(SampleSize)
+		, AverageGameThreadTime(SampleSize)
+		, AverageRenderThreadTime(SampleSize)
+		, AverageGPUFrameTime(SampleSize)
+	{
+	}
+
+	/** Get the average number of milliseconds in total over the frames that have been sampled */
+	float GetAverageFrameTime() const
+	{
+		return AverageFrameTime.GetAverage();
+	}
+
+	/** Get the average number of milliseconds the gamethread was used over the frames that have been sampled */
+	float GetAverageGameThreadTime() const
+	{
+		return AverageGameThreadTime.GetAverage();
+	}
+
+	/** Get the average number of milliseconds the renderthread was used over the frames that have been sampled */
+	float GetAverageRenderThreadTime() const
+	{
+		return AverageRenderThreadTime.GetAverage();
+	}
+
+	/** Get the average number of milliseconds the GPU was busy over the frames that have been sampled */
+	float GetAverageGPUFrameTime() const
+	{
+		return AverageGPUFrameTime.GetAverage();
+	}
+
+	/** Have we taken enough samples to get a reliable average? */
+	bool IsReliable() const
+	{
+		return AverageFrameTime.IsReliable();
+	}
+
+	/** Update the samples based on what happened last frame */
+	void Update()
+	{
+		const double CurrentTime = FApp::GetCurrentTime();
+		const double DeltaTime = CurrentTime - FApp::GetLastTime();
+
+		// Number of milliseconds in total last frame
+		const double RawFrameTime = DeltaTime * 1000.0;
+		AverageFrameTime.Tick(CurrentTime, static_cast<float>(RawFrameTime));
+
+		// Number of milliseconds the gamethread was used last frame
+		const double RawGameThreadTime = FPlatformTime::ToMilliseconds(GGameThreadTime);
+		AverageGameThreadTime.Tick(CurrentTime, static_cast<float>(RawGameThreadTime));
+
+		// Number of milliseconds the renderthread was used last frame
+		const double RawRenderThreadTime = FPlatformTime::ToMilliseconds(GRenderThreadTime);
+		AverageRenderThreadTime.Tick(CurrentTime, static_cast<float>(RawRenderThreadTime));
+
+		// Number of milliseconds the GPU was busy last frame
+		const uint32 GPUCycles = RHIGetGPUFrameCycles();
+		const double RawGPUFrameTime = FPlatformTime::ToMilliseconds(GPUCycles);
+		AverageGPUFrameTime.Tick(CurrentTime, static_cast<float>(RawGPUFrameTime));
+	}
+
+private:
+	/** Samples for the total frame time */
+	FMovingAverage AverageFrameTime;
+
+	/** Samples for the gamethread time */
+	FMovingAverage AverageGameThreadTime;
+
+	/** Samples for the renderthread time */
+	FMovingAverage AverageRenderThreadTime;
+
+	/** Samples for the GPU busy time */
+	FMovingAverage AverageGPUFrameTime;
+
+	/** Number of samples to average over */
+	static const int32 SampleSize = 10;
+};
+
 namespace PerformanceSurveyDefs
 {
-	const static int NumFrameRateSamples = 10;
+	const static int32 NumFrameRateSamples = 10;
 	const static FTimespan FrameRateSampleInterval(0, 0, 1);	// 1 second intervals
 }
 
@@ -53,6 +147,7 @@ FUnrealEdMisc::FUnrealEdMisc() :
 	bSaveLayoutOnClose( true ),
 	bDeletePreferences( false ),
 	bIsAssetAnalyticsPending( false ),
+	PerformanceAnalyticsStats(new FPerformanceAnalyticsStats()),
 	NavigationBuildingNotificationHandler(NULL)
 {
 }
@@ -422,8 +517,18 @@ void FUnrealEdMisc::EditorAnalyticsHeartbeat()
 	// Did the user interact since the last heartbeat
 	bool bIdle = LastInteractionTime < LastHeartbeatTime;
 	
+	extern ENGINE_API float GAverageFPS;
+
 	TArray< FAnalyticsEventAttribute > Attributes;
 	Attributes.Add(FAnalyticsEventAttribute(TEXT("Idle"), bIdle));
+	if(PerformanceAnalyticsStats->IsReliable())
+	{
+		Attributes.Add(FAnalyticsEventAttribute(TEXT("AverageFPS"), GAverageFPS));
+		Attributes.Add(FAnalyticsEventAttribute(TEXT("AverageFrameTime"), PerformanceAnalyticsStats->GetAverageFrameTime()));
+		Attributes.Add(FAnalyticsEventAttribute(TEXT("AverageGameThreadTime"), PerformanceAnalyticsStats->GetAverageGameThreadTime()));
+		Attributes.Add(FAnalyticsEventAttribute(TEXT("AverageRenderThreadTime"), PerformanceAnalyticsStats->GetAverageRenderThreadTime()));
+		Attributes.Add(FAnalyticsEventAttribute(TEXT("AverageGPUFrameTime"), PerformanceAnalyticsStats->GetAverageGPUFrameTime()));
+	}
 	FEngineAnalytics::GetProvider().RecordEvent(TEXT("Editor.Usage.Heartbeat"), Attributes);
 	
 	LastHeartbeatTime = FPlatformTime::Seconds();
@@ -433,7 +538,7 @@ void FUnrealEdMisc::TickAssetAnalytics()
 {
 	if( bIsAssetAnalyticsPending )
 	{
-		FAssetRegistryModule& AssetRegistryModule = FModuleManager::LoadModuleChecked<FAssetRegistryModule>(TEXT("AssetRegistry"));
+		FAssetRegistryModule& AssetRegistryModule = FModuleManager::LoadModuleChecked<FAssetRegistryModule>(AssetRegistryName);
 
 		if( !AssetRegistryModule.Get().IsLoadingAssets())
 		{
@@ -1116,7 +1221,7 @@ FString FUnrealEdMisc::GenerateURL(const FString& InUDNPage)
 
 void FUnrealEdMisc::OnGotoAsset(const FString& InAssetPath) const
 {
-	FAssetRegistryModule& AssetRegistryModule = FModuleManager::LoadModuleChecked<FAssetRegistryModule>(TEXT("AssetRegistry"));
+	FAssetRegistryModule& AssetRegistryModule = FModuleManager::LoadModuleChecked<FAssetRegistryModule>(AssetRegistryName);
 	IAssetRegistry& AssetRegistry = AssetRegistryModule.Get();
 	FAssetData AssetData = AssetRegistry.GetAssetByObjectPath( *InAssetPath );
 	if ( AssetData.IsValid() )
@@ -1275,7 +1380,7 @@ void FUnrealEdMisc::BeginPerformanceSurvey()
 	}
 
 	// Tell the level editor we want to be notified when selection changes
-	FLevelEditorModule& LevelEditor = FModuleManager::LoadModuleChecked<FLevelEditorModule>( "LevelEditor" );
+	FLevelEditorModule& LevelEditor = FModuleManager::LoadModuleChecked<FLevelEditorModule>( LevelEditorName );
 	LevelEditor.OnMapChanged().AddRaw( this, &FUnrealEdMisc::OnMapChanged );
 
 	// Initialize survey variables
@@ -1284,13 +1389,8 @@ void FUnrealEdMisc::BeginPerformanceSurvey()
 	FrameRateSamples.Empty();
 }
 
-void FUnrealEdMisc::TickPerformanceSurvey()
+void FUnrealEdMisc::TickPerformanceAnalytics()
 {
-	if (!bIsSurveyingPerformance)
-	{
-		return;
-	}
-
 	// Don't run if we've not yet loaded a project
 	if( !FApp::HasGameName() )
 	{
@@ -1298,7 +1398,7 @@ void FUnrealEdMisc::TickPerformanceSurvey()
 	}
 
 	// Before beginning the survey wait for the asset registry to load and make sure Slate is ready
-	FAssetRegistryModule& AssetRegistryModule = FModuleManager::LoadModuleChecked<FAssetRegistryModule>(TEXT("AssetRegistry"));
+	FAssetRegistryModule& AssetRegistryModule = FModuleManager::LoadModuleChecked<FAssetRegistryModule>(AssetRegistryName);
 	if (AssetRegistryModule.Get().IsLoadingAssets() || !FSlateApplication::IsInitialized())
 	{
 		return;
@@ -1318,6 +1418,15 @@ void FUnrealEdMisc::TickPerformanceSurvey()
 		return;
 	}
 
+	// Update the stats needed by the analytics heartbeat
+	PerformanceAnalyticsStats->Update();
+
+	// Also check to see if we need to run the performance survey
+	if (!bIsSurveyingPerformance)
+	{
+		return;
+	}
+
 	// Sample the frame rate until we have enough samples to take the average
 	if (FrameRateSamples.Num() < PerformanceSurveyDefs::NumFrameRateSamples)
 	{
@@ -1354,7 +1463,7 @@ void FUnrealEdMisc::CancelPerformanceSurvey()
 	bIsSurveyingPerformance = false;
 	FrameRateSamples.Empty();
 
-	FLevelEditorModule& LevelEditor = FModuleManager::LoadModuleChecked<FLevelEditorModule>( "LevelEditor" );
+	FLevelEditorModule& LevelEditor = FModuleManager::LoadModuleChecked<FLevelEditorModule>( LevelEditorName );
 	LevelEditor.OnMapChanged().RemoveRaw( this, &FUnrealEdMisc::OnMapChanged );
 }
 
diff --git a/Engine/Source/Editor/UnrealEd/Public/UnrealEdMisc.h b/Engine/Source/Editor/UnrealEd/Public/UnrealEdMisc.h
--- a/Engine/Source/Editor/UnrealEd/Public/UnrealEdMisc.h
+++ b/Engine/Source/Editor/UnrealEd/Public/UnrealEdMisc.h
@@ -2,6 +2,8 @@
 
 #pragma once
 
+class FPerformanceAnalyticsStats;
+
 namespace EMapChangeType
 {
 	enum Type
@@ -181,8 +183,8 @@ public:
 	/** Restarts the editor, reopening the current project, if any */
 	void RestartEditor(bool bWarn = true);
 
-	/** Ticks the performance survey if it's running */
-	void TickPerformanceSurvey();
+	/** Ticks the performance analytics used by the analytics heartbeat */
+	void TickPerformanceAnalytics();
 
 	/** Triggers asset analytics if it hasn't been run yet */
 	void TickAssetAnalytics();
@@ -310,6 +312,9 @@ private:
 	/** An array of frame rate samples used by the performance survey */
 	TArray<float> FrameRateSamples;
 
+	/** Statistical information needed by the analytics to report editor performance */
+	TUniquePtr<FPerformanceAnalyticsStats> PerformanceAnalyticsStats;
+
 	/** handler to notify about navigation building process */
 	TSharedPtr<FTickableEditorObject> NavigationBuildingNotificationHandler;
 
