Removed #pragma once in SViewportToolbarComboMenu.cpp file

diff --git a/Engine/Source/Editor/UnrealEd/Private/SViewportToolBarComboMenu.cpp b/Engine/Source/Editor/UnrealEd/Private/SViewportToolBarComboMenu.cpp
--- a/Engine/Source/Editor/UnrealEd/Private/SViewportToolBarComboMenu.cpp
+++ b/Engine/Source/Editor/UnrealEd/Private/SViewportToolBarComboMenu.cpp
@@ -1,7 +1,5 @@
 // Copyright 1998-2014 Epic Games, Inc. All Rights Reserved.
 
-#pragma once
-
 #include "UnrealEd.h"
 #include "SViewportToolBarComboMenu.h"
 
