ttp346814 BP: Assertion Failed when renaming Interface variable in BP.

diff --git a/Engine/Source/Editor/UnrealEd/Private/Kismet2/BlueprintEditorUtils.cpp b/Engine/Source/Editor/UnrealEd/Private/Kismet2/BlueprintEditorUtils.cpp
--- a/Engine/Source/Editor/UnrealEd/Private/Kismet2/BlueprintEditorUtils.cpp
+++ b/Engine/Source/Editor/UnrealEd/Private/Kismet2/BlueprintEditorUtils.cpp
@@ -5984,6 +5984,12 @@ bool FBlueprintEditorUtils::PropertyValueFromString(const UProperty* Property, c
 			const auto EndOfParsedBuff = Property->ImportText(Value.IsEmpty() ? TEXT("()") : *Value, Property->ContainerPtrToValuePtr<uint8>(DefaultObject), 0, NULL, &ImportError);
 			bParseSucceeded = EndOfParsedBuff && ImportError.IsEmpty();
 		}
+		else if (Property->IsA(UInterfaceProperty::StaticClass()))
+		{
+			FStringOutputDevice ImportError;
+			const auto EndOfParsedBuff = Property->ImportText(*Value, Property->ContainerPtrToValuePtr<uint8>(DefaultObject), 0, NULL, &ImportError);
+			bParseSucceeded = EndOfParsedBuff && ImportError.IsEmpty();
+		}
 		else
 		{
 			// HOOK UP NEW TYPES HERE
