Copy package guid over when duplicating and copying packages in memory for PIE

diff --git a/Engine/Source/Editor/UnrealEd/Private/PlayLevel.cpp b/Engine/Source/Editor/UnrealEd/Private/PlayLevel.cpp
--- a/Engine/Source/Editor/UnrealEd/Private/PlayLevel.cpp
+++ b/Engine/Source/Editor/UnrealEd/Private/PlayLevel.cpp
@@ -2895,6 +2895,7 @@ UWorld* UEditorEngine::CreatePIEWorldByDuplication(FWorldContext &WorldContext,
 	PlayWorldPackage->PackageFlags |= PKG_PlayInEditor;
 	PlayWorldPackage->PIEInstanceID = WorldContext.PIEInstance;
 	PlayWorldPackage->FileName = InPackage->FileName;
+	PlayWorldPackage->SetGuid( InPackage->GetGuid() );
 
 	check(GPlayInEditorID == -1 || GPlayInEditorID == WorldContext.PIEInstance);
 	GPlayInEditorID = WorldContext.PIEInstance;
diff --git a/Engine/Source/Runtime/Engine/Private/World.cpp b/Engine/Source/Runtime/Engine/Private/World.cpp
--- a/Engine/Source/Runtime/Engine/Private/World.cpp
+++ b/Engine/Source/Runtime/Engine/Private/World.cpp
@@ -2118,6 +2118,7 @@ UWorld* UWorld::DuplicateWorldForPIE(const FString& PackageName, UWorld* OwningW
 #if WITH_EDITOR
 	PIELevelPackage->PIEInstanceID = WorldContext.PIEInstance;
 #endif
+	PIELevelPackage->SetGuid( EditorLevelPackage->GetGuid() );
 
 	// Set up string asset reference fixups
 	TArray<FString> PackageNamesBeingDuplicatedForPIE;
