Back out changelist 2289149

This code was removing the prefix of the server when exporting package names to clients

diff --git a/Engine/Source/Editor/UnrealEd/Private/Editor.cpp b/Engine/Source/Editor/UnrealEd/Private/Editor.cpp
--- a/Engine/Source/Editor/UnrealEd/Private/Editor.cpp
+++ b/Engine/Source/Editor/UnrealEd/Private/Editor.cpp
@@ -6780,7 +6780,7 @@ FORCEINLINE bool NetworkRemapPath_local(FWorldContext &Context, FString &Str, bo
 bool UEditorEngine::NetworkRemapPath( UWorld *InWorld, FString &Str, bool reading)
 {
 	FWorldContext &Context = GetWorldContextFromWorldChecked(InWorld);
-	if ((Context.PIEPrefix.IsEmpty() && !reading) || (Context.PIERemapPrefix.IsEmpty() && reading))
+	if (Context.PIEPrefix.IsEmpty() || Context.PIERemapPrefix.IsEmpty())
 	{
 		return false;
 	}
@@ -6791,7 +6791,7 @@ bool UEditorEngine::NetworkRemapPath( UWorld *InWorld, FString &Str, bool readin
 bool UEditorEngine::NetworkRemapPath( UPendingNetGame *PendingNetGame, FString &Str, bool reading)
 {
 	FWorldContext &Context = GetWorldContextFromPendingNetGameChecked(PendingNetGame);
-	if ((Context.PIEPrefix.IsEmpty() && !reading) || (Context.PIERemapPrefix.IsEmpty() && reading))
+	if (Context.PIEPrefix.IsEmpty() || Context.PIERemapPrefix.IsEmpty())
 	{
 		return false;
 	}
