Add supersearch for all editors

diff --git a/Engine/Source/Editor/UnrealEd/Private/Toolkits/AssetEditorToolkit.cpp b/Engine/Source/Editor/UnrealEd/Private/Toolkits/AssetEditorToolkit.cpp
--- a/Engine/Source/Editor/UnrealEd/Private/Toolkits/AssetEditorToolkit.cpp
+++ b/Engine/Source/Editor/UnrealEd/Private/Toolkits/AssetEditorToolkit.cpp
@@ -13,6 +13,7 @@
 #include "IDocumentation.h"
 #include "ReferenceViewer.h"
 #include "IIntroTutorials.h"
+#include "SuperSearchModule.h"
 
 #define LOCTEXT_NAMESPACE "AssetEditorToolkit"
 
@@ -100,6 +101,12 @@ void FAssetEditorToolkit::InitAssetEditor( const EToolkitMode::Type Mode, const
 			FGlobalTabmanager::Get()->InsertNewDocumentTab( AssetEditorToolkitTab, FTabManager::ESearchPreference::PreferLiveTab, NewMajorTab.ToSharedRef() );
 		}
 
+#if PLATFORM_MAC
+		FSuperSearchModule& SuperSearchModule = FModuleManager::LoadModuleChecked< FSuperSearchModule >(TEXT("SuperSearch"));
+		TSharedPtr< SEditableTextBox > ExposedEditableTextBox;
+		TSharedRef<SWidget> SuperSearchWidget = SuperSearchModule.MakeSearchBox(ExposedEditableTextBox);
+#endif
+
 		IUserFeedbackModule& UserFeedback = FModuleManager::LoadModuleChecked<IUserFeedbackModule>(TEXT("UserFeedback"));
 		TSharedRef<SWidget> UserFeedbackWidget = UserFeedback.CreateFeedbackWidget(GetBaseToolkitName());
 
@@ -108,13 +115,22 @@ void FAssetEditorToolkit::InitAssetEditor( const EToolkitMode::Type Mode, const
 
 		NewMajorTab->SetRightContent(
 				SNew(SHorizontalBox)
-				+SHorizontalBox::Slot()
+				+ SHorizontalBox::Slot()
 				.AutoWidth()
 				.Padding(8.0f, 0.0f, 0.0f, 0.0f)
 				.VAlign(VAlign_Center)
 				[
 					UserFeedbackWidget
 				]
+#if PLATFORM_MAC
+				+SHorizontalBox::Slot()
+				.AutoWidth()
+				.Padding(16.0f, 0.0f, 0.0f, 0.0f)
+				.VAlign(VAlign_Center)
+				[
+					SuperSearchWidget
+				]
+#endif
 				+SHorizontalBox::Slot()
 				.AutoWidth()
 				.Padding(8.0f, 0.0f, 8.0f, 0.0f)
diff --git a/Engine/Source/Editor/UnrealEd/Private/Toolkits/SStandaloneAssetEditorToolkitHost.cpp b/Engine/Source/Editor/UnrealEd/Private/Toolkits/SStandaloneAssetEditorToolkitHost.cpp
--- a/Engine/Source/Editor/UnrealEd/Private/Toolkits/SStandaloneAssetEditorToolkitHost.cpp
+++ b/Engine/Source/Editor/UnrealEd/Private/Toolkits/SStandaloneAssetEditorToolkitHost.cpp
@@ -123,6 +123,12 @@ void SStandaloneAssetEditorToolkitHost::RestoreFromLayout( const TSharedRef<FTab
 	TSharedPtr<SWidget> RestoredUI = MyTabManager->RestoreFrom( NewLayout, ParentWindow );
 
 	checkf(RestoredUI.IsValid(), TEXT("The layout must have a primary dock area") );
+
+#if !PLATFORM_MAC
+	FSuperSearchModule& SuperSearchModule = FModuleManager::LoadModuleChecked< FSuperSearchModule >(TEXT("SuperSearch"));
+	TSharedPtr< SEditableTextBox > ExposedEditableTextBox;
+	TSharedRef<SWidget> SuperSearchWidget = SuperSearchModule.MakeSearchBox(ExposedEditableTextBox);
+#endif
 	
 	this->ChildSlot
 	[
@@ -146,9 +152,22 @@ void SStandaloneAssetEditorToolkitHost::RestoreFromLayout( const TSharedRef<FTab
 					+SOverlay::Slot()
 					.HAlign(HAlign_Right)
 					[
-						SAssignNew(MenuOverlayWidgetContent, SBorder)
-						.Padding(0)
-						.BorderImage(FEditorStyle::GetBrush("NoBorder"))
+						SNew(SHorizontalBox)
+						+ SHorizontalBox::Slot()
+						[
+							SAssignNew(MenuOverlayWidgetContent, SBorder)
+							.Padding(0)
+							.BorderImage(FEditorStyle::GetBrush("NoBorder"))
+						]
+
+#if !PLATFORM_MAC
+						+ SHorizontalBox::Slot()
+						.AutoWidth()
+						.Padding(FMargin(2.f, 0, 2.f, 0))
+						[
+							SuperSearchWidget
+						]
+#endif						
 					]
 				]
 			// Viewport/Document/Docking area
diff --git a/Engine/Source/Editor/UnrealEd/UnrealEd.Build.cs b/Engine/Source/Editor/UnrealEd/UnrealEd.Build.cs
--- a/Engine/Source/Editor/UnrealEd/UnrealEd.Build.cs
+++ b/Engine/Source/Editor/UnrealEd/UnrealEd.Build.cs
@@ -48,6 +48,7 @@ public class UnrealEd : ModuleRules
 				"SourceCodeAccess",
 				"ReferenceViewer",
                 "IntroTutorials",
+                "SuperSearch",
 				"OutputLog"
 			}
 		);
@@ -159,6 +160,7 @@ public class UnrealEd : ModuleRules
 				"ImageWrapper",
 				"Blutility",
 				"IntroTutorials",
+                "SuperSearch",
 				"DesktopPlatform",
 				"WorkspaceMenuStructure",
 				"BspMode",
