LaunchOn: Now using a tutorial link in the message log instead of always popping up the tutorial

diff --git a/Engine/Source/Editor/UnrealEd/Private/Kismet2/DebuggerCommands.cpp b/Engine/Source/Editor/UnrealEd/Private/Kismet2/DebuggerCommands.cpp
--- a/Engine/Source/Editor/UnrealEd/Private/Kismet2/DebuggerCommands.cpp
+++ b/Engine/Source/Editor/UnrealEd/Private/Kismet2/DebuggerCommands.cpp
@@ -25,6 +25,8 @@
 #include "Editor/LevelEditor/Public/ILevelViewport.h"
 
 #include "EditorAnalytics.h"
+#include "MessageLog.h"
+
 
 #define LOCTEXT_NAMESPACE "DebuggerCommands"
 
@@ -1382,23 +1384,54 @@ bool FInternalPlayWorldCommandCallbacks::IsReadyToLaunchOnDevice(FString DeviceI
 	const ITargetPlatform* Platform = GetTargetPlatformManager()->FindTargetPlatform(PlatformName);
 	if (Platform)
 	{
-		FString NotInstalledDocLink;
+		FString NotInstalledTutorialLink;
 		FString ProjectPath = FPaths::IsProjectFilePathSet() ? FPaths::ConvertRelativePathToFull(FPaths::GetProjectFilePath()) : FPaths::RootDir() / FApp::GetGameName() / FApp::GetGameName() + TEXT(".uproject");
-		int32 Result = Platform->DoesntHaveRequirements(ProjectPath, bHasCode, NotInstalledDocLink);
+		int32 Result = Platform->DoesntHaveRequirements(ProjectPath, bHasCode, NotInstalledTutorialLink);
+		
+		// report to analytics
 		FEditorAnalytics::ReportBuildRequirementsFailure(TEXT("Editor.LaunchOn.Failed"), PlatformName, bHasCode, Result);
+
+		// report to message log
+		FText MessageLogText;
+		FText MessageLogTextDetail;
+
 		switch (Result)
 		{
-		// broadcast this, and assume someone will pick it up
-		case ETargetPlatformReadyStatus::SDKNotFound:
-		case ETargetPlatformReadyStatus::ProvisionNotFound:
-		case ETargetPlatformReadyStatus::SigningKeyNotFound:
-		case (ETargetPlatformReadyStatus::ProvisionNotFound | ETargetPlatformReadyStatus::SigningKeyNotFound):
-			{
-				IMainFrameModule& MainFrameModule = FModuleManager::GetModuleChecked<IMainFrameModule>(TEXT("MainFrame"));
-				MainFrameModule.BroadcastMainFrameSDKNotInstalled(PlatformName, NotInstalledDocLink);
-			}
-			return false;
+			case ETargetPlatformReadyStatus::SDKNotFound:
+				MessageLogText = LOCTEXT("SdkNotFoundMessage", "Software Development Kit (SDK) not be found");
+				MessageLogTextDetail = FText::Format(LOCTEXT("SdkNotFoundMessageDetail", "Please install the SDK for the {0} target platform!"), Platform->DisplayName());
+				break;
+
+			case ETargetPlatformReadyStatus::ProvisionNotFound:
+				MessageLogText = LOCTEXT("ProvisionNotFoundMessage", "Provision not found");
+				MessageLogTextDetail = LOCTEXT("ProvisionNotFoundMessageDetail", "A provision is required for deploying your app to the device.");
+				break;
 
+			case ETargetPlatformReadyStatus::SigningKeyNotFound:
+				MessageLogText = LOCTEXT("SigningKeyNotFoundMessage", "Signing key not found");
+				MessageLogTextDetail = LOCTEXT("SigningKeyNotFoundMessageDetail", "The app could not be digitally signed, because the signing key is not configured.");
+				break;
+
+			default:
+				break;
+		}
+
+		if (!MessageLogText.IsEmpty())
+		{
+			TSharedRef<FTokenizedMessage> Message = FTokenizedMessage::Create(EMessageSeverity::Error);
+			Message->AddToken(FTextToken::Create(MessageLogText));
+			Message->AddToken(FTextToken::Create(MessageLogTextDetail));
+			Message->AddToken(FTutorialToken::Create(NotInstalledTutorialLink));
+			Message->AddToken(FDocumentationToken::Create(TEXT("Platforms/iOS/QuickStart/6")));
+
+			FMessageLog MessageLog("PackagingResults");
+			MessageLog.AddMessage(Message);
+			MessageLog.Open();
+		}
+
+		// report to main frame
+		switch (Result)
+		{
 #if PLATFORM_WINDOWS
 		case ETargetPlatformReadyStatus::CodeUnsupported:
 			// show the message
