Renamed "Default" and "Custom" to "Common" and "All" in the new blueprint class picker dialog

diff --git a/Engine/Source/Editor/UnrealEd/Private/SClassPickerDialog.cpp b/Engine/Source/Editor/UnrealEd/Private/SClassPickerDialog.cpp
--- a/Engine/Source/Editor/UnrealEd/Private/SClassPickerDialog.cpp
+++ b/Engine/Source/Editor/UnrealEd/Private/SClassPickerDialog.cpp
@@ -54,7 +54,7 @@ void SClassPickerDialog::Construct(const FArguments& InArgs)
 				[
 					SNew(SExpandableArea)
 					.InitiallyCollapsed(!bExpandDefaultClassPicker)
-					.AreaTitle(NSLOCTEXT("SClassPickerDialog", "StandardClassesAreaTitle", "Standard Classes"))
+					.AreaTitle(NSLOCTEXT("SClassPickerDialog", "CommonClassesAreaTitle", "Common Classes"))
 					.OnAreaExpansionChanged(this, &SClassPickerDialog::OnDefaultAreaExpansionChanged)
 					.BodyContent()
 					[
@@ -73,7 +73,7 @@ void SClassPickerDialog::Construct(const FArguments& InArgs)
 					SNew(SExpandableArea)
 					.MaxHeight(385.f)
 					.InitiallyCollapsed(!bExpandCustomClassPicker)
-					.AreaTitle(NSLOCTEXT("SClassPickerDialog", "CustomClassesAreaTitle", "Custom Classes"))
+					.AreaTitle(NSLOCTEXT("SClassPickerDialog", "AllClassesAreaTitle", "All Classes"))
 					.OnAreaExpansionChanged(this, &SClassPickerDialog::OnCustomAreaExpansionChanged)
 					.BodyContent()
 					[
@@ -117,7 +117,9 @@ bool SClassPickerDialog::PickClass(const FText& TitleText, const FClassViewerIni
 	TSharedRef<SWindow> PickerWindow = SNew(SWindow)
 		.Title(TitleText)
 		.SizingRule( ESizingRule::Autosized )
-		.ClientSize( FVector2D( 0.f, 300.f ));
+		.ClientSize( FVector2D( 0.f, 300.f ))
+		.SupportsMaximize(false)
+		.SupportsMinimize(false);
 
 	TSharedRef<SClassPickerDialog> ClassPickerDialog = SNew(SClassPickerDialog)
 		.ParentWindow(PickerWindow)
