Consolidated Play and Simulate into one button.  Play now uses universal play icon that simulate was using.  Simulate now accessible via sub-menu

diff --git a/Engine/Source/Editor/EditorStyle/Private/SlateEditorStyle.cpp b/Engine/Source/Editor/EditorStyle/Private/SlateEditorStyle.cpp
--- a/Engine/Source/Editor/EditorStyle/Private/SlateEditorStyle.cpp
+++ b/Engine/Source/Editor/EditorStyle/Private/SlateEditorStyle.cpp
@@ -4504,10 +4504,10 @@ void FSlateEditorStyle::FStyle::SetupPersonaStyle()
 		Set( "PlayWorld.Simulate", new IMAGE_BRUSH( "Icons/icon_simulate_40x", Icon40x40 ) );
 		Set( "PlayWorld.Simulate.Small", new IMAGE_BRUSH( "Icons/icon_simulate_40x", Icon20x20 ) );
 
-		Set( "PlayWorld.RepeatLastPlay", new IMAGE_BRUSH( "Icons/icon_play_40x", Icon40x40 ) );
-		Set( "PlayWorld.RepeatLastPlay.Small", new IMAGE_BRUSH( "Icons/icon_play_40x", Icon20x20 ) );
-		Set( "PlayWorld.PlayInViewport", new IMAGE_BRUSH( "Icons/icon_play_40x", Icon40x40 ) );
-		Set( "PlayWorld.PlayInViewport.Small", new IMAGE_BRUSH( "Icons/icon_play_40x", Icon20x20 ) );
+		Set( "PlayWorld.RepeatLastPlay", new IMAGE_BRUSH( "Icons/icon_simulate_40x", Icon40x40 ) );
+		Set( "PlayWorld.RepeatLastPlay.Small", new IMAGE_BRUSH( "Icons/icon_simulate_40x", Icon20x20 ) );
+		Set( "PlayWorld.PlayInViewport", new IMAGE_BRUSH( "Icons/icon_simulate_40x", Icon40x40 ) );
+		Set( "PlayWorld.PlayInViewport.Small", new IMAGE_BRUSH( "Icons/icon_simulate_40x", Icon20x20 ) );
 		Set( "PlayWorld.PlayInEditorFloating", new IMAGE_BRUSH( "Icons/icon_playInWindow_40x", Icon40x40 ) );
 		Set( "PlayWorld.PlayInEditorFloating.Small", new IMAGE_BRUSH( "Icons/icon_playInWindow_40x", Icon20x20 ) );
 		Set( "PlayWorld.PlayInMobilePreview", new IMAGE_BRUSH( "Icons/icon_PlayMobilePreview_40x", Icon40x40 ) );
@@ -4519,8 +4519,8 @@ void FSlateEditorStyle::FStyle::SetupPersonaStyle()
 		Set( "PlayWorld.PlayInCameraLocation", new IMAGE_BRUSH( "Icons/icon_PlayCameraLocation_40x", Icon40x40 ) );
 		Set( "PlayWorld.PlayInDefaultPlayerStart", new IMAGE_BRUSH( "Icons/icon_PlayDefaultPlayerStart_40x", Icon40x40 ) );
 
-		Set( "PlayWorld.ResumePlaySession", new IMAGE_BRUSH( "Icons/icon_play_40x", Icon40x40 ) );
-		Set( "PlayWorld.ResumePlaySession.Small", new IMAGE_BRUSH( "Icons/icon_play_40x", Icon20x20 ) );
+		Set( "PlayWorld.ResumePlaySession", new IMAGE_BRUSH( "Icons/icon_simulate_40x", Icon40x40 ) );
+		Set( "PlayWorld.ResumePlaySession.Small", new IMAGE_BRUSH( "Icons/icon_simulate_40x", Icon20x20 ) );
 		Set( "PlayWorld.PausePlaySession", new IMAGE_BRUSH( "Icons/icon_pause_40x", Icon40x40 ) );
 		Set( "PlayWorld.PausePlaySession.Small", new IMAGE_BRUSH( "Icons/icon_pause_40x", Icon20x20 ) );
 		Set( "PlayWorld.SingleFrameAdvance", new IMAGE_BRUSH( "Icons/icon_advance_40x", Icon40x40 ) );
diff --git a/Engine/Source/Editor/UnrealEd/Classes/Settings/LevelEditorPlaySettings.h b/Engine/Source/Editor/UnrealEd/Classes/Settings/LevelEditorPlaySettings.h
--- a/Engine/Source/Editor/UnrealEd/Classes/Settings/LevelEditorPlaySettings.h
+++ b/Engine/Source/Editor/UnrealEd/Classes/Settings/LevelEditorPlaySettings.h
@@ -58,6 +58,9 @@ enum EPlayModeType
 	/** Runs in a new process. */
 	PlayMode_InNewProcess,
 
+	/** Simulates in viewport without possessing the player */
+	PlayMode_Simulate,
+
 	/** The number of different Play Modes */
 	PlayMode_Count,
 };
diff --git a/Engine/Source/Editor/UnrealEd/Private/Kismet2/DebuggerCommands.cpp b/Engine/Source/Editor/UnrealEd/Private/Kismet2/DebuggerCommands.cpp
--- a/Engine/Source/Editor/UnrealEd/Private/Kismet2/DebuggerCommands.cpp
+++ b/Engine/Source/Editor/UnrealEd/Private/Kismet2/DebuggerCommands.cpp
@@ -202,7 +202,7 @@ FPlayWorldCommands::FPlayWorldCommands()
 void FPlayWorldCommands::RegisterCommands()
 {
 	// SIE
-	UI_COMMAND( Simulate, "Simulate", "Start simulating the game", EUserInterfaceActionType::ToggleButton, FInputGesture( EKeys::S, EModifierKey::Alt ) );
+	UI_COMMAND( Simulate, "Simulate", "Start simulating the game", EUserInterfaceActionType::Check, FInputGesture( EKeys::S, EModifierKey::Alt ) );
 
 	// PIE
 	UI_COMMAND( RepeatLastPlay, "Play", "Launches a game preview session in the same mode as the last game preview session launched from the Game Preview Modes dropdown next to the Play button on the level editor toolbar", EUserInterfaceActionType::Button, FInputGesture( EKeys::P, EModifierKey::Alt ) )
@@ -245,7 +245,7 @@ void FPlayWorldCommands::BindGlobalPlayWorldCommands()
 	ActionList.MapAction( Commands.Simulate,
 		FExecuteAction::CreateStatic( &FInternalPlayWorldCommandCallbacks::Simulate_Clicked ),
 		FCanExecuteAction::CreateStatic( &FInternalPlayWorldCommandCallbacks::Simulate_CanExecute ),
-		FIsActionChecked::CreateStatic( &FInternalPlayWorldCommandCallbacks::Simulate_IsChecked ),
+		FIsActionChecked::CreateStatic( &FInternalPlayWorldCommandCallbacks::PlayInModeIsChecked, PlayMode_Simulate ),
 		FIsActionButtonVisible::CreateStatic( &FInternalPlayWorldCommandCallbacks::CanShowNonPlayWorldOnlyActions )
 		);
 
@@ -386,9 +386,6 @@ void FPlayWorldCommands::BindGlobalPlayWorldCommands()
 
 void FPlayWorldCommands::BuildToolbar( FToolBarBuilder& ToolbarBuilder, bool bIncludeLaunchButtonAndOptions )
 {
-	// Simulate
-	ToolbarBuilder.AddToolBarButton(FPlayWorldCommands::Get().Simulate, NAME_None, TAttribute<FText>(), TAttribute<FText>(), TAttribute<FSlateIcon>(), FName(TEXT("LevelToolbarSimulate")));
-
 	// Play
 	ToolbarBuilder.AddToolBarButton( 
 		FPlayWorldCommands::Get().RepeatLastPlay, 
@@ -489,6 +486,10 @@ TSharedRef< SWidget > FPlayWorldCommands::GeneratePlayMenuContent( TSharedRef<FU
 			case PlayMode_InViewPort:
 				PlayModeCommand = FPlayWorldCommands::Get().PlayInViewport;
 				break;
+
+			case PlayMode_Simulate:
+				PlayModeCommand = FPlayWorldCommands::Get().Simulate;
+				break;
 			}
 
 			if (PlayModeCommand.IsValid())
@@ -505,9 +506,10 @@ TSharedRef< SWidget > FPlayWorldCommands::GeneratePlayMenuContent( TSharedRef<FU
 	MenuBuilder.BeginSection("LevelEditorPlayModes", LOCTEXT("PlayButtonModesSection", "Modes"));
 	{
 		FLocal::AddPlayModeMenuEntry(MenuBuilder, PlayMode_InViewPort);
+		FLocal::AddPlayModeMenuEntry(MenuBuilder, PlayMode_InMobilePreview);
 		FLocal::AddPlayModeMenuEntry(MenuBuilder, PlayMode_InEditorFloating);
 		FLocal::AddPlayModeMenuEntry(MenuBuilder, PlayMode_InNewProcess);
-		FLocal::AddPlayModeMenuEntry(MenuBuilder, PlayMode_InMobilePreview);
+		FLocal::AddPlayModeMenuEntry(MenuBuilder, PlayMode_Simulate);
 	}
 	MenuBuilder.EndSection();
 
@@ -893,6 +895,14 @@ FSlateIcon FInternalPlayWorldCommandCallbacks::GetPossessEjectImage()
 }
 
 
+void SetLastExecutedPlayMode(EPlayModeType PlayMode)
+{
+	ULevelEditorPlaySettings* PlaySettings = GetMutableDefault<ULevelEditorPlaySettings>();
+	PlaySettings->LastExecutedPlayModeType = PlayMode;
+	PlaySettings->SaveConfig();
+}
+
+
 void FInternalPlayWorldCommandCallbacks::Simulate_Clicked()
 {
 	// Is a simulation session already running?  If so, do nothing
@@ -901,6 +911,8 @@ void FInternalPlayWorldCommandCallbacks::Simulate_Clicked()
 		return;
 	}
 
+	SetLastExecutedPlayMode( PlayMode_Simulate );
+
 	FLevelEditorModule& LevelEditorModule = FModuleManager::GetModuleChecked<FLevelEditorModule>( TEXT("LevelEditor") );
 
 	TSharedPtr<ILevelViewport> ActiveLevelViewport = LevelEditorModule.GetFirstActiveViewport();
@@ -961,19 +973,14 @@ const TSharedRef < FUICommandInfo > GetLastPlaySessionCommand()
 	case PlayMode_InNewProcess:
 		Command = Commands.PlayInNewProcess.ToSharedRef();
 		break;
+	case PlayMode_Simulate:
+		Command = Commands.Simulate.ToSharedRef();
 	}
 
 	return Command;
 }
 
 
-void SetLastExecutedPlayMode( EPlayModeType PlayMode )
-{
-	ULevelEditorPlaySettings* PlaySettings = GetMutableDefault<ULevelEditorPlaySettings>();
-	PlaySettings->LastExecutedPlayModeType = PlayMode;
-	PlaySettings->SaveConfig();
-}
-
 
 /** Report PIE usage to engine analytics */
 void RecordLastExecutedPlayMode()
@@ -1020,6 +1027,10 @@ void RecordLastExecutedPlayMode()
 			PlayModeString = TEXT("InNewProcess");
 			break;
 
+		case PlayMode_Simulate:
+			PlayModeString = TEXT("Simulate");
+			break;
+
 		default:
 			PlayModeString = TEXT("<UNKNOWN>");
 		}
