LaunchOn: Fixed LaunchOn not stopping in case of SDK errors

diff --git a/Engine/Source/Editor/UnrealEd/Private/Kismet2/DebuggerCommands.cpp b/Engine/Source/Editor/UnrealEd/Private/Kismet2/DebuggerCommands.cpp
--- a/Engine/Source/Editor/UnrealEd/Private/Kismet2/DebuggerCommands.cpp
+++ b/Engine/Source/Editor/UnrealEd/Private/Kismet2/DebuggerCommands.cpp
@@ -1438,6 +1438,8 @@ bool FInternalPlayWorldCommandCallbacks::IsReadyToLaunchOnDevice(FString DeviceI
 			FMessageLog MessageLog("PackagingResults");
 			MessageLog.AddMessage(Message);
 			MessageLog.Open();
+
+			return false;
 		}
 
 		// report to main frame
