fix non unity build

diff --git a/Engine/Source/Editor/UnrealEd/Private/Toolkits/SStandaloneAssetEditorToolkitHost.cpp b/Engine/Source/Editor/UnrealEd/Private/Toolkits/SStandaloneAssetEditorToolkitHost.cpp
--- a/Engine/Source/Editor/UnrealEd/Private/Toolkits/SStandaloneAssetEditorToolkitHost.cpp
+++ b/Engine/Source/Editor/UnrealEd/Private/Toolkits/SStandaloneAssetEditorToolkitHost.cpp
@@ -7,6 +7,7 @@
 #include "Toolkits/SStandaloneAssetEditorToolkitHost.h"
 #include "Toolkits/SAssetEditorCommon.h"
 #include "Toolkits/AssetEditorToolkit.h"
+#include "SuperSearchModule.h"
 #include "MainFrame.h"
 
 #define LOCTEXT_NAMESPACE "StandaloneAssetEditorToolkit"
