Remove UFactory::ValidGameNames, any such factory should not be in shared code

diff --git a/Engine/Source/Developer/AssetTools/Private/AssetTools.cpp b/Engine/Source/Developer/AssetTools/Private/AssetTools.cpp
--- a/Engine/Source/Developer/AssetTools/Private/AssetTools.cpp
+++ b/Engine/Source/Developer/AssetTools/Private/AssetTools.cpp
@@ -326,7 +326,7 @@ TArray<UObject*> FAssetTools::ImportAssets(const FString& DestinationPath)
 		if( CurrentClass->IsChildOf(UFactory::StaticClass()) && !(CurrentClass->HasAnyClassFlags(CLASS_Abstract)) )
 		{
 			UFactory* Factory = Cast<UFactory>( CurrentClass->GetDefaultObject() );
-			if( Factory->bEditorImport && Factory->ValidForCurrentGame() )
+			if( Factory->bEditorImport )
 			{
 				Factories.Add( Factory );
 			}
@@ -446,7 +446,7 @@ TArray<UObject*> FAssetTools::ImportAssets(const TArray<FString>& Files, const F
 			if( (*ClassIt)->IsChildOf(UFactory::StaticClass()) && !((*ClassIt)->HasAnyClassFlags(CLASS_Abstract)) )
 			{
 				UFactory* Factory = Cast<UFactory>( (*ClassIt)->GetDefaultObject() );
-				if( Factory->bEditorImport && Factory->ValidForCurrentGame() )
+				if( Factory->bEditorImport )
 				{
 					TArray<FString> FactoryExtensions;
 					Factory->GetSupportedFileExtensions(FactoryExtensions);
@@ -491,7 +491,7 @@ TArray<UObject*> FAssetTools::ImportAssets(const TArray<FString>& Files, const F
 	}
 	else
 	{
-		if( SpecifiedFactory->bEditorImport && SpecifiedFactory->ValidForCurrentGame() )
+		if( SpecifiedFactory->bEditorImport )
 		{
 			TArray<FString> FactoryExtensions;
 			SpecifiedFactory->GetSupportedFileExtensions(FactoryExtensions);
diff --git a/Engine/Source/Editor/ContentBrowser/Private/NewAssetContextMenu.cpp b/Engine/Source/Editor/ContentBrowser/Private/NewAssetContextMenu.cpp
--- a/Engine/Source/Editor/ContentBrowser/Private/NewAssetContextMenu.cpp
+++ b/Engine/Source/Editor/ContentBrowser/Private/NewAssetContextMenu.cpp
@@ -26,7 +26,7 @@ TArray<FFactoryItem> FindFactoriesInCategory(EAssetTypeCategories::Type AssetTyp
 		if (Class->IsChildOf(UFactory::StaticClass()) && !Class->HasAnyClassFlags(CLASS_Abstract))
 		{
 			UFactory* Factory = Class->GetDefaultObject<UFactory>();
-			if (Factory->ShouldShowInNewMenu() && Factory->ValidForCurrentGame() && ensure(!Factory->GetDisplayName().IsEmpty()))
+			if (Factory->ShouldShowInNewMenu() && ensure(!Factory->GetDisplayName().IsEmpty()))
 			{
 				uint32 FactoryCategories = Factory->GetMenuCategories();
 
diff --git a/Engine/Source/Editor/UnrealEd/Classes/Factories/Factory.h b/Engine/Source/Editor/UnrealEd/Classes/Factories/Factory.h
--- a/Engine/Source/Editor/UnrealEd/Classes/Factories/Factory.h
+++ b/Engine/Source/Editor/UnrealEd/Classes/Factories/Factory.h
@@ -48,10 +48,6 @@ public:
 	UPROPERTY()
 	int32 AutoPriority;
 
-	/** List of game names that this factory can be used for (if empty, all games valid) */
-	UPROPERTY()
-	TArray<FString> ValidGameNames;
-
 	static FString GetCurrentFilename() { return CurrentFilename; }
 
 	static FString CurrentFilename;
@@ -149,17 +145,6 @@ public:
 	static UObject* StaticImportObject( UClass* Class, UObject* InOuter, FName Name, EObjectFlags Flags, const TCHAR* Filename=TEXT(""), UObject* Context=NULL, UFactory* Factory=NULL, const TCHAR* Parms=NULL, FFeedbackContext* Warn=GWarn, int32 MaxImportFileSize = 0xC100000 );
 	static UObject* StaticImportObject( UClass* Class, UObject* InOuter, FName Name, EObjectFlags Flags, bool& bOutOperationCanceled, const TCHAR* Filename=TEXT(""), UObject* Context=NULL, UFactory* Factory=NULL, const TCHAR* Parms=NULL, FFeedbackContext* Warn=GWarn, int32 MaxImportFileSize = 0xC100000 );
 
-	/**
-	 * UFactory::ValidForCurrentGame
-	 *
-	 * Search through list of valid game names
-	 * If list is empty or current game name is in the list - return true
-	 * Otherwise, name wasn't found in the list - return false
-	 *
-	 * author: superville
-	 */
-	bool ValidForCurrentGame();
-
 	/** Creates a list of file extensions supported by this factory */
 	void GetSupportedFileExtensions( TArray<FString>& OutExtensions ) const;
 
diff --git a/Engine/Source/Editor/UnrealEd/Private/Editor.cpp b/Engine/Source/Editor/UnrealEd/Private/Editor.cpp
--- a/Engine/Source/Editor/UnrealEd/Private/Editor.cpp
+++ b/Engine/Source/Editor/UnrealEd/Private/Editor.cpp
@@ -2730,7 +2730,7 @@ void FReimportManager::GetNewReimportPath(UObject* Obj, TArray<FString>& InOutFi
 		if( CurrentClass->IsChildOf(UFactory::StaticClass()) && !(CurrentClass->HasAnyClassFlags(CLASS_Abstract)) )
 		{
 			UFactory* Factory = Cast<UFactory>( CurrentClass->GetDefaultObject() );
-			if( Factory->bEditorImport && Factory->ValidForCurrentGame() && Obj->GetClass()->IsChildOf(Factory->GetSupportedClass()) )
+			if( Factory->bEditorImport && Obj->GetClass()->IsChildOf(Factory->GetSupportedClass()) )
 			{
 				Factories.Add( Factory );
 			}
diff --git a/Engine/Source/Editor/UnrealEd/Private/Factories/Factory.cpp b/Engine/Source/Editor/UnrealEd/Private/Factories/Factory.cpp
--- a/Engine/Source/Editor/UnrealEd/Private/Factories/Factory.cpp
+++ b/Engine/Source/Editor/UnrealEd/Private/Factories/Factory.cpp
@@ -309,24 +309,6 @@ UObject* UFactory::StaticImportObject
 }
 
 
-bool UFactory::ValidForCurrentGame()
-{
-	if( ValidGameNames.Num() > 0 )
-	{
-		for( int32 Idx = 0; Idx < ValidGameNames.Num(); Idx++ )
-		{
-			if( FCString::Stricmp( FApp::GetGameName(), *ValidGameNames[Idx] ) == 0 )
-			{
-				return 1;
-			}
-		}
-
-		return 0;
-	}
-
-	return 1;
-}
-
 void UFactory::GetSupportedFileExtensions(TArray<FString>& OutExtensions) const
 {
 	for ( int32 FormatIdx = 0; FormatIdx < Formats.Num(); ++FormatIdx )
diff --git a/Engine/Source/Editor/UnrealEd/Private/Tests/AutomationEditorCommon.cpp b/Engine/Source/Editor/UnrealEd/Private/Tests/AutomationEditorCommon.cpp
--- a/Engine/Source/Editor/UnrealEd/Private/Tests/AutomationEditorCommon.cpp
+++ b/Engine/Source/Editor/UnrealEd/Private/Tests/AutomationEditorCommon.cpp
@@ -171,7 +171,7 @@ namespace AutomationEditorCommonUtils
 			if ((*ClassIt)->IsChildOf(UFactory::StaticClass()) && !((*ClassIt)->HasAnyClassFlags(CLASS_Abstract)))
 			{
 				UFactory* Factory = Cast<UFactory>((*ClassIt)->GetDefaultObject());
-				if (Factory->bEditorImport && Factory->ValidForCurrentGame())
+				if (Factory->bEditorImport)
 				{
 					TArray<FString> FactoryExtensions;
 					Factory->GetSupportedFileExtensions(FactoryExtensions);
