Added the commandline to the Editor startup analytics

TTP# 346678 - Analytics: Add editor commandline to startup stats

The "Editor.ProgramStarted" analytics event now includes a "CommandLine" attribute containing the contents of FCommandLine.

diff --git a/Engine/Source/Editor/UnrealEd/Private/UnrealEd.cpp b/Engine/Source/Editor/UnrealEd/Private/UnrealEd.cpp
--- a/Engine/Source/Editor/UnrealEd/Private/UnrealEd.cpp
+++ b/Engine/Source/Editor/UnrealEd/Private/UnrealEd.cpp
@@ -67,6 +67,7 @@ int32 EditorInit( IEngineLoop& EngineLoop )
 		TArray<FAnalyticsEventAttribute> EventAttributes;
 		EventAttributes.Add(FAnalyticsEventAttribute(TEXT("MachineID"), FPlatformMisc::GetMachineId().ToString(EGuidFormats::Digits).ToLower()));
 		EventAttributes.Add(FAnalyticsEventAttribute(TEXT("AccountID"), FPlatformMisc::GetEpicAccountId()));
+		EventAttributes.Add(FAnalyticsEventAttribute(TEXT("CommandLine"), FCommandLine::Get()));
 
 		FEngineAnalytics::GetProvider().RecordEvent(TEXT("Editor.ProgramStarted"), EventAttributes);
 	}
