Fix remap dialog not being cleaned up properly

diff --git a/Engine/Source/Editor/UnrealEd/Private/SSkeletonWidget.cpp b/Engine/Source/Editor/UnrealEd/Private/SSkeletonWidget.cpp
--- a/Engine/Source/Editor/UnrealEd/Private/SSkeletonWidget.cpp
+++ b/Engine/Source/Editor/UnrealEd/Private/SSkeletonWidget.cpp
@@ -327,6 +327,8 @@ void SSkeletonSelectorWindow::ConstructWindow()
 		];
 }
 
+TSharedPtr<SWindow> SAnimationRemapSkeleton::DialogWindow;
+
 bool SAnimationRemapSkeleton::OnShouldFilterAsset(const class FAssetData& AssetData)
 {
 	USkeleton * AssetSkeleton = NULL;
@@ -690,6 +692,7 @@ FReply SAnimationRemapSkeleton::OnApply()
 	CloseWindow();
 	return FReply::Handled();
 }
+
 FReply SAnimationRemapSkeleton::OnCancel()
 {
 	NewSkeleton = NULL;
@@ -697,6 +700,12 @@ FReply SAnimationRemapSkeleton::OnCancel()
 	return FReply::Handled();
 }
 
+void SAnimationRemapSkeleton::OnRemapDialogClosed(const TSharedRef<SWindow>& Window)
+{
+	NewSkeleton = NULL;
+	DialogWindow = nullptr;
+}
+
 void SAnimationRemapSkeleton::CloseWindow()
 {
 	if ( WidgetWindow.IsValid() )
@@ -707,8 +716,6 @@ void SAnimationRemapSkeleton::CloseWindow()
 
 void SAnimationRemapSkeleton::ShowWindow(USkeleton * OldSkeleton, const FText& WarningMessage, FOnRetargetAnimation RetargetDelegate)
 {
-	static TSharedPtr<SWindow> DialogWindow;
-
 	if(DialogWindow.IsValid())
 	{
 		FSlateApplication::Get().DestroyWindowImmediately(DialogWindow.ToSharedRef());
@@ -736,6 +743,7 @@ void SAnimationRemapSkeleton::ShowWindow(USkeleton * OldSkeleton, const FText& W
 			.OnRetargetDelegate(RetargetDelegate)
 		];
 
+	DialogWindow->SetOnWindowClosed(FRequestDestroyWindowOverride::CreateSP(DialogWidget.Get(), &SAnimationRemapSkeleton::OnRemapDialogClosed));
 	DialogWindow->SetContent(DialogWrapper.ToSharedRef());
 
 	FSlateApplication::Get().AddWindow(DialogWindow.ToSharedRef());
diff --git a/Engine/Source/Editor/UnrealEd/Public/SSkeletonWidget.h b/Engine/Source/Editor/UnrealEd/Public/SSkeletonWidget.h
--- a/Engine/Source/Editor/UnrealEd/Public/SSkeletonWidget.h
+++ b/Engine/Source/Editor/UnrealEd/Public/SSkeletonWidget.h
@@ -383,6 +383,9 @@ private:
 
 	void CloseWindow();
 
+	/** Handler for dialog window close button */
+	void OnRemapDialogClosed(const TSharedRef<SWindow>& Window);
+
 	/**
 	 * Handler for when asset is selected
 	 */
@@ -402,6 +405,8 @@ public:
 	 * @return true if successfully selected new skeleton
 	 */
 	static UNREALED_API void ShowWindow(USkeleton * OldSkeleton, const FText& WarningMessage, FOnRetargetAnimation RetargetDelegate);
+
+	static TSharedPtr<SWindow> DialogWindow;
 };
 
 ////////////////////////////////////////////////////
