Curve graph min and max are now preserved during blueprint changing/saving.

#ttp 344632 - EDITOR: Timeline Editor: Zoom Fit resets on Play/Stop

#branch UE4

#change The æfitÆ state of the curve isnÆt maintained because itÆs only fit for the frame you press the button, if pan or add another key it may no longer be æbest-fitÆ either. The issue is, whenever the blueprint (not the curve) is modified, the curve editor widget is recreated due to a Blueprint->BroadcastChanged delegate, which is why the curve state isnÆt maintained because when you PIE or Save the blueprint is recreates the widget. However, I noticed that this is only the case with the Vertical as the Horizonal is saved externally due to attributes and delegates, so I modified the vertical to use the same system.

diff --git a/Engine/Source/Editor/Kismet/Private/STimelineEditor.cpp b/Engine/Source/Editor/Kismet/Private/STimelineEditor.cpp
--- a/Engine/Source/Editor/Kismet/Private/STimelineEditor.cpp
+++ b/Engine/Source/Editor/Kismet/Private/STimelineEditor.cpp
@@ -219,8 +219,11 @@ void STimelineEdTrack::Construct(const FArguments& InArgs, TSharedPtr<FTimelineE
 					SAssignNew(TrackWidget, SCurveEditor)
 					.ViewMinInput(TimelineRef, &STimelineEditor::GetViewMinInput)
 					.ViewMaxInput(TimelineRef, &STimelineEditor::GetViewMaxInput)
+					.ViewMinOutput(TimelineRef, &STimelineEditor::GetViewMinOutput)
+					.ViewMaxOutput(TimelineRef, &STimelineEditor::GetViewMaxOutput)
 					.TimelineLength(TimelineRef, &STimelineEditor::GetTimelineLength)
 					.OnSetInputViewRange(TimelineRef, &STimelineEditor::SetInputViewRange)
+					.OnSetOutputViewRange(TimelineRef, &STimelineEditor::SetOutputViewRange)
 					.DesiredSize(TimelineRef, &STimelineEditor::GetTimelineDesiredSize)
 					.DrawCurve(bDrawCurve)
 					.HideUI(false)
@@ -655,6 +658,8 @@ void STimelineEditor::Construct(const FArguments& InArgs, TSharedPtr<FBlueprintE
 	// Leave these uninitialized at first.  We'll zoom to fit the tracks which will set the correct values
 	ViewMinInput = 0.f;
 	ViewMaxInput = 0.f;
+	ViewMinOutput = 0.f;
+	ViewMaxOutput = 0.f;
 
 	CommandList = MakeShareable( new FUICommandList );
 
@@ -897,6 +902,16 @@ float STimelineEditor::GetViewMinInput() const
 	return ViewMinInput;
 }
 
+float STimelineEditor::GetViewMaxOutput() const
+{
+	return ViewMaxOutput;
+}
+
+float STimelineEditor::GetViewMinOutput() const
+{
+	return ViewMinOutput;
+}
+
 float STimelineEditor::GetTimelineLength() const
 {
 	return (TimelineObj != NULL) ? TimelineObj->TimelineLength : 0.f;
@@ -908,6 +923,12 @@ void STimelineEditor::SetInputViewRange(float InViewMinInput, float InViewMaxInp
 	ViewMinInput = InViewMinInput;
 }
 
+void STimelineEditor::SetOutputViewRange(float InViewMinOutput, float InViewMaxOutput)
+{
+	ViewMaxOutput = InViewMaxOutput;
+	ViewMinOutput = InViewMinOutput;
+}
+
 TSharedRef<ITableRow> STimelineEditor::MakeTrackWidget( TSharedPtr<FTimelineEdTrack> Track, const TSharedRef<STableViewBase>& OwnerTable )
 {
 	check( Track.IsValid() );
diff --git a/Engine/Source/Editor/Kismet/Private/STimelineEditor.h b/Engine/Source/Editor/Kismet/Private/STimelineEditor.h
--- a/Engine/Source/Editor/Kismet/Private/STimelineEditor.h
+++ b/Engine/Source/Editor/Kismet/Private/STimelineEditor.h
@@ -171,6 +171,12 @@ private:
 	/** Maximum input shown for tracks */
 	float ViewMaxInput;
 
+	/** Minimum output shown for tracks */
+	float ViewMinOutput;
+
+	/** Maximum output shown for tracks */
+	float ViewMaxOutput;
+
 	/** The default name of the last track created, used to identify which track needs to be renamed. */
 	FName NewTrackPendingRename;
 
@@ -193,11 +199,15 @@ public:
 
 	float GetViewMinInput() const;
 	float GetViewMaxInput() const;
+	float GetViewMinOutput() const;
+	float GetViewMaxOutput() const;
 
 	/** Return length of timeline */
 	float GetTimelineLength() const;
 
 	void SetInputViewRange(float InViewMinInput, float InViewMaxInput);
+	void SetOutputViewRange(float InViewMinOutput, float InViewMaxOutput);
+
 	/** When user attempts to commit the name of a track*/
 	bool OnVerifyTrackNameCommit(const FText& TrackName, FText& OutErrorMessage, FTTTrackBase* TrackBase, STimelineEdTrack* Track );
 	/** When user commits the name of a track*/
diff --git a/Engine/Source/Editor/Persona/Private/SAnimCurvePanel.cpp b/Engine/Source/Editor/Persona/Private/SAnimCurvePanel.cpp
--- a/Engine/Source/Editor/Persona/Private/SAnimCurvePanel.cpp
+++ b/Engine/Source/Editor/Persona/Private/SAnimCurvePanel.cpp
@@ -257,8 +257,10 @@ void SAnimCurveEd::Construct(const FArguments& InArgs)
 
 void SAnimCurveEd::SetDefaultOutput(const float MinZoomRange)
 {
-	ViewMinOutput = (ViewMinOutput);
-	ViewMaxOutput = (ViewMaxOutput + (MinZoomRange));
+	const float NewMinOutput = (ViewMinOutput.Get());
+	const float NewMaxOutput = (ViewMaxOutput.Get() + MinZoomRange);
+
+	SetOutputMinMax(NewMinOutput, NewMaxOutput);
 }
 //////////////////////////////////////////////////////////////////////////
 //  SCurveEd Track : each track for curve editing 
diff --git a/Engine/Source/Editor/UnrealEd/Private/SCurveEditor.cpp b/Engine/Source/Editor/UnrealEd/Private/SCurveEditor.cpp
--- a/Engine/Source/Editor/UnrealEd/Private/SCurveEditor.cpp
+++ b/Engine/Source/Editor/UnrealEd/Private/SCurveEditor.cpp
@@ -56,6 +56,7 @@ void SCurveEditor::Construct(const FArguments& InArgs)
 
 	TimelineLength = InArgs._TimelineLength;
 	SetInputViewRangeHandler = InArgs._OnSetInputViewRange;
+	SetOutputViewRangeHandler = InArgs._OnSetOutputViewRange;
 	bDrawCurve = InArgs._DrawCurve;
 	bHideUI = InArgs._HideUI;
 	bAllowZoomOutput = InArgs._AllowZoomOutput;
@@ -140,150 +141,150 @@ void SCurveEditor::Construct(const FArguments& InArgs)
 		+ SHorizontalBox::Slot()
 		.FillWidth(1.0f)
 		[
-			SNew( SVerticalBox )
+		SNew( SVerticalBox )
 
-			+ SVerticalBox::Slot()
-			.FillHeight(1.0f)
+		+ SVerticalBox::Slot()
+		.FillHeight(1.0f)
+		[
+			SNew(SHorizontalBox)
+			.Visibility( this, &SCurveEditor::GetCurveAreaVisibility )
+
+			+ SHorizontalBox::Slot()
+			.AutoWidth()
+			.Padding(FMargin(30, 12, 0, 0))
 			[
-				SNew(SHorizontalBox)
-				.Visibility( this, &SCurveEditor::GetCurveAreaVisibility )
+				CurveSelector
+			]
 
-				+ SHorizontalBox::Slot()
-				.AutoWidth()
-				.Padding(FMargin(30, 12, 0, 0))
+			+ SHorizontalBox::Slot()
+			.AutoWidth()
+			[
+				SNew(SBorder)
+				.VAlign(VAlign_Top)
+				.HAlign(HAlign_Left)
+				.BorderImage( FEditorStyle::GetBrush("NoBorder") )
+				.DesiredSizeScale(FVector2D(256.0f,32.0f))
+				.Visibility(this, &SCurveEditor::GetControlVisibility)
+				.Padding(FMargin(2, 12, 0, 0))
 				[
-					CurveSelector
-				]
+					SNew(SHorizontalBox)
 
-				+ SHorizontalBox::Slot()
-				.AutoWidth()
-				[
-					SNew(SBorder)
-					.VAlign(VAlign_Top)
-					.HAlign(HAlign_Left)
-					.BorderImage( FEditorStyle::GetBrush("NoBorder") )
-					.DesiredSizeScale(FVector2D(256.0f,32.0f))
-					.Visibility(this, &SCurveEditor::GetControlVisibility)
-					.Padding(FMargin(2, 12, 0, 0))
+					+ SHorizontalBox::Slot()
+					.AutoWidth()
 					[
-						SNew(SHorizontalBox)
-
-						+ SHorizontalBox::Slot()
-						.AutoWidth()
+						SNew(SButton)
+						.ToolTipText(LOCTEXT("ZoomToFitHorizontal", "Zoom To Fit Horizontal"))
+						.Visibility(this, &SCurveEditor::GetZoomButtonVisibility)
+						.OnClicked(this, &SCurveEditor::ZoomToFitHorizontalClicked)
+						.ContentPadding(1)
 						[
-							SNew(SButton)
-							.ToolTipText(LOCTEXT("ZoomToFitHorizontal", "Zoom To Fit Horizontal"))
-							.Visibility(this, &SCurveEditor::GetZoomButtonVisibility)
-							.OnClicked(this, &SCurveEditor::ZoomToFitHorizontalClicked)
-							.ContentPadding(1)
-							[
-								SNew(SImage) 
-								.Image( FEditorStyle::GetBrush("CurveEd.FitHorizontal") ) 
-								.ColorAndOpacity( FSlateColor::UseForeground() ) 
-							]
+							SNew(SImage) 
+							.Image( FEditorStyle::GetBrush("CurveEd.FitHorizontal") ) 
+							.ColorAndOpacity( FSlateColor::UseForeground() ) 
 						]
+					]
 					
-						+ SHorizontalBox::Slot()
-						.AutoWidth()
+					+ SHorizontalBox::Slot()
+					.AutoWidth()
+					[
+						SNew(SButton)
+						.ToolTipText(LOCTEXT("ZoomToFitVertical", "Zoom To Fit Vertical"))
+						.Visibility(this, &SCurveEditor::GetZoomButtonVisibility)
+						.OnClicked(this, &SCurveEditor::ZoomToFitVerticalClicked)
+						.ContentPadding(1)
 						[
-							SNew(SButton)
-							.ToolTipText(LOCTEXT("ZoomToFitVertical", "Zoom To Fit Vertical"))
-							.Visibility(this, &SCurveEditor::GetZoomButtonVisibility)
-							.OnClicked(this, &SCurveEditor::ZoomToFitVerticalClicked)
-							.ContentPadding(1)
-							[
-								SNew(SImage) 
-								.Image( FEditorStyle::GetBrush("CurveEd.FitVertical") ) 
-								.ColorAndOpacity( FSlateColor::UseForeground() ) 
-							]
+							SNew(SImage) 
+							.Image( FEditorStyle::GetBrush("CurveEd.FitVertical") ) 
+							.ColorAndOpacity( FSlateColor::UseForeground() ) 
 						]
+					]
 
-						+ SHorizontalBox::Slot()
-						.AutoWidth()
+					+ SHorizontalBox::Slot()
+					.AutoWidth()
+					[
+						SNew(SBorder)
+						.BorderImage(FEditorStyle::GetBrush("NoBorder"))
+						.Visibility(this, &SCurveEditor::GetEditVisibility)
+						.VAlign(VAlign_Center)
 						[
-							SNew(SBorder)
-							.BorderImage(FEditorStyle::GetBrush("NoBorder"))
-							.Visibility(this, &SCurveEditor::GetEditVisibility)
-							.VAlign(VAlign_Center)
+							SNew(SNumericEntryBox<float>)
+							.IsEnabled(this, &SCurveEditor::GetInputEditEnabled)
+							.Font(FEditorStyle::GetFontStyle("CurveEd.InfoFont"))
+							.Value(this, &SCurveEditor::OnGetTime)
+							.UndeterminedString(LOCTEXT("MultipleValues", "Multiple Values").ToString())
+							.OnValueCommitted(this, &SCurveEditor::OnTimeComitted)
+							.OnValueChanged(this, &SCurveEditor::OnTimeChanged)
+							.OnBeginSliderMovement(this, &SCurveEditor::OnBeginSliderMovement, LOCTEXT("SetValue", "Set New Time"))
+							.OnEndSliderMovement(this, &SCurveEditor::OnEndSliderMovement)
+							.LabelVAlign(VAlign_Center)
+							.AllowSpin(true)
+							.MinValue(TOptional<float>())
+							.MaxValue(TOptional<float>())
+							.MaxSliderValue(TOptional<float>())
+							.MinSliderValue(TOptional<float>())
+							.Delta(0.001f)
+							.MinDesiredValueWidth(60.0f)
+							.Label()
 							[
-								SNew(SNumericEntryBox<float>)
-								.IsEnabled(this, &SCurveEditor::GetInputEditEnabled)
-								.Font(FEditorStyle::GetFontStyle("CurveEd.InfoFont"))
-								.Value(this, &SCurveEditor::OnGetTime)
-								.UndeterminedString(LOCTEXT("MultipleValues", "Multiple Values").ToString())
-								.OnValueCommitted(this, &SCurveEditor::OnTimeComitted)
-								.OnValueChanged(this, &SCurveEditor::OnTimeChanged)
-								.OnBeginSliderMovement(this, &SCurveEditor::OnBeginSliderMovement, LOCTEXT("SetValue", "Set New Time"))
-								.OnEndSliderMovement(this, &SCurveEditor::OnEndSliderMovement)
-								.LabelVAlign(VAlign_Center)
-								.AllowSpin(true)
-								.MinValue(TOptional<float>())
-								.MaxValue(TOptional<float>())
-								.MaxSliderValue(TOptional<float>())
-								.MinSliderValue(TOptional<float>())
-								.Delta(0.001f)
-								.MinDesiredValueWidth(60.0f)
-								.Label()
-								[
-									SNew(STextBlock)
-									.Text(LOCTEXT("Time", "Time"))
-								]
+								SNew(STextBlock)
+								.Text(LOCTEXT("Time", "Time"))
 							]
 						]
+					]
 					
-						+ SHorizontalBox::Slot()
-						.AutoWidth()
+					+ SHorizontalBox::Slot()
+					.AutoWidth()
+					[
+						SNew(SBorder)
+						.BorderImage( FEditorStyle::GetBrush("NoBorder") )
+						.Visibility(this, &SCurveEditor::GetEditVisibility)
+						.VAlign(VAlign_Center)
 						[
-							SNew(SBorder)
-							.BorderImage( FEditorStyle::GetBrush("NoBorder") )
-							.Visibility(this, &SCurveEditor::GetEditVisibility)
-							.VAlign(VAlign_Center)
+							SNew(SNumericEntryBox<float>)
+							.Font(FEditorStyle::GetFontStyle("CurveEd.InfoFont"))
+							.Value(this, &SCurveEditor::OnGetValue)
+							.UndeterminedString(LOCTEXT("MultipleValues", "Multiple Values").ToString())
+							.OnValueCommitted(this, &SCurveEditor::OnValueComitted)
+							.OnValueChanged(this, &SCurveEditor::OnValueChanged)
+							.OnBeginSliderMovement(this, &SCurveEditor::OnBeginSliderMovement, LOCTEXT("SetValue", "Set New Value"))
+							.OnEndSliderMovement(this, &SCurveEditor::OnEndSliderMovement)
+							.LabelVAlign(VAlign_Center)
+							.AllowSpin(true)
+							.MinValue(TOptional<float>())
+							.MaxValue(TOptional<float>())
+							.MaxSliderValue(TOptional<float>())
+							.MinSliderValue(TOptional<float>())
+							.Delta(.001f)
+							.MinDesiredValueWidth(60.0f)
+							.Label()
 							[
-								SNew(SNumericEntryBox<float>)
-								.Font(FEditorStyle::GetFontStyle("CurveEd.InfoFont"))
-								.Value(this, &SCurveEditor::OnGetValue)
-								.UndeterminedString(LOCTEXT("MultipleValues", "Multiple Values").ToString())
-								.OnValueCommitted(this, &SCurveEditor::OnValueComitted)
-								.OnValueChanged(this, &SCurveEditor::OnValueChanged)
-								.OnBeginSliderMovement(this, &SCurveEditor::OnBeginSliderMovement, LOCTEXT("SetValue", "Set New Value"))
-								.OnEndSliderMovement(this, &SCurveEditor::OnEndSliderMovement)
-								.LabelVAlign(VAlign_Center)
-								.AllowSpin(true)
-								.MinValue(TOptional<float>())
-								.MaxValue(TOptional<float>())
-								.MaxSliderValue(TOptional<float>())
-								.MinSliderValue(TOptional<float>())
-								.Delta(.001f)
-								.MinDesiredValueWidth(60.0f)
-								.Label()
-								[
-									SNew(STextBlock)
-									.Text(LOCTEXT("Value", "Value"))
-								]
+								SNew(STextBlock)
+								.Text(LOCTEXT("Value", "Value"))
 							]
 						]
 					]
 				]
 			]
-		
+		]
+
 
-			+ SVerticalBox::Slot()
-			.VAlign(VAlign_Bottom)
-			.FillHeight(.75f)
+		+ SVerticalBox::Slot()
+		.VAlign(VAlign_Bottom)
+		.FillHeight(.75f)
+		[
+			SNew( SBorder )
+			.Visibility( this, &SCurveEditor::GetColorGradientVisibility )
+			.BorderImage( FEditorStyle::GetBrush("ToolPanel.GroupBorder") )
+			.BorderBackgroundColor( FLinearColor( .8f, .8f, .8f, .60f ) )
+			.Padding(1.0f)
 			[
-				SNew( SBorder )
-				.Visibility( this, &SCurveEditor::GetColorGradientVisibility )
-				.BorderImage( FEditorStyle::GetBrush("ToolPanel.GroupBorder") )
-				.BorderBackgroundColor( FLinearColor( .8f, .8f, .8f, .60f ) )
-				.Padding(1.0f)
-				[
-					SAssignNew( GradientViewer, SColorGradientEditor )
-					.ViewMinInput( ViewMinInput )
-					.ViewMaxInput( ViewMaxInput )
-					.IsEditingEnabled( this, &SCurveEditor::IsEditingEnabled )
-				]
+				SAssignNew( GradientViewer, SColorGradientEditor )
+				.ViewMinInput( ViewMinInput )
+				.ViewMaxInput( ViewMaxInput )
+				.IsEditingEnabled( this, &SCurveEditor::IsEditingEnabled )
 			]
 		]
+		]
 	];
 
 	if (GEditor != NULL)
@@ -531,7 +532,7 @@ int32 SCurveEditor::OnPaint(const FPaintArgs& Args, const FGeometry& AllottedGeo
 	FGeometry CurveAreaGeometry = AllottedGeometry;
 
 	// Positioning info
-	FTrackScaleInfo ScaleInfo(ViewMinInput.Get(),  ViewMaxInput.Get(), ViewMinOutput, ViewMaxOutput, CurveAreaGeometry.Size);
+	FTrackScaleInfo ScaleInfo(ViewMinInput.Get(),  ViewMaxInput.Get(), ViewMinOutput.Get(), ViewMaxOutput.Get(), CurveAreaGeometry.Size);
 
 	// Draw background to indicate valid timeline area
 	float ZeroInputX = ScaleInfo.InputToLocalX(0.f);
@@ -557,16 +558,16 @@ int32 SCurveEditor::OnPaint(const FPaintArgs& Args, const FGeometry& AllottedGeo
 
 	// time=0 line
 	int32 ZeroLineLayerId = GridLineLayerId + 1;
-	FSlateDrawElement::MakeBox
-	(
-		OutDrawElements,
+		FSlateDrawElement::MakeBox
+		(
+			OutDrawElements,
 		ZeroLineLayerId,
 		CurveAreaGeometry.ToPaintGeometry(FVector2D(ZeroInputX, 0), FVector2D(1, CurveAreaGeometry.Size.Y)),
-		WhiteBrush,
-		MyClippingRect,
-		DrawEffects,
+			WhiteBrush,
+			MyClippingRect,
+			DrawEffects,
 		WhiteBrush->GetTint(InWidgetStyle) * InWidgetStyle.GetColorAndOpacityTint()
-	);
+		);
 
 	// value=0 line
 	if( AreCurvesVisible() )
@@ -588,23 +589,23 @@ int32 SCurveEditor::OnPaint(const FPaintArgs& Args, const FGeometry& AllottedGeo
 
 	int32 KeyLayerId = CurveLayerId + 1;
 	int32 SelectedKeyLayerId = KeyLayerId + 1;
-	
+
 	if( AreCurvesVisible() )
 	{
 		//Paint the curves, unlocked curves will be on top
 		for ( auto CurveViewModel : CurveViewModels )
 		{
 			if (CurveViewModel->bIsVisible)
-			{
+				{
 				PaintCurve(CurveViewModel, CurveAreaGeometry, ScaleInfo, OutDrawElements, CurveViewModel->bIsLocked ? LockedCurveLayerID : CurveLayerId, MyClippingRect, DrawEffects, InWidgetStyle);
+				}
 			}
-		}
 
 		//Paint the keys on top of the curve
 		for (auto CurveViewModel : CurveViewModels)
 		{
 			if (CurveViewModel->bIsVisible)
-			{
+		{
 				PaintKeys(CurveViewModel, ScaleInfo, OutDrawElements, KeyLayerId, SelectedKeyLayerId, CurveAreaGeometry, MyClippingRect, DrawEffects, InWidgetStyle);
 			}
 		}
@@ -1082,18 +1083,18 @@ void SCurveEditor::AddNewKey(FGeometry InMyGeometry, FVector2D ScreenPosition)
 {
 	bool bKeyAdded = false;
 	for (auto CurveViewModel : CurveViewModels)
-	{
-		if (!CurveViewModel->bIsLocked)
 		{
+		if (!CurveViewModel->bIsLocked)
+	{
 			FRichCurve* SelectedCurve = CurveViewModel->CurveInfo.CurveToEdit;
-			if ( IsValidCurve(SelectedCurve) )
-			{
-				const FScopedTransaction Transaction(LOCTEXT("CurveEditor_AddKey", "Add Key(s)"));
-				CurveOwner->ModifyOwner();
+		if ( IsValidCurve(SelectedCurve) )
+		{
+			const FScopedTransaction Transaction(LOCTEXT("CurveEditor_AddKey", "Add Key(s)"));
+			CurveOwner->ModifyOwner();
 
-				FTrackScaleInfo ScaleInfo(ViewMinInput.Get(), ViewMaxInput.Get(), ViewMinOutput, ViewMaxOutput, InMyGeometry.Size);
+			FTrackScaleInfo ScaleInfo(ViewMinInput.Get(), ViewMaxInput.Get(), ViewMinOutput.Get(), ViewMaxOutput.Get(), InMyGeometry.Size);
 
-				FVector2D LocalClickPos = InMyGeometry.AbsoluteToLocal(ScreenPosition);
+			FVector2D LocalClickPos = InMyGeometry.AbsoluteToLocal(ScreenPosition);
 
 				FVector2D NewKeyLocation = SnapLocation(FVector2D
 				(
@@ -1102,8 +1103,8 @@ void SCurveEditor::AddNewKey(FGeometry InMyGeometry, FVector2D ScreenPosition)
 				);
 				FKeyHandle NewKeyHandle = SelectedCurve->AddKey(NewKeyLocation.X, NewKeyLocation.Y);
 
-				EmptySelection();
-				AddToSelection(FSelectedCurveKey(SelectedCurve, NewKeyHandle));
+			EmptySelection();
+			AddToSelection(FSelectedCurveKey(SelectedCurve, NewKeyHandle));
 				bKeyAdded = true;
 			}
 		}
@@ -1111,8 +1112,8 @@ void SCurveEditor::AddNewKey(FGeometry InMyGeometry, FVector2D ScreenPosition)
 
 	if (bKeyAdded)
 	{
-		CurveOwner->OnCurveChanged();
-	}
+	CurveOwner->OnCurveChanged();
+}
 }
 
 void SCurveEditor::OnMouseCaptureLost()
@@ -1187,11 +1188,13 @@ FReply SCurveEditor::OnMouseWheel( const FGeometry& MyGeometry, const FPointerEv
 
 	if (bAllowZoomOutput)
 	{
-		const float OutputViewSize = ViewMaxOutput - ViewMinOutput;
+		const float OutputViewSize = ViewMaxOutput.Get() - ViewMinOutput.Get();
 		const float OutputChange = OutputViewSize * ZoomDelta;
 
-		ViewMinOutput = (ViewMinOutput - (OutputChange * 0.5f));
-		ViewMaxOutput = (ViewMaxOutput + (OutputChange * 0.5f));
+		const float NewMinOutput = (ViewMinOutput.Get() - (OutputChange * 0.5f));
+		const float NewMaxOutput = (ViewMaxOutput.Get() + (OutputChange * 0.5f));
+
+		SetOutputMinMax(NewMinOutput, NewMaxOutput);
 	}
 
 	{
@@ -1291,7 +1294,7 @@ void SCurveEditor::TryStartDrag(const FGeometry& InMyGeometry, const FPointerEve
 
 void SCurveEditor::ProcessDrag(const FGeometry& InMyGeometry, const FPointerEvent& InMouseEvent)
 {
-	FTrackScaleInfo ScaleInfo(ViewMinInput.Get(), ViewMaxInput.Get(), ViewMinOutput, ViewMaxOutput, InMyGeometry.Size);
+	FTrackScaleInfo ScaleInfo(ViewMinInput.Get(), ViewMaxInput.Get(), ViewMinOutput.Get(), ViewMaxOutput.Get(), InMyGeometry.Size);
 	FVector2D ScreenDelta = InMouseEvent.GetCursorDelta();
 
 	FVector2D InputDelta;
@@ -1313,8 +1316,10 @@ void SCurveEditor::ProcessDrag(const FGeometry& InMyGeometry, const FPointerEven
 	else if (DragState == EDragState::Pan)
 	{
 		// Output is not clamped.
-		ViewMinOutput = (ViewMinOutput - InputDelta.Y);
-		ViewMaxOutput = (ViewMaxOutput - InputDelta.Y);
+		const float NewMinOutput = (ViewMinOutput.Get() - InputDelta.Y);
+		const float NewMaxOutput = (ViewMaxOutput.Get() - InputDelta.Y);
+
+		SetOutputMinMax(NewMinOutput, NewMaxOutput);
 
 		// Input maybe clamped if DataMinInput or DataMaxOutput was set.
 		float NewMinInput = ViewMinInput.Get() - InputDelta.X;
@@ -1331,21 +1336,21 @@ void SCurveEditor::EndDrag(const FGeometry& InMyGeometry, const FPointerEvent& I
 
 	if (DragState == EDragState::DragKey || DragState == EDragState::DragTangent)
 	{
-		EndDragTransaction();
-	}
+			EndDragTransaction();
+		}
 	else if (DragState == EDragState::MarqueeSelect)
 	{
 		FVector2D MarqueTopLeft
 		(
 			FMath::Min(MouseDownLocation.X, MouseMoveLocation.X),
 			FMath::Min(MouseDownLocation.Y, MouseMoveLocation.Y)
-		);
+			);
 
 		FVector2D MarqueBottomRight
 		(
 			FMath::Max(MouseDownLocation.X, MouseMoveLocation.X),
 			FMath::Max(MouseDownLocation.Y, MouseMoveLocation.Y)
-		);
+			);
 
 		TArray<FSelectedCurveKey> SelectedCurveKeys = GetEditableKeysWithinMarquee(InMyGeometry, MarqueTopLeft, MarqueBottomRight);
 
@@ -1378,59 +1383,59 @@ void SCurveEditor::ProcessClick(const FGeometry& InMyGeometry, const FPointerEve
 	const bool bShiftDown = InMouseEvent.IsShiftDown();
 
 	FSelectedCurveKey HitKey = HitTestKeys(InMyGeometry, InMouseEvent.GetScreenSpacePosition());
-	if (bLeftMouseButton)
-	{
-		// If the user left clicked a key, update selection based on modifier key state.
-		if (HitKey.IsValid())
+		if (bLeftMouseButton)
 		{
-			if (!IsKeySelected(HitKey))
+			// If the user left clicked a key, update selection based on modifier key state.
+			if (HitKey.IsValid())
 			{
-				if (!bControlDown)
+				if (!IsKeySelected(HitKey))
 				{
-					EmptySelection();
+					if (!bControlDown)
+					{
+						EmptySelection();
+					}
+					AddToSelection(HitKey);
+				}
+				else if (bControlDown)
+				{
+					RemoveFromSelection(HitKey);
 				}
-				AddToSelection(HitKey);
-			}
-			else if (bControlDown)
-			{
-				RemoveFromSelection(HitKey);
-			}
-		}
-		else
-		{
-			// If the user didn't click a key, add a new one if shift is held down, or try to select a curve.
-			if (bShiftDown)
-			{
-				AddNewKey(InMyGeometry, InMouseEvent.GetScreenSpacePosition());
 			}
 			else
 			{
-				// clicking on background clears key selection
-				EmptySelection();
+				// If the user didn't click a key, add a new one if shift is held down, or try to select a curve.
+				if (bShiftDown)
+				{
+					AddNewKey(InMyGeometry, InMouseEvent.GetScreenSpacePosition());
+				}
+				else
+				{
+					// clicking on background clears key selection
+					EmptySelection();
+				}
 			}
 		}
-	}
-	else if (bRightMouseButton)
-	{
-		// If the user right clicked, handle opening context menus.
-		if (HitKey.IsValid())
+		else if (bRightMouseButton)
 		{
-			// Make sure key is selected in readiness for context menu
-			if (!IsKeySelected(HitKey))
+			// If the user right clicked, handle opening context menus.
+			if (HitKey.IsValid())
 			{
-				EmptySelection();
-				AddToSelection(HitKey);
-			}
+				// Make sure key is selected in readiness for context menu
+				if (!IsKeySelected(HitKey))
+				{
+					EmptySelection();
+					AddToSelection(HitKey);
+				}
 			PushKeyMenu(InMyGeometry, InMouseEvent.GetScreenSpacePosition());
-		}
-		else
-		{
-			FVector2D ScreenPos = InMouseEvent.GetScreenSpacePosition();
-			if (!HitTestCubicTangents(InMyGeometry, ScreenPos).IsValid())
-			{
-				CreateContextMenu(InMyGeometry, ScreenPos);
 			}
 			else
+			{
+				FVector2D ScreenPos = InMouseEvent.GetScreenSpacePosition();
+				if (!HitTestCubicTangents(InMyGeometry, ScreenPos).IsValid())
+				{
+					CreateContextMenu(InMyGeometry, ScreenPos);
+				}
+			else
 			{
 				EmptySelection();
 			}
@@ -1578,7 +1583,7 @@ SCurveEditor::FSelectedCurveKey SCurveEditor::HitTestKeys(const FGeometry& InMyG
 
 	if( AreCurvesVisible() )
 	{
-		FTrackScaleInfo ScaleInfo(ViewMinInput.Get(),  ViewMaxInput.Get(), ViewMinOutput, ViewMaxOutput, InMyGeometry.Size);
+		FTrackScaleInfo ScaleInfo(ViewMinInput.Get(),  ViewMaxInput.Get(), ViewMinOutput.Get(), ViewMaxOutput.Get(), InMyGeometry.Size);
 
 		const FVector2D HitPosition = InMyGeometry.AbsoluteToLocal( HitScreenPosition );
 
@@ -1586,20 +1591,20 @@ SCurveEditor::FSelectedCurveKey SCurveEditor::HitTestKeys(const FGeometry& InMyG
 		for(auto CurveViewModel : CurveViewModels)
 		{
 			if (!CurveViewModel->bIsLocked && CurveViewModel->bIsVisible)
-			{
+		{
 				FRichCurve* Curve = CurveViewModel->CurveInfo.CurveToEdit;
-				if(Curve != NULL)
+			if(Curve != NULL)
+			{
+				for (auto It(Curve->GetKeyHandleIterator()); It; ++It)
 				{
-					for (auto It(Curve->GetKeyHandleIterator()); It; ++It)
-					{
-						float KeyScreenX = ScaleInfo.InputToLocalX(Curve->GetKeyTime(It.Key()));
-						float KeyScreenY = ScaleInfo.OutputToLocalY(Curve->GetKeyValue(It.Key()));
+					float KeyScreenX = ScaleInfo.InputToLocalX(Curve->GetKeyTime(It.Key()));
+					float KeyScreenY = ScaleInfo.OutputToLocalY(Curve->GetKeyValue(It.Key()));
 
-						if(	HitPosition.X > (KeyScreenX - (0.5f * CONST_KeySize.X)) && 
-							HitPosition.X < (KeyScreenX + (0.5f * CONST_KeySize.X)) &&
-							HitPosition.Y > (KeyScreenY - (0.5f * CONST_KeySize.Y)) &&
-							HitPosition.Y < (KeyScreenY + (0.5f * CONST_KeySize.Y)) )
-						{
+					if(	HitPosition.X > (KeyScreenX - (0.5f * CONST_KeySize.X)) && 
+						HitPosition.X < (KeyScreenX + (0.5f * CONST_KeySize.X)) &&
+						HitPosition.Y > (KeyScreenY - (0.5f * CONST_KeySize.Y)) &&
+						HitPosition.Y < (KeyScreenY + (0.5f * CONST_KeySize.Y)) )
+					{
 							return  FSelectedCurveKey(Curve, It.Key());
 						}
 					}
@@ -1788,8 +1793,10 @@ FReply SCurveEditor::ZoomToFitHorizontalClicked()
 /** Set Default output values when range is too small **/
 void SCurveEditor::SetDefaultOutput(const float MinZoomRange)
 {
-	ViewMinOutput = (ViewMinOutput - (0.5f*MinZoomRange));
-	ViewMaxOutput = (ViewMaxOutput + (0.5f*MinZoomRange));
+	const float NewMinOutput = (ViewMinOutput.Get() - (0.5f*MinZoomRange));
+	const float NewMaxOutput = (ViewMaxOutput.Get() + (0.5f*MinZoomRange));
+
+	SetOutputMinMax(NewMinOutput, NewMaxOutput);
 }
 
 void SCurveEditor::ZoomToFitVertical(bool OnlySelected)
@@ -1827,20 +1834,21 @@ void SCurveEditor::ZoomToFitVertical(bool OnlySelected)
 
 		const float MinZoomRange = (TotalKeys > 0 ) ? CONST_MinViewRange: CONST_DefaultZoomRange;
 
-		ViewMinOutput = (InMin);
-		ViewMaxOutput = (InMax);
-
 		// Clamp the minimum size
-		float Size = ViewMaxOutput - ViewMinOutput;
+		float Size = InMax - InMin;
 		if( Size < MinZoomRange )
 		{
 			SetDefaultOutput(MinZoomRange);
-			Size = ViewMaxOutput - ViewMinOutput;
+			InMin = ViewMinOutput.Get();
+			InMax = ViewMaxOutput.Get();
+			Size = InMax - InMin;
 		}
 
 		// add margin
-		ViewMinOutput = (ViewMinOutput - CONST_FitMargin*Size);
-		ViewMaxOutput = (ViewMaxOutput + CONST_FitMargin*Size);
+		const float NewMinOutput = (InMin - CONST_FitMargin*Size);
+		const float NewMaxOutput = (InMax + CONST_FitMargin*Size);
+
+		SetOutputMinMax(NewMinOutput, NewMaxOutput);
 	}
 }
 
@@ -2008,11 +2016,32 @@ void SCurveEditor::SetInputMinMax(float NewMin, float NewMax)
 	}
 }
 
+void SCurveEditor::SetOutputMinMax(float NewMin, float NewMax)
+{
+	if (SetOutputViewRangeHandler.IsBound())
+	{
+		SetOutputViewRangeHandler.Execute(NewMin, NewMax);
+	}
+	else
+	{
+		//if no delegate and view min output isn't using a delegate just set value directly
+		if (ViewMinOutput.IsBound() == false)
+		{
+			ViewMinOutput.Set(NewMin);
+		}
+
+		if (ViewMaxOutput.IsBound() == false)
+		{
+			ViewMaxOutput.Set(NewMax);
+		}
+	}
+}
+
 TSharedPtr<FCurveViewModel> SCurveEditor::HitTestCurves(  const FGeometry& InMyGeometry, const FPointerEvent& InMouseEvent )
 {
 	if( AreCurvesVisible() )
 	{
-		FTrackScaleInfo ScaleInfo(ViewMinInput.Get(),  ViewMaxInput.Get(), ViewMinOutput, ViewMaxOutput, InMyGeometry.Size);
+		FTrackScaleInfo ScaleInfo(ViewMinInput.Get(),  ViewMaxInput.Get(), ViewMinOutput.Get(), ViewMaxOutput.Get(), InMyGeometry.Size);
 
 		const FVector2D HitPosition = InMyGeometry.AbsoluteToLocal( InMouseEvent.GetScreenSpacePosition()  );
 
@@ -2045,7 +2074,7 @@ SCurveEditor::FSelectedTangent SCurveEditor::HitTestCubicTangents( const FGeomet
 
 	if( AreCurvesVisible() )
 	{
-		FTrackScaleInfo ScaleInfo(ViewMinInput.Get(),  ViewMaxInput.Get(), ViewMinOutput, ViewMaxOutput, InMyGeometry.Size);
+		FTrackScaleInfo ScaleInfo(ViewMinInput.Get(),  ViewMaxInput.Get(), ViewMinOutput.Get(), ViewMaxOutput.Get(), InMyGeometry.Size);
 
 		const FVector2D HitPosition = InMyGeometry.AbsoluteToLocal( HitScreenPosition);
 
@@ -2117,9 +2146,9 @@ bool SCurveEditor::IsInterpolationModeSelected(ERichCurveInterpMode InterpMode,
 	}
 	else
 	{
-		return false;
-	}
-}
+					return false;
+				}
+			}
 
 /* Given a tangent value for a key, calculates the 2D delta vector from that key in curve space */
 static inline FVector2D CalcTangentDir(float Tangent)
@@ -2205,7 +2234,7 @@ TArray<SCurveEditor::FSelectedCurveKey> SCurveEditor::GetEditableKeysWithinMarqu
 	TArray<FSelectedCurveKey> KeysWithinMarquee;
 	if (AreCurvesVisible())
 	{
-		FTrackScaleInfo ScaleInfo(ViewMinInput.Get(), ViewMaxInput.Get(), ViewMinOutput, ViewMaxOutput, InMyGeometry.Size);
+		FTrackScaleInfo ScaleInfo(ViewMinInput.Get(), ViewMaxInput.Get(), ViewMinOutput.Get(), ViewMaxOutput.Get(), InMyGeometry.Size);
 		for (auto CurveViewModel : CurveViewModels)
 		{
 			if (!CurveViewModel->bIsLocked && CurveViewModel->bIsVisible)
diff --git a/Engine/Source/Editor/UnrealEd/Public/SCurveEditor.h b/Engine/Source/Editor/UnrealEd/Public/SCurveEditor.h
--- a/Engine/Source/Editor/UnrealEd/Public/SCurveEditor.h
+++ b/Engine/Source/Editor/UnrealEd/Public/SCurveEditor.h
@@ -97,6 +97,7 @@ public:
 // SCurveEditor
 
 DECLARE_DELEGATE_TwoParams( FOnSetInputViewRange, float, float )
+DECLARE_DELEGATE_TwoParams( FOnSetOutputViewRange, float, float )
 
 class SCurveEditor : 
 	public SCompoundWidget,
@@ -146,8 +147,8 @@ public:
 		SLATE_ATTRIBUTE( float, ViewMaxInput )
 		SLATE_ATTRIBUTE( TOptional<float>, DataMinInput )
 		SLATE_ATTRIBUTE( TOptional<float>, DataMaxInput )
-		SLATE_ARGUMENT( float, ViewMinOutput )
-		SLATE_ARGUMENT( float, ViewMaxOutput )
+		SLATE_ATTRIBUTE( float, ViewMinOutput )
+		SLATE_ATTRIBUTE( float, ViewMaxOutput )
 		SLATE_ATTRIBUTE( float, InputSnap )
 		SLATE_ATTRIBUTE( float, OutputSnap )
 		SLATE_ARGUMENT( bool, SnappingEnabled )
@@ -161,6 +162,7 @@ public:
 		SLATE_ARGUMENT( bool, ZoomToFitHorizontal )
 		SLATE_ARGUMENT( bool, ShowZoomButtons )
 		SLATE_EVENT( FOnSetInputViewRange, OnSetInputViewRange )
+		SLATE_EVENT( FOnSetOutputViewRange, OnSetOutputViewRange )
 		SLATE_EVENT( FSimpleDelegate, OnCreateAsset )
 	SLATE_END_ARGS()
 
@@ -466,6 +468,7 @@ protected:
 
 	/** Update view range */
 	UNREALED_API void SetInputMinMax(float NewMin, float NewMax);
+	UNREALED_API void SetOutputMinMax(float NewMin, float NewMax);
 
 private:
 
@@ -506,6 +509,9 @@ private:
 	/** Handler for adjust timeline panning viewing */
 	FOnSetInputViewRange	SetInputViewRangeHandler;
 
+	/** Handler for adjust timeline panning viewing */
+	FOnSetOutputViewRange	SetOutputViewRangeHandler;
+
 	/** Index for the current transaction if any */
 	int32					TransactionIndex;
 
@@ -534,12 +540,12 @@ protected:
 	/** Maximum input of view range  */
 	TAttribute<float>	ViewMaxInput;
 	/** How long the overall timeline is */
-	TAttribute<float>		TimelineLength;
+	TAttribute<float>	TimelineLength;
 
 	/** Max output view range */
-	float				ViewMinOutput;
+	TAttribute<float>	ViewMinOutput;
 	/** Min output view range */
-	float				ViewMaxOutput;
+	TAttribute<float>	ViewMaxOutput;
 
 	/** The snapping value for the input domain. */
 	TAttribute<float> InputSnap;
