Added the game name to the editor startup analytics

TTP# 346677 - Analytics: Need a discrete new project event so we know when in the project browser or not

The "Editor.ProgramStarted" analytics event now includes a "GameName" attribute containing the contents of FApp::GetGameName(). If this is present and empty, then you know the editor has been opened in "project browser" mode, rather than as a full editor.

diff --git a/Engine/Source/Editor/UnrealEd/Private/UnrealEd.cpp b/Engine/Source/Editor/UnrealEd/Private/UnrealEd.cpp
--- a/Engine/Source/Editor/UnrealEd/Private/UnrealEd.cpp
+++ b/Engine/Source/Editor/UnrealEd/Private/UnrealEd.cpp
@@ -67,6 +67,7 @@ int32 EditorInit( IEngineLoop& EngineLoop )
 		TArray<FAnalyticsEventAttribute> EventAttributes;
 		EventAttributes.Add(FAnalyticsEventAttribute(TEXT("MachineID"), FPlatformMisc::GetMachineId().ToString(EGuidFormats::Digits).ToLower()));
 		EventAttributes.Add(FAnalyticsEventAttribute(TEXT("AccountID"), FPlatformMisc::GetEpicAccountId()));
+		EventAttributes.Add(FAnalyticsEventAttribute(TEXT("GameName"), FApp::GetGameName()));
 		EventAttributes.Add(FAnalyticsEventAttribute(TEXT("CommandLine"), FCommandLine::Get()));
 
 		FEngineAnalytics::GetProvider().RecordEvent(TEXT("Editor.ProgramStarted"), EventAttributes);
