Added support for cancelling cook by the book in editor.

diff --git a/Engine/Source/Developer/LauncherServices/Private/Launcher/LauncherWorker.cpp b/Engine/Source/Developer/LauncherServices/Private/Launcher/LauncherWorker.cpp
--- a/Engine/Source/Developer/LauncherServices/Private/Launcher/LauncherWorker.cpp
+++ b/Engine/Source/Developer/LauncherServices/Private/Launcher/LauncherWorker.cpp
@@ -524,13 +524,18 @@ void FLauncherWorker::CreateAndExecuteTasks( const ILauncherProfileRef& InProfil
 		class FWaitForCookInEditorToFinish : public FLauncherTask
 		{
 		public:
-			FWaitForCookInEditorToFinish() : FLauncherTask( FString(TEXT("CookByTheBookInEditor")), FString(TEXT("CookByTheBookInEditorDesk")), NULL, NULL)
+			FWaitForCookInEditorToFinish() : FLauncherTask( FString(TEXT("Cooking in the editor")), FString(TEXT("Prepairing content to run on device")), NULL, NULL)
 			{
 			}
 			virtual bool PerformTask( FLauncherTaskChainState& ChainState ) override
 			{
 				while ( !ChainState.Profile->OnIsCookFinished().Execute() )
 				{
+					if (GetStatus() == ELauncherTaskStatus::Canceling)
+					{
+						ChainState.Profile->OnCookCanceled().Execute();
+						return false;
+					}
 					FPlatformProcess::Sleep( 0.1f );
 				}
 				return true;
diff --git a/Engine/Source/Developer/LauncherServices/Private/Profiles/LauncherProfile.h b/Engine/Source/Developer/LauncherServices/Private/Profiles/LauncherProfile.h
--- a/Engine/Source/Developer/LauncherServices/Private/Profiles/LauncherProfile.h
+++ b/Engine/Source/Developer/LauncherServices/Private/Profiles/LauncherProfile.h
@@ -705,6 +705,11 @@ public:
 		return IsCookFinishedDelegate;
 	}
 
+	virtual FCookCanceledDelegate& OnCookCanceled() override
+	{
+		return CookCanceledDelegate;
+	}
+
 	virtual void SetDeploymentMode( ELauncherProfileDeploymentModes::Type Mode ) override
 	{
 		if (DeploymentMode != Mode)
@@ -1144,7 +1149,9 @@ private:
 
 private:
 
+	// cook in the editor callbacks (not valid for any other cook mode)
 	FIsCookFinishedDelegate IsCookFinishedDelegate;
+	FCookCanceledDelegate CookCanceledDelegate;
 
 	// Holds a delegate to be invoked when changing the device group to deploy to.
 	FOnLauncherProfileDeployedDeviceGroupChanged DeployedDeviceGroupChangedDelegate;
diff --git a/Engine/Source/Developer/LauncherServices/Public/Interfaces/ILauncherProfile.h b/Engine/Source/Developer/LauncherServices/Public/Interfaces/ILauncherProfile.h
--- a/Engine/Source/Developer/LauncherServices/Public/Interfaces/ILauncherProfile.h
+++ b/Engine/Source/Developer/LauncherServices/Public/Interfaces/ILauncherProfile.h
@@ -273,6 +273,11 @@ DECLARE_MULTICAST_DELEGATE(FOnProfileProjectChanged);
  */
 DECLARE_DELEGATE_RetVal(bool, FIsCookFinishedDelegate);
 
+/**
+ * Delegate type used to callback if the cook has been canceled
+ * only used for cook by the book in editor
+ */
+DECLARE_DELEGATE(FCookCanceledDelegate);
 
 /**
  * Interface for launcher profile.
@@ -374,6 +379,12 @@ public:
 	 */
 	virtual FIsCookFinishedDelegate& OnIsCookFinished() = 0;
 
+	/**
+	 * Returns the cook delegate which should be called if the cook is canceled
+	 *  Used by cook by the book in the editor
+	 */
+	virtual FCookCanceledDelegate& OnCookCanceled() = 0;
+
 public:
 
 	/**
diff --git a/Engine/Source/Editor/UnrealEd/Classes/CookOnTheSide/CookOnTheFlyServer.h b/Engine/Source/Editor/UnrealEd/Classes/CookOnTheSide/CookOnTheFlyServer.h
--- a/Engine/Source/Editor/UnrealEd/Classes/CookOnTheSide/CookOnTheFlyServer.h
+++ b/Engine/Source/Editor/UnrealEd/Classes/CookOnTheSide/CookOnTheFlyServer.h
@@ -41,45 +41,6 @@
 class UCookOnTheFlyServer : public UObject, public FTickableEditorObject
 {
 	GENERATED_UCLASS_BODY()
-
-private:
-	/** Current cook mode the cook on the fly server is running in */
-	ECookMode::Type CurrentCookMode;
-	
-	//////////////////////////////////////////////////////////////////////////
-	// Cook by the book options
-	struct FCookByTheBookOptions
-	{
-	public:
-		FCookByTheBookOptions() : bGenerateStreamingInstallManifests(false),
-			bRunning(false),
-			CookTime( 0.0f )
-		{ }
-
-		/** Should we generate streaming install manifests (only valid option in cook by the book) */
-		bool bGenerateStreamingInstallManifests;
-		/** Is cook by the book currently running */
-		bool bRunning;
-		/** Leak test: last gc items (only valid when running from commandlet requires gc between each cooked package) */
-		TSet<FWeakObjectPtr> LastGCItems;
-		/** Map of platform name to manifest generator, manifest is only used in cook by the book however it needs to be maintained across multiple cook by the books. */
-		TMap<FName, FChunkManifestGenerator*> ManifestGenerators;
-		double CookTime;
-	};
-	FCookByTheBookOptions* CookByTheBookOptions;
-	
-
-	//////////////////////////////////////////////////////////////////////////
-	// Cook on the fly options
-	/** Cook on the fly server uses the NetworkFileServer */
-	TArray<class INetworkFileServer*> NetworkFileServers;
-
-	//////////////////////////////////////////////////////////////////////////
-	// General cook options
-	ECookInitializationFlags CookFlags;
-	TAutoPtr<class FSandboxPlatformFile> SandboxFile;
-	bool bIsSavingPackage; // used to stop recursive mark package dirty functions
-
 private:
 
 	/** Array which has been made thread safe :) */
@@ -141,6 +102,12 @@
 			ScopeLockType ScopeLock( &SynchronizationObject );
 			return Items.Num();
 		}
+
+		void Empty() 
+		{
+			ScopeLockType ScopeLock( &SynchronizationObject );
+			Items.Empty();
+		}
 	};
 
 
@@ -402,6 +369,30 @@
 			return true;
 			// return FilesProcessed.Contains(Filename);
 		}
+		// two versions of this function so I don't have to create temporary FFIleplatformRequest in some cases to call the exists function
+		bool Exists( const FName& Filename, const TArray<FName>& PlatformNames )
+		{
+			FScopeLock ScopeLock(&SynchronizationObject);
+
+			const FFilePlatformRequest* OurRequest = FilesProcessed.Find( Filename );
+
+			if (!OurRequest)
+			{
+				return false;
+			}
+
+			// make sure all the platforms are completed
+			for ( const auto& Platform : PlatformNames )
+			{
+				if ( !OurRequest->GetPlatformnames().Contains( Platform ) )
+				{
+					return false;
+				}
+			}
+
+			return true;
+		}
+
 		bool GetCookedPlatforms( const FName& Filename, TArray<FName>& PlatformList )
 		{
 			FScopeLock ScopeLock( &SynchronizationObject );
@@ -419,6 +410,11 @@
 			FScopeLock ScopeLock( &SynchronizationObject );
 			return FilesProcessed.Remove( Filename );
 		}
+		void Empty()
+		{
+			FScopeLock ScopeLock( &SynchronizationObject );
+			FilesProcessed.Empty();
+		}
 	};
 
 
@@ -471,6 +467,34 @@
 			}
 			return false;
 		}
+
+		void DequeueAllRequests( TArray<FFilePlatformRequest>& RequestArray )
+		{
+			FScopeLock ScopeLock( &SynchronizationObject );
+			if ( Queue.Num() )
+			{
+				for ( const auto& Request : PlatformList )
+				{
+					RequestArray.Add( FFilePlatformRequest( Request.Key, Request.Value ) );
+				}
+				PlatformList.Empty();
+				Queue.Empty();
+			}
+		}
+
+		bool Exists( const FName& Filename, const TArray<FName>& PlatformNames ) const 
+		{
+			const TArray<FName>* Platforms = PlatformList.Find( Filename );
+			if ( Platforms == NULL )
+				return false;
+
+			for ( const auto& PlatformName : PlatformNames )
+			{
+				if ( !Platforms->Contains( PlatformName ) )
+					return false;
+			}
+			return true;
+		}
 		
 		bool HasItems() const
 		{
@@ -483,6 +507,14 @@
 			FScopeLock ScopeLock( &SynchronizationObject );
 			return Queue.Num();
 		}
+
+
+		void Empty()  
+		{
+			FScopeLock ScopeLock( &SynchronizationObject );
+			Queue.Empty();
+			PlatformList.Empty();
+		}
 	};
 
 	struct FThreadSafeUnsolicitedPackagesList
@@ -516,7 +548,54 @@
 				}
 			}
 		}
+		void Empty()
+		{
+			FScopeLock S( &SyncObject );
+			CookedPackages.Empty();
+		}
 	};
+private:
+	/** Current cook mode the cook on the fly server is running in */
+	ECookMode::Type CurrentCookMode;
+	
+	//////////////////////////////////////////////////////////////////////////
+	// Cook by the book options
+	struct FCookByTheBookOptions
+	{
+	public:
+		FCookByTheBookOptions() : bGenerateStreamingInstallManifests(false),
+			bRunning(false),
+			CookTime( 0.0f )
+		{ }
 
+		/** Should we generate streaming install manifests (only valid option in cook by the book) */
+		bool bGenerateStreamingInstallManifests;
+		/** Is cook by the book currently running */
+		bool bRunning;
+		/** Cancel has been queued will be processed next tick */
+		bool bCancel;
+		/** Leak test: last gc items (only valid when running from commandlet requires gc between each cooked package) */
+		TSet<FWeakObjectPtr> LastGCItems;
+		/** Map of platform name to manifest generator, manifest is only used in cook by the book however it needs to be maintained across multiple cook by the books. */
+		TMap<FName, FChunkManifestGenerator*> ManifestGenerators;
+		/** If a cook is cancelled next cook will need to resume cooking */ 
+		TArray<FFilePlatformRequest> PreviousCookRequests; 
+		double CookTime;
+	};
+	FCookByTheBookOptions* CookByTheBookOptions;
+	
+
+	//////////////////////////////////////////////////////////////////////////
+	// Cook on the fly options
+	/** Cook on the fly server uses the NetworkFileServer */
+	TArray<class INetworkFileServer*> NetworkFileServers;
+
+	//////////////////////////////////////////////////////////////////////////
+	// General cook options
+	ECookInitializationFlags CookFlags;
+	TAutoPtr<class FSandboxPlatformFile> SandboxFile;
+	bool bIsSavingPackage; // used to stop recursive mark package dirty functions
+
+
 	FThreadSafeQueue<struct FRecompileRequest*> RecompileRequests;
 	FFilenameQueue CookRequests; // list of requested files
@@ -577,6 +656,16 @@
 	 */
 	void StartCookByTheBook(const TArray<ITargetPlatform*>& TargetPlatforms, const TArray<FString>& CookMaps, const TArray<FString>& CookDirectories, const TArray<FString>& CookCultures, const TArray<FString>& IniMapSections );
 
+	/**
+	 * Queue a cook by the book cancel (you might want to do this instead of calling cancel directly so that you don't have to be in the game thread when canceling
+	 */
+	void QueueCancelCookByTheBook();
+
+	/**
+	 * Cancel the currently running cook by the book (needs to be called from the game thread)
+	 */
+	void CancelCookByTheBook();
+
 	bool IsCookByTheBookRunning() const;
 
 	/**
@@ -588,9 +677,14 @@
 	 */
 	uint32 TickCookOnTheSide( const float TimeSlice, uint32 &CookedPackagesCount );
 	
+	/**
+	 * Force stop whatever pending cook requests are going on and clear all the cooked data
+	 * Note cook on the side / cook on the fly clients may not be able to recover from this if they are waiting on a cook request to complete
+	 */
+	void StopAndClearCookedData();
+
 	/** 
 	 * Process any shader recompile requests
-	 *
 	 */
 	void TickRecompileShaderRequests();
 
diff --git a/Engine/Source/Editor/UnrealEd/Classes/Editor/EditorEngine.h b/Engine/Source/Editor/UnrealEd/Classes/Editor/EditorEngine.h
--- a/Engine/Source/Editor/UnrealEd/Classes/Editor/EditorEngine.h
+++ b/Engine/Source/Editor/UnrealEd/Classes/Editor/EditorEngine.h
@@ -1480,6 +1480,11 @@ public:
 	 */
 	virtual bool IsCookByTheBookInEditorFinished() const { return true; }
 
+	/**
+	 * Cancels the current cook by the book in editor
+	 */
+	virtual void CancelCookByTheBookInEditor() { }
+
 
 	/**
 	 * Makes a request to start a play from a Slate editor session
diff --git a/Engine/Source/Editor/UnrealEd/Classes/Editor/UnrealEdEngine.h b/Engine/Source/Editor/UnrealEd/Classes/Editor/UnrealEdEngine.h
--- a/Engine/Source/Editor/UnrealEd/Classes/Editor/UnrealEdEngine.h
+++ b/Engine/Source/Editor/UnrealEd/Classes/Editor/UnrealEdEngine.h
@@ -553,6 +553,12 @@ class UNREALED_API UUnrealEdEngine : public UEditorEngine, public FNotifyHook
 	virtual bool IsCookByTheBookInEditorFinished() const override;
 
 
+	/**
+	 * cancels the current cook by the book in editor
+	 */
+	virtual void CancelCookByTheBookInEditor() override;
+
+
 
 	// Hook replacements.
 	void ShowActorProperties();
diff --git a/Engine/Source/Editor/UnrealEd/Private/CookOnTheFlyServer.cpp b/Engine/Source/Editor/UnrealEd/Private/CookOnTheFlyServer.cpp
--- a/Engine/Source/Editor/UnrealEd/Private/CookOnTheFlyServer.cpp
+++ b/Engine/Source/Editor/UnrealEd/Private/CookOnTheFlyServer.cpp
@@ -81,15 +81,16 @@ public:
 	FScopeTimer( const FString &InName, bool IncrementScope = false )
 	{
 		DecrementScope = IncrementScope;
-		if( DecrementScope)
-		{
-			++GScopeDepth;
-		}
+		
 		FString Name = InName;
 		for ( int i =0; i < GScopeDepth; ++i )
 		{
 			Name = FString(TEXT("  ")) + Name;
 		}
+		if( DecrementScope)
+		{
+			++GScopeDepth;
+		}
 		Index = GTimerInfo.Emplace(MoveTemp(Name), 0.0);
 		Started = false;
 	}
@@ -270,9 +271,9 @@ public:
 	FName StandardFileFName;
 };
 
-static TMap<FString, FCachedPackageFilename> PackageFilenameCache;
+static TMap<FName, FCachedPackageFilename> PackageFilenameCache;
 
-static const FCachedPackageFilename &Cache(const FString &PackageName)
+static const FCachedPackageFilename &Cache(const FName& PackageName)
 {
 	FCachedPackageFilename *Cached = PackageFilenameCache.Find( PackageName );
 	if ( Cached != NULL )
@@ -280,12 +281,11 @@ static const FCachedPackageFilename &Cache(const FString &PackageName)
 		return *Cached;
 	}
 
-
 	FString Filename;
 	FString PackageFilename;
 	FString StandardFilename;
 	FName StandardFileFName = NAME_None;
-	if (FPackageName::DoesPackageExist(PackageName, NULL, &Filename))
+	if (FPackageName::DoesPackageExist(PackageName.ToString(), NULL, &Filename))
 	{
 		StandardFilename = PackageFilename = FPaths::ConvertRelativePathToFull(Filename);
 		
@@ -309,35 +309,53 @@ static const FCachedPackageFilename &Cache(const FString &PackageName)
 	return Filename;
 }*/
 
-static FString GetCachedPackageFilename( const FString &PackageName )
+static FString GetCachedPackageFilename( const FName& PackageName )
 {
 	return Cache( PackageName ).PackageFilename;
 }
 
-static FString GetCachedStandardPackageFilename( const FString &PackageName )
+static FString GetCachedStandardPackageFilename( const FName& PackageName )
 {
 	return Cache( PackageName ).StandardFilename;
 }
 
-static FName GetCachedStandardPackageFileFName( const FString &PackageName )
+static FName GetCachedStandardPackageFileFName( const FName& PackageName )
 {
 	return Cache( PackageName ).StandardFileFName;
 }
 
 
-static FString GetCachedPackageFilename( UPackage *Package )
+static FString GetCachedPackageFilename( UPackage* Package )
 {
-	return Cache( Package->GetName() ).PackageFilename;
+	check( Package->GetName() == Package->GetFName().ToString() );
+	return Cache( Package->GetFName() ).PackageFilename;
 }
 
 static FString GetCachedStandardPackageFilename( UPackage* Package )
 {
-	return Cache( Package->GetName() ).StandardFilename;
+	check( Package->GetName() == Package->GetFName().ToString() );
+	return Cache( Package->GetFName() ).StandardFilename;
 }
 
-static FName GetCachedStandardPackageFileFName( UPackage *Package )
+static FName GetCachedStandardPackageFileFName( UPackage* Package )
 {
-	return Cache( Package->GetName() ).StandardFileFName;
+	check( Package->GetName() == Package->GetFName().ToString() );
+	return Cache( Package->GetFName() ).StandardFileFName;
+}
+
+
+static const FString& GetCachedSandboxFilename( UPackage* Package, TAutoPtr<class FSandboxPlatformFile>& SandboxFile )
+{
+	FName PackageFName = Package->GetFName();
+	static TMap<FName, FString> CachedSandboxFilenames;
+	FString* CachedSandboxFilename = CachedSandboxFilenames.Find(PackageFName);
+	if ( CachedSandboxFilename )
+		return *CachedSandboxFilename;
+
+	const FString& PackageFilename = GetCachedPackageFilename(Package);
+	FString SandboxFilename = SandboxFile->ConvertToAbsolutePathForExternalAppForWrite(*PackageFilename );
+
+	return CachedSandboxFilenames.Add( PackageFName, MoveTemp(SandboxFilename) );
 }
 
 /* helper structs
@@ -412,8 +430,6 @@ TStatId UCookOnTheFlyServer::GetStatId() const
 	RETURN_QUICK_DECLARE_CYCLE_STAT(UCookServer, STATGROUP_Tickables);
 }
 
-
-
 bool UCookOnTheFlyServer::StartNetworkFileServer( const bool BindAnyPort )
 {
 	check( CurrentCookMode == ECookMode::CookOnTheFly );
@@ -781,6 +797,14 @@ uint32 UCookOnTheFlyServer::TickCookOnTheSide( const float TimeSlice, uint32 &Co
 			break;
 		}
 
+		if ( IsCookByTheBookMode() )
+		{
+			check( CookByTheBookOptions );
+			if ( CookByTheBookOptions->bCancel )
+			{
+				CancelCookByTheBook();
+			}
+		}
 
 		FFilePlatformRequest ToBuild;
 
@@ -990,12 +1014,10 @@ uint32 UCookOnTheFlyServer::TickCookOnTheSide( const float TimeSlice, uint32 &Co
 			PackageNames.Empty(ObjectsInOuter.Num());
 			{
 				SCOPE_TIMER(GeneratePackageNames);
+				ACCUMULATE_TIMER(UnsolicitedPackageAlreadyCooked);
 				ACCUMULATE_TIMER(PackageCast);
 				ACCUMULATE_TIMER(AddUnassignedPackageToManifest);
-				ACCUMULATE_TIMER(ConvertToAbsolutePath);
-				ACCUMULATE_TIMER(ManifestAddUnassignedPackage);
-				ACCUMULATE_TIMER(CookedPackagesExists);
-				ACCUMULATE_TIMER(PackagesToSaveAdd);
+				ACCUMULATE_TIMER(GetCachedName);
 				for( int32 Index = 0; Index < ObjectsInOuter.Num(); Index++ )
 				{
 					ACCUMULATE_TIMER_START(PackageCast);
@@ -1004,6 +1026,17 @@ uint32 UCookOnTheFlyServer::TickCookOnTheSide( const float TimeSlice, uint32 &Co
 
 					if (Package)
 					{
+						FName StandardPackageFName = GetCachedStandardPackageFileFName(Package);
+						if ( StandardPackageFName == NAME_None )
+							continue;
+
+						ACCUMULATE_TIMER_START(UnsolicitedPackageAlreadyCooked);
+						// package is already cooked don't care about processing it again here
+						if ( CookRequests.Exists( StandardPackageFName, AllTargetPlatformNames) )
+							continue;
+						if ( CookedPackages.Exists(StandardPackageFName,AllTargetPlatformNames) )
+							continue;
+						ACCUMULATE_TIMER_STOP(UnsolicitedPackageAlreadyCooked);
 						// this has been replaced to use the GenerateDependencies call
 						/*if ( CookByTheBookOptions && CookByTheBookOptions->ManifestGenerator != NULL )
 						{
@@ -1030,37 +1063,22 @@ uint32 UCookOnTheFlyServer::TickCookOnTheSide( const float TimeSlice, uint32 &Co
 								if (!Filename.IsEmpty())
 								{
 									// Populate streaming install manifests
-									ACCUMULATE_TIMER_START(ConvertToAbsolutePath);
-									FString SandboxFilename = SandboxFile->ConvertToAbsolutePathForExternalAppForWrite(*Filename);
-									ACCUMULATE_TIMER_STOP(ConvertToAbsolutePath);
-									ACCUMULATE_TIMER_START(ManifestAddUnassignedPackage);
+									// FString SandboxFilename = SandboxFile->ConvertToAbsolutePathForExternalAppForWrite(*Filename);
+									const FString& SandboxFilename = GetCachedSandboxFilename( Package, SandboxFile );
 									Manifest->AddUnassignedPackageToManifest(Package, SandboxFilename);
-									ACCUMULATE_TIMER_STOP(ManifestAddUnassignedPackage);
 								}
 							}
 							ACCUMULATE_TIMER_STOP(AddUnassignedPackageToManifest);
 						}
 
-
+						ACCUMULATE_TIMER_START(GetCachedName);
 						FName StandardFilename = GetCachedStandardPackageFileFName(Package);
 						if ( StandardFilename != NAME_None ) // if we have name none that means we are in core packages or something...
 						{
 							// check if the package has already been saved
-							bool bShouldSave = false;
-
-							ACCUMULATE_TIMER_START(CookedPackagesExists);
-							if ( !CookedPackages.Exists(FFilePlatformRequest(StandardFilename, TargetPlatformNames)))
-							{
-								bShouldSave = true;
-							}
-							ACCUMULATE_TIMER_STOP(CookedPackagesExists);
-							ACCUMULATE_TIMER_START(PackagesToSaveAdd);
-							if ( bShouldSave )
-							{
-								PackagesToSave.AddUnique( Package );
-							}
-							ACCUMULATE_TIMER_STOP(PackagesToSaveAdd);
+							PackagesToSave.AddUnique( Package );
 						}
+						ACCUMULATE_TIMER_STOP(GetCachedName);
 					}
 				}
 			}
@@ -1100,6 +1118,7 @@ uint32 UCookOnTheFlyServer::TickCookOnTheSide( const float TimeSlice, uint32 &Co
 
 					if ( bShouldFinishTick )
 					{
+						SCOPE_TIMER(EnqueueUnsavedPackages);
 						// don't save these packages because we have a real request
 						// enqueue all the packages which we were about to save
 						if (CurrentCookMode==ECookMode::CookByTheBookFromTheEditor)
@@ -1291,6 +1310,7 @@ void UCookOnTheFlyServer::EndNetworkFileServer()
 		delete NetworkFileServer;
 		NetworkFileServer = NULL;
 	}
+	NetworkFileServers.Empty();
 }
 
 void UCookOnTheFlyServer::BeginDestroy()
@@ -2021,6 +2041,7 @@ void UCookOnTheFlyServer::SaveGlobalShaderMapFiles(const TArray<ITargetPlatform*
 
 void UCookOnTheFlyServer::CookByTheBookFinished()
 {
+	check( IsInGameThread() );
 	check( IsCookByTheBookMode() );
 	check( CookByTheBookOptions->bRunning == true );
 
@@ -2046,13 +2067,51 @@ void UCookOnTheFlyServer::CookByTheBookFinished()
 
 	UE_LOG(LogCookOnTheFly, Display, TEXT("Cook by the book total time %fs"), CookByTheBookOptions->CookTime);
 
+	CookByTheBookOptions->bRunning = false;
+}
+
+void UCookOnTheFlyServer::QueueCancelCookByTheBook()
+{
+	if ( IsCookByTheBookMode() )
+	{
+		check( CookByTheBookOptions != NULL );
+		CookByTheBookOptions->bCancel = true;
+	}
+}
 
+void UCookOnTheFlyServer::CancelCookByTheBook()
+{
+	if ( IsCookByTheBookMode() && CookByTheBookOptions->bRunning )
+	{
+		check(CookByTheBookOptions);
+		check( IsInGameThread() );
 
-	CookByTheBookOptions->bRunning = false;
+		// save the cook requests 
+		CookRequests.DequeueAllRequests(CookByTheBookOptions->PreviousCookRequests);
+		CookByTheBookOptions->bRunning = false;
+	}	
+}
+
+void UCookOnTheFlyServer::StopAndClearCookedData()
+{
+	if ( IsCookByTheBookMode() )
+	{
+		check( CookByTheBookOptions != NULL );
+		check( CookByTheBookOptions->bRunning == false );
+		CancelCookByTheBook();
+		CookByTheBookOptions->PreviousCookRequests.Empty();
+	}
+
+	RecompileRequests.Empty();
+	CookRequests.Empty();
+	UnsolicitedCookedPackages.Empty();
+	CookedPackages.Empty(); // set of files which have been cooked when needing to recook a file the entry will need to be removed from here
 }
 
 void UCookOnTheFlyServer::StartCookByTheBook(const TArray<ITargetPlatform*>& TargetPlatforms, const TArray<FString> &CookMaps, const TArray<FString> &CookDirectories, const TArray<FString> &CookCultures, const TArray<FString> &IniMapSections )
 {
+	check( IsInGameThread() );
+
 	// can transition from cook on the fly to cook by the book from the editor
 	check( IsCookByTheBookMode() );
 
@@ -2085,6 +2144,7 @@ void UCookOnTheFlyServer::StartCookByTheBook(const TArray<ITargetPlatform*>& Tar
 	*/
 
 	CookByTheBookOptions->bRunning = true;
+	CookByTheBookOptions->bCancel = false;
 	CookByTheBookOptions->CookTime = 0.0f;
 
 	bool bCookAll = false;
@@ -2143,21 +2203,27 @@ void UCookOnTheFlyServer::StartCookByTheBook(const TArray<ITargetPlatform*>& Tar
 	GenerateLongPackageNames(FilesInPath);
 
 	// add all the files for the requested platform to the cook list
-	for ( auto FileName : FilesInPath )
+	for ( const auto& FileName : FilesInPath )
 	{
-		FName FileFName = GetCachedStandardPackageFileFName(FileName);
+		FName FileFName = FName(*FileName);
+		FName PackageFileFName = GetCachedStandardPackageFileFName(FileFName);
 		
-		if (FileFName != NAME_None)
+		if (PackageFileFName != NAME_None)
 		{
-			CookRequests.EnqueueUnique( MoveTemp( FFilePlatformRequest( FileFName, TargetPlatformNames ) ) );
+			CookRequests.EnqueueUnique( MoveTemp( FFilePlatformRequest( PackageFileFName, TargetPlatformNames ) ) );
 		}
 		else
 		{
 			LogCookerMessage( FString::Printf(TEXT("Unable to find package for cooking %s"), *FileName), EMessageSeverity::Warning );
 			UE_LOG(LogCookOnTheFly, Warning, TEXT("Unable to find package for cooking %s"), *FileName)
-		}
-		
+		}	
+	}
+
+	for ( const auto& PreviousRequest : CookByTheBookOptions->PreviousCookRequests )
+	{
+		CookRequests.EnqueueUnique( MoveTemp( PreviousRequest ) );
 	}
+	CookByTheBookOptions->PreviousCookRequests.Empty();
 
 }
 
diff --git a/Engine/Source/Editor/UnrealEd/Private/PlayLevel.cpp b/Engine/Source/Editor/UnrealEd/Private/PlayLevel.cpp
--- a/Engine/Source/Editor/UnrealEd/Private/PlayLevel.cpp
+++ b/Engine/Source/Editor/UnrealEd/Private/PlayLevel.cpp
@@ -1326,7 +1326,11 @@ void UEditorEngine::PlayUsingLauncher()
 
 			FIsCookFinishedDelegate &CookerFinishedDelegate = LauncherProfile->OnIsCookFinished();
 
-			CookerFinishedDelegate.BindUObject(this, &UEditorEngine::IsCookByTheBookInEditorFinished);   // AddUObject(this, &UEditorEngine::IsCookByTheBookInEditorFinished);
+			CookerFinishedDelegate.BindUObject(this, &UEditorEngine::IsCookByTheBookInEditorFinished);
+
+			FCookCanceledDelegate &CookCancelledDelegate = LauncherProfile->OnCookCanceled();
+
+			CookCancelledDelegate.BindUObject(this, &UEditorEngine::CancelCookByTheBookInEditor);
 		}
 
 		ILauncherPtr Launcher = LauncherServicesModule.CreateLauncher();
diff --git a/Engine/Source/Editor/UnrealEd/Private/UnrealEdEngine.cpp b/Engine/Source/Editor/UnrealEd/Private/UnrealEdEngine.cpp
--- a/Engine/Source/Editor/UnrealEd/Private/UnrealEdEngine.cpp
+++ b/Engine/Source/Editor/UnrealEd/Private/UnrealEdEngine.cpp
@@ -148,6 +148,11 @@ bool UUnrealEdEngine::IsCookByTheBookInEditorFinished() const
 	return !CookServer->IsCookByTheBookRunning();
 }
 
+void UUnrealEdEngine::CancelCookByTheBookInEditor()
+{
+	CookServer->QueueCancelCookByTheBook();
+}
+
 void UUnrealEdEngine::MakeSortedSpriteInfo(TArray<FSpriteCategoryInfo>& OutSortedSpriteInfo) const
 {
 	struct Local
