Increased lightmass amortized export time threshold to 10ms to reduce export latency

diff --git a/Engine/Source/Editor/UnrealEd/Private/Lightmass/Lightmass.cpp b/Engine/Source/Editor/UnrealEd/Private/Lightmass/Lightmass.cpp
--- a/Engine/Source/Editor/UnrealEd/Private/Lightmass/Lightmass.cpp
+++ b/Engine/Source/Editor/UnrealEd/Private/Lightmass/Lightmass.cpp
@@ -67,7 +67,7 @@ void FSwarmDebugOptions::Touch()
 int32 FLightmassProcessor::MaxProcessAvailableCount = 8;
 
 /** We don't want any amortization steps to take longer than this amount every tick */
-static const float AllowedAmortizationTimePerTick = 0.005f; // in seconds
+static const float AllowedAmortizationTimePerTick = 0.01f; // in seconds
 
 volatile int32 FLightmassProcessor::VolumeSampleTaskCompleted = 0;
 volatile int32 FLightmassProcessor::MeshAreaLightDataTaskCompleted = 0;
@@ -622,7 +622,11 @@ bool FLightmassExporter::WriteToMaterialChannel(FLightmassStatistics& Stats)
 			{
 			case BuildMaterials:
 				{
-					if (CurrentAmortizationIndex >= Materials.Num()) {ExportStage = ShaderCompilation; CurrentAmortizationIndex = 0;}
+					if (CurrentAmortizationIndex >= Materials.Num()) 
+					{
+						ExportStage = ShaderCompilation; 
+						CurrentAmortizationIndex = 0;
+					}
 					else
 					{
 						BuildMaterialMap(Materials[CurrentAmortizationIndex]);
@@ -639,7 +643,11 @@ bool FLightmassExporter::WriteToMaterialChannel(FLightmassStatistics& Stats)
 				break;
 			case ExportMaterials:
 				{
-					if (CurrentAmortizationIndex >= Materials.Num()) {ExportStage = CleanupMaterialExport; CurrentAmortizationIndex = 0;}
+					if (CurrentAmortizationIndex >= Materials.Num()) 
+					{
+						ExportStage = CleanupMaterialExport; 
+						CurrentAmortizationIndex = 0;
+					}
 					else
 					{
 						ExportMaterial(Materials[CurrentAmortizationIndex]);
@@ -649,7 +657,11 @@ bool FLightmassExporter::WriteToMaterialChannel(FLightmassStatistics& Stats)
 				break;
 			case CleanupMaterialExport:
 				{
-					if (CurrentAmortizationIndex >= OpenedMaterialExportChannels.Num()) {ExportStage = Complete; CurrentAmortizationIndex = 0;}
+					if (CurrentAmortizationIndex >= OpenedMaterialExportChannels.Num()) 
+					{
+						ExportStage = Complete; 
+						CurrentAmortizationIndex = 0;
+					}
 					else
 					{
 						Swarm.CloseChannel(OpenedMaterialExportChannels[CurrentAmortizationIndex]);
