Resolved TTP #343526 : PERSONA: Not prompted to name assets created via toolbar Create Asset button

diff --git a/Engine/Source/Editor/UnrealEd/Private/AnimationEditorUtils.cpp b/Engine/Source/Editor/UnrealEd/Private/AnimationEditorUtils.cpp
--- a/Engine/Source/Editor/UnrealEd/Private/AnimationEditorUtils.cpp
+++ b/Engine/Source/Editor/UnrealEd/Private/AnimationEditorUtils.cpp
@@ -7,6 +7,207 @@
 
 #define LOCTEXT_NAMESPACE "AnimationEditorUtils"
 
+/////////////////////////////////////////////////////
+/////////////////////////////////////////////////////
+/////////////////////////////////////////////////////
+/////////////////////////////////////////////////////
+
+///////////////////////////////////////////////////////////////////////////////
+// Create Animation dialog to determine a newly created asset's name
+///////////////////////////////////////////////////////////////////////////////
+
+FText SCreateAnimationAssetDlg::LastUsedAssetPath;
+
+void SCreateAnimationAssetDlg::Construct(const FArguments& InArgs)
+{
+	AssetPath = FText::FromString(FPackageName::GetLongPackagePath(InArgs._DefaultAssetPath.ToString()));
+	AssetName = FText::FromString(FPackageName::GetLongPackageAssetName(InArgs._DefaultAssetPath.ToString()));
+
+	if (AssetPath.IsEmpty())
+	{
+		AssetPath = LastUsedAssetPath;
+	}
+	else
+	{
+		LastUsedAssetPath = AssetPath;
+	}
+
+	FPathPickerConfig PathPickerConfig;
+	PathPickerConfig.DefaultPath = AssetPath.ToString();
+	PathPickerConfig.OnPathSelected = FOnPathSelected::CreateSP(this, &SCreateAnimationAssetDlg::OnPathChange);
+	PathPickerConfig.bAddDefaultPath = true;
+
+	FContentBrowserModule& ContentBrowserModule = FModuleManager::LoadModuleChecked<FContentBrowserModule>("ContentBrowser");
+
+	SWindow::Construct(SWindow::FArguments()
+		.Title(LOCTEXT("SCreateAnimationAssetDlg_Title", "Create a New Animation Asset"))
+		.SupportsMinimize(false)
+		.SupportsMaximize(false)
+		//.SizingRule( ESizingRule::Autosized )
+		.ClientSize(FVector2D(450, 450))
+		[
+			SNew(SVerticalBox)
+
+			+ SVerticalBox::Slot() // Add user input block
+			.Padding(2)
+			[
+				SNew(SBorder)
+				.BorderImage(FEditorStyle::GetBrush("ToolPanel.GroupBorder"))
+				[
+					SNew(SVerticalBox)
+
+					+ SVerticalBox::Slot()
+					.AutoHeight()
+					[
+						SNew(STextBlock)
+						.Text(LOCTEXT("SelectPath", "Select Path to create animation"))
+						.Font(FSlateFontInfo(FPaths::EngineContentDir() / TEXT("Slate/Fonts/Roboto-Regular.ttf"), 14))
+					]
+
+					+ SVerticalBox::Slot()
+						.FillHeight(1)
+						.Padding(3)
+						[
+							ContentBrowserModule.Get().CreatePathPicker(PathPickerConfig)
+						]
+
+					+ SVerticalBox::Slot()
+						.AutoHeight()
+						[
+							SNew(SSeparator)
+						]
+
+					+ SVerticalBox::Slot()
+						.AutoHeight()
+						.Padding(3)
+						[
+							SNew(SHorizontalBox)
+
+							+ SHorizontalBox::Slot()
+							.AutoWidth()
+							.Padding(0, 0, 10, 0)
+							.VAlign(VAlign_Center)
+							[
+								SNew(STextBlock)
+								.Text(LOCTEXT("AnimationName", "Animation Name"))
+							]
+
+							+ SHorizontalBox::Slot()
+								[
+									SNew(SEditableTextBox)
+									.Text(AssetName)
+									.OnTextCommitted(this, &SCreateAnimationAssetDlg::OnNameChange)
+									.MinDesiredWidth(250)
+								]
+						]
+
+				]
+			]
+
+			+ SVerticalBox::Slot()
+				.AutoHeight()
+				.HAlign(HAlign_Right)
+				.Padding(5)
+				[
+					SNew(SUniformGridPanel)
+					.SlotPadding(FEditorStyle::GetMargin("StandardDialog.SlotPadding"))
+					.MinDesiredSlotWidth(FEditorStyle::GetFloat("StandardDialog.MinDesiredSlotWidth"))
+					.MinDesiredSlotHeight(FEditorStyle::GetFloat("StandardDialog.MinDesiredSlotHeight"))
+					+ SUniformGridPanel::Slot(0, 0)
+					[
+						SNew(SButton)
+						.HAlign(HAlign_Center)
+						.ContentPadding(FEditorStyle::GetMargin("StandardDialog.ContentPadding"))
+						.Text(LOCTEXT("OK", "OK"))
+						.OnClicked(this, &SCreateAnimationAssetDlg::OnButtonClick, EAppReturnType::Ok)
+					]
+					+ SUniformGridPanel::Slot(1, 0)
+						[
+							SNew(SButton)
+							.HAlign(HAlign_Center)
+							.ContentPadding(FEditorStyle::GetMargin("StandardDialog.ContentPadding"))
+							.Text(LOCTEXT("Cancel", "Cancel"))
+							.OnClicked(this, &SCreateAnimationAssetDlg::OnButtonClick, EAppReturnType::Cancel)
+						]
+				]
+		]);
+}
+
+void SCreateAnimationAssetDlg::OnNameChange(const FText& NewName, ETextCommit::Type CommitInfo)
+{
+	AssetName = NewName;
+}
+
+void SCreateAnimationAssetDlg::OnPathChange(const FString& NewPath)
+{
+	AssetPath = FText::FromString(NewPath);
+	LastUsedAssetPath = AssetPath;
+}
+
+FReply SCreateAnimationAssetDlg::OnButtonClick(EAppReturnType::Type ButtonID)
+{
+	UserResponse = ButtonID;
+
+	if (ButtonID != EAppReturnType::Cancel)
+	{
+		if (!ValidatePackage())
+		{
+			// reject the request
+			return FReply::Handled();
+		}
+	}
+
+	RequestDestroyWindow();
+
+	return FReply::Handled();
+}
+
+/** Ensures supplied package name information is valid */
+bool SCreateAnimationAssetDlg::ValidatePackage()
+{
+	FText Reason;
+	FString FullPath = GetFullAssetPath();
+
+	if (!FPackageName::IsValidLongPackageName(FullPath, false, &Reason)
+		|| !FName(*AssetName.ToString()).IsValidObjectName(Reason))
+	{
+		FMessageDialog::Open(EAppMsgType::Ok, Reason);
+		return false;
+	}
+
+	return true;
+}
+
+EAppReturnType::Type SCreateAnimationAssetDlg::ShowModal()
+{
+	GEditor->EditorAddModalWindow(SharedThis(this));
+	return UserResponse;
+}
+
+FString SCreateAnimationAssetDlg::GetAssetPath()
+{
+	return AssetPath.ToString();
+}
+
+FString SCreateAnimationAssetDlg::GetAssetName()
+{
+	return AssetName.ToString();
+}
+
+FString SCreateAnimationAssetDlg::GetFullAssetPath()
+{
+	return AssetPath.ToString() + "/" + AssetName.ToString();
+}
+
+/////////////////////////////////////////////////////
+/////////////////////////////////////////////////////
+/////////////////////////////////////////////////////
+/////////////////////////////////////////////////////
+
+/////////////////////////////////////////////////////
+// Animation editor utility functions
+/////////////////////////////////////////////////////
+
 namespace AnimationEditorUtils
 {
 	/** Creates a unique package and asset name taking the form InBasePackageName+InSuffix */
@@ -24,10 +225,25 @@ namespace AnimationEditorUtils
 			USkeleton* Skeleton = (*SkelIt).Get();
 			if(Skeleton)
 			{
-				// Determine an appropriate name
 				FString Name;
 				FString PackageName;
-				CreateUniqueAssetName(Skeleton->GetOutermost()->GetName(), InPrefix, PackageName, Name);
+				FString AssetPath = Skeleton->GetOutermost()->GetName();
+				// Determine an appropriate name
+				CreateUniqueAssetName(AssetPath, InPrefix, PackageName, Name);
+
+				// set the unique asset as a default name
+				TSharedRef<SCreateAnimationAssetDlg> NewAnimDlg =
+					SNew(SCreateAnimationAssetDlg)
+					.DefaultAssetPath(FText::FromString(PackageName));
+
+				// show a dialog to determine a new asset name
+				if (NewAnimDlg->ShowModal() == EAppReturnType::Cancel)
+				{
+					return;
+				}
+
+				PackageName = NewAnimDlg->GetFullAssetPath();
+				Name = NewAnimDlg->GetAssetName();
 
 				// Create the asset, and assign its skeleton
 				FAssetToolsModule& AssetToolsModule = FModuleManager::GetModuleChecked<FAssetToolsModule>("AssetTools");
diff --git a/Engine/Source/Editor/UnrealEd/Public/AnimationEditorUtils.h b/Engine/Source/Editor/UnrealEd/Public/AnimationEditorUtils.h
--- a/Engine/Source/Editor/UnrealEd/Public/AnimationEditorUtils.h
+++ b/Engine/Source/Editor/UnrealEd/Public/AnimationEditorUtils.h
@@ -3,6 +3,50 @@
 #ifndef __AnimationEditorUtils_h__
 #define __AnimationEditorUtils_h__
 
+/** dialog to prompt users to decide an animation asset name */
+class SCreateAnimationAssetDlg : public SWindow
+{
+public:
+	SLATE_BEGIN_ARGS(SCreateAnimationAssetDlg)
+	{
+	}
+		SLATE_ARGUMENT(FText, DefaultAssetPath)
+	SLATE_END_ARGS()
+
+		SCreateAnimationAssetDlg()
+		: UserResponse(EAppReturnType::Cancel)
+	{
+	}
+
+	void Construct(const FArguments& InArgs);
+
+public:
+	/** Displays the dialog in a blocking fashion */
+	EAppReturnType::Type ShowModal();
+
+	/** Gets the resulting asset path */
+	FString GetAssetPath();
+
+	/** Gets the resulting asset name */
+	FString GetAssetName();
+
+	/** Gets the resulting full asset path (path+'/'+name) */
+	FString GetFullAssetPath();
+
+protected:
+	void OnPathChange(const FString& NewPath);
+	void OnNameChange(const FText& NewName, ETextCommit::Type CommitInfo);
+	FReply OnButtonClick(EAppReturnType::Type ButtonID);
+
+	bool ValidatePackage();
+
+	EAppReturnType::Type UserResponse;
+	FText AssetPath;
+	FText AssetName;
+
+	static FText LastUsedAssetPath;
+};
+
 /** Defines FCanExecuteAction delegate interface.  Returns true when an action is able to execute. */
 DECLARE_DELEGATE_OneParam(FAnimAssetCreated, TArray<class UObject*>);
 
