Fixed crash sending the "Editor.Usage.Heartbeat" analytics event

From 4.5 pre-release on GitHub:
https://answers.unrealengine.com/questions/102945/45-preview-bug-reports.html

Analytics wasn't enabled, and we were missing the check.

diff --git a/Engine/Source/Editor/UnrealEd/Private/UnrealEdMisc.cpp b/Engine/Source/Editor/UnrealEd/Private/UnrealEdMisc.cpp
--- a/Engine/Source/Editor/UnrealEd/Private/UnrealEdMisc.cpp
+++ b/Engine/Source/Editor/UnrealEd/Private/UnrealEdMisc.cpp
@@ -510,6 +510,12 @@ void FUnrealEdMisc::InitEngineAnalytics()
 
 void FUnrealEdMisc::EditorAnalyticsHeartbeat()
 {
+	// Don't attempt to send the heartbeat if analytics isn't available
+	if(!FEngineAnalytics::IsAvailable())
+	{
+		return;
+	}
+
 	static double LastHeartbeatTime = FPlatformTime::Seconds();
 	
 	double LastInteractionTime = FSlateApplication::Get().GetLastUserInteractionTime();
