Optimisation for setting unique actor names on multiple objects

Previously the logic was to iterate the entire world each time a duplicate name was found. This was very slow when setting unique actor names for large numbers of objects (especially if they all shared similar names). The function now takes an option set of actor labels to avoid having to iterate the world every time.

diff --git a/Engine/Source/Editor/UnrealEd/Classes/Editor/EditorEngine.h b/Engine/Source/Editor/UnrealEd/Classes/Editor/EditorEngine.h
--- a/Engine/Source/Editor/UnrealEd/Classes/Editor/EditorEngine.h
+++ b/Engine/Source/Editor/UnrealEd/Classes/Editor/EditorEngine.h
@@ -170,6 +170,33 @@ struct FCopySelectedInfo
 	}
 };
 
+/** A cache of actor labels */
+struct FCachedActorLabels
+{
+	/** Default constructor - does not populate the array */
+	FCachedActorLabels();
+
+	/** Constructor that populates the set of actor names */
+	explicit FCachedActorLabels(UWorld* World, const TSet<AActor*>& IgnoredActors = TSet<AActor*>());
+
+	/** Populate the set of actor names */
+	void Populate(UWorld* World, const TSet<AActor*>& IgnoredActors = TSet<AActor*>());
+
+	/** Add a new label to this set */
+	FORCEINLINE void Add(const FString& InLabel)
+	{
+		ActorLabels.Add(InLabel);
+	}
+
+	/** Check if the specified label exists */
+	FORCEINLINE bool Contains(const FString& InLabel) const
+	{
+		return ActorLabels.Contains(InLabel);
+	}
+
+private:
+	TSet<FString> ActorLabels;
+};
 
 UCLASS(config=Engine, transient)
 class UNREALED_API UEditorEngine : public UEngine
@@ -2195,11 +2222,11 @@ public:
 	/**
 	 * Assigns a new label to an actor. If the name exists it will be appended with a number to make it unique. Actor labels are only available in development builds.
 	 *
-	 * @param	Actor			The actor to change the label of
-	 * @param	NewActorLabel	The new label string to assign to the actor.  If empty, the actor will have a default label.
+	 * @param	Actor					The actor to change the label of
+	 * @param	NewActorLabel			The new label string to assign to the actor.  If empty, the actor will have a default label.
+	 * @param	InExistingActorLabels	(optional) Pointer to a set of actor labels that are currently in use
 	 */
-	void SetActorLabelUnique( AActor* Actor, const FString& NewActorLabel ) const;
-
+	void SetActorLabelUnique( AActor* Actor, const FString& NewActorLabel, const FCachedActorLabels* InExistingActorLabels = nullptr ) const;
 
 	/**
 	 * Gets the user-friendly, localized (if exists) name of a property
diff --git a/Engine/Source/Editor/UnrealEd/Private/Editor.cpp b/Engine/Source/Editor/UnrealEd/Private/Editor.cpp
--- a/Engine/Source/Editor/UnrealEd/Private/Editor.cpp
+++ b/Engine/Source/Editor/UnrealEd/Private/Editor.cpp
@@ -4602,7 +4602,7 @@ bool UEditorEngine::SplitActorLabel( FString& InOutLabel, int32& OutIdx ) const
 	return false;
 }
 
-void UEditorEngine::SetActorLabelUnique( AActor* Actor, const FString& NewActorLabel ) const
+void UEditorEngine::SetActorLabelUnique(AActor* Actor, const FString& NewActorLabel, const FCachedActorLabels* InExistingActorLabels) const
 {
 	check( Actor );
 
@@ -4610,38 +4610,33 @@ void UEditorEngine::SetActorLabelUnique( AActor* Actor, const FString& NewActorL
 	FString ModifiedActorLabel = NewActorLabel;
 	int32   LabelIdx           = 0;
 
-	for (;;)
+	FCachedActorLabels ActorLabels;
+	if (!InExistingActorLabels)
 	{
-		// Check for a duplicate actor name
-		bool Duplicate = false;
-		for (FActorIterator It(Actor->GetWorld()); It; ++It)
-		{
-			if (Actor != *It && It->GetActorLabel() == ModifiedActorLabel)
-			{
-				Duplicate = true;
-				break;
-			}
-		}
+		InExistingActorLabels = &ActorLabels;
+
+		TSet<AActor*> IgnoreActors;
+		IgnoreActors.Add(Actor);
+		ActorLabels.Populate(Actor->GetWorld(), IgnoreActors);
+	}
 
-		// If there wasn't one, we can stop looping
-		if (!Duplicate)
+
+	if (InExistingActorLabels->Contains(ModifiedActorLabel))
+	{
+		// See if the current label ends in a number, and try to create a new label based on that
+		if (!SplitActorLabel( Prefix, LabelIdx ))
 		{
-			break;
+			// If there wasn't a number on there, append an underscore and start the numbering from 2 (1 before incrementing below)
+			Prefix += TEXT('_');
+			LabelIdx = 1;
 		}
 
-		if (LabelIdx == 0)
+		// Update the actor label until we find one that doesn't already exist
+		while (InExistingActorLabels->Contains(ModifiedActorLabel))
 		{
-			// See if the current label ends in a number, and try to create a new label based on that
-			if (!SplitActorLabel( Prefix, LabelIdx ))
-			{
-				// If there wasn't a number on there, append an underscore and start the numbering from 2 (1 before incrementing below)
-				Prefix += TEXT('_');
-				LabelIdx = 1;
-			}
+			++LabelIdx;
+			ModifiedActorLabel = FString::Printf(TEXT("%s%d"), *Prefix, LabelIdx);
 		}
-
-		++LabelIdx;
-		ModifiedActorLabel = FString::Printf(TEXT("%s%d"), *Prefix, LabelIdx);
 	}
 
 	Actor->SetActorLabel( ModifiedActorLabel );
@@ -7370,3 +7365,27 @@ UWorld::InitializationValues UEditorEngine::GetEditorWorldInitializationValues()
 		.ShouldSimulatePhysics(false)
 		.EnableTraceCollision(true);
 }
+
+FCachedActorLabels::FCachedActorLabels()
+{
+	
+}
+
+FCachedActorLabels::FCachedActorLabels(UWorld* World, const TSet<AActor*>& IgnoredActors)
+{
+	Populate(World, IgnoredActors);
+}
+
+void FCachedActorLabels::Populate(UWorld* World, const TSet<AActor*>& IgnoredActors)
+{
+	ActorLabels.Empty();
+
+	for (FActorIterator It(World); It; ++It)
+	{
+		if (!IgnoredActors.Contains(*It))
+		{
+			ActorLabels.Add(It->GetActorLabel());
+		}
+	}
+	ActorLabels.Shrink();
+}
diff --git a/Engine/Source/Editor/UnrealEd/Private/EditorActor.cpp b/Engine/Source/Editor/UnrealEd/Private/EditorActor.cpp
--- a/Engine/Source/Editor/UnrealEd/Private/EditorActor.cpp
+++ b/Engine/Source/Editor/UnrealEd/Private/EditorActor.cpp
@@ -234,6 +234,8 @@ void UUnrealEdEngine::edactPasteSelected(UWorld* InWorld, bool bDuplicate, bool
 	// Create a location offset.
 	const FVector LocationOffset = CreateLocationOffset( bDuplicate, bOffsetLocations );
 
+	FCachedActorLabels ActorLabels(InWorld);
+
 	// Transact the current selection set.
 	USelection* SelectedActors = GetSelectedActors();
 	SelectedActors->Modify();
@@ -273,8 +275,9 @@ void UUnrealEdEngine::edactPasteSelected(UWorld* InWorld, bool bDuplicate, bool
 		Actor->TeleportTo(Actor->GetActorLocation() + ActorLocationOffset, Actor->GetActorRotation(), false, true);
 
 		// Re-label duplicated actors so that labels become unique
-		GEditor->SetActorLabelUnique(Actor, Actor->GetActorLabel());
-		
+		GEditor->SetActorLabelUnique(Actor, Actor->GetActorLabel(), &ActorLabels);
+		ActorLabels.Add(Actor->GetActorLabel());
+
 		GEditor->Layers->InitializeNewActorLayers( Actor );
 
 			// Ensure any layers this actor belongs to are visible
