Editor Performance raw files will now list the time stamp in which the data was taken.
AutomationEditorCommon.cpp

diff --git a/Engine/Source/Editor/UnrealEd/Private/Tests/AutomationEditorCommon.cpp b/Engine/Source/Editor/UnrealEd/Private/Tests/AutomationEditorCommon.cpp
--- a/Engine/Source/Editor/UnrealEd/Private/Tests/AutomationEditorCommon.cpp
+++ b/Engine/Source/Editor/UnrealEd/Private/Tests/AutomationEditorCommon.cpp
@@ -530,7 +530,15 @@ bool FEditorPerformanceCommand::Update()
 		//Log out to a text file the Used Virtual Memory in kilobytes.
 		WriteToTextFile(TEXT("Performance"), ShortMapName, TEXT("RAWMemoryUsedVirtual"), (float)UsedVirtualMemory / 1024, TEXT(","));
 
-
+		//Log out to a text file the Duration.
+		FString OldCurrentTimeData;
+		FDateTime CurrentDateAndTime = FDateTime::Now();
+		FString OldFileLocation = FPaths::AutomationLogDir() + TEXT("Performance/") + ShortMapName + TEXT("/RAWTimeStamp.txt");
+		FString FormatedTimeStamp = FString::Printf(TEXT("%04i-%02i-%02i %02i:%02i:%02i:%3i"), CurrentDateAndTime.GetYear(), CurrentDateAndTime.GetMonth(), CurrentDateAndTime.GetDay(), CurrentDateAndTime.GetHour(), CurrentDateAndTime.GetMinute(), CurrentDateAndTime.GetSecond(), CurrentDateAndTime.GetMillisecond());
+		FFileHelper::LoadFileToString( OldCurrentTimeData, *OldFileLocation );
+		FString FileSetup = OldCurrentTimeData + FormatedTimeStamp + TEXT(",");
+		FFileHelper::SaveStringToFile(FileSetup, *OldFileLocation);
+		
 		return false;
 	}
 	
@@ -561,6 +569,7 @@ bool FGenerateEditorPerformanceCharts::Update()
 	FString DurationFileLocation = FileLocation + TEXT("RAWDuration.txt");
 	FString MemoryUsedPeakFileLocation = FileLocation + TEXT("RAWMemoryUsedPeak.txt");
 	FString MemoryUsedPeakVirtualFileLocation = FileLocation + TEXT("RAWMemoryUsedPeakVirtual.txt");
+	FString TimeStampFileLocation = FileLocation + TEXT("RAWTimeStamp.txt");
 	//Get the data from the raw files and put them into an array.
 	TArray<FString> FPSRawArray = CreateArrayFromFile(FPSFileLocation);
 	TArray<FString> FrameTimeRawArray = CreateArrayFromFile(FrameTimeLocation);
@@ -572,11 +581,11 @@ bool FGenerateEditorPerformanceCharts::Update()
 	TArray<FString> TestRunTimeRawArray = CreateArrayFromFile(DurationFileLocation);
 	TArray<FString> MemoryUsedPeakArray = CreateArrayFromFile(MemoryUsedPeakFileLocation);
 	TArray<FString> MemoryUsedPeakVirtualArray = CreateArrayFromFile(MemoryUsedPeakVirtualFileLocation);
+	TArray<FString> TimeStampArray = CreateArrayFromFile(TimeStampFileLocation);
 
 	//Variable that will hold the old performance csv file data.
 	FString OldPerformanceCSVFile;
 
-	//The CSV files.
 	//RAW csv holds all of the information from the text files
 	//Final csv only holds the averaged information.
 	FString FinalCSVFilename = FString::Printf(TEXT("%s%s_Performance.csv"), *FileLocation, *LoadedMapName);
@@ -591,7 +600,7 @@ bool FGenerateEditorPerformanceCharts::Update()
 	}
 
 	//Add the top title row to the raw csv file
-	FString RAWCSVLine = (TEXT("Map Name, Changelist, Test Run Time, Map Load Time, Average FPS, Frame Time, Used Physical Memory, Used Virtual Memory, Used Peak Physical, Used Peak Virtual, Available Physical Memory, Available Virtual Memory\n"));
+	FString RAWCSVLine = (TEXT("Map Name, Changelist, Time Stamp, Test Run Time (in seconds), Map Load Time, Average FPS, Frame Time, Used Physical Memory, Used Virtual Memory, Used Peak Physical, Used Peak Virtual, Available Physical Memory, Available Virtual Memory\n"));
 	RAWCSVFile->Serialize(TCHAR_TO_ANSI(*RAWCSVLine), RAWCSVLine.Len());
 	
 	//Loop through the text files and add them to the raw csv file.
@@ -605,7 +614,7 @@ bool FGenerateEditorPerformanceCharts::Update()
 		}
 
 		//Create the line that will be written and then write it to the csv file.
-		RAWCSVLine = LoadedMapName + TEXT(",") + Changelist + TEXT(",") + TestRunTimeRawArray[I] + TEXT(",") + MapLoadTimeRawArray[0] + TEXT(",") + FPSRawArray[I] + TEXT(",") + FrameTimeRawArray[I] + TEXT(",") + MemoryRawArray[I] + TEXT(",") + MemoryUsedVirtualRawArray[I] + TEXT(",") + MemoryUsedPeakArray[I] + TEXT(",") + MemoryUsedPeakVirtualArray[I] + TEXT(",") + MemoryAvailPhysRawArray[I] + TEXT(",") + MemoryAvailVirtualRawArray[I] + LINE_TERMINATOR;
+		RAWCSVLine = LoadedMapName + TEXT(",") + Changelist + TEXT(",") + TimeStampArray[I] + TEXT(",") + TestRunTimeRawArray[I] + TEXT(",") + MapLoadTimeRawArray[0] + TEXT(",") + FPSRawArray[I] + TEXT(",") + FrameTimeRawArray[I] + TEXT(",") + MemoryRawArray[I] + TEXT(",") + MemoryUsedVirtualRawArray[I] + TEXT(",") + MemoryUsedPeakArray[I] + TEXT(",") + MemoryUsedPeakVirtualArray[I] + TEXT(",") + MemoryAvailPhysRawArray[I] + TEXT(",") + MemoryAvailVirtualRawArray[I] + LINE_TERMINATOR;
 		RAWCSVFile->Serialize(TCHAR_TO_ANSI(*RAWCSVLine), RAWCSVLine.Len());
 	}
 	
@@ -672,7 +681,7 @@ bool FGenerateEditorPerformanceCharts::Update()
 	FinalCSVFile->Serialize(TCHAR_TO_ANSI(*OldPerformanceCSVFile), OldPerformanceCSVFile.Len());
 
 	//Write the averaged numbers to the Performance CSV file.
-	FString FinalCSVLine = FString::Printf(TEXT("%s,%s,%s,%.0f,%s,%.3f,%.3f,%.0f,%.0f,%.0f,%.0f,%.0f,%.0f%s"), *FDateTime::Today().ToString(), *LoadedMapName, *Changelist, TestRunTime, *MapLoadTime, FPSAvg, FrameTimeAvg, MemoryAvg, MemoryUsedVirtualAvg, MemoryUsedPeak, MemoryUsedPeakVirtual, MemoryAvailPhysAvg, MemoryAvailVirtualAvg, LINE_TERMINATOR);
+	FString FinalCSVLine = FString::Printf(TEXT("%s,%s,%s,%.0f,%s,%.3f,%.3f,%.0f,%.0f,%.0f,%.0f,%.0f,%.0f%s"), *FDateTime::Now().ToString(), *LoadedMapName, *Changelist, TestRunTime, *MapLoadTime, FPSAvg, FrameTimeAvg, MemoryAvg, MemoryUsedVirtualAvg, MemoryUsedPeak, MemoryUsedPeakVirtual, MemoryAvailPhysAvg, MemoryAvailVirtualAvg, LINE_TERMINATOR);
 	FinalCSVFile->Serialize(TCHAR_TO_ANSI(*FinalCSVLine), FinalCSVLine.Len());
 
 	//Close the csv files so we can use them while the editor is still open.
@@ -698,6 +707,7 @@ bool FGenerateEditorPerformanceCharts::Update()
 	IFileManager::Get().Delete(*MemoryUsedVirtualFileLocation);
 	IFileManager::Get().Delete(*MemoryUsedPeakFileLocation);
 	IFileManager::Get().Delete(*MemoryUsedPeakVirtualFileLocation);
+	IFileManager::Get().Delete(*TimeStampFileLocation);
 
 	return true;
 }
