Switched launchon to not use pak files.

diff --git a/Engine/Source/Editor/UnrealEd/Private/PlayLevel.cpp b/Engine/Source/Editor/UnrealEd/Private/PlayLevel.cpp
--- a/Engine/Source/Editor/UnrealEd/Private/PlayLevel.cpp
+++ b/Engine/Source/Editor/UnrealEd/Private/PlayLevel.cpp
@@ -1261,12 +1261,6 @@ void UEditorEngine::PlayUsingLauncher()
 		LauncherProfile->SetDeployedDeviceGroup(DeviceGroup);
 		LauncherProfile->SetEditorExe(FUnrealEdMisc::Get().GetExecutableForCommandlets());
 
-
-		// hack to make launch on use pak files and generate manifests
-		LauncherProfile->SetDeployWithUnrealPak( true );
-
-
-
 		const FString DummyDeviceName(FString::Printf(TEXT("All_iOS_On_%s"), FPlatformProcess::ComputerName()));
 		if (PlayUsingLauncherDeviceId.Left(PlayUsingLauncherDeviceId.Find(TEXT("@"))) != TEXT("IOS") || !PlayUsingLauncherDeviceName.Contains(DummyDeviceName))
 		{
