Fix UMG widgets clipping outside of viewports

diff --git a/Engine/Source/Editor/LevelEditor/Private/SLevelViewport.cpp b/Engine/Source/Editor/LevelEditor/Private/SLevelViewport.cpp
--- a/Engine/Source/Editor/LevelEditor/Private/SLevelViewport.cpp
+++ b/Engine/Source/Editor/LevelEditor/Private/SLevelViewport.cpp
@@ -30,6 +30,7 @@
 #include "BufferVisualizationData.h"
 #include "EditorViewportCommands.h"
 #include "Runtime/Engine/Classes/Engine/RendererSettings.h"
+#include "SScissorRectBox.h"
 
 static const FName LevelEditorName("LevelEditor");
 
@@ -169,10 +170,13 @@ void SLevelViewport::ConstructViewportOverlayContent()
 #if ALLOW_PLAY_IN_VIEWPORT_GAMEUI
 		ViewportOverlay->AddSlot( SlotIndex )
 		[
-			SNew(SDPIScaler)
-			.DPIScale(this, &SLevelViewport::GetGameViewportDPIScale)
+			SNew( SScissorRectBox )
 			[
-				PIEViewportOverlayWidget.ToSharedRef()
+				SNew(SDPIScaler)
+				.DPIScale(this, &SLevelViewport::GetGameViewportDPIScale)
+				[
+					PIEViewportOverlayWidget.ToSharedRef()
+				]
 			]
 		];
 
diff --git a/Engine/Source/Editor/UnrealEd/Private/PlayLevel.cpp b/Engine/Source/Editor/UnrealEd/Private/PlayLevel.cpp
--- a/Engine/Source/Editor/UnrealEd/Private/PlayLevel.cpp
+++ b/Engine/Source/Editor/UnrealEd/Private/PlayLevel.cpp
@@ -32,7 +32,7 @@
 #include "Engine/GameInstance.h"
 #include "EditorAnalytics.h"
 #include "Runtime/Engine/Classes/Engine/RendererSettings.h"
-
+#include "SScissorRectBox.h"
 #include "Online.h"
 
 DEFINE_LOG_CATEGORY_STATIC(LogPlayLevel, Log, All);
@@ -2431,9 +2431,12 @@ UGameInstance* UEditorEngine::CreatePIEGameInstance(int32 PIEInstance, bool bInS
 						.IsEnabled(FSlateApplication::Get().GetNormalExecutionAttribute())
 						.EnableGammaCorrection( false )// Gamma correction in the game is handled in post processing in the scene renderer
 						[
-							SAssignNew(PIEDPIScaler, SDPIScaler)
+							SNew(SScissorRectBox)
 							[
-								ViewportOverlayWidgetRef
+								SAssignNew(PIEDPIScaler, SDPIScaler)
+								[
+									ViewportOverlayWidgetRef
+								]
 							]
 						];
 
