More minor fixes to the EditorPerformanceTest tool.  Updates to the time stamp.

diff --git a/Engine/Source/Editor/UnrealEd/Private/Tests/AutomationEditorCommon.cpp b/Engine/Source/Editor/UnrealEd/Private/Tests/AutomationEditorCommon.cpp
--- a/Engine/Source/Editor/UnrealEd/Private/Tests/AutomationEditorCommon.cpp
+++ b/Engine/Source/Editor/UnrealEd/Private/Tests/AutomationEditorCommon.cpp
@@ -530,11 +530,11 @@ bool FEditorPerformanceCommand::Update()
 		//Log out to a text file the Used Virtual Memory in kilobytes.
 		WriteToTextFile(TEXT("Performance"), ShortMapName, TEXT("RAWMemoryUsedVirtual"), (float)UsedVirtualMemory / 1024, TEXT(","));
 
-		//Log out to a text file the Duration.
+		//Log out to a text file the current time stamp.
 		FString OldCurrentTimeData;
 		FDateTime CurrentDateAndTime = FDateTime::Now();
+		FString FormatedTimeStamp = FString::Printf(TEXT("%04i-%02i-%02i %02i:%02i:%02i.%03i"), CurrentDateAndTime.GetYear(), CurrentDateAndTime.GetMonth(), CurrentDateAndTime.GetDay(), CurrentDateAndTime.GetHour(), CurrentDateAndTime.GetMinute(), CurrentDateAndTime.GetSecond(), CurrentDateAndTime.GetMillisecond());
 		FString OldFileLocation = FPaths::AutomationLogDir() + TEXT("Performance/") + ShortMapName + TEXT("/RAWTimeStamp.txt");
-		FString FormatedTimeStamp = FString::Printf(TEXT("%04i-%02i-%02i %02i:%02i:%02i:%3i"), CurrentDateAndTime.GetYear(), CurrentDateAndTime.GetMonth(), CurrentDateAndTime.GetDay(), CurrentDateAndTime.GetHour(), CurrentDateAndTime.GetMinute(), CurrentDateAndTime.GetSecond(), CurrentDateAndTime.GetMillisecond());
 		FFileHelper::LoadFileToString( OldCurrentTimeData, *OldFileLocation );
 		FString FileSetup = OldCurrentTimeData + FormatedTimeStamp + TEXT(",");
 		FFileHelper::SaveStringToFile(FileSetup, *OldFileLocation);
@@ -604,7 +604,7 @@ bool FGenerateEditorPerformanceCharts::Update()
 	RAWCSVFile->Serialize(TCHAR_TO_ANSI(*RAWCSVLine), RAWCSVLine.Len());
 	
 	//Loop through the text files and add them to the raw csv file.
-	for (int32 I = 0; I < FPSRawArray.Num(); I++)
+	for (int32 I = 0; I < FPSRawArray.Num() - 1; I++)
 	{
 		//If the raw file isn't available to write to then we'll fail back this test.
 		if (!RAWCSVFile)
