CookCommandlet: Don't add script packages to a list of packages to cook.

diff --git a/Engine/Source/Editor/UnrealEd/Classes/Commandlets/CookCommandlet.h b/Engine/Source/Editor/UnrealEd/Classes/Commandlets/CookCommandlet.h
--- a/Engine/Source/Editor/UnrealEd/Classes/Commandlets/CookCommandlet.h
+++ b/Engine/Source/Editor/UnrealEd/Classes/Commandlets/CookCommandlet.h
@@ -144,5 +144,12 @@ private:
 	/**	Process deferred commands */
 	void ProcessDeferredCommands();
 
-
+	/** Adds a unique package filename to cook. Rejects script packages. */
+	FORCEINLINE void AddFileToCook(TArray<FString>& InOutFilesToCook, const FString& InFilename) const
+	{
+		if (!FPackageName::IsScriptPackage(InFilename))
+		{
+			InOutFilesToCook.AddUnique(InFilename);
+		}
+	}
 };
\ No newline at end of file
diff --git a/Engine/Source/Editor/UnrealEd/Private/Commandlets/CookCommandlet.cpp b/Engine/Source/Editor/UnrealEd/Private/Commandlets/CookCommandlet.cpp
--- a/Engine/Source/Editor/UnrealEd/Private/Commandlets/CookCommandlet.cpp
+++ b/Engine/Source/Editor/UnrealEd/Private/Commandlets/CookCommandlet.cpp
@@ -26,7 +26,6 @@
 
 DEFINE_LOG_CATEGORY_STATIC(LogCookCommandlet, Log, All);
 
-
 UCookerSettings::UCookerSettings(const FPostConstructInitializeProperties& PCIP)
 	: Super(PCIP)
 {
@@ -751,7 +750,7 @@ void UCookCommandlet::CollectFilesToCook(TArray<FString>& FilesInPath)
 	GEditor->ParseMapSectionIni(*Params, MapList);
 	for (int32 MapIdx = 0; MapIdx < MapList.Num(); MapIdx++)
 	{
-		FilesInPath.AddUnique(MapList[MapIdx]);
+		AddFileToCook(FilesInPath, MapList[MapIdx]);
 	}
 
 	TArray<FString> CmdLineMapEntries;
@@ -818,12 +817,12 @@ void UCookCommandlet::CollectFilesToCook(TArray<FString>& FilesInPath)
 			}
 			else
 			{
-				FilesInPath.AddUnique(CurrEntry);
+				AddFileToCook(FilesInPath, CurrEntry);
 			}
 		}
 		else
 		{
-			FilesInPath.AddUnique(CurrEntry);
+			AddFileToCook(FilesInPath, CurrEntry);
 		}
 	}
 
@@ -838,7 +837,7 @@ void UCookCommandlet::CollectFilesToCook(TArray<FString>& FilesInPath)
 		{
 			FString StdFile = Files[Index];
 			FPaths::MakeStandardFilename(StdFile);
-			FilesInPath.AddUnique(StdFile);
+			AddFileToCook(FilesInPath, StdFile);
 
 			// this asset may not be in our currently mounted content directories, so try to mount a new one now
 			FString LongPackageName;
@@ -879,7 +878,7 @@ void UCookCommandlet::CollectFilesToCook(TArray<FString>& FilesInPath)
 
 			for (int32 TokenFileIndex = 0; TokenFileIndex < TokenFiles.Num(); ++TokenFileIndex)
 			{
-				FilesInPath.AddUnique(TokenFiles[TokenFileIndex]);
+				AddFileToCook(FilesInPath, TokenFiles[TokenFileIndex]);
 			}
 		}
 	}
@@ -899,28 +898,28 @@ void UCookCommandlet::CollectFilesToCook(TArray<FString>& FilesInPath)
 		{
 			if (Obj != FName(NAME_None).ToString())
 			{
-				FilesInPath.AddUnique(Obj);
+				AddFileToCook(FilesInPath, Obj);
 			}
 		}
 		if (PlatformEngineIni.GetString(TEXT("/Script/EngineSettings.GameMapsSettings"), TEXT("ServerDefaultMap"), Obj))
 		{
 			if (Obj != FName(NAME_None).ToString())
 			{
-				FilesInPath.AddUnique(Obj);
+				AddFileToCook(FilesInPath, Obj);
 			}
 		}
 		if (PlatformEngineIni.GetString(TEXT("/Script/EngineSettings.GameMapsSettings"), TEXT("GlobalDefaultGameMode"), Obj))
 		{
 			if (Obj != FName(NAME_None).ToString())
 			{
-				FilesInPath.AddUnique(Obj);
+				AddFileToCook(FilesInPath, Obj);
 			}
 		}
 		if (PlatformEngineIni.GetString(TEXT("/Script/EngineSettings.GameMapsSettings"), TEXT("GlobalDefaultServerGameMode"), Obj))
 		{
 			if (Obj != FName(NAME_None).ToString())
 			{
-				FilesInPath.AddUnique(Obj);
+				AddFileToCook(FilesInPath, Obj);
 			}
 		}
 	}
@@ -934,7 +933,7 @@ void UCookCommandlet::CollectFilesToCook(TArray<FString>& FilesInPath)
 	{
 		if (InterfaceFile != TEXT("None") && InterfaceFile != TEXT(""))
 		{
-			FilesInPath.AddUnique(InterfaceFile);
+			AddFileToCook(FilesInPath, InterfaceFile);
 		}
 	}
 
@@ -957,7 +956,7 @@ void UCookCommandlet::CollectFilesToCook(TArray<FString>& FilesInPath)
 				{
 					FString StdFile = Files[Index];
 					FPaths::MakeStandardFilename(StdFile);
-					FilesInPath.AddUnique(StdFile);
+					AddFileToCook(FilesInPath, StdFile);
 				}
 			}
 		}
@@ -973,14 +972,14 @@ void UCookCommandlet::GenerateLongPackageNames(TArray<FString>& FilesInPath)
 		const FString& FileInPath = FilesInPath[FilesInPath.Num() - FileIndex - 1];
 		if (FPackageName::IsValidLongPackageName(FileInPath))
 		{
-			FilesInPathReverse.AddUnique(FileInPath);
+			AddFileToCook(FilesInPathReverse, FileInPath);
 		}
 		else
 		{
 			FString LongPackageName;
 			if (FPackageName::TryConvertFilenameToLongPackageName(FileInPath, LongPackageName))
 			{
-				FilesInPathReverse.AddUnique(LongPackageName);
+				AddFileToCook(FilesInPathReverse, LongPackageName);
 			}
 			else
 			{
