#Automation #Analytics

Added a test analytic event to the EditorAutomationTests.cpp that fires off an Editor.Usage.TestEvent whenever Load All Maps is executed, and also spits out a line in the log. Records no parameters.

Added DEFINE_LOG_CATEGORY_STATIC(LogAnalytics, Log, All); to the AutomationCommon.h

Also made a quick typo fix

diff --git a/Engine/Source/Editor/UnrealEd/Private/Tests/EditorAutomationTests.cpp b/Engine/Source/Editor/UnrealEd/Private/Tests/EditorAutomationTests.cpp
--- a/Engine/Source/Editor/UnrealEd/Private/Tests/EditorAutomationTests.cpp
+++ b/Engine/Source/Editor/UnrealEd/Private/Tests/EditorAutomationTests.cpp
@@ -26,8 +26,6 @@
 #include "PackageTools.h"
 
 
-//DEFINE_LOG_CATEGORY_STATIC(LogEditorAutomationTests, Log, All);
-
 /**
 * Change the attributes for a point light in the level.
 */
@@ -835,6 +833,14 @@ bool FLoadAllMapsInEditorTest::RunTest(const FString& Parameters)
 	FString MapName = Parameters;
 	double MapLoadStartTime = 0;
 
+	//Test event for analytics. This should fire anytime this automation procedure is started.
+	if (FEngineAnalytics::IsAvailable())
+	{
+		FEngineAnalytics::GetProvider().RecordEvent( TEXT("Editor.Usage.TestEvent"));
+		UE_LOG(LogAnalytics, Log, TEXT("AnayticsTest: Load All Maps automation triggered and Editor.Usage.TestEvent analytic event has been fired."));
+	}
+	
+
 	const bool bTakeScreenshots = FAutomationTestFramework::GetInstance().IsScreenshotAllowed();
 	if( bTakeScreenshots )
 	{
diff --git a/Engine/Source/Runtime/Engine/Public/Tests/AutomationCommon.h b/Engine/Source/Runtime/Engine/Public/Tests/AutomationCommon.h
--- a/Engine/Source/Runtime/Engine/Public/Tests/AutomationCommon.h
+++ b/Engine/Source/Runtime/Engine/Public/Tests/AutomationCommon.h
@@ -8,11 +8,12 @@
 // Common Latent commands which are used across test type. I.e. Engine, Network, etc...
 
 DEFINE_LOG_CATEGORY_STATIC(LogEditorAutomationTests, Log, All);
+DEFINE_LOG_CATEGORY_STATIC(LogAnalytics, Log, All);
 
 /** Common automation functions */
 namespace AutomationCommon
 {
-	/** Get a string conatins the render mode we are currently in */
+	/** Get a string contains the render mode we are currently in */
 	static FString GetRenderDetailsString()
 	{
 		FString HardwareDetailsString;
