Added cooked package and asset registry for stats viewer.

diff --git a/Engine/Source/Editor/UnrealEd/Classes/CookOnTheSide/CookOnTheFlyServer.h b/Engine/Source/Editor/UnrealEd/Classes/CookOnTheSide/CookOnTheFlyServer.h
--- a/Engine/Source/Editor/UnrealEd/Classes/CookOnTheSide/CookOnTheFlyServer.h
+++ b/Engine/Source/Editor/UnrealEd/Classes/CookOnTheSide/CookOnTheFlyServer.h
@@ -136,7 +136,7 @@ public:
 		 * Add functions / functionality to the FUnsynchronizedQueue
 		 */
 	};
-
+	
 	template<typename Type>
 	struct FQueue : public FUnsynchronizedQueue<Type, FDummyCriticalSection, FDummyScopeLock>
 	{
@@ -146,7 +146,8 @@ public:
 		 */
 	};
 
-	template<typename Type>
+	// pending delete: don't think I need this anymore
+	/*template<typename Type>
 	struct FLookupQueue
 	{
 	private:
@@ -199,7 +200,7 @@ public:
 			return Queue.Num();
 		}
 
-	};
+	};*/
 
 public:
 	/** cooked file requests which includes platform which file is requested for */
diff --git a/Engine/Source/Editor/UnrealEd/Private/Commandlets/ChunkManifestGenerator.cpp b/Engine/Source/Editor/UnrealEd/Private/Commandlets/ChunkManifestGenerator.cpp
--- a/Engine/Source/Editor/UnrealEd/Private/Commandlets/ChunkManifestGenerator.cpp
+++ b/Engine/Source/Editor/UnrealEd/Private/Commandlets/ChunkManifestGenerator.cpp
@@ -14,6 +14,10 @@
 #include "ChunkManifestGenerator.h"
 #include "IPlatformFileSandboxWrapper.h"
 
+#include "JsonWriter.h"
+#include "JsonReader.h"
+#include "JsonSerializer.h"
+
 DEFINE_LOG_CATEGORY_STATIC(LogChunkManifestGenerator, Log, All);
 
 FChunkManifestGenerator::FChunkManifestGenerator(const TArray<ITargetPlatform*>& InPlatforms)
@@ -354,6 +358,230 @@ bool FChunkManifestGenerator::SaveAssetRegistry(const FString& SandboxPath)
 	return true;
 }
 
+typedef TSharedRef< TJsonWriter<TCHAR, TPrettyJsonPrintPolicy<TCHAR> > > JsonWriter;
+typedef TSharedRef< TJsonReader<TCHAR> > JsonReader;
+
+void CopyJsonValueToWriter( JsonWriter &Json, const FString& ValueName, const TSharedPtr<FJsonValue>& JsonValue )
+{
+	if ( JsonValue->Type == EJson::String )
+	{
+		Json->WriteValue( ValueName, JsonValue->AsString() );
+	}
+	else if ( JsonValue->Type == EJson::Array )
+	{
+		if (ValueName.IsEmpty())
+		{
+			Json->WriteArrayStart();
+		}
+		else
+		{
+			Json->WriteArrayStart(ValueName);
+		}
+		
+		const TArray<TSharedPtr<FJsonValue>>& Array = JsonValue->AsArray();
+		for ( const auto& ArrayValue : Array )
+		{
+			CopyJsonValueToWriter(Json, FString(), ArrayValue);
+		}
+
+		Json->WriteArrayEnd();
+	}
+	else if ( JsonValue->Type == EJson::Object )
+	{
+		if (ValueName.IsEmpty())
+		{
+			Json->WriteObjectStart();
+		}
+		else
+		{
+			Json->WriteObjectStart(ValueName);
+		}
+
+		const TSharedPtr<FJsonObject>& Object = JsonValue->AsObject();
+		for ( const auto& ObjectProperty : Object->Values)
+		{
+			CopyJsonValueToWriter(Json, ObjectProperty.Key, ObjectProperty.Value );
+		}
+
+		Json->WriteObjectEnd();
+	}
+	else
+	{
+		
+		UE_LOG(LogChunkManifestGenerator, Warning, TEXT("Unrecognized json value type %d in object %s"), *UEnum::GetValueAsString(TEXT("Json.EJson"), JsonValue->Type), *ValueName)
+	}
+}
+
+// cooked package asset registry saves information about all the cooked packages and assets contained within for stats purposes
+// in json format
+bool FChunkManifestGenerator::SaveCookedPackageAssetRegistry( const FString& SandboxCookedRegistryFilename, const bool Append )
+{
+	bool bSuccess = false;
+	for ( const auto& Platform : Platforms )
+	{
+		TSet<FName> CookedPackages;
+
+		// save the file 
+		const FString CookedAssetRegistryFilename = SandboxCookedRegistryFilename.Replace(TEXT("[Platform]"), *Platform->PlatformName());
+
+		FString JsonOutString;
+		JsonWriter Json = TJsonWriterFactory<TCHAR, TPrettyJsonPrintPolicy<TCHAR> >::Create(&JsonOutString);
+
+		Json->WriteObjectStart();
+		Json->WriteArrayStart(TEXT("Packages"));
+
+		for ( const auto& Package : AllCookedPackages )
+		{
+			Json->WriteObjectStart(); // unnamed package start
+			const FName& PackageName = Package.Key;
+			const FString& SandboxPath = Package.Value;
+
+			CookedPackages.Add( PackageName );
+
+			FString PlatformSandboxPath = SandboxPath.Replace(TEXT("[Platform]"), *Platform->PlatformName());
+			
+			FDateTime TimeStamp = IFileManager::Get().GetTimeStamp( *PlatformSandboxPath );
+
+			Json->WriteValue( "SourcePackageName", PackageName.ToString() );
+			Json->WriteValue( "CookedPackageName", PlatformSandboxPath );
+			Json->WriteValue( "CookedPackageTimeStamp", TimeStamp.ToString() );
+			
+			Json->WriteArrayStart("AssetData");
+			for (const auto& AssetData : AssetRegistryData)
+			{	// Add only assets that have actually been cooked and belong to any chunk
+				if (AssetData.ChunkIDs.Num() > 0 && (AssetData.PackageName == PackageName))
+				{
+					Json->WriteObjectStart();
+					// save all their infos 
+					Json->WriteValue(TEXT("ObjectPath"), AssetData.ObjectPath.ToString() );
+					Json->WriteValue(TEXT("PackageName"), AssetData.PackageName.ToString() );
+					Json->WriteValue(TEXT("PackagePath"), AssetData.PackagePath.ToString() );
+					Json->WriteValue(TEXT("GroupNames"), AssetData.GroupNames.ToString() );
+					Json->WriteValue(TEXT("AssetName"), AssetData.AssetName.ToString() );
+					Json->WriteValue(TEXT("AssetClass"), AssetData.AssetClass.ToString() );
+					Json->WriteObjectStart("TagsAndValues");
+					for ( const auto& Tag : AssetData.TagsAndValues )
+					{
+						Json->WriteValue( Tag.Key.ToString(), Tag.Value );
+					}
+					Json->WriteObjectEnd(); // end tags and values object
+					Json->WriteObjectEnd(); // end unnamed array object
+				}
+			}
+			Json->WriteArrayEnd();
+			Json->WriteObjectEnd(); // unnamed package
+		}
+
+		if ( Append )
+		{
+			FString JsonInString;
+			if ( FFileHelper::LoadFileToString(JsonInString, *CookedAssetRegistryFilename) )
+			{
+				// load up previous package asset registry and fill in any packages which weren't recooked on this run
+				JsonReader Reader = TJsonReaderFactory<TCHAR>::Create(JsonInString);
+				TSharedPtr<FJsonObject> JsonObject;
+				bool shouldRead = FJsonSerializer::Deserialize(Reader, JsonObject) && JsonObject.IsValid() && JsonObject->HasTypedField<EJson::Array>(TEXT("Packages"));
+				if ( shouldRead )
+				{
+					TArray<TSharedPtr<FJsonValue>> PackageList = JsonObject->GetArrayField(TEXT("Packages"));
+					for (auto PackageListIt = PackageList.CreateConstIterator(); PackageListIt && shouldRead; ++PackageListIt)
+					{
+						const TSharedPtr<FJsonValue>& JsonValue = *PackageListIt;
+						shouldRead = JsonValue->Type == EJson::Object;
+						if ( shouldRead )
+						{
+							const TSharedPtr<FJsonObject>& JsonPackage = JsonValue->AsObject();
+
+							// get the package name and see if we have already written it out this run
+							
+							FString CookedPackageName;
+							check( JsonPackage->TryGetStringField(TEXT("CookedPackageName"), CookedPackageName) );
+
+							const FName CookedPackageFName(*CookedPackageName);
+							if ( CookedPackages.Contains(CookedPackageFName))
+							{
+								// don't need to process this package
+								continue;
+							}
+
+
+							// check that the on disk version is still valid
+							FString SourcePackageName;
+							check( JsonPackage->TryGetStringField( TEXT("SourcePackageName"), SourcePackageName) );
+
+							// if our timestamp is different then don't copy the information over
+							FDateTime CurrentTimeStamp = IFileManager::Get().GetTimeStamp( *CookedPackageName );
+
+							FString SavedTimeString;
+							check( JsonPackage->TryGetStringField(TEXT("CookedPackageTimeStamp"), SavedTimeString) );
+							FDateTime SavedTimeStamp;
+							FDateTime::Parse(SavedTimeString, SavedTimeStamp);
+
+							if ( SavedTimeStamp != CurrentTimeStamp )
+							{
+								continue;
+							}
+
+
+
+							CopyJsonValueToWriter(Json, FString(), JsonValue);
+							// read in all the other stuff and copy it over to the new registry
+							/*Json->WriteObjectStart(); // open package
+
+							// copy all the values over
+							for ( const auto& JsonPackageValue : JsonPackage->Values)
+							{
+								CopyJsonValueToWriter(Json, JsonPackageValue.Key, JsonPackageValue.Value);
+							}
+
+							Json->WriteObjectEnd();*/
+						}
+						
+					}
+				}
+				else
+				{
+					UE_LOG(LogChunkManifestGenerator, Warning, TEXT("Unable to read or json is invalid format %s"), *CookedAssetRegistryFilename);
+				}
+			}
+		}
+
+
+		Json->WriteArrayEnd();
+		Json->WriteObjectEnd();
+
+		if (Json->Close())
+		{
+			FArchive* ItemTemplatesFile = IFileManager::Get().CreateFileWriter(*CookedAssetRegistryFilename);
+			if (ItemTemplatesFile)
+			{
+				// serialize the file contents
+				TStringConversion<FTCHARToUTF8_Convert> Convert(*JsonOutString);
+				ItemTemplatesFile->Serialize(const_cast<ANSICHAR*>(Convert.Get()), Convert.Length());
+				ItemTemplatesFile->Close();
+				if ( !ItemTemplatesFile->IsError() )
+				{
+					bSuccess = true;
+				}
+				else
+				{
+					UE_LOG(LogChunkManifestGenerator, Error, TEXT("Unable to write to %s"), *CookedAssetRegistryFilename);
+				}
+				delete ItemTemplatesFile;
+			}
+			else
+			{
+				UE_LOG(LogChunkManifestGenerator, Error, TEXT("Unable to open %s for writing."), *CookedAssetRegistryFilename);
+			}
+		}
+		else
+		{
+			UE_LOG(LogChunkManifestGenerator, Error, TEXT("Error closing Json Writer"));
+		}
+	}
+	return bSuccess;
+}
+
 bool FChunkManifestGenerator::GetPackageDependencies(FName PackageName, TArray<FName>& DependentPackageNames)
 {
 	return AssetRegistry.GetDependencies(PackageName, DependentPackageNames);
diff --git a/Engine/Source/Editor/UnrealEd/Private/Commandlets/ChunkManifestGenerator.h b/Engine/Source/Editor/UnrealEd/Private/Commandlets/ChunkManifestGenerator.h
--- a/Engine/Source/Editor/UnrealEd/Private/Commandlets/ChunkManifestGenerator.h
+++ b/Engine/Source/Editor/UnrealEd/Private/Commandlets/ChunkManifestGenerator.h
@@ -231,19 +231,6 @@ public:
 	 */
 	void AddUnassignedPackageToManifest(UPackage* Package, const FString& PackageSandboxPath );
 
-	/**
-	* Adds a package to chunk manifest
-	*
-	* @param PackageFName Package name (Package->GetFName())
-	* @param PackagePathName Package path name (Package->GetPathName())
-	* @param SandboxFilename Cooked sandbox path of the package to add to a manifest
-	* @param LastLoadedMapName Name of the last loaded map (can be empty)
-	* @param the SandboxPlatformFile used during cook
-	*/
-	void AddPackageToChunkManifest(const FName& PackageFName, const FString& PackagePathName, const FString& SandboxFilename, const FString& LastLoadedMapName, FSandboxPlatformFile* SandboxFile);
-
-
-
 	/**
 	 * Collects all the packages loaded 
 	 *
@@ -273,4 +260,14 @@ public:
 	*/
 	bool SaveAssetRegistry(const FString& SandboxPath);
 
+
+	/**
+	 * Saves cooked package and asset information about all the cooked packages and assets contained within for stats purposes
+	 * in json format
+	 *
+	 * @param Sandbox path to save the registry to
+	 * @param Append to the previous registry if there is one
+	 */
+	bool SaveCookedPackageAssetRegistry( const FString& SandboxPath, const bool Append );
+
 };
diff --git a/Engine/Source/Editor/UnrealEd/Private/CookOnTheFlyServer.cpp b/Engine/Source/Editor/UnrealEd/Private/CookOnTheFlyServer.cpp
--- a/Engine/Source/Editor/UnrealEd/Private/CookOnTheFlyServer.cpp
+++ b/Engine/Source/Editor/UnrealEd/Private/CookOnTheFlyServer.cpp
@@ -1782,7 +1782,6 @@ void UCookOnTheFlyServer::GenerateAssetRegistry(const TArray<ITargetPlatform*>&
 		}
 	}
 	UE_LOG(LogCookOnTheFly, Display, TEXT("Done creating registry. It took %5.2fs."), GenerateAssetRegistryTime);
-	
 }
 
 void UCookOnTheFlyServer::GenerateLongPackageNames(TArray<FString>& FilesInPath)
@@ -2054,12 +2053,17 @@ void UCookOnTheFlyServer::CookByTheBookFinished()
 		FString RegistryFilename = FPaths::GameDir() / TEXT("AssetRegistry.bin");
 		FString SandboxRegistryFilename = SandboxFile->ConvertToAbsolutePathForExternalAppForWrite(*RegistryFilename);
 		// the registry filename will be modified when we call save asset registry
+		
+		FString CookedAssetRegistry = FPaths::GameDir() / TEXT("CookedAssetRegistry.json");
+		FString SandboxCookedAssetRegistryFilename = SandboxFile->ConvertToAbsolutePathForExternalAppForWrite(*CookedAssetRegistry);
 
 		for ( auto& Manifest : CookByTheBookOptions->ManifestGenerators )
 		{
 			// Always try to save the manifests, this is required to make the asset registry work, but doesn't necessarily write a file
 			Manifest.Value->SaveManifests(SandboxFile.GetOwnedPointer());
 			Manifest.Value->SaveAssetRegistry(SandboxRegistryFilename);
+
+			Manifest.Value->SaveCookedPackageAssetRegistry(SandboxCookedAssetRegistryFilename, true);
 		}
 	}
 
