Fixed TTP #344536 : LIVE: EDITOR: Socket Snapping does not attach Actors to Sockets of Skeletal Mesh components in Blueprint Actors

diff --git a/Engine/Source/Editor/UnrealEd/Classes/Editor/EditorEngine.h b/Engine/Source/Editor/UnrealEd/Classes/Editor/EditorEngine.h
--- a/Engine/Source/Editor/UnrealEd/Classes/Editor/EditorEngine.h
+++ b/Engine/Source/Editor/UnrealEd/Classes/Editor/EditorEngine.h
@@ -2124,8 +2124,9 @@ public:
 	 * @param ParentActor	Actor to parent the child to
 	 * @param ChildActor	Actor being parented
 	 * @param SocketName	An optional socket name to attach to in the parent (use NAME_None when there is no socket)
+	 * @param Component		Actual Component included in the ParentActor which is a usually blueprint actor
 	 */
-	void ParentActors( AActor* ParentActor, AActor* ChildActor, const FName SocketName );
+	void ParentActors( AActor* ParentActor, AActor* ChildActor, const FName SocketName, USceneComponent* Component=NULL );
 
 	/**
 	 * Detaches selected actors from there parents
@@ -2427,15 +2428,35 @@ public:
 
 
 private:
+	/**
+	* snap a child actor to a named socket in a specific SkeletalMeshComponent
+	*
+	* @param ParentActor	Actor to parent the child to
+	* @param ChildActor	Actor being parented
+	* @param SkeletalMeshComponent	Component which has the Socket
+	* @param SocketName	An optional socket name to attach to in the parent (use NAME_None when there is no socket)
+	* @return true if successful, false otherwise
+	*/
+	bool AttachActorToComponent(AActor* ParentActor, AActor* ChildActor, USkeletalMeshComponent* SkeletalMeshComponent, const FName SocketName);
+	/**
+	* snap a child actor to a named socket in a specific StaticMeshComponent
+	*
+	* @param ChildActor	Actor being parented
+	* @param StaticMeshComponent	Component which has the Socket
+	* @param SocketName	An optional socket name to attach to in the parent (use NAME_None when there is no socket)
+	* @return true if successful, false otherwise
+	*/
+	bool AttachActorToComponent(AActor* ChildActor, UStaticMeshComponent* StaticMeshComponent, const FName SocketName);
 	/**
 	 * Utility method to try and snap a child actor to a named socket in the parent
 	 *
 	 * @param ParentActor	Actor to parent the child to
 	 * @param ChildActor	Actor being parented
 	 * @param SocketName	An optional socket name to attach to in the parent (use NAME_None when there is no socket)
+	 * @param Component		Actual Component which has the Socket, this component will be used when ParentActor is a blueprint actor
 	 * @return true if successful, false otherwise
 	 */
-	bool SnapToSocket( AActor* ParentActor, AActor* ChildActor, const FName SocketName );
+	bool SnapToSocket( AActor* ParentActor, AActor* ChildActor, const FName SocketName, USceneComponent* Component=NULL );
 
 	/**
 	 * Delegate definition for the execute function that follows
diff --git a/Engine/Source/Editor/UnrealEd/Private/Editor.cpp b/Engine/Source/Editor/UnrealEd/Private/Editor.cpp
--- a/Engine/Source/Editor/UnrealEd/Private/Editor.cpp
+++ b/Engine/Source/Editor/UnrealEd/Private/Editor.cpp
@@ -3975,7 +3975,54 @@ UBrushBuilder* UEditorEngine::FindBrushBuilder( UClass* BrushBuilderClass )
 	return Builder;
 }
 
-bool UEditorEngine::SnapToSocket (AActor* ParentActor, AActor* ChildActor, const FName SocketName)
+bool UEditorEngine::AttachActorToComponent(AActor* ParentActor, AActor* ChildActor, USkeletalMeshComponent* SkeletalMeshComponent, const FName SocketName)
+{
+	bool bAttachedToSocket = false;
+
+	if(SkeletalMeshComponent && SkeletalMeshComponent->SkeletalMesh)
+	{
+	    USkeletalMeshSocket const* const Socket = SkeletalMeshComponent->SkeletalMesh->FindSocket(SocketName);
+	    if (Socket)
+	    {
+		    bAttachedToSocket = Socket->AttachActor(ChildActor, SkeletalMeshComponent);
+	    }
+	    else
+	    {
+		    // now search bone, if bone exists, snap to the bone
+		    int32 BoneIndex = SkeletalMeshComponent->SkeletalMesh->RefSkeleton.FindBoneIndex(SocketName);
+		    if (BoneIndex != INDEX_NONE)
+		    {
+			    ChildActor->GetRootComponent()->SnapTo(ParentActor->GetRootComponent(), SocketName);
+			    bAttachedToSocket = true;
+    
+    #if WITH_EDITOR
+			    ChildActor->PreEditChange(NULL);
+			    ChildActor->PostEditChange();
+    #endif // WITH_EDITOR
+		    }
+	    }
+	}
+
+	return bAttachedToSocket;
+}
+
+bool UEditorEngine::AttachActorToComponent(AActor* ChildActor, UStaticMeshComponent* StaticMeshComponent, const FName SocketName)
+{
+	bool bAttachedToSocket = false;
+
+	if(StaticMeshComponent && StaticMeshComponent->StaticMesh)
+	{
+	    UStaticMeshSocket const* const Socket = StaticMeshComponent->StaticMesh->FindSocket(SocketName);
+	    if (Socket)
+	    {
+		    bAttachedToSocket = Socket->AttachActor(ChildActor, StaticMeshComponent);
+	    }
+	}
+
+	return bAttachedToSocket;
+}
+
+bool UEditorEngine::SnapToSocket (AActor* ParentActor, AActor* ChildActor, const FName SocketName, USceneComponent* Component)
 {
 	bool bAttachedToSocket = false;
 
@@ -3984,27 +4031,7 @@ bool UEditorEngine::SnapToSocket (AActor* ParentActor, AActor* ChildActor, const
 		ParentSkelMeshActor->SkeletalMeshComponent &&
 		ParentSkelMeshActor->SkeletalMeshComponent->SkeletalMesh)
 	{
-		USkeletalMeshComponent* const SkeletalMeshComponent = ParentSkelMeshActor->SkeletalMeshComponent;
-		USkeletalMeshSocket const* const Socket = SkeletalMeshComponent->SkeletalMesh->FindSocket(SocketName);
-		if (Socket)
-		{
-			bAttachedToSocket = Socket->AttachActor(ChildActor, SkeletalMeshComponent);
-		}
-		else
-		{
-			// now search bone, if bone exists, snap to the bone
-			int32 BoneIndex = SkeletalMeshComponent->SkeletalMesh->RefSkeleton.FindBoneIndex(SocketName);
-			if (BoneIndex!=INDEX_NONE)
-			{
-				ChildActor->GetRootComponent()->SnapTo(ParentActor->GetRootComponent(), SocketName);
-				bAttachedToSocket = true;
-
-#if WITH_EDITOR
-				ChildActor->PreEditChange( NULL );
-				ChildActor->PostEditChange();
-#endif // WITH_EDITOR
-			}
-		}
+		bAttachedToSocket = AttachActorToComponent(ParentActor, ChildActor, ParentSkelMeshActor->SkeletalMeshComponent, SocketName);
 	}
 
 	AStaticMeshActor* ParentStaticMeshActor = Cast<AStaticMeshActor>(ParentActor);
@@ -4012,18 +4039,29 @@ bool UEditorEngine::SnapToSocket (AActor* ParentActor, AActor* ChildActor, const
 		ParentStaticMeshActor->StaticMeshComponent &&
 		ParentStaticMeshActor->StaticMeshComponent->StaticMesh)
 	{
-		UStaticMeshComponent* const StaticMeshComponent = ParentStaticMeshActor->StaticMeshComponent;
-		UStaticMeshSocket const* const Socket = StaticMeshComponent->StaticMesh->FindSocket(SocketName);
-		if (Socket)
+		bAttachedToSocket = AttachActorToComponent(ChildActor, ParentStaticMeshActor->StaticMeshComponent, SocketName);
+	}
+
+	// if the component is in a blueprint actor
+	if (Component && ParentActor->OwnsComponent(Component))
+	{
+		USkeletalMeshComponent* const SkeletalMeshComponent = Cast<USkeletalMeshComponent>(Component);
+		if (SkeletalMeshComponent)
+		{
+			bAttachedToSocket = AttachActorToComponent(ParentActor, ChildActor, SkeletalMeshComponent, SocketName);
+		}
+
+		UStaticMeshComponent* const StaticMeshComponent = Cast<UStaticMeshComponent>(Component);		
+		if (StaticMeshComponent)
 		{
-			bAttachedToSocket = Socket->AttachActor(ChildActor, StaticMeshComponent);
+			bAttachedToSocket = AttachActorToComponent(ChildActor, StaticMeshComponent, SocketName);
 		}
 	}
 
 	return bAttachedToSocket;
 }
 
-void UEditorEngine::ParentActors( AActor* ParentActor, AActor* ChildActor, const FName SocketName)
+void UEditorEngine::ParentActors( AActor* ParentActor, AActor* ChildActor, const FName SocketName, USceneComponent* Component)
 {
 	if (CanParentActors(ParentActor, ChildActor))
 	{
@@ -4056,7 +4094,7 @@ void UEditorEngine::ParentActors( AActor* ParentActor, AActor* ChildActor, const
 		}
 
 		// Try snapping to socket if requested
-		if( (SocketName != NAME_None) && SnapToSocket(ParentActor, ChildActor, SocketName))
+		if ((SocketName != NAME_None) && SnapToSocket(ParentActor, ChildActor, SocketName, Component))
 		{
 			// Refresh editor if snapping to socket was successful
 			RedrawLevelEditingViewports();
diff --git a/Engine/Source/Editor/UnrealEd/Private/LevelViewportClickHandlers.cpp b/Engine/Source/Editor/UnrealEd/Private/LevelViewportClickHandlers.cpp
--- a/Engine/Source/Editor/UnrealEd/Private/LevelViewportClickHandlers.cpp
+++ b/Engine/Source/Editor/UnrealEd/Private/LevelViewportClickHandlers.cpp
@@ -952,7 +952,7 @@ namespace ClickHandlers
 				}
 				else
 				{
-					GEditor->ParentActors(SocketProxy->Actor, Actor, SocketProxy->SocketName);
+					GEditor->ParentActors(SocketProxy->Actor, Actor, SocketProxy->SocketName, SocketProxy->SceneComponent);
 				}
 			}
 		}
