Change the CycleTransformGizmoCoordSystem default key equivalent on Mac OS X to Ctrl+~ as Cmd+~ is already used by OS X for window cycling.

diff --git a/Engine/Source/Editor/UnrealEd/Private/EditorViewportCommands.cpp b/Engine/Source/Editor/UnrealEd/Private/EditorViewportCommands.cpp
--- a/Engine/Source/Editor/UnrealEd/Private/EditorViewportCommands.cpp
+++ b/Engine/Source/Editor/UnrealEd/Private/EditorViewportCommands.cpp
@@ -51,7 +51,11 @@ void FEditorViewportCommands::RegisterCommands()
 	UI_COMMAND( RelativeCoordinateSystem_World, "World-relative Transform", "Move and rotate objects relative to the cardinal world axes", EUserInterfaceActionType::RadioButton, FInputGesture() );
 	UI_COMMAND( RelativeCoordinateSystem_Local, "Local-relative Transform", "Move and rotate objects relative to the object's local axes", EUserInterfaceActionType::RadioButton, FInputGesture() );
 
+#if PLATFORM_MAC
+	UI_COMMAND( CycleTransformGizmoCoordSystem, "Cycle Transform Coordinate System", "Cycles the transform gizmo coordinate systems between world and local (object) space", EUserInterfaceActionType::Button, FInputGesture(EModifierKey::Command, EKeys::Tilde));
+#else
 	UI_COMMAND( CycleTransformGizmoCoordSystem, "Cycle Transform Coordinate System", "Cycles the transform gizmo coordinate systems between world and local (object) space", EUserInterfaceActionType::Button, FInputGesture(EModifierKey::Control, EKeys::Tilde));
+#endif
 	UI_COMMAND( CycleTransformGizmos, "Cycle Between Translate, Rotate, and Scale", "Cycles the transform gizmos between translate, rotate, and scale", EUserInterfaceActionType::Button, FInputGesture(EKeys::SpaceBar) );
 	
 	UI_COMMAND( FocusViewportToSelection, "Focus Selected", "Moves the camera in front of the selection", EUserInterfaceActionType::Button, FInputGesture( EKeys::F ) );
