Fix Auto Generate Lightmap UV's being ignored

diff --git a/Engine/Source/Editor/UnrealEd/Private/Fbx/FbxStaticMeshImport.cpp b/Engine/Source/Editor/UnrealEd/Private/Fbx/FbxStaticMeshImport.cpp
--- a/Engine/Source/Editor/UnrealEd/Private/Fbx/FbxStaticMeshImport.cpp
+++ b/Engine/Source/Editor/UnrealEd/Private/Fbx/FbxStaticMeshImport.cpp
@@ -1200,6 +1200,10 @@ UStaticMesh* UnFbx::FFbxImporter::ImportStaticMeshAsSingle(UObject* InParent, TA
 			SrcModel.BuildSettings.DstLightmapIndex = FirstOpenUVChannel;
 			StaticMesh->LightMapCoordinateIndex = FirstOpenUVChannel;
 		}
+		else
+		{
+			SrcModel.BuildSettings.bGenerateLightmapUVs = false;
+		}
 
 		StaticMesh->LODGroup = ImportOptions->StaticMeshLODGroup;
 		StaticMesh->Build(false);
