Return the correct world from LevelEditorViewport::GetWorld() when in Simulate

diff --git a/Engine/Source/Editor/UnrealEd/Private/EditorViewportClient.cpp b/Engine/Source/Editor/UnrealEd/Private/EditorViewportClient.cpp
--- a/Engine/Source/Editor/UnrealEd/Private/EditorViewportClient.cpp
+++ b/Engine/Source/Editor/UnrealEd/Private/EditorViewportClient.cpp
@@ -2474,7 +2474,7 @@ void FEditorViewportClient::Draw(FViewport* InViewport, FCanvas* Canvas)
 	float RealTimeSeconds;
 	float DeltaTimeSeconds;
 
-	UWorld* World = GWorld;
+	UWorld* World = GetWorld();
 	if (( GetScene() != World->Scene) || (IsRealtime() == true))
 	{
 		// Use time relative to start time to avoid issues with float vs double
diff --git a/Engine/Source/Editor/UnrealEd/Private/LevelEditorViewport.cpp b/Engine/Source/Editor/UnrealEd/Private/LevelEditorViewport.cpp
--- a/Engine/Source/Editor/UnrealEd/Private/LevelEditorViewport.cpp
+++ b/Engine/Source/Editor/UnrealEd/Private/LevelEditorViewport.cpp
@@ -4197,11 +4197,6 @@ void FLevelEditorViewportClient::UpdateLinkedOrthoViewports( bool bInvalidate )
 //	FLevelEditorViewportClient::GetScene
 //
 
-FSceneInterface* FLevelEditorViewportClient::GetScene() const
-{
-	return GWorld != NULL ? GWorld->Scene : NULL;
-}
-
 FLinearColor FLevelEditorViewportClient::GetBackgroundColor() const
 {
 	return IsPerspective() ? GEditor->C_WireBackground : GEditor->C_OrthoBackground;
@@ -4370,7 +4365,12 @@ void FLevelEditorViewportClient::SetAllSpriteCategoryVisibility( bool bVisible )
 
 UWorld* FLevelEditorViewportClient::GetWorld() const
 {
-	if (World)
+	if (bIsSimulateInEditorViewport)
+	{
+		// TODO: Find a proper way to get this
+		return GWorld;
+	}
+	else if (World)
 	{
 		return World;
 	}
diff --git a/Engine/Source/Editor/UnrealEd/Public/LevelEditorViewport.h b/Engine/Source/Editor/UnrealEd/Public/LevelEditorViewport.h
--- a/Engine/Source/Editor/UnrealEd/Public/LevelEditorViewport.h
+++ b/Engine/Source/Editor/UnrealEd/Public/LevelEditorViewport.h
@@ -156,7 +156,6 @@ public:
 	virtual FMatrix GetWidgetCoordSystem() const;
 	virtual ECoordSystem GetWidgetCoordSystemSpace() const;
 	virtual void SetupViewForRendering( FSceneViewFamily& ViewFamily, FSceneView& View ) override;
-	virtual FSceneInterface* GetScene() const override;
 	virtual FLinearColor GetBackgroundColor() const override;
 	virtual int32 GetCameraSpeedSetting() const override;
 	virtual void SetCameraSpeedSetting(int32 SpeedSetting) override;
