Adding the abiliity to gather data for whatever level you currently have open in the editor for the Map Performance in Editor automation test.
The test run time amount can be changed be changed by  updating it in the Editor Perferences automation tab or in the defaultengine.ini
added the default setting to the BaseEngine.ini as well so that all projects can use this without having to manually enter it.

diff --git a/Engine/Config/BaseEngine.ini b/Engine/Config/BaseEngine.ini
--- a/Engine/Config/BaseEngine.ini
+++ b/Engine/Config/BaseEngine.ini
@@ -2068,6 +2068,7 @@ TestAllBlueprints=false
 +EditorTestModules=LandscapeEditor
 +EditorTestModules=GameProjectGeneration
 +TestLevelFolders=TestMaps
++EditorPerformanceTestMaps=(PerformanceTestmap=(FilePath="CurrentLevel"),TestTimer=60)
 
 [AutomationTesting.StaticMeshEditorTest]
 +EditorViewButtons=Wireframe
diff --git a/Engine/Source/Editor/UnrealEd/Private/Tests/AutomationEditorCommon.cpp b/Engine/Source/Editor/UnrealEd/Private/Tests/AutomationEditorCommon.cpp
--- a/Engine/Source/Editor/UnrealEd/Private/Tests/AutomationEditorCommon.cpp
+++ b/Engine/Source/Editor/UnrealEd/Private/Tests/AutomationEditorCommon.cpp
@@ -583,6 +583,12 @@ bool FGenerateEditorPerformanceCharts::Update()
 	TArray<FString> MemoryUsedPeakVirtualArray = CreateArrayFromFile(MemoryUsedPeakVirtualFileLocation);
 	TArray<FString> TimeStampArray = CreateArrayFromFile(TimeStampFileLocation);
 
+	//If the map load time doesn't exist then we give add a '0' to the Map Load time array.
+	if (IFileManager::Get().FileSize(MapLoadTimeFileLocation.GetCharArray().GetData()) <= 0)
+	{
+		MapLoadTimeRawArray.Add(TEXT("0"));
+	}
+
 	//Variable that will hold the old performance csv file data.
 	FString OldPerformanceCSVFile;
 
diff --git a/Engine/Source/Editor/UnrealEd/Private/Tests/EditorPerformanceTests.cpp b/Engine/Source/Editor/UnrealEd/Private/Tests/EditorPerformanceTests.cpp
--- a/Engine/Source/Editor/UnrealEd/Private/Tests/EditorPerformanceTests.cpp
+++ b/Engine/Source/Editor/UnrealEd/Private/Tests/EditorPerformanceTests.cpp
@@ -33,6 +33,14 @@ void FMapPerformanceInEditor::GetTests(TArray<FString>& OutBeautifiedNames, TArr
 		
 		for (int32 indexloop = 0; indexloop < PerformanceMapName.Num(); ++indexloop)
 		{
+
+			if (PerformanceMapName[indexloop].Equals(TEXT("CurrentLevel")))
+			{
+				OutBeautifiedNames.Add(TEXT("CurrentLevel"));
+				OutTestCommands.Add(TEXT("CurrentLevel"));
+				continue;
+			}
+		
 			//Get the location of the map being used.
 			FString Filename = FPaths::ConvertRelativePathToFull(PerformanceMapName[indexloop]);
 
@@ -44,8 +52,8 @@ void FMapPerformanceInEditor::GetTests(TArray<FString>& OutBeautifiedNames, TArr
 				{
 					FString ShortName = FPaths::GetBaseFilename(Filename);
 					FString PathName = FPaths::GetPath(Filename);
-					OutBeautifiedNames.Add(ShortName);
 					FString AssetName = FString::Printf(TEXT("/Engine/%s/%s.%s"), *PathName, *ShortName, *ShortName);
+					OutBeautifiedNames.Add(ShortName);
 					OutTestCommands.Add(AssetName);
 				}
 				else
@@ -61,8 +69,8 @@ void FMapPerformanceInEditor::GetTests(TArray<FString>& OutBeautifiedNames, TArr
 				{
 					FString ShortName = FPaths::GetBaseFilename(Filename);
 					FString PathName = FPaths::GetPath(Filename);
-					OutBeautifiedNames.Add(ShortName);
 					FString AssetName = FString::Printf(TEXT("/Game/%s/%s.%s"), *PathName, *ShortName, *ShortName);
+					OutBeautifiedNames.Add(ShortName);
 					OutTestCommands.Add(AssetName);
 				}
 				else
@@ -71,33 +79,40 @@ void FMapPerformanceInEditor::GetTests(TArray<FString>& OutBeautifiedNames, TArr
 				}
 			}
 		}
-
 	}
+
 }
 
-/*
-*
-*
-*/ 
 bool FMapPerformanceInEditor::RunTest(const FString& Parameters)
 {
 	//Get the map name from the parameters.
 	FString MapName = Parameters;
-	//Get the base filename for the map that will be used.
-	FString ShortMapName = FPaths::GetBaseFilename(MapName);
-	
+
+	//Gets the main frame module to get the name of our current level.
+	const IMainFrameModule& MainFrameModule = FModuleManager::GetModuleChecked< IMainFrameModule >("MainFrame");
+	FString ShortMapName = MainFrameModule.GetLoadedLevelName();
+
+	//Duration variable will be used to indicate how long the test will run for.  Defaults to a minute.
+	float Duration = 0;
+
+	if (!MapName.Equals(TEXT("CurrentLevel")))
+	{
+		//Get the base filename for the map that will be used.
+		ShortMapName = FPaths::GetBaseFilename(MapName);
+
+		//Load Map and get the time it took to take to load the map.
+		ADD_LATENT_AUTOMATION_COMMAND(FEditorLoadMap(MapName));
+	}
+
 	//This gets the info we need from the automation test settings in the engine.ini.
 	UAutomationTestSettings const* AutomationTestSettings = GetDefault<UAutomationTestSettings>();
 	check(AutomationTestSettings);
 
-	//Duration variable will be used to indicate how long the test will run for.
-	float Duration = 0.0f;
-
 	//Now we find the test timer(aka duration) for our test.
 	for (auto Iter = AutomationTestSettings->EditorPerformanceTestMaps.CreateConstIterator(); Iter; ++Iter)
 	{
 		FString IterMapName = FPaths::GetBaseFilename(Iter->PerformanceTestmap.FilePath);
-		if (IterMapName == ShortMapName)
+		if ((IterMapName == ShortMapName) || (IterMapName == TEXT("CurrentLevel")))
 		{
 			Duration = ((Iter->TestTimer));
 			//If the duration is equal to 0 then we simply warn the user that they need to set the test timer option for the performance test.
@@ -115,13 +130,11 @@ bool FMapPerformanceInEditor::RunTest(const FString& Parameters)
 		}
 	}
 
-	UE_LOG(LogEditorAutomationTests, Log, TEXT("Running the performance capture test for %.0f seconds on %s"), Duration, *ShortMapName);
-
-	//Load Map and get the time it took to take to load the map.
-	ADD_LATENT_AUTOMATION_COMMAND(FEditorLoadMap(MapName));
-	ADD_LATENT_AUTOMATION_COMMAND(FWaitLatentCommand(1.0f));
 
+	UE_LOG(LogEditorAutomationTests, Log, TEXT("Running the performance capture test for %.0f seconds on %s"), Duration, *ShortMapName);
+	
 	//Grab the performance numbers based on the duration.
+	ADD_LATENT_AUTOMATION_COMMAND(FWaitLatentCommand(1.0f));
 	ADD_LATENT_AUTOMATION_COMMAND(FEditorPerformanceCommand(Duration));
 	ADD_LATENT_AUTOMATION_COMMAND(FWaitLatentCommand(0.5f));
 
