Fixed crash when attempting to rename a level asset which has been recently moved from the level browser.
#ttp 347330 - EDITOR: CRASH: Attempting to rename a level asset that was recently removed from the level browser can result in a crash
#branch UE4
#proj Editor.UnrealEd
#add Clear the dirty flag and unload the package in EditorLevelUtils::EditorDestroyLevel after destroying the ULevel and UWorld objects.

diff --git a/Engine/Source/Editor/UnrealEd/Private/EditorLevelUtils.cpp b/Engine/Source/Editor/UnrealEd/Private/EditorLevelUtils.cpp
--- a/Engine/Source/Editor/UnrealEd/Private/EditorLevelUtils.cpp
+++ b/Engine/Source/Editor/UnrealEd/Private/EditorLevelUtils.cpp
@@ -17,6 +17,7 @@
 #include "ActorEditorUtils.h"
 #include "Editor/Levels/Public/LevelEdMode.h"
 #include "ContentStreaming.h"
+#include "PackageTools.h"
 
 #include "Runtime/AssetRegistry/Public/AssetRegistryModule.h"
 
@@ -530,8 +531,21 @@ namespace EditorLevelUtils
 		}
 
 		InLevel->GetOuter()->MarkPendingKill();
-		InLevel->MarkPendingKill();		
-		InLevel->GetOuter()->ClearFlags(RF_Public|RF_Standalone);
+		InLevel->MarkPendingKill();
+ 		InLevel->GetOuter()->ClearFlags(RF_Public|RF_Standalone);
+
+		UPackage* Package = Cast<UPackage>(InLevel->GetOutermost());
+		// Destroying actors in the world will have set the package dirty flag - clear it here so it can be unloaded successfully
+		Package->SetDirtyFlag(false);
+
+		TArray<UPackage*> Packages;
+		Packages.Add(Package);
+		if (!PackageTools::UnloadPackages(Packages))
+		{
+			FFormatNamedArguments Args;
+			Args.Add(TEXT("Package"), FText::FromString(Package->GetName()));
+			FMessageDialog::Open(EAppMsgType::Ok, FText::Format(LOCTEXT("UnloadPackagesFail", "Unable to unload package '{Package}'."), Args));
+		}
 
 		World->MarkPackageDirty();
 		World->BroadcastLevelsChanged();
