From John Barrett: Adjusted editor message box dialog, to allow non-modal dialogs with a callback

diff --git a/Engine/Source/Editor/UnrealEd/Private/Dialogs/Dialogs.cpp b/Engine/Source/Editor/UnrealEd/Private/Dialogs/Dialogs.cpp
--- a/Engine/Source/Editor/UnrealEd/Private/Dialogs/Dialogs.cpp
+++ b/Engine/Source/Editor/UnrealEd/Private/Dialogs/Dialogs.cpp
@@ -233,6 +233,10 @@ private:
 	FReply HandleButtonClicked( EAppReturnType::Type InResponse )
 	{
 		Response = InResponse;
+
+		ResultCallback.ExecuteIfBound(ParentWindow.ToSharedRef(), Response);
+
+
 		ParentWindow->RequestDestroyWindow();
 
 		return FReply::Handled();
@@ -244,6 +248,9 @@ private:
 		CopyMessageToClipboard();
 	}
 		
+public:
+	/** Callback delegate that is triggered, when the dialog is run in non-modal mode */
+	FOnMsgDlgResult ResultCallback;
 
 private:
 
@@ -253,29 +260,53 @@ private:
 };
 
 
-EAppReturnType::Type OpenMsgDlgInt( EAppMsgType::Type InMessageType, const FText& InMessage, const FText& InTitle )
+void CreateMsgDlgWindow(TSharedPtr<SWindow>& OutWindow, TSharedPtr<SChoiceDialog>& OutDialog, EAppMsgType::Type InMessageType,
+						const FText& InMessage, const FText& InTitle, FOnMsgDlgResult ResultCallback=NULL)
 {
-	auto ModalWindow = SNew(SWindow)
-		.Title( InTitle )
-		.SizingRule( ESizingRule::Autosized )
+	OutWindow = SNew(SWindow)
+		.Title(InTitle)
+		.SizingRule(ESizingRule::Autosized)
 		.AutoCenter(EAutoCenter::PreferredWorkArea)
-		.SupportsMinimize(false) .SupportsMaximize(false);
+		.SupportsMinimize(false).SupportsMaximize(false);
 
-	auto MessageBox = SNew(SChoiceDialog)
-		.ParentWindow(ModalWindow)
-		.Message( InMessage.ToString() )
+	OutDialog = SNew(SChoiceDialog)
+		.ParentWindow(OutWindow)
+		.Message(InMessage.ToString())
 		.WrapMessageAt(512.0f)
 		.MessageType(InMessageType);
 
-	ModalWindow->SetContent( MessageBox );
+	OutDialog->ResultCallback = ResultCallback;
+
+	OutWindow->SetContent(OutDialog.ToSharedRef());
+}
+
+EAppReturnType::Type OpenMsgDlgInt(EAppMsgType::Type InMessageType, const FText& InMessage, const FText& InTitle)
+{
+	TSharedPtr<SWindow> MsgWindow = NULL;
+	TSharedPtr<SChoiceDialog> MsgDialog = NULL;
+
+	CreateMsgDlgWindow(MsgWindow, MsgDialog, InMessageType, InMessage, InTitle);
 
-	GEditor->EditorAddModalWindow(ModalWindow);
+	GEditor->EditorAddModalWindow(MsgWindow.ToSharedRef());
 
-	EAppReturnType::Type Response = MessageBox->GetResponse();
+	EAppReturnType::Type Response = MsgDialog->GetResponse();
 
 	return Response;
 }
 
+TSharedRef<SWindow> OpenMsgDlgInt_NonModal(EAppMsgType::Type InMessageType, const FText& InMessage, const FText& InTitle,
+							FOnMsgDlgResult ResultCallback)
+{
+	TSharedPtr<SWindow> MsgWindow = NULL;
+	TSharedPtr<SChoiceDialog> MsgDialog = NULL;
+
+	CreateMsgDlgWindow(MsgWindow, MsgDialog, InMessageType, InMessage, InTitle, ResultCallback);
+
+	FSlateApplication::Get().AddWindow(MsgWindow.ToSharedRef());
+
+	return MsgWindow.ToSharedRef();
+}
+
 class SModalDialog : public SCompoundWidget
 {
 public:
diff --git a/Engine/Source/Editor/UnrealEd/Public/Dialogs/Dialogs.h b/Engine/Source/Editor/UnrealEd/Public/Dialogs/Dialogs.h
--- a/Engine/Source/Editor/UnrealEd/Public/Dialogs/Dialogs.h
+++ b/Engine/Source/Editor/UnrealEd/Public/Dialogs/Dialogs.h
@@ -7,8 +7,31 @@
 #include "Engine.h"
 #include "Slate.h"
 
+/**
+ * Opens a modal/blocking message box dialog (with an additional 'copy message text' button), and returns the result immediately
+ *
+ * @param InMessageType		The type of message box to display (e.g. 'ok', or 'yes'/'no' etc.)
+ * @param InMessage			The message to display in the message box
+ * @param InTitle			The title to display for the message box
+ * @return					Returns the result of the user input
+ */
 EAppReturnType::Type UNREALED_API OpenMsgDlgInt(EAppMsgType::Type InMessageType, const FText& InMessage, const FText& InTitle);
 
+DECLARE_DELEGATE_TwoParams(FOnMsgDlgResult, const TSharedRef<SWindow>&, EAppReturnType::Type);
+
+/**
+ * Opens a non-modal/non-blocking message box, which returns its result through a delegate/callback,
+ * using a reference to the created window, to identify which dialog has returned a result (in case there are multiple dialog windows)
+ *
+ * @param InMessageType		The type of message box to display (e.g. 'ok', or 'yes'/'no' etc.)
+ * @param InMessage			The message to display in the message box
+ * @param InTitle			The title to display for the message box
+ * @param ResultCallback	The delegate/callback instance, where results should be returned
+ * @return					Returns the dialog window reference, which the calling code should store, to identify which dialog returned
+ */
+TSharedRef<SWindow> UNREALED_API OpenMsgDlgInt_NonModal(EAppMsgType::Type InMessageType, const FText& InMessage, const FText& InTitle,
+											FOnMsgDlgResult ResultCallback);
+
 /*-----------------------------------------------------------------------------
 	FDragDropConfirmationDialog
 -----------------------------------------------------------------------------*/
