Added cook on the side and cook in the editor to experimental.

diff --git a/Engine/Source/Editor/UnrealEd/Classes/Settings/EditorExperimentalSettings.h b/Engine/Source/Editor/UnrealEd/Classes/Settings/EditorExperimentalSettings.h
--- a/Engine/Source/Editor/UnrealEd/Classes/Settings/EditorExperimentalSettings.h
+++ b/Engine/Source/Editor/UnrealEd/Classes/Settings/EditorExperimentalSettings.h
@@ -72,6 +72,15 @@ public:
 	UPROPERTY(EditAnywhere, config, Category = Rendering, meta = (DisplayName = "Feature Level Preview"))
 	bool bFeatureLevelPreview;
 
+	/** Enable cook in the editor (overrides cook on the side setting) */
+	UPROPERTY(EditAnywhere, config, Category = Cooking, meta = (DisplayName = "Cook In The Editor (cooks from launchon will be run inside the main editor process instead of starting a new editor)"))
+	bool bCookInTheEditor;
+
+	/** Enable cook on the side */
+	UPROPERTY(EditAnywhere, config, Category = Cooking, meta = (DisplayName = "Cook On The Side (Run a cook on the fly server in the background of the editor)"))
+	bool bCookOnTheSide;
+
+
 	/**
 	 * Returns an event delegate that is executed when a setting has changed.
 	 *
diff --git a/Engine/Source/Editor/UnrealEd/Private/UnrealEdEngine.cpp b/Engine/Source/Editor/UnrealEd/Private/UnrealEdEngine.cpp
--- a/Engine/Source/Editor/UnrealEd/Private/UnrealEdEngine.cpp
+++ b/Engine/Source/Editor/UnrealEd/Private/UnrealEdEngine.cpp
@@ -110,19 +110,21 @@ void UUnrealEdEngine::Init(IEngineLoop* InEngineLoop)
 		PropertyModule.RegisterCustomClassLayout("ProjectPackagingSettings", FOnGetDetailCustomizationInstance::CreateStatic(&FProjectPackagingSettingsCustomization::MakeInstance));
 	}
 
-	if ( FParse::Param( FCommandLine::Get(),TEXT("COOKONTHESIDE")) )
+	UEditorExperimentalSettings const* ExperimentalSettings =  GetDefault<UEditorExperimentalSettings>();
+
+	if ( ExperimentalSettings->bCookInTheEditor )
 	{
 		CookServer = ConstructObject<UCookOnTheFlyServer>( UCookOnTheFlyServer::StaticClass() );
-		CookServer->Initialize( ECookMode::CookOnTheFly, ECookInitializationFlags::AutoTick | ECookInitializationFlags::AsyncSave );
-		CookServer->StartNetworkFileServer( false );
+		CookServer->Initialize( ECookMode::CookByTheBookFromTheEditor, ECookInitializationFlags::AutoTick | ECookInitializationFlags::AsyncSave );
 
 		FCoreDelegates::OnObjectPropertyChanged.AddUObject(CookServer, &UCookOnTheFlyServer::OnObjectPropertyChanged);
 		FCoreDelegates::OnObjectModified.AddUObject(CookServer, &UCookOnTheFlyServer::OnObjectModified);
 	}
-	else if ( FParse::Param( FCommandLine::Get(), TEXT("COOKINTHEEDITOR")))
+	else if ( ExperimentalSettings->bCookOnTheSide )
 	{
 		CookServer = ConstructObject<UCookOnTheFlyServer>( UCookOnTheFlyServer::StaticClass() );
-		CookServer->Initialize( ECookMode::CookByTheBookFromTheEditor, ECookInitializationFlags::AutoTick | ECookInitializationFlags::AsyncSave );
+		CookServer->Initialize( ECookMode::CookOnTheFly, ECookInitializationFlags::AutoTick | ECookInitializationFlags::AsyncSave );
+		CookServer->StartNetworkFileServer( false );
 
 		FCoreDelegates::OnObjectPropertyChanged.AddUObject(CookServer, &UCookOnTheFlyServer::OnObjectPropertyChanged);
 		FCoreDelegates::OnObjectModified.AddUObject(CookServer, &UCookOnTheFlyServer::OnObjectModified);
