Editor: Forwarding launch and deployment errors to message log

diff --git a/Engine/Source/Editor/MainFrame/Private/Frame/MainFrameActions.cpp b/Engine/Source/Editor/MainFrame/Private/Frame/MainFrameActions.cpp
--- a/Engine/Source/Editor/MainFrame/Private/Frame/MainFrameActions.cpp
+++ b/Engine/Source/Editor/MainFrame/Private/Frame/MainFrameActions.cpp
@@ -538,17 +538,17 @@ void FMainFrameActionCallbacks::PackageProject( const FName InPlatformInfoName )
 			switch (Result)
 			{
 				case ETargetPlatformReadyStatus::SDKNotFound:
-					MessageLogText = LOCTEXT("SdkNotFoundMessage", "Software Development Kit (SDK) not be found");
+					MessageLogText = LOCTEXT("SdkNotFoundMessage", "Software Development Kit (SDK) not found.");
 					MessageLogTextDetail = FText::Format(LOCTEXT("SdkNotFoundMessageDetail", "Please install the SDK for the {0} target platform!"), Platform->DisplayName());
 					break;
 
 				case ETargetPlatformReadyStatus::ProvisionNotFound:
-					MessageLogText = LOCTEXT("ProvisionNotFoundMessage", "Provision not found");
+					MessageLogText = LOCTEXT("ProvisionNotFoundMessage", "Provision not found.");
 					MessageLogTextDetail = LOCTEXT("ProvisionNotFoundMessageDetail", "A provision is required for deploying your app to the device.");
 					break;
 
 				case ETargetPlatformReadyStatus::SigningKeyNotFound:
-					MessageLogText = LOCTEXT("SigningKeyNotFoundMessage", "Signing key not found");
+					MessageLogText = LOCTEXT("SigningKeyNotFoundMessage", "Signing key not found.");
 					MessageLogTextDetail = LOCTEXT("SigningKeyNotFoundMessageDetail", "The app could not be digitally signed, because the signing key is not configured.");
 					break;
 
diff --git a/Engine/Source/Editor/UnrealEd/Classes/Editor/EditorEngine.h b/Engine/Source/Editor/UnrealEd/Classes/Editor/EditorEngine.h
--- a/Engine/Source/Editor/UnrealEd/Classes/Editor/EditorEngine.h
+++ b/Engine/Source/Editor/UnrealEd/Classes/Editor/EditorEngine.h
@@ -2658,7 +2658,7 @@ protected:
 	void HandleStageStarted(const FString& InStage, TWeakPtr<SNotificationItem> NotificationItemPtr);
 	void HandleStageCompleted(const FString& InStage, double StageTime, bool bHasCode, TWeakPtr<SNotificationItem> NotificationItemPtr);
 	void HandleLaunchCanceled(double TotalTime, bool bHasCode, TWeakPtr<SNotificationItem> NotificationItemPtr);
-	void HandleLaunchCompleted(bool Succeeded, double TotalTime, int32 ErrorCode, bool bHasCode, TWeakPtr<SNotificationItem> NotificationItemPtr);
+	void HandleLaunchCompleted(bool Succeeded, double TotalTime, int32 ErrorCode, bool bHasCode, TWeakPtr<SNotificationItem> NotificationItemPtr, TSharedPtr<class FMessageLog> MessageLog);
 
 public:
 	/** True if world assets are enabled */
diff --git a/Engine/Source/Editor/UnrealEd/Private/Kismet2/DebuggerCommands.cpp b/Engine/Source/Editor/UnrealEd/Private/Kismet2/DebuggerCommands.cpp
--- a/Engine/Source/Editor/UnrealEd/Private/Kismet2/DebuggerCommands.cpp
+++ b/Engine/Source/Editor/UnrealEd/Private/Kismet2/DebuggerCommands.cpp
@@ -1409,17 +1409,17 @@ bool FInternalPlayWorldCommandCallbacks::IsReadyToLaunchOnDevice(FString DeviceI
 		switch (Result)
 		{
 			case ETargetPlatformReadyStatus::SDKNotFound:
-				MessageLogText = LOCTEXT("SdkNotFoundMessage", "Software Development Kit (SDK) not be found");
+				MessageLogText = LOCTEXT("SdkNotFoundMessage", "Software Development Kit (SDK) not found.");
 				MessageLogTextDetail = FText::Format(LOCTEXT("SdkNotFoundMessageDetail", "Please install the SDK for the {0} target platform!"), Platform->DisplayName());
 				break;
 
 			case ETargetPlatformReadyStatus::ProvisionNotFound:
-				MessageLogText = LOCTEXT("ProvisionNotFoundMessage", "Provision not found");
+				MessageLogText = LOCTEXT("ProvisionNotFoundMessage", "Provision not found.");
 				MessageLogTextDetail = LOCTEXT("ProvisionNotFoundMessageDetail", "A provision is required for deploying your app to the device.");
 				break;
 
 			case ETargetPlatformReadyStatus::SigningKeyNotFound:
-				MessageLogText = LOCTEXT("SigningKeyNotFoundMessage", "Signing key not found");
+				MessageLogText = LOCTEXT("SigningKeyNotFoundMessage", "Signing key not found.");
 				MessageLogTextDetail = LOCTEXT("SigningKeyNotFoundMessageDetail", "The app could not be digitally signed, because the signing key is not configured.");
 				break;
 
diff --git a/Engine/Source/Editor/UnrealEd/Private/PlayLevel.cpp b/Engine/Source/Editor/UnrealEd/Private/PlayLevel.cpp
--- a/Engine/Source/Editor/UnrealEd/Private/PlayLevel.cpp
+++ b/Engine/Source/Editor/UnrealEd/Private/PlayLevel.cpp
@@ -1183,7 +1183,7 @@ void UEditorEngine::HandleLaunchCanceled(double TotalTime, bool bHasCode, TWeakP
 	bPlayUsingLauncher = false;	
 }
 
-void UEditorEngine::HandleLaunchCompleted(bool Succeeded, double TotalTime, int32 ErrorCode, bool bHasCode, TWeakPtr<SNotificationItem> NotificationItemPtr)
+void UEditorEngine::HandleLaunchCompleted(bool Succeeded, double TotalTime, int32 ErrorCode, bool bHasCode, TWeakPtr<SNotificationItem> NotificationItemPtr, TSharedPtr<class FMessageLog> MessageLog)
 {
 	if (Succeeded)
 	{
@@ -1217,17 +1217,26 @@ void UEditorEngine::HandleLaunchCompleted(bool Succeeded, double TotalTime, int3
 		const FString DummyDeviceName(FString::Printf(TEXT("All_iOS_On_%s"), FPlatformProcess::ComputerName()));
 		if (PlayUsingLauncherDeviceId.Left(PlayUsingLauncherDeviceId.Find(TEXT("@"))) == TEXT("IOS") && PlayUsingLauncherDeviceName.Contains(DummyDeviceName))
 		{
-			CompletionMsg = LOCTEXT("LauncherTaskFailed", "Deployment failed!!");
+			CompletionMsg = LOCTEXT("LauncherTaskFailed", "Deployment failed!");
 		}
 		else
 		{
-			CompletionMsg = LOCTEXT("LauncherTaskFailed", "Launch failed!!");
+			CompletionMsg = LOCTEXT("LauncherTaskFailed", "Launch failed!");
 		}
+
+		MessageLog->Error()
+			->AddToken(FTextToken::Create(CompletionMsg))
+			->AddToken(FTextToken::Create(FText::FromString(FEditorAnalytics::TranslateErrorCode(ErrorCode))));
+
+		// flush log, because it won't be destroyed until the notification popup closes
+		MessageLog->NumMessages(EMessageSeverity::Info);
+
 		TGraphTask<FLauncherNotificationTask>::CreateTask().ConstructAndDispatchWhenReady(
 			NotificationItemPtr,
 			SNotificationItem::CS_Fail,
 			CompletionMsg
-			);
+		);
+
 		TArray<FAnalyticsEventAttribute> ParamArray;
 		ParamArray.Add(FAnalyticsEventAttribute(TEXT("Time"), TotalTime));
 		FEditorAnalytics::ReportEvent(TEXT( "Editor.LaunchOn.Failed" ), PlayUsingLauncherDeviceId.Left(PlayUsingLauncherDeviceId.Find(TEXT("@"))), bHasCode, ErrorCode, ParamArray);
@@ -1371,11 +1380,13 @@ void UEditorEngine::PlayUsingLauncher()
 		TWeakPtr<SNotificationItem> NotificationItemPtr(NotificationItem);
 		if (LauncherWorker.IsValid() && LauncherWorker->GetStatus() != ELauncherWorkerStatus::Completed)
 		{
+			TSharedPtr<FMessageLog> MessageLog = MakeShareable(new FMessageLog("PackagingResults"));
+
 			GEditor->PlayEditorSound(TEXT("/Engine/EditorSounds/Notifications/CompileStart_Cue.CompileStart_Cue"));
 			LauncherWorker->OnOutputReceived().AddStatic(HandleOutputReceived);
 			LauncherWorker->OnStageStarted().AddUObject(this, &UEditorEngine::HandleStageStarted, NotificationItemPtr);
 			LauncherWorker->OnStageCompleted().AddUObject(this, &UEditorEngine::HandleStageCompleted, bHasCode, NotificationItemPtr);
-			LauncherWorker->OnCompleted().AddUObject(this, &UEditorEngine::HandleLaunchCompleted, bHasCode, NotificationItemPtr);
+			LauncherWorker->OnCompleted().AddUObject(this, &UEditorEngine::HandleLaunchCompleted, bHasCode, NotificationItemPtr, MessageLog);
 			LauncherWorker->OnCanceled().AddUObject(this, &UEditorEngine::HandleLaunchCanceled, bHasCode, NotificationItemPtr);
 		}
 		else
diff --git a/Engine/Source/Editor/UnrealEd/Public/EditorAnalytics.h b/Engine/Source/Editor/UnrealEd/Public/EditorAnalytics.h
--- a/Engine/Source/Editor/UnrealEd/Public/EditorAnalytics.h
+++ b/Engine/Source/Editor/UnrealEd/Public/EditorAnalytics.h
@@ -97,7 +97,8 @@ public:
 	*/
 	UNREALED_API static void ReportBuildRequirementsFailure(FString EventName, FString PlatformName, bool bHasCode, int32 Requirements);
 
-private:
+public:
+
 	/* Translates the error code to something human readable */
 	static FString TranslateErrorCode(int32 ErrorCode);
 };
