Fixed DiffuseColor input still being used when exporting translucent materials to lightmass

diff --git a/Engine/Source/Editor/UnrealEd/Private/Lightmass/LightmassRender.cpp b/Engine/Source/Editor/UnrealEd/Private/Lightmass/LightmassRender.cpp
--- a/Engine/Source/Editor/UnrealEd/Private/Lightmass/LightmassRender.cpp
+++ b/Engine/Source/Editor/UnrealEd/Private/Lightmass/LightmassRender.cpp
@@ -307,6 +307,7 @@ public:
 	int32 CompilePropertyAndSetMaterialPropertyWithoutCast(EMaterialProperty Property, FMaterialCompiler* Compiler) const
 	{
 		static const auto UseDiffuseSpecularMaterialInputs = IConsoleManager::Get().FindTConsoleVariableDataInt(TEXT("r.UseDiffuseSpecularMaterialInputs"));
+		EMaterialProperty DiffuseInput = UseDiffuseSpecularMaterialInputs->GetValueOnGameThread() == 1 ? MP_DiffuseColor : MP_BaseColor;
 
 		// MAKE SURE THIS MATCHES THE CHART IN WillFillData
 		// 						  RETURNED VALUES (F16 'textures')
@@ -334,7 +335,6 @@ public:
 				// Only return for Opaque and Masked...
 				if (BlendMode == BLEND_Opaque || BlendMode == BLEND_Masked)
 				{
-					EMaterialProperty DiffuseInput = UseDiffuseSpecularMaterialInputs->GetValueOnGameThread() == 1 ? MP_DiffuseColor : MP_BaseColor;
 					return Compiler->ForceCast(MaterialInterface->CompileProperty(&ProxyCompiler, DiffuseInput),MCT_Float3,true,true);
 				}
 				break;
@@ -372,7 +372,7 @@ public:
 					}
 					else
 					{
-						return Compiler->ForceCast(MaterialInterface->CompileProperty(Compiler, MP_DiffuseColor),MCT_Float3,true,true);
+						return Compiler->ForceCast(MaterialInterface->CompileProperty(Compiler, DiffuseInput),MCT_Float3,true,true);
 					}
 				}
 				else if ((BlendMode == BLEND_Translucent) || (BlendMode == BLEND_Additive))
@@ -384,7 +384,7 @@ public:
 					}
 					else
 					{
-						ColoredOpacity = Compiler->ForceCast(MaterialInterface->CompileProperty(Compiler, MP_DiffuseColor),MCT_Float3,true,true);
+						ColoredOpacity = Compiler->ForceCast(MaterialInterface->CompileProperty(Compiler, DiffuseInput),MCT_Float3,true,true);
 					}
 					return Compiler->Lerp(Compiler->Constant3(1.0f, 1.0f, 1.0f), ColoredOpacity, Compiler->ForceCast( MaterialInterface->CompileProperty(&ProxyCompiler,MP_Opacity), MCT_Float1));
 				}
diff --git a/Engine/Source/Programs/UnrealLightmass/Public/ImportExport.h b/Engine/Source/Programs/UnrealLightmass/Public/ImportExport.h
--- a/Engine/Source/Programs/UnrealLightmass/Public/ImportExport.h
+++ b/Engine/Source/Programs/UnrealLightmass/Public/ImportExport.h
@@ -49,7 +49,7 @@ namespace Lightmass
 	static const FGuid LM_DEBUGOUTPUT_VERSION			= FGuid(0xfa520216, 0xd5aa49eb, 0xa9bbc1f5, 0x8ab715fc);
 	static const FGuid LM_SCENE_VERSION					= FGuid(0x8ad6acb1, 0xe15e4fec, 0x93c72936, 0xf2b67090);
 	static const FGuid LM_STATICMESH_VERSION			= FGuid(0x9b3b1bc6, 0x9270497a, 0xa07521fe, 0xc7045e91);
-	static const FGuid LM_MATERIAL_VERSION				= FGuid(0xa6630d72, 0x33bf47b3, 0xbbcf76bb, 0x4350a2b5);
+	static const FGuid LM_MATERIAL_VERSION				= FGuid(0xa6640d72, 0x33bf47b3, 0xbbcf76bb, 0x4350a2b5);
 
 
 	/** Alert source object type identifiers... */
