TArray documentation and a coupld of methods deprecation.

diff --git a/Engine/Source/Editor/UnrealEd/Private/Factories/EditorFactories.cpp b/Engine/Source/Editor/UnrealEd/Private/Factories/EditorFactories.cpp
--- a/Engine/Source/Editor/UnrealEd/Private/Factories/EditorFactories.cpp
+++ b/Engine/Source/Editor/UnrealEd/Private/Factories/EditorFactories.cpp
@@ -169,7 +169,7 @@ UObject* UMaterialFactoryNew::FactoryCreateNew(UClass* Class,UObject* InParent,F
 
 		TArray<FExpressionOutput> Outputs;
 		Outputs = Expression->GetOutputs();
-		FExpressionOutput* Output = Outputs.GetTypedData();
+		FExpressionOutput* Output = Outputs.GetData();
 		NewMaterial->BaseColor.Mask = Output->Mask;
 		NewMaterial->BaseColor.MaskR = Output->MaskR;
 		NewMaterial->BaseColor.MaskG = Output->MaskG;
@@ -1684,7 +1684,7 @@ UObject* USoundFactory::FactoryCreateBinary
 		// Read the wave info and make sure we have valid wave data
 		FWaveModInfo WaveInfo;
 		FString ErrorMessage;
-		if( WaveInfo.ReadWaveInfo( RawWaveData.GetTypedData(), RawWaveData.Num(), &ErrorMessage ) )
+		if( WaveInfo.ReadWaveInfo( RawWaveData.GetData(), RawWaveData.Num(), &ErrorMessage ) )
 		{
 			if( *WaveInfo.pBitsPerSample != 16 )
 			{
@@ -3496,7 +3496,7 @@ UTexture* UTextureFactory::ImportTexture(UClass* Class, UObject* InParent, FName
 			if ( PngImageWrapper->GetRaw( Format, BitDepth, RawPNG ) )
 			{
 				uint8* MipData = Texture->Source.LockMip(0);
-				FMemory::Memcpy( MipData, RawPNG->GetTypedData(), RawPNG->Num() );
+				FMemory::Memcpy( MipData, RawPNG->GetData(), RawPNG->Num() );
 
 				// Replace the pixels with 0.0 alpha with a color value from the nearest neighboring color which has a non-zero alpha
 				FillZeroAlphaPNGData( Texture->Source, MipData );
@@ -3570,7 +3570,7 @@ UTexture* UTextureFactory::ImportTexture(UClass* Class, UObject* InParent, FName
 				Texture->SRGB = true;
 			
 				uint8* MipData = Texture->Source.LockMip( 0 );
-				FMemory::Memcpy( MipData, RawJPEG->GetTypedData(), RawJPEG->Num() );
+				FMemory::Memcpy( MipData, RawJPEG->GetData(), RawJPEG->Num() );
 				Texture->Source.UnlockMip( 0 );
 			}
 			else
@@ -4404,7 +4404,7 @@ UObject* UTextureFactory::FactoryCreateBinary
 			((UMaterialExpressionTextureSample*)Material->BaseColor.Expression)->Texture = Texture;
 
 			Outputs = Material->BaseColor.Expression->GetOutputs();
-			FExpressionOutput* Output = Outputs.GetTypedData();
+			FExpressionOutput* Output = Outputs.GetData();
 			Material->BaseColor.Mask = Output->Mask;
 			Material->BaseColor.MaskR = Output->MaskR;
 			Material->BaseColor.MaskG = Output->MaskG;
@@ -4418,7 +4418,7 @@ UObject* UTextureFactory::FactoryCreateBinary
 			((UMaterialExpressionTextureSample*)Material->EmissiveColor.Expression)->Texture = Texture;
 
 			Outputs = Material->EmissiveColor.Expression->GetOutputs();
-			FExpressionOutput* Output = Outputs.GetTypedData();
+			FExpressionOutput* Output = Outputs.GetData();
 			Material->EmissiveColor.Mask = Output->Mask;
 			Material->EmissiveColor.MaskR = Output->MaskR;
 			Material->EmissiveColor.MaskG = Output->MaskG;
@@ -4432,7 +4432,7 @@ UObject* UTextureFactory::FactoryCreateBinary
 			((UMaterialExpressionTextureSample*)Material->Roughness.Expression)->Texture = Texture;
 
 			Outputs = Material->Roughness.Expression->GetOutputs();
-			FExpressionOutput* Output = Outputs.GetTypedData();
+			FExpressionOutput* Output = Outputs.GetData();
 			Material->Roughness.Mask = Output->Mask;
 			Material->Roughness.MaskR = 0;
 			Material->Roughness.MaskG = 0;
@@ -4446,7 +4446,7 @@ UObject* UTextureFactory::FactoryCreateBinary
 			((UMaterialExpressionTextureSample*)Material->EmissiveColor.Expression)->Texture = Texture;
 
 			Outputs = Material->EmissiveColor.Expression->GetOutputs();
-			FExpressionOutput* Output = Outputs.GetTypedData();
+			FExpressionOutput* Output = Outputs.GetData();
 			Material->EmissiveColor.Mask = Output->Mask;
 			Material->EmissiveColor.MaskR = 0;
 			Material->EmissiveColor.MaskG = 0;
@@ -4460,7 +4460,7 @@ UObject* UTextureFactory::FactoryCreateBinary
 			((UMaterialExpressionTextureSample*)Material->Opacity.Expression)->Texture = Texture;
 
 			Outputs = Material->Opacity.Expression->GetOutputs();
-			FExpressionOutput* Output = Outputs.GetTypedData();
+			FExpressionOutput* Output = Outputs.GetData();
 			Material->Opacity.Mask = Output->Mask;
 			Material->Opacity.MaskR = 0;
 			Material->Opacity.MaskG = 0;
@@ -4474,7 +4474,7 @@ UObject* UTextureFactory::FactoryCreateBinary
 			((UMaterialExpressionTextureSample*)Material->OpacityMask.Expression)->Texture = Texture;
 
 			Outputs = Material->OpacityMask.Expression->GetOutputs();
-			FExpressionOutput* Output = Outputs.GetTypedData();
+			FExpressionOutput* Output = Outputs.GetData();
 			Material->OpacityMask.Mask = Output->Mask;
 			Material->OpacityMask.MaskR = 0;
 			Material->OpacityMask.MaskG = 0;
@@ -4590,7 +4590,7 @@ bool UTextureExporterPCX::ExportBinary( UObject* Object, const TCHAR* Type, FArc
 	{
 		for( int32 ColorPlane = 2; ColorPlane >= 0; ColorPlane-- )
 		{
-			uint8* ScreenPtr = RawData.GetTypedData() + (Line * SizeX * 4) + ColorPlane;
+			uint8* ScreenPtr = RawData.GetData() + (Line * SizeX * 4) + ColorPlane;
 			for( int32 Row=0; Row<SizeX; Row++ )
 			{
 				if( (*ScreenPtr&0xc0)==0xc0 )
@@ -4681,8 +4681,8 @@ class FHDRExportHelper
 {
 	void WriteScanLine(FArchive& Ar, const TArray<uint8>& ScanLine)
 	{
-		const uint8* LineEnd    = ScanLine.GetTypedData() + ScanLine.Num();
-		const uint8* LineSource = ScanLine.GetTypedData();
+		const uint8* LineEnd    = ScanLine.GetData() + ScanLine.Num();
+		const uint8* LineSource = ScanLine.GetData();
 		TArray<uint8> Output;
 		Output.Reserve(ScanLine.Num() * 2);
 		while (LineSource < LineEnd)
@@ -4726,7 +4726,7 @@ class FHDRExportHelper
 			}
 			LineSource += NextPos;
 		}
-		Ar.Serialize(Output.GetTypedData(), Output.Num());
+		Ar.Serialize(Output.GetData(), Output.Num());
 	}
 
 	static FColor ToRGBEDithered(const FLinearColor& ColorIN, const FRandomStream& Rand)
@@ -4830,11 +4830,11 @@ class FHDRExportHelper
 			{
 				TArray<FFloat16Color> FloatColors;
 				bReadSuccess = RenderTarget->ReadFloat16Pixels(FloatColors);
-				FMemory::Memcpy(RawData.GetTypedData(), FloatColors.GetTypedData(), ImageBytes);
+				FMemory::Memcpy(RawData.GetData(), FloatColors.GetData(), ImageBytes);
 			}
 			break;
 			case PF_B8G8R8A8:
-				bReadSuccess = RenderTarget->ReadPixelsPtr((FColor*)RawData.GetTypedData());
+				bReadSuccess = RenderTarget->ReadPixelsPtr((FColor*)RawData.GetData());
 			break;
 		}
 		if (bReadSuccess == false)
@@ -4849,11 +4849,11 @@ class FHDRExportHelper
 		WriteHDRHeader(Ar);
 		if (Format == PF_FloatRGBA)
 		{
-			WriteHDRBits(Ar, (FFloat16Color*)RawData.GetTypedData());
+			WriteHDRBits(Ar, (FFloat16Color*)RawData.GetData());
 		}
 		else
 		{
-			WriteHDRBits(Ar, (FColor*)RawData.GetTypedData());
+			WriteHDRBits(Ar, (FColor*)RawData.GetData());
 		}
 	}
 
diff --git a/Engine/Source/Editor/UnrealEd/Private/Factories/SkeletalMeshImport.cpp b/Engine/Source/Editor/UnrealEd/Private/Factories/SkeletalMeshImport.cpp
--- a/Engine/Source/Editor/UnrealEd/Private/Factories/SkeletalMeshImport.cpp
+++ b/Engine/Source/Editor/UnrealEd/Private/Factories/SkeletalMeshImport.cpp
@@ -401,7 +401,7 @@ FSavedCustomSortSectionInfo::FSavedCustomSortSectionInfo(USkeletalMesh* Existing
 		}
 		else
 		{
-			FMemory::Memcpy(SavedIndices.GetTypedData(), LODModel.MultiSizeIndexContainer.GetIndexBuffer()->GetPointerTo(Section.BaseIndex), NumIndices*LODModel.MultiSizeIndexContainer.GetDataTypeSize());
+			FMemory::Memcpy(SavedIndices.GetData(), LODModel.MultiSizeIndexContainer.GetIndexBuffer()->GetPointerTo(Section.BaseIndex), NumIndices*LODModel.MultiSizeIndexContainer.GetDataTypeSize());
 		}
 	}
 }
diff --git a/Engine/Source/Editor/UnrealEd/Private/Factories/VectorFieldFactory.cpp b/Engine/Source/Editor/UnrealEd/Private/Factories/VectorFieldFactory.cpp
--- a/Engine/Source/Editor/UnrealEd/Private/Factories/VectorFieldFactory.cpp
+++ b/Engine/Source/Editor/UnrealEd/Private/Factories/VectorFieldFactory.cpp
@@ -144,7 +144,7 @@ static bool ParseFGA(
 	const int32 ValueCount = VectorCount * 3;
 	OutContents->Values.Empty( ValueCount );
 	OutContents->Values.AddZeroed( ValueCount );
-	float* Values = OutContents->Values.GetTypedData();
+	float* Values = OutContents->Values.GetData();
 	float* ValuesEnd = Values + OutContents->Values.Num();
 
 	// Parse each value.
@@ -239,7 +239,7 @@ UObject* UVectorFieldStaticFactory::FactoryCreateBinary(
 				const int32 DestBufferSize = VectorCount * sizeof(FFloat16Color);
 				VectorField->SourceData.Lock( LOCK_READ_WRITE );
 				FFloat16Color* RESTRICT DestValues = (FFloat16Color*)VectorField->SourceData.Realloc( DestBufferSize );
-				const FVector* RESTRICT SrcValues = (FVector*)FileContents.Values.GetTypedData();
+				const FVector* RESTRICT SrcValues = (FVector*)FileContents.Values.GetData();
 				for ( int32 VectorIndex = 0; VectorIndex < VectorCount; ++VectorIndex )
 				{
 					DestValues->R = SrcValues->X;
diff --git a/Engine/Source/Editor/UnrealEd/Private/Fbx/FbxMaterialImport.cpp b/Engine/Source/Editor/UnrealEd/Private/Fbx/FbxMaterialImport.cpp
--- a/Engine/Source/Editor/UnrealEd/Private/Fbx/FbxMaterialImport.cpp
+++ b/Engine/Source/Editor/UnrealEd/Private/Fbx/FbxMaterialImport.cpp
@@ -110,7 +110,7 @@ UTexture* UnFbx::FFbxImporter::ImportTexture( FbxFileTexture* FbxTexture, bool b
 	if (DataBinary.Num()>0)
 	{
 		UE_LOG(LogFbxMaterialImport, Verbose, TEXT("Loading texture file %s"),*Filename);
-		const uint8* PtrTexture = DataBinary.GetTypedData();
+		const uint8* PtrTexture = DataBinary.GetData();
 		UTextureFactory* TextureFact = new UTextureFactory(FPostConstructInitializeProperties());
 		TextureFact->AddToRoot();
 
@@ -289,7 +289,7 @@ bool UnFbx::FFbxImporter::CreateAndLinkExpressionForMaterialProperty(
 			if (MaterialInput.Expression)
 			{
 				TArray<FExpressionOutput> Outputs = MaterialInput.Expression->GetOutputs();
-				FExpressionOutput* Output = Outputs.GetTypedData();
+				FExpressionOutput* Output = Outputs.GetData();
 				MaterialInput.Mask = Output->Mask;
 				MaterialInput.MaskR = Output->MaskR;
 				MaterialInput.MaskG = Output->MaskG;
@@ -345,7 +345,7 @@ void UnFbx::FFbxImporter::FixupMaterial( FbxSurfaceMaterial& FbxMaterial, UMater
 		}
 
 		TArray<FExpressionOutput> Outputs = UnrealMaterial->BaseColor.Expression->GetOutputs();
-		FExpressionOutput* Output = Outputs.GetTypedData();
+		FExpressionOutput* Output = Outputs.GetData();
 		UnrealMaterial->BaseColor.Mask = Output->Mask;
 		UnrealMaterial->BaseColor.MaskR = Output->MaskR;
 		UnrealMaterial->BaseColor.MaskG = Output->MaskG;
diff --git a/Engine/Source/Editor/UnrealEd/Private/Fbx/FbxSkeletalMeshImport.cpp b/Engine/Source/Editor/UnrealEd/Private/Fbx/FbxSkeletalMeshImport.cpp
--- a/Engine/Source/Editor/UnrealEd/Private/Fbx/FbxSkeletalMeshImport.cpp
+++ b/Engine/Source/Editor/UnrealEd/Private/Fbx/FbxSkeletalMeshImport.cpp
@@ -1224,7 +1224,7 @@ USkeletalMesh* UnFbx::FFbxImporter::ImportSkeletalMesh(UObject* InParent, TArray
 	SkeletalMesh->LODInfo[0].ReductionSettings = Settings;
 
 	// Create initial bounding box based on expanded version of reference pose for meshes without physics assets. Can be overridden by artist.
-	FBox BoundingBox( SkelMeshImportDataPtr->Points.GetTypedData(), SkelMeshImportDataPtr->Points.Num() );
+	FBox BoundingBox(SkelMeshImportDataPtr->Points.GetData(), SkelMeshImportDataPtr->Points.Num());
 	FBox Temp = BoundingBox;
 	FVector MidMesh		= 0.5f*(Temp.Min + Temp.Max);
 	BoundingBox.Min		= Temp.Min + 1.0f*(Temp.Min - MidMesh);
diff --git a/Engine/Source/Editor/UnrealEd/Private/Kismet2/KismetDebugUtilities.cpp b/Engine/Source/Editor/UnrealEd/Private/Kismet2/KismetDebugUtilities.cpp
--- a/Engine/Source/Editor/UnrealEd/Private/Kismet2/KismetDebugUtilities.cpp
+++ b/Engine/Source/Editor/UnrealEd/Private/Kismet2/KismetDebugUtilities.cpp
@@ -178,7 +178,7 @@ void FKismetDebugUtilities::OnScriptException(const UObject* ActiveObject, const
 
 		if (BlueprintObj->GetObjectBeingDebugged() == ActiveObject)
 		{
-			const int32 BreakpointOffset = StackFrame.Code - StackFrame.Node->Script.GetTypedData() - 1; //@TODO: Might want to make this a parameter of Info
+			const int32 BreakpointOffset = StackFrame.Code - StackFrame.Node->Script.GetData() - 1; //@TODO: Might want to make this a parameter of Info
 
 			// Record into the trace log
 			FKismetTraceSample& Tracer = TraceStackSamples.WriteNewElementUninitialized();
@@ -388,7 +388,7 @@ void FKismetDebugUtilities::AttemptToBreakExecution(UBlueprint* BlueprintObj, co
 	}
 	else
 	{
-		UE_LOG(LogBlueprintDebug, Warning, TEXT("Tried to break execution in an unknown spot at object %s:%04X"), *StackFrame.Node->GetFullName(), StackFrame.Code - StackFrame.Node->Script.GetTypedData());
+		UE_LOG(LogBlueprintDebug, Warning, TEXT("Tried to break execution in an unknown spot at object %s:%04X"), *StackFrame.Node->GetFullName(), StackFrame.Code - StackFrame.Node->Script.GetData());
 	}
 
 	if ((GUnrealEd->PlayWorld != NULL) && !GIsAutomationTesting)
diff --git a/Engine/Source/Editor/UnrealEd/Private/Lightmass/Lightmass.cpp b/Engine/Source/Editor/UnrealEd/Private/Lightmass/Lightmass.cpp
--- a/Engine/Source/Editor/UnrealEd/Private/Lightmass/Lightmass.cpp
+++ b/Engine/Source/Editor/UnrealEd/Private/Lightmass/Lightmass.cpp
@@ -3892,7 +3892,7 @@ bool FLightmassProcessor::ImportLightMapData2DData(int32 Channel, FQuantizedLigh
 	// make space for the samples
 	QuantizedData->Data.Empty(SizeX * SizeY);
 	QuantizedData->Data.AddUninitialized(SizeX * SizeY);
-	FLightMapCoefficients* DataBuffer = QuantizedData->Data.GetTypedData();
+	FLightMapCoefficients* DataBuffer = QuantizedData->Data.GetData();
 
 	int32 DataBufferSize = SizeX * SizeY * sizeof(FLightMapCoefficients);
 
diff --git a/Engine/Source/Editor/UnrealEd/Private/Lightmass/LightmassRender.cpp b/Engine/Source/Editor/UnrealEd/Private/Lightmass/LightmassRender.cpp
--- a/Engine/Source/Editor/UnrealEd/Private/Lightmass/LightmassRender.cpp
+++ b/Engine/Source/Editor/UnrealEd/Private/Lightmass/LightmassRender.cpp
@@ -1003,7 +1003,7 @@ bool FLightmassMaterialRenderer::GenerateMaterialPropertyData(
 					case MP_Opacity:		OutputName += TEXT("Transmissive");		break;
 					}
 					OutputName += TEXT(".BMP");
-					FFileHelper::CreateBitmap(*OutputName,InOutSizeX,InOutSizeY,OutputBuffer.GetTypedData());
+					FFileHelper::CreateBitmap(*OutputName,InOutSizeX,InOutSizeY,OutputBuffer.GetData());
 				}
 			}
 		}
diff --git a/Engine/Source/Editor/UnrealEd/Private/ObjectTools.cpp b/Engine/Source/Editor/UnrealEd/Private/ObjectTools.cpp
--- a/Engine/Source/Editor/UnrealEd/Private/ObjectTools.cpp
+++ b/Engine/Source/Editor/UnrealEd/Private/ObjectTools.cpp
@@ -3487,7 +3487,7 @@ namespace ThumbnailTools
 
 					// Copy the contents of the remote texture to system memory
 					// NOTE: OutRawImageData must be a preallocated buffer!
-					RenderTargetResource->ReadPixelsPtr((FColor*)OutData.GetTypedData(), FReadSurfaceDataFlags(), InSrcRect);
+					RenderTargetResource->ReadPixelsPtr((FColor*)OutData.GetData(), FReadSurfaceDataFlags(), InSrcRect);
 				}
 			}
 		}
diff --git a/Engine/Source/Editor/UnrealEd/Private/SComponentClassCombo.cpp b/Engine/Source/Editor/UnrealEd/Private/SComponentClassCombo.cpp
--- a/Engine/Source/Editor/UnrealEd/Private/SComponentClassCombo.cpp
+++ b/Engine/Source/Editor/UnrealEd/Private/SComponentClassCombo.cpp
@@ -390,7 +390,7 @@ void SComponentClassCombo::UpdateComponentClassList()
 
 	if (SortedClassList.Num() > 0)
 	{
-		Sort(SortedClassList.GetTypedData(),SortedClassList.Num(), SortComboEntry());
+		Sort(SortedClassList.GetData(), SortedClassList.Num(), SortComboEntry());
 
 		FString PreviousHeading;
 		for ( int32 ClassIndex = 0; ClassIndex < SortedClassList.Num(); ClassIndex++ )
diff --git a/Engine/Source/Editor/UnrealEd/Private/ScriptDisassembler.cpp b/Engine/Source/Editor/UnrealEd/Private/ScriptDisassembler.cpp
--- a/Engine/Source/Editor/UnrealEd/Private/ScriptDisassembler.cpp
+++ b/Engine/Source/Editor/UnrealEd/Private/ScriptDisassembler.cpp
@@ -116,7 +116,7 @@ uint8 FKismetBytecodeDisassembler::ReadBYTE(int32& ScriptIndex)
 
 FString FKismetBytecodeDisassembler::ReadName(int32& ScriptIndex)
 {
-	const FScriptName ConstValue = *(FScriptName*)(Script.GetTypedData() + ScriptIndex);
+	const FScriptName ConstValue = *(FScriptName*)(Script.GetData() + ScriptIndex);
 	ScriptIndex += sizeof(FScriptName);
 
 	return ScriptNameToName(ConstValue).ToString();
diff --git a/Engine/Source/Editor/UnrealEd/Private/StaticMeshEdit.cpp b/Engine/Source/Editor/UnrealEd/Private/StaticMeshEdit.cpp
--- a/Engine/Source/Editor/UnrealEd/Private/StaticMeshEdit.cpp
+++ b/Engine/Source/Editor/UnrealEd/Private/StaticMeshEdit.cpp
@@ -771,7 +771,7 @@ void CreateModelFromStaticMesh(UModel* Model,AStaticMeshActor* StaticMeshActor)
 			FPoly*						Polygon		= new(Model->Polys->Element) FPoly;
 
 			Polygon->Init();
-			Polygon->iLink = Polygon - Model->Polys->Element.GetTypedData();
+			Polygon->iLink = Polygon - Model->Polys->Element.GetData();
 			Polygon->Material = StaticMesh->LODModels[0].Elements[Triangle.MaterialIndex].Material;
 			Polygon->PolyFlags = PF_DefaultFlags;
 			Polygon->SmoothingMask = Triangle.SmoothingMask;
diff --git a/Engine/Source/Editor/UnrealEd/Private/UnrealEdSrv.cpp b/Engine/Source/Editor/UnrealEd/Private/UnrealEdSrv.cpp
--- a/Engine/Source/Editor/UnrealEd/Private/UnrealEdSrv.cpp
+++ b/Engine/Source/Editor/UnrealEd/Private/UnrealEdSrv.cpp
@@ -1262,7 +1262,7 @@ bool UUnrealEdEngine::Exec( UWorld* InWorld, const TCHAR* Stream, FOutputDevice&
 				FString FileName;
 				const FString BaseFileName = FPaths::ScreenShotDir() / TEXT("EditorScreenshot");
 				FFileHelper::GenerateNextBitmapFilename(*BaseFileName, FileName);
-				FFileHelper::CreateBitmap(*FileName, OutImageSize.X, OutImageSize.Y, OutImageData.GetTypedData());
+				FFileHelper::CreateBitmap(*FileName, OutImageSize.X, OutImageSize.Y, OutImageData.GetData());
 			}
 		};
 
diff --git a/Engine/Source/Editor/WorldBrowser/Private/Tiles/WorldTileCollectionModel.cpp b/Engine/Source/Editor/WorldBrowser/Private/Tiles/WorldTileCollectionModel.cpp
--- a/Engine/Source/Editor/WorldBrowser/Private/Tiles/WorldTileCollectionModel.cpp
+++ b/Engine/Source/Editor/WorldBrowser/Private/Tiles/WorldTileCollectionModel.cpp
@@ -1397,7 +1397,7 @@ static bool ReadRawFile(TArray<DataType>& Result, const TCHAR* Filename, uint32
 	}
 	Result.Reset();
 	Result.AddUninitialized(Reader->TotalSize()/Result.GetTypeSize());
-	Reader->Serialize(Result.GetTypedData(), Result.Num()*Result.GetTypeSize());
+	Reader->Serialize(Result.GetData(), Result.Num()*Result.GetTypeSize());
 	return Reader->Close();
 }
 
diff --git a/Engine/Source/Programs/CrashReporter/CrashReportClient/Private/CrashUpload.cpp b/Engine/Source/Programs/CrashReporter/CrashReportClient/Private/CrashUpload.cpp
--- a/Engine/Source/Programs/CrashReporter/CrashReportClient/Private/CrashUpload.cpp
+++ b/Engine/Source/Programs/CrashReporter/CrashReportClient/Private/CrashUpload.cpp
@@ -134,7 +134,7 @@ bool FCrashUpload::SendCheckReportRequest()
 		const int32 Length = Converter.Length();
 		PostData.Reset( Length );
 		PostData.AddUninitialized( Length );
-		CopyAssignItems( (ANSICHAR*)PostData.GetTypedData(), Converter.Get(), Length );
+		CopyAssignItems( (ANSICHAR*)PostData.GetData(), Converter.Get(), Length );
 
 		Request->SetURL(UrlPrefix / TEXT("CheckReportDetail"));
 		Request->SetHeader(TEXT("Content-Type"), TEXT("text/plain; charset=utf-8"));
diff --git a/Engine/Source/Programs/UnrealHeaderTool/Private/HeaderParser.cpp b/Engine/Source/Programs/UnrealHeaderTool/Private/HeaderParser.cpp
--- a/Engine/Source/Programs/UnrealHeaderTool/Private/HeaderParser.cpp
+++ b/Engine/Source/Programs/UnrealHeaderTool/Private/HeaderParser.cpp
@@ -5072,7 +5072,7 @@ void FHeaderParser::CompileDelegateDeclaration(FClasses& AllClasses, const TCHAR
 		ParseParameterList(AllClasses, DelegateSignatureFunction, /*bExpectCommaBeforeName=*/ true);
 
 		// Check the expected versus actual number of parameters
-		int32 ParamCount = FoundParamCount - DelegateParameterCountStrings.GetTypedData() + 1;
+		int32 ParamCount = FoundParamCount - DelegateParameterCountStrings.GetData() + 1;
 		if (DelegateSignatureFunction->NumParms != ParamCount)
 			FError::Throwf(TEXT("Expected %d parameters but found %d parameters"), ParamCount, DelegateSignatureFunction->NumParms);
 	}
diff --git a/Engine/Source/Programs/UnrealHeaderTool/Private/Manifest.cpp b/Engine/Source/Programs/UnrealHeaderTool/Private/Manifest.cpp
--- a/Engine/Source/Programs/UnrealHeaderTool/Private/Manifest.cpp
+++ b/Engine/Source/Programs/UnrealHeaderTool/Private/Manifest.cpp
@@ -150,9 +150,9 @@ FManifest FManifest::LoadFromFile(const FString& Filename)
 		KnownModule.PublicUObjectHeaders       .AddZeroed(PublicHeaders .Num());
 		KnownModule.PrivateUObjectHeaders      .AddZeroed(PrivateHeaders.Num());
 
-		ProcessHeaderArray(KnownModule.PublicUObjectClassesHeaders.GetTypedData(), ClassesHeaders, *(Outer + TEXT(".ClassHeaders"  )));
-		ProcessHeaderArray(KnownModule.PublicUObjectHeaders       .GetTypedData(), PublicHeaders , *(Outer + TEXT(".PublicHeaders" )));
-		ProcessHeaderArray(KnownModule.PrivateUObjectHeaders      .GetTypedData(), PrivateHeaders, *(Outer + TEXT(".PrivateHeaders")));
+		ProcessHeaderArray(KnownModule.PublicUObjectClassesHeaders.GetData(), ClassesHeaders, *(Outer + TEXT(".ClassHeaders")));
+		ProcessHeaderArray(KnownModule.PublicUObjectHeaders       .GetData(), PublicHeaders , *(Outer + TEXT(".PublicHeaders" )));
+		ProcessHeaderArray(KnownModule.PrivateUObjectHeaders      .GetData(), PrivateHeaders, *(Outer + TEXT(".PrivateHeaders")));
 
 		// Sort the headers alphabetically.  This is just to add determinism to the compilation dependency order, since we currently
 		// don't rely on explicit includes (but we do support 'dependson')
diff --git a/Engine/Source/Programs/UnrealLightmass/Private/Lighting/LightmapData.h b/Engine/Source/Programs/UnrealLightmass/Private/Lighting/LightmapData.h
--- a/Engine/Source/Programs/UnrealLightmass/Private/Lighting/LightmapData.h
+++ b/Engine/Source/Programs/UnrealLightmass/Private/Lighting/LightmapData.h
@@ -55,8 +55,8 @@ namespace Lightmass
 		*/
 		void Quantize(int32 DebugSampleIndex);
 
-		const FLightSample* GetData() { return Data.GetTypedData(); }
-		const FQuantizedLightSampleData* GetQuantizedData() { return QuantizedData.GetTypedData(); }
+		const FLightSample* GetData() { return Data.GetData(); }
+		const FQuantizedLightSampleData* GetQuantizedData() { return QuantizedData.GetData(); }
 
 		/**
 		 * Compresses the raw lightmap data to a buffer for writing over Swarm
@@ -129,8 +129,8 @@ namespace Lightmass
 
 		void Quantize(int32 DebugSampleIndex);
 
-		const FShadowSample* GetData() { return Data.GetTypedData(); }
-		const FQuantizedShadowSampleData* GetQuantizedData() { return QuantizedData.GetTypedData(); }
+		const FShadowSample* GetData() { return Data.GetData(); }
+		const FQuantizedShadowSampleData* GetQuantizedData() { return QuantizedData.GetData(); }
 
 		void Compress(int32 DebugSampleIndex);
 		uint8* GetCompressedData()
@@ -193,8 +193,8 @@ namespace Lightmass
 
 		void Quantize(int32 DebugSampleIndex);
 
-		const FSignedDistanceFieldShadowSample* GetData() { return Data.GetTypedData(); }
-		const FQuantizedSignedDistanceFieldShadowSampleData* GetQuantizedData() { return QuantizedData.GetTypedData(); }
+		const FSignedDistanceFieldShadowSample* GetData() { return Data.GetData(); }
+		const FQuantizedSignedDistanceFieldShadowSampleData* GetQuantizedData() { return QuantizedData.GetData(); }
 
 		void Compress(int32 DebugSampleIndex);
 		uint8* GetCompressedData()
diff --git a/Engine/Source/Runtime/AIModule/Classes/BehaviorTree/BTNode.h b/Engine/Source/Runtime/AIModule/Classes/BehaviorTree/BTNode.h
--- a/Engine/Source/Runtime/AIModule/Classes/BehaviorTree/BTNode.h
+++ b/Engine/Source/Runtime/AIModule/Classes/BehaviorTree/BTNode.h
@@ -256,13 +256,13 @@ const T* UBTNode::GetNodeMemory(const struct FBehaviorTreeSearchData& SearchData
 template<typename T>
 T* UBTNode::GetNodeMemory(struct FBehaviorTreeInstance& BTInstance) const
 {
-	return (T*)(BTInstance.InstanceMemory.GetTypedData() + MemoryOffset);
+	return (T*)(BTInstance.InstanceMemory.GetData() + MemoryOffset);
 }
 
 template<typename T>
 const T* UBTNode::GetNodeMemory(const struct FBehaviorTreeInstance& BTInstance) const
 {
-	return (const T*)(BTInstance.InstanceMemory.GetTypedData() + MemoryOffset);
+	return (const T*)(BTInstance.InstanceMemory.GetData() + MemoryOffset);
 }
 
 template<typename T>
diff --git a/Engine/Source/Runtime/AIModule/Classes/BehaviorTree/BlackboardComponent.h b/Engine/Source/Runtime/AIModule/Classes/BehaviorTree/BlackboardComponent.h
--- a/Engine/Source/Runtime/AIModule/Classes/BehaviorTree/BlackboardComponent.h
+++ b/Engine/Source/Runtime/AIModule/Classes/BehaviorTree/BlackboardComponent.h
@@ -182,10 +182,10 @@ class AIMODULE_API UBlackboardComponent : public UActorComponent
 
 	/** get pointer to raw data for given key */
 	FORCEINLINE uint8* GetKeyRawData(const FName& KeyName) { return GetKeyRawData(GetKeyID(KeyName)); }
-	FORCEINLINE uint8* GetKeyRawData(FBlackboard::FKey KeyID) { return ValueMemory.Num() && ValueOffsets.IsValidIndex(KeyID) ? (ValueMemory.GetTypedData() + ValueOffsets[KeyID]) : NULL; }
+	FORCEINLINE uint8* GetKeyRawData(FBlackboard::FKey KeyID) { return ValueMemory.Num() && ValueOffsets.IsValidIndex(KeyID) ? (ValueMemory.GetData() + ValueOffsets[KeyID]) : NULL; }
 
 	FORCEINLINE const uint8* GetKeyRawData(const FName& KeyName) const { return GetKeyRawData(GetKeyID(KeyName)); }
-	FORCEINLINE const uint8* GetKeyRawData(FBlackboard::FKey KeyID) const { return ValueMemory.Num() && ValueOffsets.IsValidIndex(KeyID) ? (ValueMemory.GetTypedData() + ValueOffsets[KeyID]) : NULL; }
+	FORCEINLINE const uint8* GetKeyRawData(FBlackboard::FKey KeyID) const { return ValueMemory.Num() && ValueOffsets.IsValidIndex(KeyID) ? (ValueMemory.GetData() + ValueOffsets[KeyID]) : NULL; }
 
 	FORCEINLINE bool IsValidKey(FBlackboard::FKey KeyID) const { check(BlackboardAsset); return KeyID != FBlackboard::InvalidKey && BlackboardAsset->Keys.IsValidIndex(KeyID); }
 
diff --git a/Engine/Source/Runtime/AIModule/Classes/EnvironmentQuery/EnvQueryTypes.h b/Engine/Source/Runtime/AIModule/Classes/EnvironmentQuery/EnvQueryTypes.h
--- a/Engine/Source/Runtime/AIModule/Classes/EnvironmentQuery/EnvQueryTypes.h
+++ b/Engine/Source/Runtime/AIModule/Classes/EnvironmentQuery/EnvQueryTypes.h
@@ -679,7 +679,7 @@ public:
 		DEC_MEMORY_STAT_BY(STAT_AI_EQS_InstanceMemory, RawData.GetAllocatedSize() + Items.GetAllocatedSize());
 
 		const int32 DataOffset = RawData.AddUninitialized(ValueSize);
-		TypeItem::SetValue(RawData.GetTypedData() + DataOffset, ItemValue);
+		TypeItem::SetValue(RawData.GetData() + DataOffset, ItemValue);
 		Items.Add(FEnvQueryItem(DataOffset));
 
 		INC_MEMORY_STAT_BY(STAT_AI_EQS_InstanceMemory, RawData.GetAllocatedSize() + Items.GetAllocatedSize());
@@ -873,7 +873,7 @@ public:
 
 		uint8* GetItemData()
 		{
-			return Instance->RawData.GetTypedData() + Instance->Items[CurrentItem].DataOffset;
+			return Instance->RawData.GetData() + Instance->Items[CurrentItem].DataOffset;
 		}
 
 		void DiscardItem()
diff --git a/Engine/Source/Runtime/AIModule/Private/BehaviorTree/BlackboardComponent.cpp b/Engine/Source/Runtime/AIModule/Private/BehaviorTree/BlackboardComponent.cpp
--- a/Engine/Source/Runtime/AIModule/Private/BehaviorTree/BlackboardComponent.cpp
+++ b/Engine/Source/Runtime/AIModule/Private/BehaviorTree/BlackboardComponent.cpp
@@ -793,7 +793,7 @@ bool UBlackboardComponent::GetLocationFromEntry(FBlackboard::FKey KeyID, FVector
 		const FBlackboardEntry* EntryInfo = BlackboardAsset->GetKey(KeyID);
 		if (EntryInfo && EntryInfo->KeyType)
 		{
-			const uint8* ValueData = ValueMemory.GetTypedData() + ValueOffsets[KeyID];
+			const uint8* ValueData = ValueMemory.GetData() + ValueOffsets[KeyID];
 			return EntryInfo->KeyType->GetLocation(ValueData, ResultLocation);
 		}
 	}
@@ -814,7 +814,7 @@ bool UBlackboardComponent::GetRotationFromEntry(FBlackboard::FKey KeyID, FRotato
 		const FBlackboardEntry* EntryInfo = BlackboardAsset->GetKey(KeyID);
 		if (EntryInfo && EntryInfo->KeyType)
 		{
-			const uint8* ValueData = ValueMemory.GetTypedData() + ValueOffsets[KeyID];
+			const uint8* ValueData = ValueMemory.GetData() + ValueOffsets[KeyID];
 			return EntryInfo->KeyType->GetRotation(ValueData, ResultRotation);
 		}
 	}
diff --git a/Engine/Source/Runtime/AIModule/Private/BehaviorTree/Tasks/BTTask_RunEQSQuery.cpp b/Engine/Source/Runtime/AIModule/Private/BehaviorTree/Tasks/BTTask_RunEQSQuery.cpp
--- a/Engine/Source/Runtime/AIModule/Private/BehaviorTree/Tasks/BTTask_RunEQSQuery.cpp
+++ b/Engine/Source/Runtime/AIModule/Private/BehaviorTree/Tasks/BTTask_RunEQSQuery.cpp
@@ -104,7 +104,7 @@ void UBTTask_RunEQSQuery::OnQueryFinished(TSharedPtr<struct FEnvQueryResult> Res
 		UBlackboardComponent* MyBlackboard = MyComp->GetBlackboardComponent();
 		UEnvQueryItemType* ItemTypeCDO = (UEnvQueryItemType*)Result->ItemType->GetDefaultObject();
 
-		bSuccess = ItemTypeCDO->StoreInBlackboard(BlackboardKey, MyBlackboard, Result->RawData.GetTypedData() + Result->Items[0].DataOffset);		
+		bSuccess = ItemTypeCDO->StoreInBlackboard(BlackboardKey, MyBlackboard, Result->RawData.GetData() + Result->Items[0].DataOffset);		
 		if (!bSuccess)
 		{
 			UE_VLOG(MyOwner, LogBehaviorTree, Warning, TEXT("Failed to store query result! item:%s key:%s"),
diff --git a/Engine/Source/Runtime/AIModule/Private/EnvironmentQuery/EQSRenderingComponent.cpp b/Engine/Source/Runtime/AIModule/Private/EnvironmentQuery/EQSRenderingComponent.cpp
--- a/Engine/Source/Runtime/AIModule/Private/EnvironmentQuery/EQSRenderingComponent.cpp
+++ b/Engine/Source/Runtime/AIModule/Private/EnvironmentQuery/EQSRenderingComponent.cpp
@@ -16,7 +16,7 @@ namespace FEQSRenderingHelper
 			ItemType->IsChildOf(UEnvQueryItemType_VectorBase::StaticClass()))
 		{
 			UEnvQueryItemType_VectorBase* DefTypeOb = (UEnvQueryItemType_VectorBase*)ItemType->GetDefaultObject();
-			return DefTypeOb->GetLocation(RawData.GetTypedData() + Items[Index].DataOffset);
+			return DefTypeOb->GetLocation(RawData.GetData() + Items[Index].DataOffset);
 		}
 		return FVector::ZeroVector;
 	}
@@ -148,7 +148,7 @@ void FEQSSceneProxy::CollectEQSData(const struct FEnvQueryResult* ResultItems, c
 	float MaxScore = -BIG_NUMBER;
 	if (bUseMidResults)
 	{
-		const FEnvQueryItem* ItemInfo = Items.GetTypedData();
+		const FEnvQueryItem* ItemInfo = Items.GetData();
 		for (int32 ItemIndex = 0; ItemIndex < Items.Num(); ItemIndex++, ItemInfo++)
 		{
 			if (ItemInfo->IsValid())
diff --git a/Engine/Source/Runtime/AIModule/Private/EnvironmentQuery/EnvQueryDebugHelpers.cpp b/Engine/Source/Runtime/AIModule/Private/EnvironmentQuery/EnvQueryDebugHelpers.cpp
--- a/Engine/Source/Runtime/AIModule/Private/EnvironmentQuery/EnvQueryDebugHelpers.cpp
+++ b/Engine/Source/Runtime/AIModule/Private/EnvironmentQuery/EnvQueryDebugHelpers.cpp
@@ -28,7 +28,7 @@ void UEnvQueryDebugHelpers::QueryToBlobArray(struct FEnvQueryInstance* Query, TA
 		BlobArray.Init(0, HeaderSize + FMath::TruncToInt(1.1f * UncompressedSize));
 
 		int32 CompressedSize = BlobArray.Num() - HeaderSize;
-		uint8* DestBuffer = BlobArray.GetTypedData();
+		uint8* DestBuffer = BlobArray.GetData();
 		FMemory::Memcpy(DestBuffer, &UncompressedSize, HeaderSize);
 		DestBuffer += HeaderSize;
 
@@ -65,7 +65,7 @@ void UEnvQueryDebugHelpers::QueryToDebugData(struct FEnvQueryInstance* Query, EQ
 		ItemInfo.ItemIdx = ItemIdx + FirstItemIndex;
 		ItemInfo.TotalScore = Query->Items[ItemInfo.ItemIdx].Score;
 
-		const uint8* ItemData = Query->RawData.GetTypedData() + Query->Items[ItemInfo.ItemIdx].DataOffset;
+		const uint8* ItemData = Query->RawData.GetData() + Query->Items[ItemInfo.ItemIdx].DataOffset;
 		ItemInfo.Desc = FString::Printf(TEXT("[%d] %s"), ItemInfo.ItemIdx, *ItemCDO->GetDescription(ItemData));
 
 		ItemInfo.TestValues.Reserve(NumTests);
@@ -115,14 +115,14 @@ void  UEnvQueryDebugHelpers::BlobArrayToDebugData(const TArray<uint8>& BlobArray
 		TArray<uint8> UncompressedBuffer;
 		int32 UncompressedSize = 0;
 		const int32 HeaderSize = sizeof(int32);
-		uint8* SrcBuffer = (uint8*)BlobArray.GetTypedData();
+		uint8* SrcBuffer = (uint8*)BlobArray.GetData();
 		FMemory::Memcpy(&UncompressedSize, SrcBuffer, HeaderSize);
 		SrcBuffer += HeaderSize;
 		const int32 CompressedSize = BlobArray.Num() - HeaderSize;
 
 		UncompressedBuffer.AddZeroed(UncompressedSize);
 
-		FCompression::UncompressMemory((ECompressionFlags)(COMPRESS_ZLIB), (void*)UncompressedBuffer.GetTypedData(), UncompressedSize, SrcBuffer, CompressedSize);
+		FCompression::UncompressMemory((ECompressionFlags)(COMPRESS_ZLIB), (void*)UncompressedBuffer.GetData(), UncompressedSize, SrcBuffer, CompressedSize);
 		FMemoryReader ArReader(UncompressedBuffer);
 
 		ArReader << EQSLocalData;
diff --git a/Engine/Source/Runtime/AIModule/Private/EnvironmentQuery/EnvQueryInstance.cpp b/Engine/Source/Runtime/AIModule/Private/EnvironmentQuery/EnvQueryInstance.cpp
--- a/Engine/Source/Runtime/AIModule/Private/EnvironmentQuery/EnvQueryInstance.cpp
+++ b/Engine/Source/Runtime/AIModule/Private/EnvironmentQuery/EnvQueryInstance.cpp
@@ -81,7 +81,7 @@ bool FEnvQueryInstance::PrepareContext(UClass* Context, TArray<FEnvQuerySpatialD
 	{
 		UEnvQueryItemType_VectorBase* DefTypeOb = (UEnvQueryItemType_VectorBase*)ContextData.ValueType->GetDefaultObject();
 		const uint16 DefTypeValueSize = DefTypeOb->GetValueSize();
-		uint8* RawData = ContextData.RawData.GetTypedData();
+		uint8* RawData = ContextData.RawData.GetData();
 
 		Data.Init(ContextData.NumValues);
 		for (int32 ValueIndex = 0; ValueIndex < ContextData.NumValues; ValueIndex++)
@@ -109,7 +109,7 @@ bool FEnvQueryInstance::PrepareContext(UClass* Context, TArray<FVector>& Data)
 	{
 		UEnvQueryItemType_VectorBase* DefTypeOb = (UEnvQueryItemType_VectorBase*)ContextData.ValueType->GetDefaultObject();
 		const uint16 DefTypeValueSize = DefTypeOb->GetValueSize();
-		uint8* RawData = (uint8*)ContextData.RawData.GetTypedData();
+		uint8* RawData = (uint8*)ContextData.RawData.GetData();
 
 		Data.Init(ContextData.NumValues);
 		for (int32 ValueIndex = 0; ValueIndex < ContextData.NumValues; ValueIndex++)
@@ -136,7 +136,7 @@ bool FEnvQueryInstance::PrepareContext(UClass* Context, TArray<FRotator>& Data)
 	{
 		UEnvQueryItemType_VectorBase* DefTypeOb = (UEnvQueryItemType_VectorBase*)ContextData.ValueType->GetDefaultObject();
 		const uint16 DefTypeValueSize = DefTypeOb->GetValueSize();
-		uint8* RawData = ContextData.RawData.GetTypedData();
+		uint8* RawData = ContextData.RawData.GetData();
 
 		Data.Init(ContextData.NumValues);
 		for (int32 ValueIndex = 0; ValueIndex < ContextData.NumValues; ValueIndex++)
@@ -163,7 +163,7 @@ bool FEnvQueryInstance::PrepareContext(UClass* Context, TArray<AActor*>& Data)
 	{
 		UEnvQueryItemType_ActorBase* DefTypeOb = (UEnvQueryItemType_ActorBase*)ContextData.ValueType->GetDefaultObject();
 		const uint16 DefTypeValueSize = DefTypeOb->GetValueSize();
-		uint8* RawData = ContextData.RawData.GetTypedData();
+		uint8* RawData = ContextData.RawData.GetData();
 
 		Data.Init(ContextData.NumValues);
 		for (int32 ValueIndex = 0; ValueIndex < ContextData.NumValues; ValueIndex++)
@@ -370,7 +370,7 @@ void FEnvQueryInstance::NormalizeScores()
 	float MinScore = 0.f;
 	float MaxScore = -BIG_NUMBER;
 
-	FEnvQueryItem* ItemInfo = Items.GetTypedData();
+	FEnvQueryItem* ItemInfo = Items.GetData();
 	for (int32 ItemIndex = 0; ItemIndex < NumValidItems; ItemIndex++, ItemInfo++)
 	{
 		ensure(ItemInfo->IsValid());
@@ -379,7 +379,7 @@ void FEnvQueryInstance::NormalizeScores()
 		MaxScore = FMath::Max(MaxScore, ItemInfo->Score);
 	}
 
-	ItemInfo = Items.GetTypedData();
+	ItemInfo = Items.GetData();
 	if (MinScore == MaxScore)
 	{
 		for (int32 ItemIndex = 0; ItemIndex < NumValidItems; ItemIndex++, ItemInfo++)
@@ -657,7 +657,7 @@ FBox FEnvQueryInstance::GetBoundingBox() const
 
 		for (int32 Index = 0; Index < Items.Num(); ++Index)
 		{		
-			BBox += DefTypeOb->GetLocation(RawData.GetTypedData() + Items[Index].DataOffset);
+			BBox += DefTypeOb->GetLocation(RawData.GetData() + Items[Index].DataOffset);
 		}
 	}
 
diff --git a/Engine/Source/Runtime/AIModule/Private/EnvironmentQuery/EnvQueryTest.cpp b/Engine/Source/Runtime/AIModule/Private/EnvironmentQuery/EnvQueryTest.cpp
--- a/Engine/Source/Runtime/AIModule/Private/EnvironmentQuery/EnvQueryTest.cpp
+++ b/Engine/Source/Runtime/AIModule/Private/EnvironmentQuery/EnvQueryTest.cpp
@@ -80,7 +80,7 @@ void UEnvQueryTest::NormalizeItemScores(FEnvQueryInstance& QueryInstance)
 		}
 	}
 
-	FEnvQueryItemDetails* DetailInfo = QueryInstance.ItemDetails.GetTypedData();
+	FEnvQueryItemDetails* DetailInfo = QueryInstance.ItemDetails.GetData();
 	if ((ClampMinType == EEnvQueryTestClamping::None) ||
 		(ClampMaxType == EEnvQueryTestClamping::None)
 	   )
@@ -108,7 +108,7 @@ void UEnvQueryTest::NormalizeItemScores(FEnvQueryInstance& QueryInstance)
 		}
 	}
 
-	DetailInfo = QueryInstance.ItemDetails.GetTypedData();
+	DetailInfo = QueryInstance.ItemDetails.GetData();
 	if (bNormalizeFromZero && MinScore > 0.0f)
 	{
 		MinScore = 0.0f;
@@ -184,21 +184,21 @@ bool UEnvQueryTest::IsContextPerItem(TSubclassOf<UEnvQueryContext> CheckContext)
 FVector UEnvQueryTest::GetItemLocation(FEnvQueryInstance& QueryInstance, int32 ItemIndex) const
 {
 	return QueryInstance.ItemTypeVectorCDO ?
-		QueryInstance.ItemTypeVectorCDO->GetLocation(QueryInstance.RawData.GetTypedData() + QueryInstance.Items[ItemIndex].DataOffset) :
+		QueryInstance.ItemTypeVectorCDO->GetLocation(QueryInstance.RawData.GetData() + QueryInstance.Items[ItemIndex].DataOffset) :
 		FVector::ZeroVector;
 }
 
 FRotator UEnvQueryTest::GetItemRotation(FEnvQueryInstance& QueryInstance, int32 ItemIndex) const
 {
 	return QueryInstance.ItemTypeVectorCDO ?
-		QueryInstance.ItemTypeVectorCDO->GetRotation(QueryInstance.RawData.GetTypedData() + QueryInstance.Items[ItemIndex].DataOffset) :
+		QueryInstance.ItemTypeVectorCDO->GetRotation(QueryInstance.RawData.GetData() + QueryInstance.Items[ItemIndex].DataOffset) :
 		FRotator::ZeroRotator;
 }
 
 AActor* UEnvQueryTest::GetItemActor(FEnvQueryInstance& QueryInstance, int32 ItemIndex) const
 {
 	return QueryInstance.ItemTypeActorCDO ?
-		QueryInstance.ItemTypeActorCDO->GetActor(QueryInstance.RawData.GetTypedData() + QueryInstance.Items[ItemIndex].DataOffset) :
+		QueryInstance.ItemTypeActorCDO->GetActor(QueryInstance.RawData.GetData() + QueryInstance.Items[ItemIndex].DataOffset) :
 		NULL;
 }
 
diff --git a/Engine/Source/Runtime/AIModule/Private/EnvironmentQuery/EnvQueryTypes.cpp b/Engine/Source/Runtime/AIModule/Private/EnvironmentQuery/EnvQueryTypes.cpp
--- a/Engine/Source/Runtime/AIModule/Private/EnvironmentQuery/EnvQueryTypes.cpp
+++ b/Engine/Source/Runtime/AIModule/Private/EnvironmentQuery/EnvQueryTypes.cpp
@@ -16,7 +16,7 @@ AActor* FEnvQueryResult::GetItemAsActor(int32 Index) const
 		ItemType->IsChildOf(UEnvQueryItemType_ActorBase::StaticClass()))
 	{
 		UEnvQueryItemType_ActorBase* DefTypeOb = (UEnvQueryItemType_ActorBase*)ItemType->GetDefaultObject();
-		return DefTypeOb->GetActor(RawData.GetTypedData() + Items[Index].DataOffset);
+		return DefTypeOb->GetActor(RawData.GetData() + Items[Index].DataOffset);
 	}
 
 	return NULL;
@@ -28,7 +28,7 @@ FVector FEnvQueryResult::GetItemAsLocation(int32 Index) const
 		ItemType->IsChildOf(UEnvQueryItemType_VectorBase::StaticClass()))
 	{
 		UEnvQueryItemType_VectorBase* DefTypeOb = (UEnvQueryItemType_VectorBase*)ItemType->GetDefaultObject();
-		return DefTypeOb->GetLocation(RawData.GetTypedData() + Items[Index].DataOffset);
+		return DefTypeOb->GetLocation(RawData.GetData() + Items[Index].DataOffset);
 	}
 
 	return FVector::ZeroVector;
diff --git a/Engine/Source/Runtime/AIModule/Private/EnvironmentQuery/Items/EnvQueryItemType_Actor.cpp b/Engine/Source/Runtime/AIModule/Private/EnvironmentQuery/Items/EnvQueryItemType_Actor.cpp
--- a/Engine/Source/Runtime/AIModule/Private/EnvironmentQuery/Items/EnvQueryItemType_Actor.cpp
+++ b/Engine/Source/Runtime/AIModule/Private/EnvironmentQuery/Items/EnvQueryItemType_Actor.cpp
@@ -43,7 +43,7 @@ void UEnvQueryItemType_Actor::SetContextHelper(struct FEnvQueryContextData& Cont
 	ContextData.NumValues = 1;
 	ContextData.RawData.Init(sizeof(FWeakObjectPtr));
 
-	UEnvQueryItemType_Actor::SetValue(ContextData.RawData.GetTypedData(), SingleActor);
+	UEnvQueryItemType_Actor::SetValue(ContextData.RawData.GetData(), SingleActor);
 }
 
 void UEnvQueryItemType_Actor::SetContextHelper(struct FEnvQueryContextData& ContextData, const TArray<const AActor*>& MultipleActors)
@@ -52,7 +52,7 @@ void UEnvQueryItemType_Actor::SetContextHelper(struct FEnvQueryContextData& Cont
 	ContextData.NumValues = MultipleActors.Num();
 	ContextData.RawData.Init(sizeof(FWeakObjectPtr) * MultipleActors.Num());
 
-	uint8* RawData = (uint8*)ContextData.RawData.GetTypedData();
+	uint8* RawData = (uint8*)ContextData.RawData.GetData();
 	for (int32 ActorIndex = 0; ActorIndex < MultipleActors.Num(); ActorIndex++)
 	{
 		UEnvQueryItemType_Actor::SetValue(RawData, MultipleActors[ActorIndex]);
@@ -66,7 +66,7 @@ void UEnvQueryItemType_Actor::SetContextHelper(struct FEnvQueryContextData& Cont
 	ContextData.NumValues = MultipleActors.Num();
 	ContextData.RawData.Init(sizeof(FWeakObjectPtr)* MultipleActors.Num());
 
-	uint8* RawData = (uint8*)ContextData.RawData.GetTypedData();
+	uint8* RawData = (uint8*)ContextData.RawData.GetData();
 	for (int32 ActorIndex = 0; ActorIndex < MultipleActors.Num(); ActorIndex++)
 	{
 		UEnvQueryItemType_Actor::SetValue(RawData, MultipleActors[ActorIndex]);
diff --git a/Engine/Source/Runtime/AIModule/Private/EnvironmentQuery/Items/EnvQueryItemType_Direction.cpp b/Engine/Source/Runtime/AIModule/Private/EnvironmentQuery/Items/EnvQueryItemType_Direction.cpp
--- a/Engine/Source/Runtime/AIModule/Private/EnvironmentQuery/Items/EnvQueryItemType_Direction.cpp
+++ b/Engine/Source/Runtime/AIModule/Private/EnvironmentQuery/Items/EnvQueryItemType_Direction.cpp
@@ -45,7 +45,7 @@ void UEnvQueryItemType_Direction::SetContextHelper(struct FEnvQueryContextData&
 	ContextData.NumValues = 1;
 	ContextData.RawData.Init(sizeof(FVector));
 
-	UEnvQueryItemType_Direction::SetValue((uint8*)ContextData.RawData.GetTypedData(), SingleDirection);
+	UEnvQueryItemType_Direction::SetValue((uint8*)ContextData.RawData.GetData(), SingleDirection);
 }
 
 void UEnvQueryItemType_Direction::SetContextHelper(struct FEnvQueryContextData& ContextData, const FRotator& SingleRotation)
@@ -54,7 +54,7 @@ void UEnvQueryItemType_Direction::SetContextHelper(struct FEnvQueryContextData&
 	ContextData.NumValues = 1;
 	ContextData.RawData.Init(sizeof(FVector));
 
-	UEnvQueryItemType_Direction::SetValueRot((uint8*)ContextData.RawData.GetTypedData(), SingleRotation);
+	UEnvQueryItemType_Direction::SetValueRot((uint8*)ContextData.RawData.GetData(), SingleRotation);
 }
 
 void UEnvQueryItemType_Direction::SetContextHelper(struct FEnvQueryContextData& ContextData, const TArray<FVector>& MultipleDirections)
@@ -63,7 +63,7 @@ void UEnvQueryItemType_Direction::SetContextHelper(struct FEnvQueryContextData&
 	ContextData.NumValues = MultipleDirections.Num();
 	ContextData.RawData.Init(sizeof(FVector) * MultipleDirections.Num());
 
-	uint8* RawData = (uint8*)ContextData.RawData.GetTypedData();
+	uint8* RawData = (uint8*)ContextData.RawData.GetData();
 	for (int32 DirectionIndex = 0; DirectionIndex < MultipleDirections.Num(); DirectionIndex++)
 	{
 		UEnvQueryItemType_Direction::SetValue(RawData, MultipleDirections[DirectionIndex]);
@@ -77,7 +77,7 @@ void UEnvQueryItemType_Direction::SetContextHelper(struct FEnvQueryContextData&
 	ContextData.NumValues = MultipleRotations.Num();
 	ContextData.RawData.Init(sizeof(FVector) * MultipleRotations.Num());
 
-	uint8* RawData = (uint8*)ContextData.RawData.GetTypedData();
+	uint8* RawData = (uint8*)ContextData.RawData.GetData();
 	for (int32 RotationIndex = 0; RotationIndex < MultipleRotations.Num(); RotationIndex++)
 	{
 		UEnvQueryItemType_Direction::SetValueRot(RawData, MultipleRotations[RotationIndex]);
diff --git a/Engine/Source/Runtime/AIModule/Private/EnvironmentQuery/Items/EnvQueryItemType_Point.cpp b/Engine/Source/Runtime/AIModule/Private/EnvironmentQuery/Items/EnvQueryItemType_Point.cpp
--- a/Engine/Source/Runtime/AIModule/Private/EnvironmentQuery/Items/EnvQueryItemType_Point.cpp
+++ b/Engine/Source/Runtime/AIModule/Private/EnvironmentQuery/Items/EnvQueryItemType_Point.cpp
@@ -29,7 +29,7 @@ void UEnvQueryItemType_Point::SetContextHelper(struct FEnvQueryContextData& Cont
 	ContextData.NumValues = 1;
 	ContextData.RawData.Init(sizeof(FVector));
 
-	UEnvQueryItemType_Point::SetValue((uint8*)ContextData.RawData.GetTypedData(), SinglePoint);
+	UEnvQueryItemType_Point::SetValue((uint8*)ContextData.RawData.GetData(), SinglePoint);
 }
 
 void UEnvQueryItemType_Point::SetContextHelper(struct FEnvQueryContextData& ContextData, const TArray<FVector>& MultiplePoints)
@@ -38,7 +38,7 @@ void UEnvQueryItemType_Point::SetContextHelper(struct FEnvQueryContextData& Cont
 	ContextData.NumValues = MultiplePoints.Num();
 	ContextData.RawData.Init(sizeof(FVector) * MultiplePoints.Num());
 
-	uint8* RawData = (uint8*)ContextData.RawData.GetTypedData();
+	uint8* RawData = (uint8*)ContextData.RawData.GetData();
 	for (int32 PointIndex = 0; PointIndex < MultiplePoints.Num(); PointIndex++)
 	{
 		UEnvQueryItemType_Point::SetValue(RawData, MultiplePoints[PointIndex]);
diff --git a/Engine/Source/Runtime/AIModule/Private/Navigation/CrowdManager.cpp b/Engine/Source/Runtime/AIModule/Private/Navigation/CrowdManager.cpp
--- a/Engine/Source/Runtime/AIModule/Private/Navigation/CrowdManager.cpp
+++ b/Engine/Source/Runtime/AIModule/Private/Navigation/CrowdManager.cpp
@@ -447,7 +447,7 @@ bool UCrowdManager::SetAgentMovePath(const class UCrowdFollowingComponent* Agent
 		bSuccess = DetourCrowd->requestMoveTarget(AgentData->AgentIndex, PathRefs.Last(), &RcTargetPos.X);
 		if (bSuccess)
 		{
-			bSuccess = DetourCrowd->setAgentCorridor(AgentData->AgentIndex, PathRefs.GetTypedData(), PathRefs.Num());
+			bSuccess = DetourCrowd->setAgentCorridor(AgentData->AgentIndex, PathRefs.GetData(), PathRefs.Num());
 		}
 	}
 #endif
@@ -813,7 +813,7 @@ void UCrowdManager::CreateCrowdManager()
 			const FCrowdAvoidanceSamplingPattern& Info = SamplingPatterns[Idx];
 			if (Info.Angles.Num() > 0 && Info.Angles.Num() == Info.Radii.Num())
 			{
-				DetourCrowd->setObstacleAvoidancePattern(Idx, Info.Angles.GetTypedData(), Info.Radii.GetTypedData(), Info.Angles.Num());
+				DetourCrowd->setObstacleAvoidancePattern(Idx, Info.Angles.GetData(), Info.Radii.GetData(), Info.Angles.Num());
 			}
 		}
 
diff --git a/Engine/Source/Runtime/AIModule/Private/Navigation/PathFollowingComponent.cpp b/Engine/Source/Runtime/AIModule/Private/Navigation/PathFollowingComponent.cpp
--- a/Engine/Source/Runtime/AIModule/Private/Navigation/PathFollowingComponent.cpp
+++ b/Engine/Source/Runtime/AIModule/Private/Navigation/PathFollowingComponent.cpp
@@ -682,7 +682,7 @@ int32 UPathFollowingComponent::OptimizeSegmentVisibility(int32 StartIndex)
 			Path->ShortcutNodeRefs.Reserve(RaycastResult.CorridorPolysCount);
 			Path->ShortcutNodeRefs.SetNumUninitialized(RaycastResult.CorridorPolysCount);
 		}
-		FPlatformMemory::Memcmp(Path->ShortcutNodeRefs.GetTypedData(), RaycastResult.CorridorPolys, RaycastResult.CorridorPolysCount * sizeof(NavNodeRef));
+		FPlatformMemory::Memcmp(Path->ShortcutNodeRefs.GetData(), RaycastResult.CorridorPolys, RaycastResult.CorridorPolysCount * sizeof(NavNodeRef));
 #else
 		const bool RaycastHitResult = NavData->Raycast(PawnLocation, AdjustedDestination, HitLocation, QueryFilter);
 #endif
diff --git a/Engine/Source/Runtime/AIModule/Private/Perception/AIPerceptionComponent.cpp b/Engine/Source/Runtime/AIModule/Private/Perception/AIPerceptionComponent.cpp
--- a/Engine/Source/Runtime/AIModule/Private/Perception/AIPerceptionComponent.cpp
+++ b/Engine/Source/Runtime/AIModule/Private/Perception/AIPerceptionComponent.cpp
@@ -233,7 +233,7 @@ void UAIPerceptionComponent::ProcessStimuli()
 		return;
 	}
 	
-	FStimulusToProcess* SourcedStimulus = StimuliToProcess.GetTypedData();
+	FStimulusToProcess* SourcedStimulus = StimuliToProcess.GetData();
 	TArray<AActor*> UpdatedActors;
 	UpdatedActors.Reserve(StimuliToProcess.Num());
 
diff --git a/Engine/Source/Runtime/AIModule/Private/Perception/AIPerceptionSystem.cpp b/Engine/Source/Runtime/AIModule/Private/Perception/AIPerceptionSystem.cpp
--- a/Engine/Source/Runtime/AIModule/Private/Perception/AIPerceptionSystem.cpp
+++ b/Engine/Source/Runtime/AIModule/Private/Perception/AIPerceptionSystem.cpp
@@ -75,7 +75,7 @@ void UAIPerceptionSystem::PerformSourceRegistration()
 
 void UAIPerceptionSystem::OnNewListener(const FPerceptionListener& NewListener)
 {
-	UAISense** SenseInstance = Senses.GetTypedData();
+	UAISense** SenseInstance = Senses.GetData();
 	for (int32 SenseIndex = 0; SenseIndex < Senses.Num(); ++SenseIndex, ++SenseInstance)
 	{
 		if (*SenseInstance != NULL)
@@ -87,7 +87,7 @@ void UAIPerceptionSystem::OnNewListener(const FPerceptionListener& NewListener)
 
 void UAIPerceptionSystem::OnListenerUpdate(const FPerceptionListener& UpdatedListener)
 {
-	UAISense** SenseInstance = Senses.GetTypedData();
+	UAISense** SenseInstance = Senses.GetData();
 	for (int32 SenseIndex = 0; SenseIndex < Senses.Num(); ++SenseIndex, ++SenseInstance)
 	{
 		if (*SenseInstance != NULL)
@@ -117,7 +117,7 @@ void UAIPerceptionSystem::ManagerTick(float DeltaSeconds)
 	}
 
 	{
-		UAISense** SenseInstance = Senses.GetTypedData();
+		UAISense** SenseInstance = Senses.GetData();
 		for (int32 SenseIndex = 0; SenseIndex < Senses.Num(); ++SenseIndex, ++SenseInstance)
 		{
 			bNeedsUpdate |= *SenseInstance != NULL && (*SenseInstance)->ProgressTime(DeltaSeconds);
@@ -140,7 +140,7 @@ void UAIPerceptionSystem::ManagerTick(float DeltaSeconds)
 			}
 		}
 
-		UAISense** SenseInstance = Senses.GetTypedData();
+		UAISense** SenseInstance = Senses.GetData();
 		for (int32 SenseIndex = 0; SenseIndex < Senses.Num(); ++SenseIndex, ++SenseInstance)
 		{
 			if (*SenseInstance != NULL)
@@ -253,7 +253,7 @@ void UAIPerceptionSystem::UnregisterListener(UAIPerceptionComponent* Listener)
 
 void UAIPerceptionSystem::OnListenerRemoved(const FPerceptionListener& NewListener)
 {
-	UAISense** SenseInstance = Senses.GetTypedData();
+	UAISense** SenseInstance = Senses.GetData();
 	for (int32 SenseIndex = 0; SenseIndex < Senses.Num(); ++SenseIndex, ++SenseInstance)
 	{
 		if (*SenseInstance != NULL)
diff --git a/Engine/Source/Runtime/AIModule/Private/Perception/AISense_Sight.cpp b/Engine/Source/Runtime/AIModule/Private/Perception/AISense_Sight.cpp
--- a/Engine/Source/Runtime/AIModule/Private/Perception/AISense_Sight.cpp
+++ b/Engine/Source/Runtime/AIModule/Private/Perception/AISense_Sight.cpp
@@ -93,7 +93,7 @@ float UAISense_Sight::Update()
 
 	AIPerception::FListenerMap& ListenersMap = *GetListeners();
 
-	FAISightQuery* SightQuery = SightQueryQueue.GetTypedData();
+	FAISightQuery* SightQuery = SightQueryQueue.GetData();
 	for (int32 QueryIndex = 0; QueryIndex < SightQueryQueue.Num(); ++QueryIndex, ++SightQuery)
 	{
 		if (TracesCount < MaxTracesPerTick)
diff --git a/Engine/Source/Runtime/AssetRegistry/Private/AssetRegistry.cpp b/Engine/Source/Runtime/AssetRegistry/Private/AssetRegistry.cpp
--- a/Engine/Source/Runtime/AssetRegistry/Private/AssetRegistry.cpp
+++ b/Engine/Source/Runtime/AssetRegistry/Private/AssetRegistry.cpp
@@ -532,7 +532,7 @@ bool FAssetRegistry::GetAssets(const FARFilter& Filter, TArray<FAssetData>& OutA
 	if ( DiskFilterSets.Num() > 0 )
 	{
 		// Initialize the combined filter set to the first set, in case we can skip combining.
-		TArray<FAssetData*>* CombinedFilterSet = DiskFilterSets.GetTypedData();
+		TArray<FAssetData*>* CombinedFilterSet = DiskFilterSets.GetData();
 		TArray<FAssetData*> Intersection;
 
 		// If we have more than one set, we must combine them. We take the intersection
diff --git a/Engine/Source/Runtime/Core/Private/Async/TaskGraph.cpp b/Engine/Source/Runtime/Core/Private/Async/TaskGraph.cpp
--- a/Engine/Source/Runtime/Core/Private/Async/TaskGraph.cpp
+++ b/Engine/Source/Runtime/Core/Private/Async/TaskGraph.cpp
@@ -56,7 +56,7 @@ public:
 			{
 				checkThreadGraph(Tasks[StartIndex - 1]==NULL);
 				checkThreadGraph(Tasks[0]==NULL);
-				FMemory::Memmove(Tasks.GetTypedData(),&Tasks[StartIndex],(EndIndex - StartIndex) * sizeof(FBaseGraphTask*));
+				FMemory::Memmove(Tasks.GetData(), &Tasks[StartIndex], (EndIndex - StartIndex) * sizeof(FBaseGraphTask*));
 				EndIndex -= StartIndex;
 				FMemory::Memzero(&Tasks[EndIndex],StartIndex * sizeof(FBaseGraphTask*)); // not needed in final
 				StartIndex = 0;
diff --git a/Engine/Source/Runtime/Core/Private/Containers/String.cpp b/Engine/Source/Runtime/Core/Private/Containers/String.cpp
--- a/Engine/Source/Runtime/Core/Private/Containers/String.cpp
+++ b/Engine/Source/Runtime/Core/Private/Containers/String.cpp
@@ -10,7 +10,7 @@ void FString::TrimToNullTerminator()
 {
 	if( Data.Num() )
 	{
-		int32 DataLen = FCString::Strlen(Data.GetTypedData());
+		int32 DataLen = FCString::Strlen(Data.GetData());
 		check(DataLen == 0 || DataLen < Data.Num());
 		int32 Len = DataLen > 0 ? DataLen+1 : 0;
 
@@ -539,7 +539,7 @@ int32 FString::ParseIntoArray( TArray<FString>* InArray, const TCHAR* pchDelim,
 	check(pchDelim);
 	check(InArray);
 	InArray->Empty();
-	const TCHAR *Start = Data.GetTypedData();
+	const TCHAR *Start = Data.GetData();
 	int32 DelimLength = FCString::Strlen(pchDelim);
 	if (Start && DelimLength)
 	{
@@ -1143,13 +1143,13 @@ FArchive& operator<<( FArchive& Ar, FString& A )
 			if (LoadUCS2Char)
 			{
 				// read in the unicode string and byteswap it, etc
-				auto Passthru = StringMemoryPassthru<UCS2CHAR>(A.Data.GetTypedData(), SaveNum, SaveNum);
+				auto Passthru = StringMemoryPassthru<UCS2CHAR>(A.Data.GetData(), SaveNum, SaveNum);
 				Ar.Serialize(Passthru.Get(), SaveNum * sizeof(UCS2CHAR));
 				// Ensure the string has a null terminator
 				Passthru.Get()[SaveNum-1] = '\0';
 				Passthru.Apply();
 
-				INTEL_ORDER_TCHARARRAY(A.Data.GetTypedData())
+				INTEL_ORDER_TCHARARRAY(A.Data.GetData())
 
 				// Since Microsoft's vsnwprintf implementation raises an invalid parameter warning
 				// with a character of 0xffff, scan for it and terminate the string there.
@@ -1163,7 +1163,7 @@ FArchive& operator<<( FArchive& Ar, FString& A )
 			}
 			else
 			{
-				auto Passthru = StringMemoryPassthru<ANSICHAR>(A.Data.GetTypedData(), SaveNum, SaveNum);
+				auto Passthru = StringMemoryPassthru<ANSICHAR>(A.Data.GetData(), SaveNum, SaveNum);
 				Ar.Serialize(Passthru.Get(), SaveNum * sizeof(ANSICHAR));
 				// Ensure the string has a null terminator
 				Passthru.Get()[SaveNum-1] = '\0';
@@ -1195,14 +1195,14 @@ FArchive& operator<<( FArchive& Ar, FString& A )
 				#if !PLATFORM_LITTLE_ENDIAN
 					FString  ATemp  = A;
 					FString& A      = ATemp;
-					INTEL_ORDER_TCHARARRAY(A.Data.GetTypedData());
+					INTEL_ORDER_TCHARARRAY(A.Data.GetData());
 				#endif
 
-				Ar.Serialize((void*)StringCast<UCS2CHAR>(A.Data.GetTypedData(), Num).Get(), sizeof(UCS2CHAR) * Num);
+				Ar.Serialize((void*)StringCast<UCS2CHAR>(A.Data.GetData(), Num).Get(), sizeof(UCS2CHAR)* Num);
 			}
 			else
 			{
-				Ar.Serialize((void*)StringCast<ANSICHAR>(A.Data.GetTypedData(), Num).Get(), sizeof(ANSICHAR) * Num);
+				Ar.Serialize((void*)StringCast<ANSICHAR>(A.Data.GetData(), Num).Get(), sizeof(ANSICHAR)* Num);
 			}
 		}
 	}
diff --git a/Engine/Source/Runtime/Core/Private/IOS/IOSPlatformMisc.cpp b/Engine/Source/Runtime/Core/Private/IOS/IOSPlatformMisc.cpp
--- a/Engine/Source/Runtime/Core/Private/IOS/IOSPlatformMisc.cpp
+++ b/Engine/Source/Runtime/Core/Private/IOS/IOSPlatformMisc.cpp
@@ -152,8 +152,8 @@ void FIOSPlatformMisc::ClipboardPaste(class FString& Result)
 	{
 		TArray<TCHAR> Ch;
 		Ch.AddUninitialized([CocoaString length] + 1);
-		FPlatformString::CFStringToTCHAR((CFStringRef)CocoaString, Ch.GetTypedData());
-		Result = Ch.GetTypedData();
+		FPlatformString::CFStringToTCHAR((CFStringRef)CocoaString, Ch.GetData());
+		Result = Ch.GetData();
 	}
 	else
 	{
@@ -664,7 +664,7 @@ void GetBytesForFont(const NSString* InFontName, OUT TArray<uint8>& OutBytes)
 		OutBytes.AddZeroed( TotalSize );
 
 		// Start copying the table data into our buffer
-		uint8* DataStart = OutBytes.GetTypedData();
+		uint8* DataStart = OutBytes.GetData();
 		uint8* DataPtr = DataStart;
 
 		// Compute font header entries
diff --git a/Engine/Source/Runtime/Core/Private/Mac/MacPlatformMisc.cpp b/Engine/Source/Runtime/Core/Private/Mac/MacPlatformMisc.cpp
--- a/Engine/Source/Runtime/Core/Private/Mac/MacPlatformMisc.cpp
+++ b/Engine/Source/Runtime/Core/Private/Mac/MacPlatformMisc.cpp
@@ -600,8 +600,8 @@ void FMacPlatformMisc::ClipboardPaste(class FString& Result)
 	{
 		TArray<TCHAR> Ch;
 		Ch.AddUninitialized([CocoaString length] + 1);
-		FPlatformString::CFStringToTCHAR((CFStringRef)CocoaString, Ch.GetTypedData());
-		Result = Ch.GetTypedData();
+		FPlatformString::CFStringToTCHAR((CFStringRef)CocoaString, Ch.GetData());
+		Result = Ch.GetData();
 	}
 	else
 	{
diff --git a/Engine/Source/Runtime/Core/Private/Misc/Base64.cpp b/Engine/Source/Runtime/Core/Private/Misc/Base64.cpp
--- a/Engine/Source/Runtime/Core/Private/Misc/Base64.cpp
+++ b/Engine/Source/Runtime/Core/Private/Misc/Base64.cpp
@@ -38,7 +38,7 @@ uint8 FBase64::DecodingAlphabet[256] = { 0xFF, 0xFF, 0xFF, 0xFF, 0xFF, 0xFF, 0xF
  */
 FString FBase64::Encode(const TArray<uint8>& Source)
 {
-	return Encode((uint8*)Source.GetTypedData(), Source.Num());
+	return Encode((uint8*)Source.GetData(), Source.Num());
 }
 
 /**
@@ -61,9 +61,9 @@ bool FBase64::Decode(const FString& Source, TArray<uint8>& Dest)
 	uint32 ExpectedLength = Length / 4 * 3;
 	// Add the number we need for output
 	Dest.AddZeroed(ExpectedLength);
-	uint8* Buffer = Dest.GetTypedData();
+	uint8* Buffer = Dest.GetData();
 	uint32 PadCount = 0;
-	bool bWasSuccessful = Decode(TCHAR_TO_ANSI(*Source), Length, Dest.GetTypedData(), PadCount);
+	bool bWasSuccessful = Decode(TCHAR_TO_ANSI(*Source), Length, Dest.GetData(), PadCount);
 	if (bWasSuccessful)
 	{
 		if (PadCount > 0)
@@ -104,7 +104,7 @@ bool FBase64::Decode(const FString& Source, FString& Dest)
 	uint32 ExpectedLength = Length / 4 * 3;
 	TArray<ANSICHAR> TempDest;
 	TempDest.AddZeroed(ExpectedLength);
-	uint8* Buffer = (uint8*)TempDest.GetTypedData();
+	uint8* Buffer = (uint8*)TempDest.GetData();
 	uint32 PadCount = 0;
 
 	bool bWasSuccessful = Decode(TCHAR_TO_ANSI(*Source), Length, Buffer, PadCount);
@@ -118,7 +118,7 @@ bool FBase64::Decode(const FString& Source, FString& Dest)
 		{
 			TempDest.Add('\0');
 		}
-		Dest = ANSI_TO_TCHAR(TempDest.GetTypedData());
+		Dest = ANSI_TO_TCHAR(TempDest.GetData());
 	}
 	return bWasSuccessful;
 }
diff --git a/Engine/Source/Runtime/Core/Private/Misc/CoreMisc.cpp b/Engine/Source/Runtime/Core/Private/Misc/CoreMisc.cpp
--- a/Engine/Source/Runtime/Core/Private/Misc/CoreMisc.cpp
+++ b/Engine/Source/Runtime/Core/Private/Misc/CoreMisc.cpp
@@ -101,7 +101,7 @@ bool FFileHelper::LoadFileToArray( TArray<uint8>& Result, const TCHAR* Filename,
 	}
 	Result.Reset();
 	Result.AddUninitialized( Reader->TotalSize() );
-	Reader->Serialize( Result.GetTypedData(), Result.Num() );
+	Reader->Serialize(Result.GetData(), Result.Num());
 	bool Success = Reader->Close();
 	delete Reader;
 	return Success;
@@ -146,7 +146,7 @@ void FFileHelper::BufferToString( FString& Result, const uint8* Buffer, int32 Si
 		FUTF8ToTCHAR Conv((const ANSICHAR*)Buffer, Size);
 		int32 Length = Conv.Length();
 		ResultArray.AddUninitialized(Length + 1); // For the null terminator
-		CopyAssignItems(ResultArray.GetTypedData(), Conv.Get(), Length);
+		CopyAssignItems(ResultArray.GetData(), Conv.Get(), Length);
 	}
 
 	if (ResultArray.Num() == 1)
@@ -208,7 +208,7 @@ bool FFileHelper::SaveArrayToFile( const TArray<uint8>& Array, const TCHAR* File
 	{
 		return 0;
 	}
-	Ar->Serialize( const_cast<uint8*>(Array.GetTypedData()), Array.Num() );
+	Ar->Serialize(const_cast<uint8*>(Array.GetData()), Array.Num());
 	delete Ar;
 	return true;
 }
diff --git a/Engine/Source/Runtime/Core/Private/ProfilingDebugging/MallocProfiler.cpp b/Engine/Source/Runtime/Core/Private/ProfilingDebugging/MallocProfiler.cpp
--- a/Engine/Source/Runtime/Core/Private/ProfilingDebugging/MallocProfiler.cpp
+++ b/Engine/Source/Runtime/Core/Private/ProfilingDebugging/MallocProfiler.cpp
@@ -547,7 +547,7 @@ void FMallocProfiler::EndProfiling()
 		TArray<FStackWalkModuleInfo> ProcModules;
 		ProcModules.Reserve(Header.ModuleEntries);
 
-		Header.ModuleEntries = FPlatformStackWalk::GetProcessModuleSignatures(ProcModules.GetTypedData(), ProcModules.Num());
+		Header.ModuleEntries = FPlatformStackWalk::GetProcessModuleSignatures(ProcModules.GetData(), ProcModules.Num());
 
 		for(uint32 ModuleIndex = 0; ModuleIndex < Header.ModuleEntries; ++ModuleIndex)
 		{
diff --git a/Engine/Source/Runtime/Core/Private/Serialization/Archive.cpp b/Engine/Source/Runtime/Core/Private/Serialization/Archive.cpp
--- a/Engine/Source/Runtime/Core/Private/Serialization/Archive.cpp
+++ b/Engine/Source/Runtime/Core/Private/Serialization/Archive.cpp
@@ -738,7 +738,7 @@ void FArchive::SerializeIntPacked(uint32& Value)
 				break;
 			}
 		}
-		Serialize(PackedBytes.GetTypedData(), PackedBytes.Num()); // Actually serialize the bytes we made
+		Serialize(PackedBytes.GetData(), PackedBytes.Num()); // Actually serialize the bytes we made
 	}
 }
 
diff --git a/Engine/Source/Runtime/Core/Private/Serialization/Bitstreams.cpp b/Engine/Source/Runtime/Core/Private/Serialization/Bitstreams.cpp
--- a/Engine/Source/Runtime/Core/Private/Serialization/Bitstreams.cpp
+++ b/Engine/Source/Runtime/Core/Private/Serialization/Bitstreams.cpp
@@ -122,7 +122,7 @@ FBitWriter::FBitWriter( int64 InMaxBits, bool InAllowResize /*=false*/ )
 	Buffer.AddUninitialized( (InMaxBits+7)>>3 );
 
 	AllowResize = InAllowResize;
-	FMemory::Memzero( Buffer.GetTypedData(), Buffer.Num() );
+	FMemory::Memzero(Buffer.GetData(), Buffer.Num());
 	ArIsPersistent = ArIsSaving = 1;
 	ArNetVer |= 0x80000000;
 }
@@ -141,7 +141,7 @@ void FBitWriter::Reset(void)
 {
 	FArchive::Reset();
 	Num = 0;
-	FMemory::Memzero( Buffer.GetTypedData(), Buffer.Num() );
+	FMemory::Memzero(Buffer.GetData(), Buffer.Num());
 	ArIsPersistent = ArIsSaving = 1;
 	ArNetVer |= 0x80000000;
 }
@@ -161,7 +161,7 @@ void FBitWriter::SerializeBits( void* Src, int64 LengthBits )
 		}
 		else
 		{
-			appBitsCpy( Buffer.GetTypedData(), Num, (uint8*)Src, 0, LengthBits);
+			appBitsCpy(Buffer.GetData(), Num, (uint8*)Src, 0, LengthBits);
 			Num += LengthBits;
 		}
 	}
@@ -173,7 +173,7 @@ void FBitWriter::Serialize( void* Src, int64 LengthBytes )
 	int64 LengthBits = LengthBytes*8;
 	if( AllowAppend(LengthBits) )
 	{
-		appBitsCpy( Buffer.GetTypedData(), Num, (uint8*)Src, 0, LengthBits);
+		appBitsCpy(Buffer.GetData(), Num, (uint8*)Src, 0, LengthBits);
 		Num += LengthBits;
 	}
 	else ArIsError = 1;
@@ -268,7 +268,7 @@ void FBitWriterMark::Copy( FBitWriter& Writer, TArray<uint8> &Buffer )
 	{
 		Buffer.Init(Bytes);		// This makes room but doesnt zero
 		Buffer[Bytes-1] = 0;	// Make sure the last byte is 0 out, because appBitsCpy wont touch the last bits
-		appBitsCpy( Buffer.GetTypedData(), 0, Writer.Buffer.GetTypedData(), Num, Writer.Num - Num );
+		appBitsCpy(Buffer.GetData(), 0, Writer.Buffer.GetData(), Num, Writer.Num - Num);
 	}
 }
 
@@ -294,7 +294,7 @@ FBitReader::FBitReader( uint8* Src, int64 CountBits )
 	ArNetVer |= 0x80000000;
 	if( Src )
 	{
-		FMemory::Memcpy( Buffer.GetTypedData(), Src, (CountBits+7)>>3 );
+		FMemory::Memcpy(Buffer.GetData(), Src, (CountBits + 7) >> 3);
 	}
 }
 void FBitReader::SetData( FBitReader& Src, int64 CountBits )
@@ -304,7 +304,7 @@ void FBitReader::SetData( FBitReader& Src, int64 CountBits )
 	ArIsError  = 0;
 	Buffer.Empty();
 	Buffer.AddUninitialized( (CountBits+7)>>3 );
-	Src.SerializeBits( Buffer.GetTypedData(), CountBits );
+	Src.SerializeBits(Buffer.GetData(), CountBits);
 }
 /** This appends data from another BitReader. It checks that this bit reader is byte-aligned so it can just do a TArray::Append instead of a bitcopy.
  *	It is intended to be used by performance minded code that wants to ensure an appBitCpy is avoided.
@@ -360,7 +360,7 @@ void FBitReader::SerializeBits( void* Dest, int64 LengthBits )
 		}
 		else
 		{
-			appBitsCpy( (uint8*)Dest, 0, Buffer.GetTypedData(), Pos, LengthBits );
+			appBitsCpy((uint8*)Dest, 0, Buffer.GetData(), Pos, LengthBits);
 			Pos += LengthBits;
 		}
 	}
@@ -416,7 +416,7 @@ void FBitReader::Serialize( void* Dest, int64 LengthBytes )
 }
 uint8* FBitReader::GetData()
 {
-	return Buffer.GetTypedData();
+	return Buffer.GetData();
 }
 uint8* FBitReader::GetDataPosChecked()
 {
@@ -466,6 +466,6 @@ void FBitReaderMark::Copy( FBitReader& Reader, TArray<uint8> &Buffer )
 	{
 		Buffer.Init(Bytes);		// This makes room but doesnt zero
 		Buffer[Bytes-1] = 0;	// Make sure the last byte is 0 out, because appBitsCpy wont touch the last bits
-		appBitsCpy( Buffer.GetTypedData(), 0, Reader.Buffer.GetTypedData(), Pos, Reader.Pos - Pos );
+		appBitsCpy(Buffer.GetData(), 0, Reader.Buffer.GetData(), Pos, Reader.Pos - Pos);
 	}
 }
diff --git a/Engine/Source/Runtime/Core/Private/Serialization/CustomVersion.cpp b/Engine/Source/Runtime/Core/Private/Serialization/CustomVersion.cpp
--- a/Engine/Source/Runtime/Core/Private/Serialization/CustomVersion.cpp
+++ b/Engine/Source/Runtime/Core/Private/Serialization/CustomVersion.cpp
@@ -159,6 +159,6 @@ FCustomVersionRegistration::~FCustomVersionRegistration()
 	--FoundKey->ReferenceCount;
 	if (FoundKey->ReferenceCount == 0)
 	{
-		Versions.RemoveAt(FoundKey - Versions.GetTypedData());
+		Versions.RemoveAt(FoundKey - Versions.GetData());
 	}
 }
diff --git a/Engine/Source/Runtime/Core/Private/Windows/WindowsPlatformProcess.cpp b/Engine/Source/Runtime/Core/Private/Windows/WindowsPlatformProcess.cpp
--- a/Engine/Source/Runtime/Core/Private/Windows/WindowsPlatformProcess.cpp
+++ b/Engine/Source/Runtime/Core/Private/Windows/WindowsPlatformProcess.cpp
@@ -304,7 +304,7 @@ FProcHandle FWindowsPlatformProcess::CreateProc( const TCHAR* URL, const TCHAR*
 		CW_USEDEFAULT, CW_USEDEFAULT, CW_USEDEFAULT, CW_USEDEFAULT,
 		0, 0, 0, dwFlags, ShowWindowFlags, 0, NULL,
 		::GetStdHandle(ProcessConstants::WIN_STD_INPUT_HANDLE), HANDLE(PipeWrite), HANDLE(PipeWrite) };
-	if( !CreateProcess( NULL, CommandLine.GetCharArray().GetTypedData(), &Attr, &Attr, true, CreateFlags,
+	if (!CreateProcess(NULL, CommandLine.GetCharArray().GetData(), &Attr, &Attr, true, CreateFlags,
 		NULL, OptionalWorkingDirectory, &StartupInfo, &ProcInfo ) )
 	{
 		if (OutProcessID)
@@ -553,7 +553,7 @@ bool FWindowsPlatformProcess::ExecProcess( const TCHAR* URL, const TCHAR* Params
 		CW_USEDEFAULT, CW_USEDEFAULT, CW_USEDEFAULT, CW_USEDEFAULT,
 		0, 0, 0, dwFlags, ShowWindowFlags, 0, NULL,
 		::GetStdHandle(ProcessConstants::WIN_STD_INPUT_HANDLE), WritablePipes[0], WritablePipes[1] };
-	if (CreateProcess(NULL, CommandLine.GetCharArray().GetTypedData(), &Attr, &Attr, true, CreateFlags,
+	if (CreateProcess(NULL, CommandLine.GetCharArray().GetData(), &Attr, &Attr, true, CreateFlags,
 		NULL, NULL, &StartupInfo, &ProcInfo))
 	{
 		if (bRedirectOutput)
