Fixed orthographic viewports defaulting to lit

This addresses TTP#347842 - Regression: EDITOR: Ortho viewports default to Lit

diff --git a/Engine/Source/Editor/LevelEditor/Private/SLevelViewport.cpp b/Engine/Source/Editor/LevelEditor/Private/SLevelViewport.cpp
--- a/Engine/Source/Editor/LevelEditor/Private/SLevelViewport.cpp
+++ b/Engine/Source/Editor/LevelEditor/Private/SLevelViewport.cpp
@@ -3384,7 +3384,8 @@ void SLevelViewport::ResetNewLevelViewFlags()
 {
 	ChangeExposureSetting(FEditorViewportCommands::AutoExposureRadioID);
 
-	LevelViewportClient->SetViewMode(VMI_Lit);
+	const bool bIsPerspective = LevelViewportClient->IsPerspective();
+	LevelViewportClient->SetViewMode(bIsPerspective ? FEditorViewportClient::DefaultPerspectiveViewMode : FEditorViewportClient::DefaultOrthoViewMode);
 
 	OnUseDefaultShowFlags();
 }
diff --git a/Engine/Source/Editor/UnrealEd/Private/EditorViewportClient.cpp b/Engine/Source/Editor/UnrealEd/Private/EditorViewportClient.cpp
--- a/Engine/Source/Editor/UnrealEd/Private/EditorViewportClient.cpp
+++ b/Engine/Source/Editor/UnrealEd/Private/EditorViewportClient.cpp
@@ -19,6 +19,9 @@
 
 #define LOCTEXT_NAMESPACE "EditorViewportClient"
 
+const EViewModeIndex FEditorViewportClient::DefaultPerspectiveViewMode = VMI_Lit;
+const EViewModeIndex FEditorViewportClient::DefaultOrthoViewMode = VMI_BrushWireframe;
+
 static TAutoConsoleVariable<int32> CVarAlignedOrthoZoom(
 	TEXT("r.Editor.AlignedOrthoZoom"),
 	1,
@@ -239,8 +242,8 @@ FEditorViewportClient::FEditorViewportClient(FEditorModeTools& InModeTools, FPre
 	, bIsSimulateInEditorViewport(false)
 	, bCameraLock(false)
 	, PreviewScene(InPreviewScene)
-	, PerspViewModeIndex(VMI_Lit)
-	, OrthoViewModeIndex(VMI_BrushWireframe)
+	, PerspViewModeIndex(DefaultPerspectiveViewMode)
+	, OrthoViewModeIndex(DefaultOrthoViewMode)
 	, NearPlane(-1.0f)
 	, FarPlane(0.0f)
 	, bInGameViewMode(false)
@@ -274,7 +277,7 @@ FEditorViewportClient::FEditorViewportClient(FEditorModeTools& InModeTools, FPre
 
 	EngineShowFlags.SetSnap(1);
 
-	SetViewMode(VMI_Lit);
+	SetViewMode(IsPerspective() ? PerspViewModeIndex : OrthoViewModeIndex);
 
 	ModeTools->OnEditorModeChanged().AddRaw(this, &FEditorViewportClient::OnEditorModeChanged);
 }
diff --git a/Engine/Source/Editor/UnrealEd/Public/EditorViewportClient.h b/Engine/Source/Editor/UnrealEd/Public/EditorViewportClient.h
--- a/Engine/Source/Editor/UnrealEd/Public/EditorViewportClient.h
+++ b/Engine/Source/Editor/UnrealEd/Public/EditorViewportClient.h
@@ -1345,6 +1345,13 @@ protected:
 	FVector DefaultOrbitZoom;
 	FVector DefaultOrbitLookAt;
 	
+public:
+	/* Default view mode for perspective viewports */
+	static const EViewModeIndex DefaultPerspectiveViewMode;
+
+	/* Default view mode for orthographic viewports */
+	static const EViewModeIndex DefaultOrthoViewMode;
+
 private:
 	/* View mode to set when this viewport is of type LVT_Perspective */
 	EViewModeIndex PerspViewModeIndex;
