Fix TTP #347766 Mac: Networking: MBP: Launching as Client goes fullscreen by default - the GameSettingsINI command line argument must be escaped as paths may contain spaces, but also the CreateProc code must be sure to remove the quotes before handing it off to Cocoa. Otherwise Cocoa will add more quotes and insanity will ensue.

diff --git a/Engine/Source/Editor/UnrealEd/Private/PlayLevel.cpp b/Engine/Source/Editor/UnrealEd/Private/PlayLevel.cpp
--- a/Engine/Source/Editor/UnrealEd/Private/PlayLevel.cpp
+++ b/Engine/Source/Editor/UnrealEd/Private/PlayLevel.cpp
@@ -700,7 +700,7 @@ FString GenerateCmdLineForNextPieInstance(int32 &WinX, int32 &WinY, int32 &Insta
 	//	-Override GameUserSettings.ini
 	//	-Force no steam
 	//	-Allow saving of config files (since we are giving them an override INI)
-	FString CmdLine = FString::Printf(TEXT("GameUserSettingsINI=%s -MultiprocessSaveConfig %s -MultiprocessOSS "), *GameUserSettingsOverride, *PlayInSettings->AdditionalLaunchOptions);
+	FString CmdLine = FString::Printf(TEXT("GameUserSettingsINI=\"%s\" -MultiprocessSaveConfig %s -MultiprocessOSS "), *GameUserSettingsOverride, *PlayInSettings->AdditionalLaunchOptions);
 
 	if (bIsDedicatedServer)
 	{
diff --git a/Engine/Source/Runtime/Core/Private/Mac/MacPlatformProcess.cpp b/Engine/Source/Runtime/Core/Private/Mac/MacPlatformProcess.cpp
--- a/Engine/Source/Runtime/Core/Private/Mac/MacPlatformProcess.cpp
+++ b/Engine/Source/Runtime/Core/Private/Mac/MacPlatformProcess.cpp
@@ -408,7 +408,7 @@ FProcHandle FMacPlatformProcess::CreateProc( const TCHAR* URL, const TCHAR* Parm
 						}
 						else
 						{
-							Arg = (NSString*)FPlatformString::TCHARToCFString(*MultiPartArg);
+							Arg = (NSString*)FPlatformString::TCHARToCFString(*MultiPartArg.Replace(TEXT("\""), TEXT("")));
 						}
 						[Arguments addObject: Arg];
 						CFRelease((CFStringRef)Arg);
