Fixed missing json Lightmass dependency for windows

diff --git a/Engine/Source/Editor/UnrealEd/Private/Lightmass/Lightmass.cpp b/Engine/Source/Editor/UnrealEd/Private/Lightmass/Lightmass.cpp
--- a/Engine/Source/Editor/UnrealEd/Private/Lightmass/Lightmass.cpp
+++ b/Engine/Source/Editor/UnrealEd/Private/Lightmass/Lightmass.cpp
@@ -2441,7 +2441,8 @@ bool FLightmassProcessor::BeginRun()
 		TEXT("../Win32/UnrealLightmass-SwarmInterface.dll"),
 		TEXT("../Win32/UnrealLightmass-Core.dll"),
 		TEXT("../Win32/UnrealLightmass-CoreUObject.dll"),
-		TEXT("../Win32/UnrealLightmass-Projects.dll")
+		TEXT("../Win32/UnrealLightmass-Projects.dll"),
+		TEXT("../Win32/UnrealLightmass-Json.dll")
 	};
 	const int32 RequiredDependencyPaths32Count = ARRAY_COUNT(RequiredDependencyPaths32);
 
@@ -2455,7 +2456,8 @@ bool FLightmassProcessor::BeginRun()
 		TEXT("../Win64/UnrealLightmass-SwarmInterface.dll"),
 		TEXT("../Win64/UnrealLightmass-Core.dll"),
 		TEXT("../Win64/UnrealLightmass-CoreUObject.dll"),
-		TEXT("../Win64/UnrealLightmass-Projects.dll")
+		TEXT("../Win64/UnrealLightmass-Projects.dll"),
+		TEXT("../Win64/UnrealLightmass-Json.dll")
 	};
 #elif PLATFORM_MAC
 #if UE_BUILD_DEBUG
