Map in Editor Performance test will wait for the shaders to compile before running the performance part.
The map will load and then it will wait for any shaders to finish first.

diff --git a/Engine/Source/Editor/UnrealEd/Private/Tests/AutomationEditorCommon.cpp b/Engine/Source/Editor/UnrealEd/Private/Tests/AutomationEditorCommon.cpp
--- a/Engine/Source/Editor/UnrealEd/Private/Tests/AutomationEditorCommon.cpp
+++ b/Engine/Source/Editor/UnrealEd/Private/Tests/AutomationEditorCommon.cpp
@@ -7,6 +7,7 @@
 #include "LevelEditor.h"
 #include "Editor/MainFrame/Public/MainFrame.h"
 #include "EngineVersion.h"
+#include "ShaderCompiler.h"
 
 
 DEFINE_LOG_CATEGORY_STATIC(LogAutomationEditorCommon, Log, All);
@@ -752,6 +753,18 @@ bool FEditorLoadMap::Update()
 	return true;
 }
 
+/**
+* This will cause the test to wait for the shaders to finish compiling before moving on.
+*/
+bool FWaitForShadersToFinishCompiling::Update()
+{
+	UE_LOG(LogEditorAutomationTests, Log, TEXT("Waiting for %i shaders to finish."), GShaderCompilingManager->GetNumRemainingJobs());
+	GShaderCompilingManager->FinishAllCompilation();
+	UE_LOG(LogEditorAutomationTests, Log, TEXT("Done waiting for shaders to finish."));
+	UE_LOG(LogEditorAutomationTests, Log, TEXT("%i shaders need to be recompiled."), GShaderCompilingManager->GetNumRemainingJobs());
+	return true;
+}
+
 //////////////////////////////////////////////////////////////////////
 //Find Asset Commands
 
diff --git a/Engine/Source/Editor/UnrealEd/Private/Tests/EditorPerformanceTests.cpp b/Engine/Source/Editor/UnrealEd/Private/Tests/EditorPerformanceTests.cpp
--- a/Engine/Source/Editor/UnrealEd/Private/Tests/EditorPerformanceTests.cpp
+++ b/Engine/Source/Editor/UnrealEd/Private/Tests/EditorPerformanceTests.cpp
@@ -133,12 +133,16 @@ bool FMapPerformanceInEditor::RunTest(const FString& Parameters)
 
 	UE_LOG(LogEditorAutomationTests, Log, TEXT("Running the performance capture test for %.0f seconds on %s"), Duration, *ShortMapName);
 	
+	//Wait for shaders to finish compile.
+	ADD_LATENT_AUTOMATION_COMMAND(FWaitForShadersToFinishCompiling);
+	
 	//Grab the performance numbers based on the duration.
-	ADD_LATENT_AUTOMATION_COMMAND(FWaitLatentCommand(1.0f));
-	ADD_LATENT_AUTOMATION_COMMAND(FEditorPerformanceCommand(Duration));
 	ADD_LATENT_AUTOMATION_COMMAND(FWaitLatentCommand(0.5f));
+	ADD_LATENT_AUTOMATION_COMMAND(FEditorPerformanceCommand(Duration));
+	
 
 	//Combine performance data into one chart.
+	ADD_LATENT_AUTOMATION_COMMAND(FWaitLatentCommand(0.5f));
 	ADD_LATENT_AUTOMATION_COMMAND(FGenerateEditorPerformanceCharts(ShortMapName));
 
 	return true;
diff --git a/Engine/Source/Editor/UnrealEd/Public/Tests/AutomationEditorCommon.h b/Engine/Source/Editor/UnrealEd/Public/Tests/AutomationEditorCommon.h
--- a/Engine/Source/Editor/UnrealEd/Public/Tests/AutomationEditorCommon.h
+++ b/Engine/Source/Editor/UnrealEd/Public/Tests/AutomationEditorCommon.h
@@ -143,6 +143,11 @@ DEFINE_LATENT_AUTOMATION_COMMAND_ONE_PARAMETER(FEditorLoadMap, FString, MapName)
 */
 DEFINE_LATENT_AUTOMATION_COMMAND_ONE_PARAMETER(FGenerateEditorPerformanceCharts, FString, MapName);
 
+/**
+* Waits for shaders to finish compiling before moving on to the next thing.
+*/
+DEFINE_LATENT_AUTOMATION_COMMAND(FWaitForShadersToFinishCompiling);
+
 //////////////////////////////////////////////////////////////////////////
 // FEditorAutomationTestUtilities
 
