Lightmass: Added missing modules

diff --git a/Engine/Source/Editor/UnrealEd/Private/Lightmass/Lightmass.cpp b/Engine/Source/Editor/UnrealEd/Private/Lightmass/Lightmass.cpp
--- a/Engine/Source/Editor/UnrealEd/Private/Lightmass/Lightmass.cpp
+++ b/Engine/Source/Editor/UnrealEd/Private/Lightmass/Lightmass.cpp
@@ -2466,6 +2466,8 @@ bool FLightmassProcessor::BeginRun()
 	{
 		TEXT("../DotNET/Mac/AgentInterface.dll"),
 		TEXT("../Mac/UnrealLightmass-Core-Mac-Debug.dylib"),
+		TEXT("../Mac/UnrealLightmass-CoreUObject-Mac-Debug.dylib"),
+		TEXT("../Mac/UnrealLightmass-Json-Mac-Debug.dylib"),
 		TEXT("../Mac/UnrealLightmass-Projects-Mac-Debug.dylib"),
 		TEXT("../Mac/UnrealLightmass-SwarmInterface-Mac-Debug.dylib")
 	};
@@ -2475,6 +2477,8 @@ bool FLightmassProcessor::BeginRun()
 	{
 		TEXT("../DotNET/Mac/AgentInterface.dll"),
 		TEXT("../Mac/UnrealLightmass-Core.dylib"),
+		TEXT("../Mac/UnrealLightmass-CoreUObject.dylib"),
+		TEXT("../Mac/UnrealLightmass-Json.dylib"),
 		TEXT("../Mac/UnrealLightmass-Projects.dylib"),
 		TEXT("../Mac/UnrealLightmass-SwarmInterface.dylib")
 	};
@@ -2486,6 +2490,8 @@ bool FLightmassProcessor::BeginRun()
 	{
 		TEXT("../DotNET/Linux/AgentInterface.dll"),
 		TEXT("../Linux/libUnrealLightmass-Core-Linux-Debug.so"),
+		TEXT("../Linux/libUnrealLightmass-CoreUObject-Linux-Debug.so"),
+		TEXT("../Linux/libUnrealLightmass-Json-Linux-Debug.so"),
 		TEXT("../Linux/libUnrealLightmass-Projects-Linux-Debug.so"),
 		TEXT("../Linux/libUnrealLightmass-SwarmInterface-Linux-Debug.so"),
 		TEXT("../Linux/libUnrealLightmass-Networking-Linux-Debug.so"),
@@ -2498,6 +2504,8 @@ bool FLightmassProcessor::BeginRun()
 	{
 		TEXT("../DotNET/Linux/AgentInterface.dll"),
 		TEXT("../Linux/libUnrealLightmass-Core.so"),
+		TEXT("../Linux/libUnrealLightmass-CoreUObject.so"),
+		TEXT("../Linux/libUnrealLightmass-Json.so"),
 		TEXT("../Linux/libUnrealLightmass-Projects.so"),
 		TEXT("../Linux/libUnrealLightmass-SwarmInterface.so"),
 		TEXT("../Linux/libUnrealLightmass-Networking.so"),
