Fix editor snapping not working for destructible mesh. Looks like bounding box was calculated incorrectly.

TTP#329818

codereview(James.Golding)

diff --git a/Engine/Source/Editor/UnrealEd/Private/Factories/ApexDestructibleAssetImport.cpp b/Engine/Source/Editor/UnrealEd/Private/Factories/ApexDestructibleAssetImport.cpp
--- a/Engine/Source/Editor/UnrealEd/Private/Factories/ApexDestructibleAssetImport.cpp
+++ b/Engine/Source/Editor/UnrealEd/Private/Factories/ApexDestructibleAssetImport.cpp
@@ -830,16 +830,6 @@ bool SetApexDestructibleAsset(UDestructibleMesh& DestructibleMesh, NxDestructibl
 
 	// Create initial bounding box based on expanded version of reference pose for meshes without physics assets. Can be overridden by artist.
 	FBox BoundingBox(SkelMeshImportDataPtr->Points.GetData(), SkelMeshImportDataPtr->Points.Num());
-	FBox Temp = BoundingBox;
-	FVector MidMesh		= 0.5f*(Temp.Min + Temp.Max);
-	BoundingBox.Min		= Temp.Min + 1.0f*(Temp.Min - MidMesh);
-	BoundingBox.Max		= Temp.Max + 1.0f*(Temp.Max - MidMesh);
-	// BRGTODO : what is this?
-	// Tuck up the bottom as this rarely extends lower than a reference pose's (e.g. having its feet on the floor).
-	// Maya has Y in the vertical, other packages have Z.
-	//BEN const int32 CoordToTuck = bAssumeMayaCoordinates ? 1 : 2;
-	//BEN BoundingBox.Min[CoordToTuck]	= Temp.Min[CoordToTuck] + 0.1f*(Temp.Min[CoordToTuck] - MidMesh[CoordToTuck]);
-	BoundingBox.Min[2]	= Temp.Min[2] + 0.1f*(Temp.Min[2] - MidMesh[2]);
 	DestructibleMesh.Bounds= FBoxSphereBounds(BoundingBox);
 
 	// Store whether or not this mesh has vertex colors
