Lightmass material properties are clamped to valid ranges before exporting, prevents light explosions with photon mapping

diff --git a/Engine/Source/Editor/UnrealEd/Private/Lightmass/LightmassRender.cpp b/Engine/Source/Editor/UnrealEd/Private/Lightmass/LightmassRender.cpp
--- a/Engine/Source/Editor/UnrealEd/Private/Lightmass/LightmassRender.cpp
+++ b/Engine/Source/Editor/UnrealEd/Private/Lightmass/LightmassRender.cpp
@@ -330,12 +330,12 @@ public:
 			{
 			case MP_EmissiveColor:
 				// Emissive is ALWAYS returned...
-				return Compiler->ForceCast(MaterialInterface->CompileProperty(&ProxyCompiler,MP_EmissiveColor),MCT_Float3,true,true);
+				return Compiler->Max(Compiler->ForceCast(MaterialInterface->CompileProperty(&ProxyCompiler,MP_EmissiveColor),MCT_Float3,true,true), Compiler->Constant3(0, 0, 0));
 			case MP_DiffuseColor:
 				// Only return for Opaque and Masked...
 				if (BlendMode == BLEND_Opaque || BlendMode == BLEND_Masked)
 				{
-					return Compiler->ForceCast(MaterialInterface->CompileProperty(&ProxyCompiler, DiffuseInput),MCT_Float3,true,true);
+					return Compiler->Saturate(Compiler->ForceCast(MaterialInterface->CompileProperty(&ProxyCompiler, DiffuseInput),MCT_Float3,true,true));
 				}
 				break;
 			case MP_SpecularColor: 
@@ -343,8 +343,8 @@ public:
 				if (BlendMode == BLEND_Opaque || BlendMode == BLEND_Masked)
 				{
 					return Compiler->AppendVector(
-						Compiler->ForceCast(MaterialInterface->CompileProperty(&ProxyCompiler, MP_SpecularColor),MCT_Float3,true,true), 
-						Compiler->ForceCast(MaterialInterface->CompileProperty(&ProxyCompiler,MP_Roughness),MCT_Float1));
+						Compiler->Saturate(Compiler->ForceCast(MaterialInterface->CompileProperty(&ProxyCompiler, MP_SpecularColor),MCT_Float3,true,true)), 
+						Compiler->Saturate(Compiler->ForceCast(MaterialInterface->CompileProperty(&ProxyCompiler,MP_Roughness),MCT_Float1)));
 				}
 				break;
 			case MP_Normal:
@@ -372,7 +372,7 @@ public:
 					}
 					else
 					{
-						return Compiler->ForceCast(MaterialInterface->CompileProperty(Compiler, DiffuseInput),MCT_Float3,true,true);
+						return Compiler->Saturate(Compiler->ForceCast(MaterialInterface->CompileProperty(Compiler, DiffuseInput),MCT_Float3,true,true));
 					}
 				}
 				else if ((BlendMode == BLEND_Translucent) || (BlendMode == BLEND_Additive))
@@ -384,9 +384,9 @@ public:
 					}
 					else
 					{
-						ColoredOpacity = Compiler->ForceCast(MaterialInterface->CompileProperty(Compiler, DiffuseInput),MCT_Float3,true,true);
+						ColoredOpacity = Compiler->Saturate(Compiler->ForceCast(MaterialInterface->CompileProperty(Compiler, DiffuseInput),MCT_Float3,true,true));
 					}
-					return Compiler->Lerp(Compiler->Constant3(1.0f, 1.0f, 1.0f), ColoredOpacity, Compiler->ForceCast( MaterialInterface->CompileProperty(&ProxyCompiler,MP_Opacity), MCT_Float1));
+					return Compiler->Lerp(Compiler->Constant3(1.0f, 1.0f, 1.0f), ColoredOpacity, Compiler->Saturate(Compiler->ForceCast( MaterialInterface->CompileProperty(&ProxyCompiler,MP_Opacity), MCT_Float1)));
 				}
 				break;
 			default:
diff --git a/Engine/Source/Programs/UnrealLightmass/Public/ImportExport.h b/Engine/Source/Programs/UnrealLightmass/Public/ImportExport.h
--- a/Engine/Source/Programs/UnrealLightmass/Public/ImportExport.h
+++ b/Engine/Source/Programs/UnrealLightmass/Public/ImportExport.h
@@ -49,7 +49,7 @@ namespace Lightmass
 	static const FGuid LM_DEBUGOUTPUT_VERSION			= FGuid(0xfa520216, 0xd5aa49eb, 0xa9bbc1f5, 0x8ab715fc);
 	static const FGuid LM_SCENE_VERSION					= FGuid(0x8ad6acb1, 0xe15e4fec, 0x93c72936, 0xf2b67090);
 	static const FGuid LM_STATICMESH_VERSION			= FGuid(0x9b3b1bc6, 0x9270497a, 0xa07521fe, 0xc7045e91);
-	static const FGuid LM_MATERIAL_VERSION				= FGuid(0xa6640d72, 0x33bf47b3, 0xbbcf76bb, 0x4350a2b5);
+	static const FGuid LM_MATERIAL_VERSION				= FGuid(0xa6640d72, 0x33bf47b4, 0xbbcf76bb, 0x4350a2b5);
 
 
 	/** Alert source object type identifiers... */
diff --git a/Engine/Source/Runtime/Engine/Private/Materials/HLSLMaterialTranslator.h b/Engine/Source/Runtime/Engine/Private/Materials/HLSLMaterialTranslator.h
--- a/Engine/Source/Runtime/Engine/Private/Materials/HLSLMaterialTranslator.h
+++ b/Engine/Source/Runtime/Engine/Private/Materials/HLSLMaterialTranslator.h
@@ -2819,6 +2819,23 @@ protected:
 		}
 	}
 
+	virtual int32 Saturate(int32 X) override
+	{
+		if(X == INDEX_NONE)
+		{
+			return INDEX_NONE;
+		}
+
+		if(GetParameterUniformExpression(X))
+		{
+			return AddUniformExpression(new FMaterialUniformExpressionSaturate(GetParameterUniformExpression(X)),GetParameterType(X),TEXT("saturate(%s)"),*GetParameterCode(X));
+		}
+		else
+		{
+			return AddCodeChunk(GetParameterType(X),TEXT("saturate(%s)"),*GetParameterCode(X));
+		}
+	}
+
 	virtual int32 ComponentMask(int32 Vector,bool R,bool G,bool B,bool A) override
 	{
 		if(Vector == INDEX_NONE)
diff --git a/Engine/Source/Runtime/Engine/Private/Materials/MaterialShared.cpp b/Engine/Source/Runtime/Engine/Private/Materials/MaterialShared.cpp
--- a/Engine/Source/Runtime/Engine/Private/Materials/MaterialShared.cpp
+++ b/Engine/Source/Runtime/Engine/Private/Materials/MaterialShared.cpp
@@ -2104,6 +2104,8 @@ FMaterialUpdateContext::~FMaterialUpdateContext()
 
 EMaterialProperty GetMaterialPropertyFromInputOutputIndex(int32 Index)
 {
+	// Warning: If you change this mapping, you will need to handle backwards compatibility for nodes that use it, otherwise pins will get swapped
+
 	switch(Index)
 	{
 	case 0: return MP_BaseColor;
@@ -2142,6 +2144,8 @@ EMaterialProperty GetMaterialPropertyFromInputOutputIndex(int32 Index)
 
 int32 GetInputOutputIndexFromMaterialProperty(EMaterialProperty Property)
 {
+	// Warning: If you change this mapping, you will need to handle backwards compatibility for nodes that use it, otherwise pins will get swapped
+
 	switch(Property)
 	{
 	case MP_BaseColor: return 0;
diff --git a/Engine/Source/Runtime/Engine/Private/Materials/MaterialUniformExpressions.cpp b/Engine/Source/Runtime/Engine/Private/Materials/MaterialUniformExpressions.cpp
--- a/Engine/Source/Runtime/Engine/Private/Materials/MaterialUniformExpressions.cpp
+++ b/Engine/Source/Runtime/Engine/Private/Materials/MaterialUniformExpressions.cpp
@@ -419,6 +419,7 @@ IMPLEMENT_MATERIALUNIFORMEXPRESSION_TYPE(FMaterialUniformExpressionAppendVector)
 IMPLEMENT_MATERIALUNIFORMEXPRESSION_TYPE(FMaterialUniformExpressionMin);
 IMPLEMENT_MATERIALUNIFORMEXPRESSION_TYPE(FMaterialUniformExpressionMax);
 IMPLEMENT_MATERIALUNIFORMEXPRESSION_TYPE(FMaterialUniformExpressionClamp);
+IMPLEMENT_MATERIALUNIFORMEXPRESSION_TYPE(FMaterialUniformExpressionSaturate);
 IMPLEMENT_MATERIALUNIFORMEXPRESSION_TYPE(FMaterialUniformExpressionFloor);
 IMPLEMENT_MATERIALUNIFORMEXPRESSION_TYPE(FMaterialUniformExpressionCeil);
 IMPLEMENT_MATERIALUNIFORMEXPRESSION_TYPE(FMaterialUniformExpressionFrac);
diff --git a/Engine/Source/Runtime/Engine/Private/Materials/MaterialUniformExpressions.h b/Engine/Source/Runtime/Engine/Private/Materials/MaterialUniformExpressions.h
--- a/Engine/Source/Runtime/Engine/Private/Materials/MaterialUniformExpressions.h
+++ b/Engine/Source/Runtime/Engine/Private/Materials/MaterialUniformExpressions.h
@@ -783,6 +783,51 @@ private:
 	TRefCountPtr<FMaterialUniformExpression> Max;
 };
 
+/**
+ */
+class FMaterialUniformExpressionSaturate: public FMaterialUniformExpression
+{
+	DECLARE_MATERIALUNIFORMEXPRESSION_TYPE(FMaterialUniformExpressionSaturate);
+public:
+
+	FMaterialUniformExpressionSaturate() {}
+	FMaterialUniformExpressionSaturate(FMaterialUniformExpression* InInput):
+		Input(InInput)
+	{}
+
+	// FMaterialUniformExpression interface.
+	virtual void Serialize(FArchive& Ar)
+	{
+		Ar << Input;
+	}
+	virtual void GetNumberValue(const FMaterialRenderContext& Context,FLinearColor& OutValue) const
+	{
+		FLinearColor ValueInput = FLinearColor::Black;
+		Input->GetNumberValue(Context, ValueInput);
+
+		OutValue.R = FMath::Clamp<float>(ValueInput.R, 0, 1);
+		OutValue.G = FMath::Clamp<float>(ValueInput.G, 0, 1);
+		OutValue.B = FMath::Clamp<float>(ValueInput.B, 0, 1);
+		OutValue.A = FMath::Clamp<float>(ValueInput.A, 0, 1);
+	}
+	virtual bool IsConstant() const
+	{
+		return Input->IsConstant();
+	}
+	virtual bool IsIdentical(const FMaterialUniformExpression* OtherExpression) const
+	{
+		if (GetType() != OtherExpression->GetType())
+		{
+			return false;
+		}
+		FMaterialUniformExpressionSaturate* OtherClamp = (FMaterialUniformExpressionSaturate*)OtherExpression;
+		return Input->IsIdentical(OtherClamp->Input);
+	}
+
+private:
+	TRefCountPtr<FMaterialUniformExpression> Input;
+};
+
 /**
  */
 class FMaterialUniformExpressionFloor: public FMaterialUniformExpression
diff --git a/Engine/Source/Runtime/Engine/Public/MaterialCompiler.h b/Engine/Source/Runtime/Engine/Public/MaterialCompiler.h
--- a/Engine/Source/Runtime/Engine/Public/MaterialCompiler.h
+++ b/Engine/Source/Runtime/Engine/Public/MaterialCompiler.h
@@ -140,6 +140,7 @@ public:
 	virtual int32 Min(int32 A,int32 B) = 0;
 	virtual int32 Max(int32 A,int32 B) = 0;
 	virtual int32 Clamp(int32 X,int32 A,int32 B) = 0;
+	virtual int32 Saturate(int32 X) = 0;
 
 	virtual int32 ComponentMask(int32 Vector,bool R,bool G,bool B,bool A) = 0;
 	virtual int32 AppendVector(int32 A,int32 B) = 0;
@@ -289,6 +290,7 @@ public:
 	virtual int32 Min(int32 A,int32 B) override { return Compiler->Min(A,B); }
 	virtual int32 Max(int32 A,int32 B) override { return Compiler->Max(A,B); }
 	virtual int32 Clamp(int32 X,int32 A,int32 B) override { return Compiler->Clamp(X,A,B); }
+	virtual int32 Saturate(int32 X) override { return Compiler->Saturate(X); }
 
 	virtual int32 ComponentMask(int32 Vector,bool R,bool G,bool B,bool A) override { return Compiler->ComponentMask(Vector,R,G,B,A); }
 	virtual int32 AppendVector(int32 A,int32 B) override { return Compiler->AppendVector(A,B); }
