Added ability to use long package names with the -map flag when cooking

eg) -map=/Game/MyMap

https://udn.unrealengine.com/questions/222051/requests-for-unrealfrontend.html

diff --git a/Engine/Source/Editor/UnrealEd/Private/Commandlets/CookCommandlet.cpp b/Engine/Source/Editor/UnrealEd/Private/Commandlets/CookCommandlet.cpp
--- a/Engine/Source/Editor/UnrealEd/Private/Commandlets/CookCommandlet.cpp
+++ b/Engine/Source/Editor/UnrealEd/Private/Commandlets/CookCommandlet.cpp
@@ -820,6 +820,11 @@ void UCookCommandlet::CollectFilesToCook(TArray<FString>& FilesInPath)
 				AddFileToCook(FilesInPath, CurrEntry);
 			}
 		}
+		else if (FPackageName::IsValidLongPackageName(CurrEntry))
+		{
+			CurrEntry = FPackageName::LongPackageNameToFilename(CurrEntry, TEXT(".umap"));
+			AddFileToCook(FilesInPath, CurrEntry);
+		}
 		else
 		{
 			AddFileToCook(FilesInPath, CurrEntry);
