Added FBX window current file display

diff --git a/Engine/Source/Editor/UnrealEd/Private/Fbx/FbxOptionWindow.cpp b/Engine/Source/Editor/UnrealEd/Private/Fbx/FbxOptionWindow.cpp
--- a/Engine/Source/Editor/UnrealEd/Private/Fbx/FbxOptionWindow.cpp
+++ b/Engine/Source/Editor/UnrealEd/Private/Fbx/FbxOptionWindow.cpp
@@ -19,6 +19,33 @@ void SFbxOptionWindow::Construct(const FArguments& InArgs)
 	this->ChildSlot
 	[
 		SNew(SVerticalBox)
+		+SVerticalBox::Slot()
+		.AutoHeight()
+		.Padding(2)
+		[
+			SNew(SBorder)
+			.Padding(FMargin(3))
+			.BorderImage(FEditorStyle::GetBrush("ToolPanel.GroupBorder"))
+			[
+				SNew(SHorizontalBox)
+				+SHorizontalBox::Slot()
+				.AutoWidth()
+				[
+					SNew(STextBlock)
+					.Font(FEditorStyle::GetFontStyle("CurveEd.LabelFont"))
+					.Text(LOCTEXT("Import_CurrentFileTitle", "Current File: "))
+				]
+				+SHorizontalBox::Slot()
+				.Padding(5, 0, 0, 0)
+				.AutoWidth()
+				.VAlign(VAlign_Center)
+				[
+					SNew(STextBlock)
+					.Font(FEditorStyle::GetFontStyle("CurveEd.InfoFont"))
+					.Text(InArgs._FullPath)
+				]
+			]
+		]
 		+ SVerticalBox::Slot()
 		.AutoHeight()
 		.Padding(2)
