Removed underscores from suffix on actor labels when there is a name clash.
#ttp 345516 - NUX:Level Editing: Sometimes newly placed actors get a suffix with an underscore and a number, and sometimes only a number
#branch UE4
#proj Editor.UnrealEd

diff --git a/Engine/Source/Editor/UnrealEd/Private/Editor.cpp b/Engine/Source/Editor/UnrealEd/Private/Editor.cpp
--- a/Engine/Source/Editor/UnrealEd/Private/Editor.cpp
+++ b/Engine/Source/Editor/UnrealEd/Private/Editor.cpp
@@ -4664,8 +4664,7 @@ void UEditorEngine::SetActorLabelUnique(AActor* Actor, const FString& NewActorLa
 		// See if the current label ends in a number, and try to create a new label based on that
 		if (!SplitActorLabel( Prefix, LabelIdx ))
 		{
-			// If there wasn't a number on there, append an underscore and start the numbering from 2 (1 before incrementing below)
-			Prefix += TEXT('_');
+			// If there wasn't a number on there, append a number, starting from 2 (1 before incrementing below)
 			LabelIdx = 1;
 		}
 
