Fixed TTP #348093 "Editor: Toggling realtime has no effect" related to navmesh rendering performance with disabled realtime in editor. #ue4

diff --git a/Engine/Source/Editor/UnrealEd/Private/NavigationBuildingNotification.cpp b/Engine/Source/Editor/UnrealEd/Private/NavigationBuildingNotification.cpp
--- a/Engine/Source/Editor/UnrealEd/Private/NavigationBuildingNotification.cpp
+++ b/Engine/Source/Editor/UnrealEd/Private/NavigationBuildingNotification.cpp
@@ -44,6 +44,7 @@ void FNavigationBuildingNotificationImpl::BuildStarted()
 void FNavigationBuildingNotificationImpl::BuildFinished()
 {
 	// Finished all requests! Notify the UI.
+	UEditorEngine* const EEngine = Cast<UEditorEngine>(GEngine);
 	TSharedPtr<SNotificationItem> NotificationItem = NavigationBuildNotificationPtr.Pin();
 	if (NotificationItem.IsValid())
 	{
@@ -51,10 +52,22 @@ void FNavigationBuildingNotificationImpl::BuildFinished()
 		NotificationItem->SetCompletionState(SNotificationItem::CS_Success);
 		NotificationItem->ExpireAndFadeout();
 
+		if (EEngine)
+		{
+			// request update for all viewports with disabled real time but with visible navmesh
+			for (auto Viewport : EEngine->AllViewportClients)
+			{
+				if (Viewport && Viewport->IsRealtime() == false && Viewport->EngineShowFlags.Navigation)
+				{
+					Viewport->bNeedsRedraw = true;
+					EEngine->UpdateSingleViewportClient(Viewport, true, false);
+				}
+			}
+		}
+
 		NavigationBuildNotificationPtr.Reset();
 	}
 
-	UEditorEngine* const EEngine = Cast<UEditorEngine>(GEngine);
 	if (EEngine != NULL && (FEditorBuildUtils::IsBuildingNavigationFromUserRequest()))
 	{
 		FNotificationInfo Info( NSLOCTEXT("NavigationBuild", "NavigationBuildDoneMessage", "Navigation building completed.") );
@@ -74,11 +87,6 @@ void FNavigationBuildingNotificationImpl::BuildFinished()
 		}
 	}
 
-	if (EEngine && EEngine->IsAnyViewportRealtime() == false)
-	{
-		EEngine->RedrawLevelEditingViewports();
-	}
-
 	FEditorBuildUtils::PathBuildingFinished();
 }
 
