The editor performance tool will now change the viewports to the first bookmarked view.

diff --git a/Engine/Source/Editor/UnrealEd/Private/Tests/AutomationEditorCommon.cpp b/Engine/Source/Editor/UnrealEd/Private/Tests/AutomationEditorCommon.cpp
--- a/Engine/Source/Editor/UnrealEd/Private/Tests/AutomationEditorCommon.cpp
+++ b/Engine/Source/Editor/UnrealEd/Private/Tests/AutomationEditorCommon.cpp
@@ -761,10 +761,36 @@ bool FWaitForShadersToFinishCompiling::Update()
 	UE_LOG(LogEditorAutomationTests, Log, TEXT("Waiting for %i shaders to finish."), GShaderCompilingManager->GetNumRemainingJobs());
 	GShaderCompilingManager->FinishAllCompilation();
 	UE_LOG(LogEditorAutomationTests, Log, TEXT("Done waiting for shaders to finish."));
-	UE_LOG(LogEditorAutomationTests, Log, TEXT("%i shaders need to be recompiled."), GShaderCompilingManager->GetNumRemainingJobs());
 	return true;
 }
 
+/**
+* Latent command that changes the editor viewport to the first available bookmarked view.
+*/
+bool FChangeViewportToFirstAvailableBookmarkCommand::Update()
+{
+	FEditorModeTools EditorModeTools;
+	FLevelEditorViewportClient* ViewportClient;
+
+	UE_LOG(LogEditorAutomationTests, Log, TEXT("Attempting to change the editor viewports view to the first set bookmark."));
+
+	//Move the perspective viewport view to show the test.
+	for (int32 i = 0; i < GEditor->LevelViewportClients.Num(); i++)
+	{
+		ViewportClient = GEditor->LevelViewportClients[i];
+
+		for (uint32 ViewportIndex = 0; ViewportIndex <= AWorldSettings::MAX_BOOKMARK_NUMBER; ViewportIndex++)
+		{
+			if (EditorModeTools.CheckBookmark(ViewportIndex, ViewportClient))
+			{
+				UE_LOG(LogEditorAutomationTests, Log, TEXT("Changing a viewport view to the set bookmark %i"), ViewportIndex);
+				EditorModeTools.JumpToBookmark(ViewportIndex, true, ViewportClient);
+				break;
+			}
+		}
+	}
+	return true;
+}
 //////////////////////////////////////////////////////////////////////
 //Find Asset Commands
 
diff --git a/Engine/Source/Editor/UnrealEd/Private/Tests/EditorPerformanceTests.cpp b/Engine/Source/Editor/UnrealEd/Private/Tests/EditorPerformanceTests.cpp
--- a/Engine/Source/Editor/UnrealEd/Private/Tests/EditorPerformanceTests.cpp
+++ b/Engine/Source/Editor/UnrealEd/Private/Tests/EditorPerformanceTests.cpp
@@ -133,6 +133,10 @@ bool FMapPerformanceInEditor::RunTest(const FString& Parameters)
 
 	UE_LOG(LogEditorAutomationTests, Log, TEXT("Running the performance capture test for %.0f seconds on %s"), Duration, *ShortMapName);
 	
+	//Move the viewport views to the first bookmark
+	ADD_LATENT_AUTOMATION_COMMAND(FChangeViewportToFirstAvailableBookmarkCommand);
+	ADD_LATENT_AUTOMATION_COMMAND(FWaitLatentCommand(0.5f));
+
 	//Wait for shaders to finish compile.
 	ADD_LATENT_AUTOMATION_COMMAND(FWaitForShadersToFinishCompiling);
 	
diff --git a/Engine/Source/Editor/UnrealEd/Public/Tests/AutomationEditorCommon.h b/Engine/Source/Editor/UnrealEd/Public/Tests/AutomationEditorCommon.h
--- a/Engine/Source/Editor/UnrealEd/Public/Tests/AutomationEditorCommon.h
+++ b/Engine/Source/Editor/UnrealEd/Public/Tests/AutomationEditorCommon.h
@@ -148,6 +148,11 @@ DEFINE_LATENT_AUTOMATION_COMMAND_ONE_PARAMETER(FGenerateEditorPerformanceCharts,
 */
 DEFINE_LATENT_AUTOMATION_COMMAND(FWaitForShadersToFinishCompiling);
 
+/**
+* Latent command that changes the editor viewport to the first available bookmarked view.
+*/
+DEFINE_LATENT_AUTOMATION_COMMAND(FChangeViewportToFirstAvailableBookmarkCommand);
+
 //////////////////////////////////////////////////////////////////////////
 // FEditorAutomationTestUtilities
 
