Updating the open asset in editor tests.  Add more asset types to the excluded list so they don't get added to the Misc. category.

diff --git a/Engine/Source/Editor/UnrealEd/Private/Tests/AutomationEditorCommon.cpp b/Engine/Source/Editor/UnrealEd/Private/Tests/AutomationEditorCommon.cpp
--- a/Engine/Source/Editor/UnrealEd/Private/Tests/AutomationEditorCommon.cpp
+++ b/Engine/Source/Editor/UnrealEd/Private/Tests/AutomationEditorCommon.cpp
@@ -8,6 +8,8 @@
 #include "Editor/MainFrame/Public/MainFrame.h"
 #include "EngineVersion.h"
 #include "ShaderCompiler.h"
+#include "BehaviorTree/BehaviorTree.h"
+
 
 
 DEFINE_LOG_CATEGORY_STATIC(LogAutomationEditorCommon, Log, All);
@@ -876,41 +878,45 @@ void FEditorAutomationTestUtilities::CollectMiscGameContentTestsByClass(TArray<F
 
 	//This is a list of classes that we don't want to be in the misc category.
 	TArray<FName> ExcludeClassesList;
-	ExcludeClassesList.Add(UTexture::StaticClass()->GetFName());
-	ExcludeClassesList.Add(UTexture2D::StaticClass()->GetFName());
-	ExcludeClassesList.Add(UTextureRenderTarget::StaticClass()->GetFName());
-	ExcludeClassesList.Add(UTextureRenderTarget2D::StaticClass()->GetFName());
-	ExcludeClassesList.Add(UTextureCube::StaticClass()->GetFName());
 	ExcludeClassesList.Add(UAimOffsetBlendSpace::StaticClass()->GetFName());
 	ExcludeClassesList.Add(UAimOffsetBlendSpace1D::StaticClass()->GetFName());
+	ExcludeClassesList.Add(UAnimBlueprint::StaticClass()->GetFName());
 	ExcludeClassesList.Add(UAnimMontage::StaticClass()->GetFName());
 	ExcludeClassesList.Add(UAnimSequence::StaticClass()->GetFName());
+	ExcludeClassesList.Add(UBehaviorTree::StaticClass()->GetFName());
 	ExcludeClassesList.Add(UBlendSpace::StaticClass()->GetFName());
 	ExcludeClassesList.Add(UBlendSpace1D::StaticClass()->GetFName());
+	ExcludeClassesList.Add(UBlueprint::StaticClass()->GetFName());
 	ExcludeClassesList.Add(UDestructibleMesh::StaticClass()->GetFName());
 	ExcludeClassesList.Add(UDialogueVoice::StaticClass()->GetFName());
 	ExcludeClassesList.Add(UDialogueWave::StaticClass()->GetFName());
+	ExcludeClassesList.Add(UFont::StaticClass()->GetFName());
 	ExcludeClassesList.Add(UMaterial::StaticClass()->GetFName());
 	ExcludeClassesList.Add(UMaterialFunction::StaticClass()->GetFName());
 	ExcludeClassesList.Add(UMaterialInstanceConstant::StaticClass()->GetFName());
 	ExcludeClassesList.Add(UMaterialParameterCollection::StaticClass()->GetFName());
 	ExcludeClassesList.Add(UParticleSystem::StaticClass()->GetFName());
 	ExcludeClassesList.Add(UPhysicalMaterial::StaticClass()->GetFName());
+	ExcludeClassesList.Add(UPhysicsAsset::StaticClass()->GetFName());
 	ExcludeClassesList.Add(UReverbEffect::StaticClass()->GetFName());
-	ExcludeClassesList.Add(USlateWidgetStyleAsset::StaticClass()->GetFName());
-	ExcludeClassesList.Add(USlateBrushAsset::StaticClass()->GetFName());
-	ExcludeClassesList.Add(USoundAttenuation::StaticClass()->GetFName());
-	ExcludeClassesList.Add(UStaticMesh::StaticClass()->GetFName());
 	ExcludeClassesList.Add(USkeletalMesh::StaticClass()->GetFName());
 	ExcludeClassesList.Add(USkeleton::StaticClass()->GetFName());
+	ExcludeClassesList.Add(USlateBrushAsset::StaticClass()->GetFName());
+	ExcludeClassesList.Add(USlateWidgetStyleAsset::StaticClass()->GetFName());
+	ExcludeClassesList.Add(USoundAttenuation::StaticClass()->GetFName());
 	ExcludeClassesList.Add(USoundClass::StaticClass()->GetFName());
 	ExcludeClassesList.Add(USoundCue::StaticClass()->GetFName());
 	ExcludeClassesList.Add(USoundMix::StaticClass()->GetFName());
 	ExcludeClassesList.Add(USoundWave::StaticClass()->GetFName());
+	ExcludeClassesList.Add(UStaticMesh::StaticClass()->GetFName());
 	ExcludeClassesList.Add(USubsurfaceProfile::StaticClass()->GetFName());
-	ExcludeClassesList.Add(UFont::StaticClass()->GetFName());
-	ExcludeClassesList.Add(UBlueprint::StaticClass()->GetFName());
-	ExcludeClassesList.Add(UAnimBlueprint::StaticClass()->GetFName());
+	ExcludeClassesList.Add(UTexture::StaticClass()->GetFName());
+	ExcludeClassesList.Add(UTexture2D::StaticClass()->GetFName());
+	ExcludeClassesList.Add(UTextureCube::StaticClass()->GetFName());
+	ExcludeClassesList.Add(UTextureRenderTarget::StaticClass()->GetFName());
+	ExcludeClassesList.Add(UTextureRenderTarget2D::StaticClass()->GetFName());
+	ExcludeClassesList.Add(UUserDefinedEnum::StaticClass()->GetFName());
+	ExcludeClassesList.Add(UWorld::StaticClass()->GetFName());
 
 	//Generating the list of assets.
 	//This list isn't expected to have anything that may be obtained with another function.
diff --git a/Engine/Source/Editor/UnrealEd/Private/Tests/EditorOpenAssetAutomationTests.cpp b/Engine/Source/Editor/UnrealEd/Private/Tests/EditorOpenAssetAutomationTests.cpp
--- a/Engine/Source/Editor/UnrealEd/Private/Tests/EditorOpenAssetAutomationTests.cpp
+++ b/Engine/Source/Editor/UnrealEd/Private/Tests/EditorOpenAssetAutomationTests.cpp
@@ -11,6 +11,7 @@
 //Includes needed for opening certain assets
 #include "Materials/MaterialFunction.h"
 #include "Slate/SlateBrushAsset.h"
+#include "BehaviorTree/BehaviorTree.h"
 
 
 
@@ -212,6 +213,25 @@ bool FOpenAnimSequenceAssetEditors::RunTest(const FString& Parameters)
 	return bDidTheTestPass;
 }
 
+//////////////////////////////////////////////////////////////////////////
+/**
+* This test opens each BEHAVIOR TREE into its sub-editor.
+*/
+IMPLEMENT_COMPLEX_AUTOMATION_TEST(FOpenBehaviorTreeAssetEditors, "Tools.AI.Open Behavior Tree Assets", (EAutomationTestFlags::ATF_Editor | EAutomationTestFlags::ATF_RequiresUser));
+
+void FOpenBehaviorTreeAssetEditors::GetTests(TArray<FString>& OutBeautifiedNames, TArray<FString>& OutTestCommands) const
+{
+	//This grabs each Behavior Tree asset in the Game/Content directory
+	FEditorAutomationTestUtilities::CollectGameContentTestsByClass(UBehaviorTree::StaticClass(), true, OutBeautifiedNames, OutTestCommands);
+}
+
+bool FOpenBehaviorTreeAssetEditors::RunTest(const FString& Parameters)
+{
+	bool bDidTheTestPass = FOpenActualAssetEditors(Parameters);
+	CollectGarbage(GARBAGE_COLLECTION_KEEPFLAGS);
+	return bDidTheTestPass;
+}
+
 //////////////////////////////////////////////////////////////////////////
 /**
 * This test opens each BlendSpace into its sub-editor.
@@ -270,6 +290,25 @@ bool FOpenBlueprintAssetEditors::RunTest(const FString& Parameters)
 	return bDidTheTestPass;
 }
 
+//////////////////////////////////////////////////////////////////////////
+/**
+* This test opens each UserDefinedEnum into its sub-editor.
+*/
+IMPLEMENT_COMPLEX_AUTOMATION_TEST(FOpenUserDefinedEnumAssetEditors, "Tools.Blueprint.Open User Defined Enum Assets", (EAutomationTestFlags::ATF_Editor | EAutomationTestFlags::ATF_RequiresUser));
+
+void FOpenUserDefinedEnumAssetEditors::GetTests(TArray<FString>& OutBeautifiedNames, TArray<FString>& OutTestCommands) const
+{
+	//This grabs each UserDefinedEnum asset in the Game/Content directory
+	FEditorAutomationTestUtilities::CollectGameContentTestsByClass(UUserDefinedEnum::StaticClass(), true, OutBeautifiedNames, OutTestCommands);
+}
+
+bool FOpenUserDefinedEnumAssetEditors::RunTest(const FString& Parameters)
+{
+	bool bDidTheTestPass = FOpenActualAssetEditors(Parameters);
+	CollectGarbage(GARBAGE_COLLECTION_KEEPFLAGS);
+	return bDidTheTestPass;
+}
+
 //////////////////////////////////////////////////////////////////////////
 /**
 * This test opens each DESTRUCTIBLE MESH into its sub-editor.
@@ -446,7 +485,7 @@ bool FOpenPhysicalMaterialAssetEditors::RunTest(const FString& Parameters)
 /**
 * This test opens each PhysicsAsset into its sub-editor.
 */
-IMPLEMENT_COMPLEX_AUTOMATION_TEST(FOpenPhysicsAssetAssetEditors, "Tools.Physics.Open Physics Asset Assets", (EAutomationTestFlags::ATF_Editor | EAutomationTestFlags::ATF_RequiresUser));
+IMPLEMENT_COMPLEX_AUTOMATION_TEST(FOpenPhysicsAssetAssetEditors, "Tools.Physics.Open Physics Assets", (EAutomationTestFlags::ATF_Editor | EAutomationTestFlags::ATF_RequiresUser));
 
 void FOpenPhysicsAssetAssetEditors::GetTests(TArray<FString>& OutBeautifiedNames, TArray<FString>& OutTestCommands) const
 {
