Fixing a mispelling in the test name.

diff --git a/Engine/Source/Editor/UnrealEd/Private/Tests/EditorPerformanceTests.cpp b/Engine/Source/Editor/UnrealEd/Private/Tests/EditorPerformanceTests.cpp
--- a/Engine/Source/Editor/UnrealEd/Private/Tests/EditorPerformanceTests.cpp
+++ b/Engine/Source/Editor/UnrealEd/Private/Tests/EditorPerformanceTests.cpp
@@ -13,7 +13,7 @@
 * Map Performance in Editor tests
 * Grabs certain performance numbers and saves it to a file.
 */
-IMPLEMENT_COMPLEX_AUTOMATION_TEST(FMapPerformanceInEditor, "Performance.Map Perfomance in Editor", EAutomationTestFlags::ATF_Editor | EAutomationTestFlags::ATF_RequiresUser)
+IMPLEMENT_COMPLEX_AUTOMATION_TEST(FMapPerformanceInEditor, "Performance.Map Performance in Editor", EAutomationTestFlags::ATF_Editor | EAutomationTestFlags::ATF_RequiresUser)
 
 /**
 * Requests a enumeration of all maps to be loaded
