Updated FEditorViewportClient::IsFlightCameraActive() to return true if user is moving the camera with WASD and trackpad and used that instead of simple IsUsingTrackpad to prevent changing widget mode on W/E key down events.

diff --git a/Engine/Source/Editor/UnrealEd/Private/EditorViewportClient.cpp b/Engine/Source/Editor/UnrealEd/Private/EditorViewportClient.cpp
--- a/Engine/Source/Editor/UnrealEd/Private/EditorViewportClient.cpp
+++ b/Engine/Source/Editor/UnrealEd/Private/EditorViewportClient.cpp
@@ -1695,8 +1695,10 @@ bool FEditorViewportClient::IsFlightCameraActive() const
 		|| Viewport->KeyState( FViewportNavigationCommands::Get().FovZoomIn->GetActiveGesture()->Key )
 		|| Viewport->KeyState( FViewportNavigationCommands::Get().FovZoomOut->GetActiveGesture()->Key ) );
 
+	const bool bIsUsingTrackpad = FSlateApplication::Get().IsUsingTrackpad();
+
 	// Movement key pressed and automatic movement enabled
-	bIsFlightMovementKey &= GetDefault<ULevelEditorViewportSettings>()->FlightCameraControlType == WASD_Always;
+	bIsFlightMovementKey &= (GetDefault<ULevelEditorViewportSettings>()->FlightCameraControlType == WASD_Always) | bIsUsingTrackpad;
 
 	// Not using automatic movement but the flight camera is active
 	bIsFlightMovementKey |= IsFlightCameraInputModeActive() && (GetDefault<ULevelEditorViewportSettings>()->FlightCameraControlType == WASD_RMBOnly );
diff --git a/Engine/Source/Editor/UnrealEd/Private/LevelEditorViewport.cpp b/Engine/Source/Editor/UnrealEd/Private/LevelEditorViewport.cpp
--- a/Engine/Source/Editor/UnrealEd/Private/LevelEditorViewport.cpp
+++ b/Engine/Source/Editor/UnrealEd/Private/LevelEditorViewport.cpp
@@ -3058,7 +3058,7 @@ bool FLevelEditorViewportClient::IsVolumeVisibleInViewport( const AActor& Volume
 
 void FLevelEditorViewportClient::SetWidgetMode( FWidget::EWidgetMode ActivatedMode )
 {
-	if( !GLevelEditorModeTools().IsTracking() && !FSlateApplication::Get().IsUsingTrackpad() )
+	if( !GLevelEditorModeTools().IsTracking() && !IsFlightCameraActive() )
 	{
 		GLevelEditorModeTools().SetWidgetMode( ActivatedMode );
 
