Fix for broken end pie camera

diff --git a/Engine/Source/Editor/UnrealEd/Classes/Editor/EditorEngine.h b/Engine/Source/Editor/UnrealEd/Classes/Editor/EditorEngine.h
--- a/Engine/Source/Editor/UnrealEd/Classes/Editor/EditorEngine.h
+++ b/Engine/Source/Editor/UnrealEd/Classes/Editor/EditorEngine.h
@@ -1603,6 +1603,9 @@ public:
 	 */
 	virtual void EndPlayMap();
 
+	/** 
+	 * Destroy the current play session and perform miscellaneous cleanup
+	 */
 	virtual void TeardownPlaySession(FWorldContext &PieWorldContext);
 
 	/**
@@ -2383,6 +2386,13 @@ public:
 	 */
 	void CreatePIEWorldFromLogin(FWorldContext& PieWorldContext, EPlayNetMode PlayNetMode, FPieLoginStruct& DataStruct);
 
+	/*
+	 * Handler for when viewport close request is made. 
+	 *
+	 * @param InViewport the viewport being closed.
+	 */
+	void OnViewportCloseRequested(FViewport* InViewport);
+
 private:
 	/** Gets the DPI Scale for the game viewport in the editor. */
 	float GetGameViewportDPIScale(UGameViewportClient* ViewportClient) const;
@@ -2633,6 +2643,15 @@ private:
 	/** Number of currently running instances logged into an online platform */
 	int32 NumOnlinePIEInstances;
 
+	/** Cached version of the view location at the point the PIE session was ended */
+	FVector LastViewLocation;
+	
+	/** Cached version of the view location at the point the PIE session was ended */
+	FRotator LastViewRotation;
+	
+	/** Are the lastview/rotation variables valid */
+	bool	bLastViewAndLocationValid;
+
 protected:
 
 	/**
diff --git a/Engine/Source/Editor/UnrealEd/Private/PlayLevel.cpp b/Engine/Source/Editor/UnrealEd/Private/PlayLevel.cpp
--- a/Engine/Source/Editor/UnrealEd/Private/PlayLevel.cpp
+++ b/Engine/Source/Editor/UnrealEd/Private/PlayLevel.cpp
@@ -110,6 +110,9 @@ void UEditorEngine::EndPlayMap()
 	// clean up any previous Play From Here sessions
 	if ( GameViewport != NULL && GameViewport->Viewport != NULL )
 	{
+		// Remove close handler binding
+		GameViewport->OnCloseRequested().Unbind();
+
 		GameViewport->CloseRequested(GameViewport->Viewport);
 	}
 	CleanupGameViewport();
@@ -313,17 +316,15 @@ void UEditorEngine::TeardownPlaySession(FWorldContext &PieWorldContext)
 				if( !bIsSimulatingInEditor)
 				{
 					// Set the editor viewport location to match that of Play in Viewport if we aren't simulating in the editor, we have a valid player to get the location from 
-					if (SlatePlayInEditorSession->EditorPlayer.IsValid() && SlatePlayInEditorSession->EditorPlayer.Get()->PlayerController)
+					if (bLastViewAndLocationValid == true)
 					{
-						FVector ViewLocation;
-						FRotator ViewRotation;
-						SlatePlayInEditorSession->EditorPlayer.Get()->PlayerController->GetPlayerViewPoint(ViewLocation, ViewRotation);
-						Viewport->GetLevelViewportClient().SetViewLocation( ViewLocation );
+						bLastViewAndLocationValid = false;
+						Viewport->GetLevelViewportClient().SetViewLocation( LastViewLocation );
 
 						if( Viewport->GetLevelViewportClient().IsPerspective() )
 						{
 							// Rotation only matters for perspective viewports not orthographic
-							Viewport->GetLevelViewportClient().SetViewRotation( ViewRotation );
+							Viewport->GetLevelViewportClient().SetViewRotation( LastViewRotation );
 						}
 					}
 				}
@@ -1542,6 +1543,26 @@ void UEditorEngine::RequestEndPlayMap()
 	if( PlayWorld )
 	{
 		bRequestEndPlayMapQueued = true;
+
+		// Cache the postion and rotation of the camera (the controller may be destroyed before we end the pie session and we need them to preserve the camera position)
+		if (bLastViewAndLocationValid == false)
+		{
+			bLastViewAndLocationValid = false;
+			for (int32 WorldIdx = WorldList.Num() - 1; WorldIdx >= 0; --WorldIdx)
+			{
+				FWorldContext &ThisContext = WorldList[WorldIdx];
+				if (ThisContext.WorldType == EWorldType::PIE)
+				{
+					FSlatePlayInEditorInfo* const SlatePlayInEditorSession = SlatePlayInEditorMap.Find(ThisContext.ContextHandle);
+					if (SlatePlayInEditorSession != nullptr)
+					{
+						SlatePlayInEditorSession->EditorPlayer.Get()->PlayerController->GetPlayerViewPoint(LastViewLocation, LastViewRotation);
+						bLastViewAndLocationValid = true;
+						break;
+					}
+				}
+			}
+		}
 	}
 }
 
@@ -2367,6 +2388,9 @@ UGameInstance* UEditorEngine::CreatePIEGameInstance(int32 PIEInstance, bool bInS
 		GameViewport->bIsPlayInEditorViewport = true;
 		PieWorldContext->GameViewport = ViewportClient;
 
+		// Add a handler for viewport close requests
+		ViewportClient->OnCloseRequested().BindUObject(this, &UEditorEngine::OnViewportCloseRequested);
+			
 		FSlatePlayInEditorInfo& SlatePlayInEditorSession = SlatePlayInEditorMap.Add(PieWorldContext->ContextHandle, FSlatePlayInEditorInfo());
 		SlatePlayInEditorSession.DestinationSlateViewport = RequestedDestinationSlateViewport;	// Might be invalid depending how pie was launched. Code below handles this.
 		RequestedDestinationSlateViewport = NULL;
@@ -2573,6 +2597,11 @@ UGameInstance* UEditorEngine::CreatePIEGameInstance(int32 PIEInstance, bool bInS
 	return GameInstance;
 }
 
+void UEditorEngine::OnViewportCloseRequested(FViewport* InViewport)
+{
+	RequestEndPlayMap();
+}
+
 float UEditorEngine::GetGameViewportDPIScale(UGameViewportClient* ViewportClient) const
 {
 	FVector2D ViewportSize;
diff --git a/Engine/Source/Runtime/Engine/Classes/Engine/GameViewportClient.h b/Engine/Source/Runtime/Engine/Classes/Engine/GameViewportClient.h
--- a/Engine/Source/Runtime/Engine/Classes/Engine/GameViewportClient.h
+++ b/Engine/Source/Runtime/Engine/Classes/Engine/GameViewportClient.h
@@ -172,6 +172,13 @@ DECLARE_MULTICAST_DELEGATE_ThreeParams(FOnScreenshotCaptured, int32 /*Width*/, i
  */
 DECLARE_DELEGATE_FourParams(FOnPNGScreenshotCaptured, int32, int32, const TArray<FColor>&, const FString&);
 
+/**
+* Delegate type for when call is made to close a viewport
+*
+* The first parameter is the viewport being closed.
+*/
+DECLARE_DELEGATE_OneParam(FOnCloseRequested, FViewport*);
+
 UCLASS(Within=Engine, transient, config=Engine)
 class ENGINE_API UGameViewportClient : public UScriptViewportClient, public FExec
 {
@@ -552,6 +559,12 @@ public:
 		return ScreenshotCapturedDelegate;
 	}
 
+	/* Accessor for the delegate called when a viewport is asked to close. */
+	FOnCloseRequested& OnCloseRequested()
+	{
+		return CloseRequestedDelegate;
+	}
+
 	/** Return the engine show flags for this viewport */
 	virtual FEngineShowFlags* GetEngineShowFlags() 
 	{ 
@@ -731,6 +744,9 @@ private:
 	/** Delegate called at the end of the frame when a screenshot is captured */
 	static FOnScreenshotCaptured ScreenshotCapturedDelegate;
 
+	/** Delegate called when a request to close the viewport is received */
+	FOnCloseRequested CloseRequestedDelegate;
+
 	/** Data needed to display perframe stat tracking when STAT UNIT is enabled */
 	FStatUnitData* StatUnitData;
 
diff --git a/Engine/Source/Runtime/Engine/Private/GameViewportClient.cpp b/Engine/Source/Runtime/Engine/Private/GameViewportClient.cpp
--- a/Engine/Source/Runtime/Engine/Private/GameViewportClient.cpp
+++ b/Engine/Source/Runtime/Engine/Private/GameViewportClient.cpp
@@ -1237,6 +1237,9 @@ void UGameViewportClient::CloseRequested(FViewport* InViewport)
 	check(InViewport == Viewport);
 
 	FSlateApplication::Get().SetGameIsFakingTouchEvents(false);
+	
+	// broadcast close request to anyone that registered an interest
+	CloseRequestedDelegate.ExecuteIfBound(InViewport);
 
 	SetViewportFrame(NULL);
 
