Fix crash when exiting simulate. (SIE, no editor player)

diff --git a/Engine/Source/Editor/UnrealEd/Private/PlayLevel.cpp b/Engine/Source/Editor/UnrealEd/Private/PlayLevel.cpp
--- a/Engine/Source/Editor/UnrealEd/Private/PlayLevel.cpp
+++ b/Engine/Source/Editor/UnrealEd/Private/PlayLevel.cpp
@@ -1554,7 +1554,7 @@ void UEditorEngine::RequestEndPlayMap()
 				if (ThisContext.WorldType == EWorldType::PIE)
 				{
 					FSlatePlayInEditorInfo* const SlatePlayInEditorSession = SlatePlayInEditorMap.Find(ThisContext.ContextHandle);
-					if (SlatePlayInEditorSession != nullptr)
+					if (SlatePlayInEditorSession != nullptr && SlatePlayInEditorSession->EditorPlayer.IsValid())
 					{
 						SlatePlayInEditorSession->EditorPlayer.Get()->PlayerController->GetPlayerViewPoint(LastViewLocation, LastViewRotation);
 						bLastViewAndLocationValid = true;
