Stopped tick cook on the side from waiting on save package for async cooked items to cache.
Now editor will tick at a decent frame rate while building shaders / caching textures.

diff --git a/Engine/Source/Developer/Android/AndroidTargetPlatform/Private/AndroidTargetPlatform.h b/Engine/Source/Developer/Android/AndroidTargetPlatform/Private/AndroidTargetPlatform.h
--- a/Engine/Source/Developer/Android/AndroidTargetPlatform/Private/AndroidTargetPlatform.h
+++ b/Engine/Source/Developer/Android/AndroidTargetPlatform/Private/AndroidTargetPlatform.h
@@ -130,7 +130,7 @@ public:
 
 	virtual const struct FTextureLODSettings& GetTextureLODSettings( ) const override;
 
-	virtual FName GetWaveFormat( class USoundWave* Wave ) const override;
+	virtual FName GetWaveFormat( const class USoundWave* Wave ) const override;
 #endif //WITH_ENGINE
 
 	virtual bool SupportsVariants() const override;
diff --git a/Engine/Source/Developer/Android/AndroidTargetPlatform/Private/AndroidTargetPlatform.inl b/Engine/Source/Developer/Android/AndroidTargetPlatform/Private/AndroidTargetPlatform.inl
--- a/Engine/Source/Developer/Android/AndroidTargetPlatform/Private/AndroidTargetPlatform.inl
+++ b/Engine/Source/Developer/Android/AndroidTargetPlatform/Private/AndroidTargetPlatform.inl
@@ -275,7 +275,7 @@ const FTextureLODSettings& FAndroidTargetPlatform<TPlatformProperties>::GetTextu
 
 
 template<class TPlatformProperties>
-FName FAndroidTargetPlatform<TPlatformProperties>::GetWaveFormat( class USoundWave* Wave ) const
+FName FAndroidTargetPlatform<TPlatformProperties>::GetWaveFormat( const class USoundWave* Wave ) const
 {
 	static bool formatRead = false;
 	static FName NAME_FORMAT;
diff --git a/Engine/Source/Developer/HTML5/HTML5TargetPlatform/Private/HTML5TargetPlatform.cpp b/Engine/Source/Developer/HTML5/HTML5TargetPlatform/Private/HTML5TargetPlatform.cpp
--- a/Engine/Source/Developer/HTML5/HTML5TargetPlatform/Private/HTML5TargetPlatform.cpp
+++ b/Engine/Source/Developer/HTML5/HTML5TargetPlatform/Private/HTML5TargetPlatform.cpp
@@ -231,7 +231,7 @@ const struct FTextureLODSettings& FHTML5TargetPlatform::GetTextureLODSettings( )
 }
 
 
-FName FHTML5TargetPlatform::GetWaveFormat( USoundWave* Wave ) const
+FName FHTML5TargetPlatform::GetWaveFormat( const USoundWave* Wave ) const
 {
 	static FName NAME_OGG(TEXT("OGG"));
 	return NAME_OGG;
diff --git a/Engine/Source/Developer/HTML5/HTML5TargetPlatform/Private/HTML5TargetPlatform.h b/Engine/Source/Developer/HTML5/HTML5TargetPlatform/Private/HTML5TargetPlatform.h
--- a/Engine/Source/Developer/HTML5/HTML5TargetPlatform/Private/HTML5TargetPlatform.h
+++ b/Engine/Source/Developer/HTML5/HTML5TargetPlatform/Private/HTML5TargetPlatform.h
@@ -72,7 +72,7 @@ public:
 
 	virtual const struct FTextureLODSettings& GetTextureLODSettings( ) const override;
 
-	virtual FName GetWaveFormat( class USoundWave* Wave ) const override;
+	virtual FName GetWaveFormat( const class USoundWave* Wave ) const override;
 #endif // WITH_ENGINE
 
 	DECLARE_DERIVED_EVENT(FHTML5TargetPlatform, ITargetPlatform::FOnTargetDeviceDiscovered, FOnTargetDeviceDiscovered);
diff --git a/Engine/Source/Developer/Linux/LinuxTargetPlatform/Private/LinuxTargetPlatform.h b/Engine/Source/Developer/Linux/LinuxTargetPlatform/Private/LinuxTargetPlatform.h
--- a/Engine/Source/Developer/Linux/LinuxTargetPlatform/Private/LinuxTargetPlatform.h
+++ b/Engine/Source/Developer/Linux/LinuxTargetPlatform/Private/LinuxTargetPlatform.h
@@ -172,7 +172,7 @@ public:
 		return TextureLODSettings;
 	}
 
-	virtual FName GetWaveFormat( class USoundWave* Wave ) const override
+	virtual FName GetWaveFormat( const class USoundWave* Wave ) const override
 	{
 		static FName NAME_OGG(TEXT("OGG"));
 		static FName NAME_OPUS(TEXT("OPUS"));
diff --git a/Engine/Source/Developer/Mac/MacTargetPlatform/Private/GenericMacTargetPlatform.h b/Engine/Source/Developer/Mac/MacTargetPlatform/Private/GenericMacTargetPlatform.h
--- a/Engine/Source/Developer/Mac/MacTargetPlatform/Private/GenericMacTargetPlatform.h
+++ b/Engine/Source/Developer/Mac/MacTargetPlatform/Private/GenericMacTargetPlatform.h
@@ -136,7 +136,7 @@ return TSuper::SupportsFeature(Feature);
 	}
 
 
-	virtual FName GetWaveFormat( class USoundWave* Wave ) const override
+	virtual FName GetWaveFormat( const class USoundWave* Wave ) const override
 	{
 		static FName NAME_OGG(TEXT("OGG"));
 
diff --git a/Engine/Source/Developer/TargetPlatform/Public/Interfaces/ITargetPlatform.h b/Engine/Source/Developer/TargetPlatform/Public/Interfaces/ITargetPlatform.h
--- a/Engine/Source/Developer/TargetPlatform/Public/Interfaces/ITargetPlatform.h
+++ b/Engine/Source/Developer/TargetPlatform/Public/Interfaces/ITargetPlatform.h
@@ -347,7 +347,7 @@ public:
 	 * @param Wave The sound node wave to get the format for.
 	 * @return Name of the wave format.
 	 */
-	virtual FName GetWaveFormat( class USoundWave* Wave ) const = 0;
+	virtual FName GetWaveFormat( const class USoundWave* Wave ) const = 0;
 
 	/**
 	 * Gets the texture LOD settings used by this platform.
diff --git a/Engine/Source/Developer/Windows/WindowsTargetPlatform/Private/GenericWindowsTargetPlatform.h b/Engine/Source/Developer/Windows/WindowsTargetPlatform/Private/GenericWindowsTargetPlatform.h
--- a/Engine/Source/Developer/Windows/WindowsTargetPlatform/Private/GenericWindowsTargetPlatform.h
+++ b/Engine/Source/Developer/Windows/WindowsTargetPlatform/Private/GenericWindowsTargetPlatform.h
@@ -152,7 +152,7 @@ public:
 		return TextureLODSettings;
 	}
 
-	virtual FName GetWaveFormat( class USoundWave* Wave ) const override
+	virtual FName GetWaveFormat( const class USoundWave* Wave ) const override
 	{
 		static FName NAME_OGG(TEXT("OGG"));
 		static FName NAME_OPUS(TEXT("OPUS"));
diff --git a/Engine/Source/Developer/iOS/IOSTargetPlatform/Private/IOSTargetPlatform.cpp b/Engine/Source/Developer/iOS/IOSTargetPlatform/Private/IOSTargetPlatform.cpp
--- a/Engine/Source/Developer/iOS/IOSTargetPlatform/Private/IOSTargetPlatform.cpp
+++ b/Engine/Source/Developer/iOS/IOSTargetPlatform/Private/IOSTargetPlatform.cpp
@@ -435,7 +435,7 @@ const FTextureLODSettings& FIOSTargetPlatform::GetTextureLODSettings() const
 }
 
 
-FName FIOSTargetPlatform::GetWaveFormat( class USoundWave* Wave ) const
+FName FIOSTargetPlatform::GetWaveFormat( const class USoundWave* Wave ) const
 {
 	static FName NAME_ADPCM(TEXT("ADPCM"));
 	return NAME_ADPCM;
diff --git a/Engine/Source/Developer/iOS/IOSTargetPlatform/Private/IOSTargetPlatform.h b/Engine/Source/Developer/iOS/IOSTargetPlatform/Private/IOSTargetPlatform.h
--- a/Engine/Source/Developer/iOS/IOSTargetPlatform/Private/IOSTargetPlatform.h
+++ b/Engine/Source/Developer/iOS/IOSTargetPlatform/Private/IOSTargetPlatform.h
@@ -106,7 +106,7 @@ public:
 
 	virtual const struct FTextureLODSettings& GetTextureLODSettings( ) const override;
 
-	virtual FName GetWaveFormat( class USoundWave* Wave ) const override;
+	virtual FName GetWaveFormat( const class USoundWave* Wave ) const override;
 #endif // WITH_ENGINE
 
 
diff --git a/Engine/Source/Editor/UnrealEd/Classes/CookOnTheSide/CookOnTheFlyServer.h b/Engine/Source/Editor/UnrealEd/Classes/CookOnTheSide/CookOnTheFlyServer.h
--- a/Engine/Source/Editor/UnrealEd/Classes/CookOnTheSide/CookOnTheFlyServer.h
+++ b/Engine/Source/Editor/UnrealEd/Classes/CookOnTheSide/CookOnTheFlyServer.h
@@ -447,13 +447,14 @@ private:
 			if ( ForceEnqueFront )
 			{
 				int32 Index = Queue.Find(Request.GetFilename());
+				check( Index != INDEX_NONE );
 				if ( Index != 0 )
 				{
 					Queue.Swap(0, Index);
 				}
 			}
 		}
-		
+
 		bool Dequeue(FFilePlatformRequest* Result)
 		{
 			FScopeLock ScopeLock( &SynchronizationObject );
@@ -566,7 +567,8 @@ private:
 	public:
 		FCookByTheBookOptions() : bGenerateStreamingInstallManifests(false),
 			bRunning(false),
-			CookTime( 0.0f )
+			CookTime( 0.0 ),
+			CookStartTime( 0.0 )
 		{ }
 
 		/** Should we generate streaming install manifests (only valid option in cook by the book) */
@@ -582,6 +584,7 @@ private:
 		/** If a cook is cancelled next cook will need to resume cooking */ 
 		TArray<FFilePlatformRequest> PreviousCookRequests; 
 		double CookTime;
+		double CookStartTime;
 	};
 	FCookByTheBookOptions* CookByTheBookOptions;
 	
@@ -697,7 +700,13 @@ public:
 	uint32 NumConnections() const;
 
 	
-
+	/**
+	 * Is this cooker running in real time mode (where it needs to respect the timeslice) 
+	 * 
+	 * @return returns true if this cooker is running in realtime mode like in the editor
+	 */
+	bool IsRealtimeMode();
+	
 
 	virtual void BeginDestroy() override;
 	
diff --git a/Engine/Source/Editor/UnrealEd/Private/CookOnTheFlyServer.cpp b/Engine/Source/Editor/UnrealEd/Private/CookOnTheFlyServer.cpp
--- a/Engine/Source/Editor/UnrealEd/Private/CookOnTheFlyServer.cpp
+++ b/Engine/Source/Editor/UnrealEd/Private/CookOnTheFlyServer.cpp
@@ -29,12 +29,13 @@
 #include "TokenizedMessage.h"
 #include "MessageLog.h"
 
-
+// shader compiler processAsyncResults
+#include "ShaderCompiler.h"
 
 DEFINE_LOG_CATEGORY_STATIC(LogCookOnTheFly, Log, All);
 
 #define DEBUG_COOKONTHEFLY 0
-#define OUTPUT_TIMING 0
+#define OUTPUT_TIMING 1
 
 #if OUTPUT_TIMING
 
@@ -779,12 +780,34 @@ void UCookOnTheFlyServer::GenerateManifestInfo( UPackage* Package, const TArray<
 	}
 }
 
+bool UCookOnTheFlyServer::IsRealtimeMode()
+{
+	return CurrentCookMode == ECookMode::CookByTheBookFromTheEditor || CurrentCookMode == ECookMode::CookByTheBook;
+}
 
 uint32 UCookOnTheFlyServer::TickCookOnTheSide( const float TimeSlice, uint32 &CookedPackageCount )
 {
-	double StartTime = FPlatformTime::Seconds();
-	uint32 Result = 0;
+	struct FCookerTimer
+	{
+		const bool bIsRealtimeMode;
+		const double StartTime;
+		const float &TimeSlice;
+		FCookerTimer(const float &InTimeSlice, bool bInIsRealtimeMode ) : 
+			bIsRealtimeMode( bInIsRealtimeMode), StartTime(FPlatformTime::Seconds()), TimeSlice(InTimeSlice)
+		{
+		}
+		double GetTimeTillNow()
+		{
+			return FPlatformTime::Seconds() - StartTime;
+		}
+		bool IsTimeUp()
+		{
+			return (FPlatformTime::Seconds() - StartTime) > TimeSlice && bIsRealtimeMode;
+		}
+	};
+	FCookerTimer Timer(TimeSlice, IsRealtimeMode());
 
+	uint32 Result = 0;
 	// This is all the target platforms which we needed to process requests for this iteration
 	// we use this in the unsolicited packages processing below
 	TArray<FName> AllTargetPlatformNames;
@@ -863,7 +886,9 @@ uint32 UCookOnTheFlyServer::TickCookOnTheSide( const float TimeSlice, uint32 &Co
 		
 		if ( bShouldCook ) // if we should cook the package then cook it otherwise add it to the list of already cooked packages below
 		{
-			FString PackageName;	
+			SCOPE_TIMER(AllOfLoadPackage);
+
+			FString PackageName;
 			UPackage *Package = NULL;
 			if ( FPackageName::TryConvertFilenameToLongPackageName(BuildFilename, PackageName) )
 			{
@@ -998,6 +1023,55 @@ uint32 UCookOnTheFlyServer::TickCookOnTheSide( const float TimeSlice, uint32 &Co
 			GRedirectCollector.ResolveStringAssetReference();
 		}
 
+
+		bool bIsAllDataCached = true;
+		if ( PackagesToSave.Num() )
+		{
+			SCOPE_TIMER(CallBeginCacheForCookedPlatformData);
+			// cache the resources for this package for each platform
+			ITargetPlatformManagerModule& TPM = GetTargetPlatformManagerRef();
+			TArray<UObject*> ObjectsInPackage;
+			GetObjectsWithOuter( PackagesToSave[0], ObjectsInPackage );
+			for ( const auto& TargetPlatformName : AllTargetPlatformNames )
+			{
+				const ITargetPlatform* TargetPlatform = TPM.FindTargetPlatform( TargetPlatformName.ToString() );
+				
+				for ( const auto& Obj : ObjectsInPackage )
+				{
+					if ( Timer.IsTimeUp() && IsRealtimeMode() )
+					{
+						bIsAllDataCached = false;
+						break;
+					}
+
+					Obj->BeginCacheForCookedPlatformData( TargetPlatform );
+					if ( Obj->IsCachedCookedPlatformDataLoaded(TargetPlatform) == false )
+					{
+						bIsAllDataCached = false;
+						break;
+					}
+				}
+
+				if ( bIsAllDataCached == false )
+					break;
+			}
+
+			if ( !bIsAllDataCached )
+			{
+				GShaderCompilingManager->ProcessAsyncResults(true, false, true);
+			}
+		}
+
+		if ( IsRealtimeMode() )
+		{
+			if ( bIsAllDataCached == false )
+			{
+				// reque the current task and process it next tick
+				CookRequests.EnqueueUnique(ToBuild, true);
+				break;
+			}
+		}
+
 		int32 FirstUnsolicitedPackage = PackagesToSave.Num();
 
 		// generate a list of other packages which were loaded with this one
@@ -1084,7 +1158,6 @@ uint32 UCookOnTheFlyServer::TickCookOnTheSide( const float TimeSlice, uint32 &Co
 			}
 		}
 
-
 		{
 			SCOPE_TIMER(GenerateManifestInfo);
 			// update manifest with cooked package info
@@ -1109,7 +1182,7 @@ uint32 UCookOnTheFlyServer::TickCookOnTheSide( const float TimeSlice, uint32 &Co
 						bShouldFinishTick = true;
 					}
 
-					if ( (FPlatformTime::Seconds() - StartTime) > TimeSlice)
+					if ( Timer.IsTimeUp() )
 					{
 						bShouldFinishTick = true;
 						// our timeslice is up
@@ -1237,7 +1310,7 @@ uint32 UCookOnTheFlyServer::TickCookOnTheSide( const float TimeSlice, uint32 &Co
 		}
 		}*/
 
-		if ( (FPlatformTime::Seconds() - StartTime) > TimeSlice)
+		if ( Timer.IsTimeUp() )
 		{
 			break;
 		}
@@ -1247,7 +1320,7 @@ uint32 UCookOnTheFlyServer::TickCookOnTheSide( const float TimeSlice, uint32 &Co
 
 	if ( CookByTheBookOptions )
 	{
-		CookByTheBookOptions->CookTime += FPlatformTime::Seconds() - StartTime;
+		CookByTheBookOptions->CookTime += Timer.GetTimeTillNow();
 	}
 
 
@@ -2068,8 +2141,8 @@ void UCookOnTheFlyServer::CookByTheBookFinished()
 	}
 
 	CookByTheBookOptions->LastGCItems.Empty();
-
-	UE_LOG(LogCookOnTheFly, Display, TEXT("Cook by the book total time %fs"), CookByTheBookOptions->CookTime);
+	const float TotalCookTime = (float)(FPlatformTime::Seconds() - CookByTheBookOptions->CookStartTime);
+	UE_LOG(LogCookOnTheFly, Display, TEXT("Cook by the book total time in tick %fs total time %f"), CookByTheBookOptions->CookTime, TotalCookTime);
 
 	CookByTheBookOptions->bRunning = false;
 }
@@ -2150,6 +2223,7 @@ void UCookOnTheFlyServer::StartCookByTheBook(const TArray<ITargetPlatform*>& Tar
 	CookByTheBookOptions->bRunning = true;
 	CookByTheBookOptions->bCancel = false;
 	CookByTheBookOptions->CookTime = 0.0f;
+	CookByTheBookOptions->CookStartTime = FPlatformTime::Seconds();
 
 	bool bCookAll = false;
 	bool bMapsOnly = false;
diff --git a/Engine/Source/Runtime/CoreUObject/Public/UObject/UObject.h b/Engine/Source/Runtime/CoreUObject/Public/UObject/UObject.h
--- a/Engine/Source/Runtime/CoreUObject/Public/UObject/UObject.h
+++ b/Engine/Source/Runtime/CoreUObject/Public/UObject/UObject.h
@@ -568,14 +568,14 @@ public:
 	 * 
 	 * @param	TargetPlatform	target platform to cache platform specific data for
 	 */
-	virtual void BeginCacheForCookedPlatformData( const ITargetPlatform *TargetPlatform ) {  }
+	virtual void BeginCacheForCookedPlatformData( const ITargetPlatform* TargetPlatform ) {  }
 
 	/**
 	 * Clears cached cooked platform data for specific platform
 	 * 
 	 * @param	TargetPlatform	target platform to cache platform specific data for
 	 */
-	virtual void ClearCachedCookedPlatformData( const ITargetPlatform *TargetPlatform ) {  }
+	virtual void ClearCachedCookedPlatformData( const ITargetPlatform* TargetPlatform ) {  }
 
 	/**
 	 * Clear all cached cooked platform data
@@ -584,6 +584,13 @@ public:
 	 */
 	virtual void ClearAllCachedCookedPlatformData() { }
 
+	/**
+	 * Have we finished loading all the cooked platform data for the target platforms requested in BeginCacheForCookedPlatformData
+	 * 
+	 * @param	TargetPlatform target platform to check for cooked platform data
+	 */
+	virtual bool IsCachedCookedPlatformDataLoaded( const ITargetPlatform* TargetPlatform ) { return true; }
+
 
 	/**
 	 * Determine if this object has SomeObject in its archetype chain.
diff --git a/Engine/Source/Runtime/Engine/Classes/Engine/Texture.h b/Engine/Source/Runtime/Engine/Classes/Engine/Texture.h
--- a/Engine/Source/Runtime/Engine/Classes/Engine/Texture.h
+++ b/Engine/Source/Runtime/Engine/Classes/Engine/Texture.h
@@ -593,8 +593,14 @@ public:
 	/**
 	 * Begins caching platform data in the background for the platform requested
 	 */
-	ENGINE_API virtual void BeginCacheForCookedPlatformData(  const ITargetPlatform *TargetPlatform ) override;
+	ENGINE_API virtual void BeginCacheForCookedPlatformData(  const ITargetPlatform* TargetPlatform ) override;
 
+	/**
+	 * Have we finished loading all the cooked platform data for the target platforms requested in BeginCacheForCookedPlatformData
+	 * 
+	 * @param	TargetPlatform target platform to check for cooked platform data
+	 */
+	ENGINE_API virtual bool IsCachedCookedPlatformDataLoaded( const ITargetPlatform* TargetPlatform ) override;
 
 	/**
 	 * Begins caching platform data in the background.
diff --git a/Engine/Source/Runtime/Engine/Classes/Materials/Material.h b/Engine/Source/Runtime/Engine/Classes/Materials/Material.h
--- a/Engine/Source/Runtime/Engine/Classes/Materials/Material.h
+++ b/Engine/Source/Runtime/Engine/Classes/Materials/Material.h
@@ -796,6 +796,7 @@ public:
 	ENGINE_API virtual void PostDuplicate(bool bDuplicateForPIE) override;
 	ENGINE_API virtual void PostLoad() override;
 	ENGINE_API virtual void BeginCacheForCookedPlatformData( const ITargetPlatform *TargetPlatform ) override;
+	ENGINE_API virtual bool IsCachedCookedPlatformDataLoaded( const ITargetPlatform* TargetPlatform ) override;
 	ENGINE_API virtual void ClearCachedCookedPlatformData( const ITargetPlatform *TargetPlatform ) override;
 	ENGINE_API virtual void ClearAllCachedCookedPlatformData() override;
 
diff --git a/Engine/Source/Runtime/Engine/Classes/Materials/MaterialInstance.h b/Engine/Source/Runtime/Engine/Classes/Materials/MaterialInstance.h
--- a/Engine/Source/Runtime/Engine/Classes/Materials/MaterialInstance.h
+++ b/Engine/Source/Runtime/Engine/Classes/Materials/MaterialInstance.h
@@ -272,6 +272,7 @@ public:
 	virtual SIZE_T GetResourceSize(EResourceSizeMode::Type Mode) override;
 	virtual void PostInitProperties() override;	
 	virtual void BeginCacheForCookedPlatformData(  const ITargetPlatform *TargetPlatform ) override;
+	virtual bool IsCachedCookedPlatformDataLoaded(  const ITargetPlatform* TargetPlatform ) override;
 	virtual void ClearCachedCookedPlatformData( const ITargetPlatform *TargetPlatform ) override;
 	virtual void ClearAllCachedCookedPlatformData() override;
 	virtual void Serialize(FArchive& Ar) override;
diff --git a/Engine/Source/Runtime/Engine/Classes/Sound/SoundWave.h b/Engine/Source/Runtime/Engine/Classes/Sound/SoundWave.h
--- a/Engine/Source/Runtime/Engine/Classes/Sound/SoundWave.h
+++ b/Engine/Source/Runtime/Engine/Classes/Sound/SoundWave.h
@@ -101,6 +101,7 @@ struct FStreamedAudioPlatformData
 #if WITH_EDITORONLY_DATA
 	void Cache(class USoundWave& InSoundWave, FName AudioFormatName, uint32 InFlags);
 	void FinishCache();
+	bool IsFinishedCache() const;
 	ENGINE_API bool TryInlineChunkData();
 	bool AreDerivedChunksAvailable() const;
 #endif
@@ -389,6 +390,8 @@ public:
 	 */
 	virtual void BeginCacheForCookedPlatformData(  const ITargetPlatform *TargetPlatform ) override;
 
+	virtual bool IsCachedCookedPlatformDataLoaded( const ITargetPlatform* TargetPlatform ) override;
+
 	/**
 	 * Begins caching platform data in the background.
 	 */
diff --git a/Engine/Source/Runtime/Engine/Private/AudioDerivedData.cpp b/Engine/Source/Runtime/Engine/Private/AudioDerivedData.cpp
--- a/Engine/Source/Runtime/Engine/Private/AudioDerivedData.cpp
+++ b/Engine/Source/Runtime/Engine/Private/AudioDerivedData.cpp
@@ -414,6 +414,11 @@ void FStreamedAudioPlatformData::Cache(USoundWave& InSoundWave, FName AudioForma
 	}
 }
 
+bool FStreamedAudioPlatformData::IsFinishedCache() const
+{
+	return AsyncTask == NULL ? true : false;
+}
+
 void FStreamedAudioPlatformData::FinishCache()
 {
 	if (AsyncTask)
@@ -1005,6 +1010,28 @@ void USoundWave::BeginCacheForCookedPlatformData(const ITargetPlatform *TargetPl
 	Super::BeginCacheForCookedPlatformData(TargetPlatform);
 }
 
+bool USoundWave::IsCachedCookedPlatformDataLoaded( const ITargetPlatform* TargetPlatform ) 
+{
+	if (TargetPlatform->SupportsFeature(ETargetPlatformFeatures::AudioStreaming) && IsStreaming())
+	{
+		// Retrieve format to cache for targetplatform.
+		FName PlatformFormat = TargetPlatform->GetWaveFormat(this);
+
+		// find format data by comparing derived data keys.
+		FString DerivedDataKey;
+		GetStreamedAudioDerivedDataKeySuffix(*this, PlatformFormat, DerivedDataKey);
+
+		FStreamedAudioPlatformData *PlatformData = CookedPlatformData.FindRef(DerivedDataKey);
+		if (PlatformData == NULL)
+		{
+			// we havne't called begincache
+			return false;
+		}
+		return PlatformData->IsFinishedCache();
+	}
+	return true; 
+}
+
 void USoundWave::FinishCachePlatformData()
 {
 	if (RunningPlatformData == NULL)
diff --git a/Engine/Source/Runtime/Engine/Private/Materials/Material.cpp b/Engine/Source/Runtime/Engine/Private/Materials/Material.cpp
--- a/Engine/Source/Runtime/Engine/Private/Materials/Material.cpp
+++ b/Engine/Source/Runtime/Engine/Private/Materials/Material.cpp
@@ -2157,8 +2157,26 @@ void UMaterial::BeginCacheForCookedPlatformData( const ITargetPlatform *TargetPl
 	}
 }
 
+bool UMaterial::IsCachedCookedPlatformDataLoaded( const ITargetPlatform* TargetPlatform ) 
+{
+	TArray<FName> DesiredShaderFormats;
+	TargetPlatform->GetAllTargetedShaderFormats(DesiredShaderFormats);
 
+	const TArray<FMaterialResource*>* CachedMaterialResourcesForPlatform = CachedMaterialResourcesForCooking.Find( TargetPlatform );
 
+	if ( CachedMaterialResourcesForPlatform != NULL ) // this should always succeed if begincacheforcookedcplatformdata is called first
+	{
+		for ( const auto& MaterialResource : *CachedMaterialResourcesForPlatform )
+		{
+			if ( MaterialResource->IsCompilationFinished() == false )
+			{
+				return false;
+			}
+		}
+		return true;
+	}
+	return false;
+}
 
 void UMaterial::ClearCachedCookedPlatformData( const ITargetPlatform *TargetPlatform )
 {
diff --git a/Engine/Source/Runtime/Engine/Private/Materials/MaterialInstance.cpp b/Engine/Source/Runtime/Engine/Private/Materials/MaterialInstance.cpp
--- a/Engine/Source/Runtime/Engine/Private/Materials/MaterialInstance.cpp
+++ b/Engine/Source/Runtime/Engine/Private/Materials/MaterialInstance.cpp
@@ -1558,6 +1558,23 @@ void UMaterialInstance::BeginCacheForCookedPlatformData( const ITargetPlatform *
 	}
 }
 
+bool UMaterialInstance::IsCachedCookedPlatformDataLoaded( const ITargetPlatform* TargetPlatform ) 
+{
+	const TArray<FMaterialResource*> *CachedMaterialResourcesForPlatform = CachedMaterialResourcesForCooking.Find( TargetPlatform );
+	if ( CachedMaterialResourcesForPlatform != NULL )
+	{
+		for ( const auto& MaterialResource : *CachedMaterialResourcesForPlatform )
+		{
+			if ( MaterialResource->IsCompilationFinished() == false )
+				return false;
+		}
+
+		return true;
+	}
+	return true; // this happens if we haven't started caching (begincache hasn't been called yet)
+}
+
+
 void UMaterialInstance::ClearCachedCookedPlatformData( const ITargetPlatform *TargetPlatform )
 {
 	TArray<FMaterialResource*> *CachedMaterialResourcesForPlatform = CachedMaterialResourcesForCooking.Find( TargetPlatform );
diff --git a/Engine/Source/Runtime/Engine/Private/Materials/MaterialShared.cpp b/Engine/Source/Runtime/Engine/Private/Materials/MaterialShared.cpp
--- a/Engine/Source/Runtime/Engine/Private/Materials/MaterialShared.cpp
+++ b/Engine/Source/Runtime/Engine/Private/Materials/MaterialShared.cpp
@@ -281,12 +281,26 @@ void FMaterial::GetShaderMapIDsWithUnfinishedCompilation(TArray<int32>& ShaderMa
 	{
 		ShaderMapIds.Add(GameThreadShaderMap->GetCompilingId());
 	}
-	else if (OutstandingCompileShaderMapId != INDEX_NONE)
+	else if (OutstandingCompileShaderMapIds.Num() != 0 )
 	{
-		ShaderMapIds.Add(OutstandingCompileShaderMapId);
+		ShaderMapIds.Append(OutstandingCompileShaderMapIds);
 	}
 }
 
+bool FMaterial::IsCompilationFinished()
+{
+	// Build an array of the shader map Id's are not finished compiling.
+	if (GameThreadShaderMap && !GameThreadShaderMap->IsCompilationFinalized())
+	{
+		return false;
+	}
+	else if (OutstandingCompileShaderMapIds.Num() != 0 )
+	{
+		return false;
+	}
+	return true;
+}
+
 void FMaterial::CancelCompilation()
 {
 	TArray<int32> ShaderMapIdsToCancel;
@@ -1217,7 +1231,6 @@ bool FMaterial::CacheShaders(EShaderPlatform Platform, bool bApplyCompletedShade
 bool FMaterial::CacheShaders(const FMaterialShaderMapId& ShaderMapId, EShaderPlatform Platform, bool bApplyCompletedShaderMapForRendering)
 {
 	bool bSucceeded = false;
-	OutstandingCompileShaderMapId = INDEX_NONE;
 
 	check(ShaderMapId.BaseMaterialId.IsValid());
 
@@ -1370,7 +1383,8 @@ bool FMaterial::BeginCompileShaderMap(
 		}
 		else
 		{
-			OutstandingCompileShaderMapId = NewShaderMap->GetCompilingId();
+			UE_LOG(LogMaterial, Display, TEXT("Adding new outstanding shader map compilation id %d."), NewShaderMap->GetCompilingId() );
+			OutstandingCompileShaderMapIds.Add( NewShaderMap->GetCompilingId() );
 			// Async compile, use NULL so that rendering will fall back to the default material.
 			OutShaderMap = NULL;
 		}
diff --git a/Engine/Source/Runtime/Engine/Private/ShaderCompiler/ShaderCompiler.cpp b/Engine/Source/Runtime/Engine/Private/ShaderCompiler/ShaderCompiler.cpp
--- a/Engine/Source/Runtime/Engine/Private/ShaderCompiler/ShaderCompiler.cpp
+++ b/Engine/Source/Runtime/Engine/Private/ShaderCompiler/ShaderCompiler.cpp
@@ -1408,6 +1408,13 @@ void FShaderCompilingManager::ProcessCompiledShaderMaps(
 				bShaderMapComplete = ShaderMap->ProcessCompilationResults(ResultArray, CompileResults.FinalizeJobIndex, TimeBudget);
 			}
 
+			for ( auto& Material : MaterialsArray )
+			{
+				Material->RemoveOutstandingCompileId( ShaderMap->CompilingId );
+			}
+
+			UE_LOG(LogShaderCompilers, Display, TEXT("Received shader map compiling id %d"), ShaderMap->CompilingId );
+
 			if (bShaderMapComplete)
 			{
 				ShaderMap->bCompiledSuccessfully = bSuccess;
@@ -1445,6 +1452,8 @@ void FShaderCompilingManager::ProcessCompiledShaderMaps(
 						// Propagate error messages
 						CurrentMaterial.CompileErrors = Errors;
 
+						MaterialsToUpdate.Add( &CurrentMaterial, NULL );
+
 						if (CurrentMaterial.IsDefaultMaterial())
 						{
 							// Log the errors unsuppressed before the fatal error, so it's always obvious from the log what the compile error was
@@ -1510,7 +1519,6 @@ void FShaderCompilingManager::ProcessCompiledShaderMaps(
 			check(!ShaderMap || ShaderMap->IsValidForRendering());
 
 			Material->SetGameThreadShaderMap(It.Value());
-			Material->ClearOutstandingCompileId();
 		}
 
 		const TSet<FSceneInterface*>& AllocatedScenes = GetRendererModule().GetAllocatedScenes();
@@ -1901,7 +1909,7 @@ void FShaderCompilingManager::FinishAllCompilation()
 	UE_LOG(LogShaders, Log, TEXT("FinishAllCompilation %.3fs"), (float)(EndTime - StartTime));
 }
 
-void FShaderCompilingManager::ProcessAsyncResults(bool bLimitExecutionTime, bool bBlockOnGlobalShaderCompletion)
+void FShaderCompilingManager::ProcessAsyncResults(bool bLimitExecutionTime, bool bBlockOnGlobalShaderCompletion, bool bProcessShaderMapNotForRendering)
 {
 	if (bAllowAsynchronousShaderCompiling)
 	{
@@ -1938,7 +1946,7 @@ void FShaderCompilingManager::ProcessAsyncResults(bool bLimitExecutionTime, bool
 
 					if (Results.FinishedJobs.Num() == Results.NumJobsQueued 
 						// Shader maps that don't want to be applied can only be processed by a call to FinishCompilation
-						&& Results.bApplyCompletedShaderMapForRendering)
+						&& (Results.bApplyCompletedShaderMapForRendering || bProcessShaderMapNotForRendering) )
 					{
 						PendingFinalizeShaderMaps.Add(It.Key(), FShaderMapFinalizeResults(Results));
 						ShaderMapsToRemove.Add(It.Key());
diff --git a/Engine/Source/Runtime/Engine/Private/TextureDerivedData.cpp b/Engine/Source/Runtime/Engine/Private/TextureDerivedData.cpp
--- a/Engine/Source/Runtime/Engine/Private/TextureDerivedData.cpp
+++ b/Engine/Source/Runtime/Engine/Private/TextureDerivedData.cpp
@@ -1473,6 +1473,37 @@ void UTexture::BeginCacheForCookedPlatformData( const ITargetPlatform *TargetPla
 	}
 }
 
+bool UTexture::IsCachedCookedPlatformDataLoaded( const ITargetPlatform* TargetPlatform ) 
+{ 
+	const TMap<FString, FTexturePlatformData*> *CookedPlatformDataPtr = GetCookedPlatformData();
+	if ( CookedPlatformDataPtr == NULL )
+		return true; // we should always have cookedplatformDataPtr in the case of WITH_EDITOR
+
+	FTextureBuildSettings BuildSettings;
+	TArray<FName> PlatformFormats;
+	TArray<FTexturePlatformData*> PlatformDataToSerialize;
+
+	GetTextureBuildSettings(*this, TargetPlatform->GetTextureLODSettings(), BuildSettings);
+	TargetPlatform->GetTextureFormats(this, PlatformFormats);
+
+	for (int32 FormatIndex = 0; FormatIndex < PlatformFormats.Num(); FormatIndex++)
+	{
+		FString DerivedDataKey;
+		BuildSettings.TextureFormatName = PlatformFormats[FormatIndex];
+		GetTextureDerivedDataKey(*this, BuildSettings, DerivedDataKey);
+
+		FTexturePlatformData *PlatformData= (*CookedPlatformDataPtr).FindRef(DerivedDataKey);
+
+		if (PlatformData)
+		{
+			if ( (PlatformData->AsyncTask != NULL) && ( PlatformData->AsyncTask->IsWorkDone() == false ) )
+				return false;
+		}
+	}
+	// if we get here all our stuff is cached :)
+	return true;
+}
+
 bool UTexture::IsAsyncCacheComplete()
 {
 	bool bComplete = true;
@@ -1486,8 +1517,6 @@ bool UTexture::IsAsyncCacheComplete()
 		}
 	}
 
-	
-
 	TMap<FString,FTexturePlatformData*>* CookedPlatformDataPtr = GetCookedPlatformData();
 	if (CookedPlatformDataPtr)
 	{
diff --git a/Engine/Source/Runtime/Engine/Public/MaterialShared.h b/Engine/Source/Runtime/Engine/Public/MaterialShared.h
--- a/Engine/Source/Runtime/Engine/Public/MaterialShared.h
+++ b/Engine/Source/Runtime/Engine/Public/MaterialShared.h
@@ -817,7 +817,6 @@ public:
 	FMaterial():
 		RenderingThreadShaderMap(NULL),
 		Id_DEPRECATED(0,0,0,0),
-		OutstandingCompileShaderMapId(INDEX_NONE),
 		QualityLevel(EMaterialQualityLevel::High),
 		bQualityLevelHasDifferentNodes(false),
 		FeatureLevel(ERHIFeatureLevel::SM4),
@@ -940,6 +939,13 @@ public:
 	 */
 	ENGINE_API void FinishCompilation();
 
+	/**
+	 * Checks if the compilation for this shader is finished
+	 * 
+	 * @return returns true if compilation is complete false otherwise
+	 */
+	ENGINE_API bool IsCompilationFinished();
+
 	EMaterialQualityLevel::Type GetQualityLevel() const 
 	{
 		return QualityLevel;
@@ -997,9 +1003,10 @@ public:
 	/** Note: SetGameThreadShaderMap must also be called with the same value, but from the game thread. */
 	ENGINE_API void SetRenderingThreadShaderMap(FMaterialShaderMap* InMaterialShaderMap);
 
-	void ClearOutstandingCompileId()
+	void RemoveOutstandingCompileId(const int32 OldOutstandingCompileShaderMapId )
 	{
-		OutstandingCompileShaderMapId = INDEX_NONE;
+		OutstandingCompileShaderMapIds.Remove( OldOutstandingCompileShaderMapId );
+		// UE_LOG(LogMaterial, Display, TEXT("Removing compile shader map id %d"), OldOutstandingCompileShaderMapId);
 	}
 
 	void AddReferencedObjects(FReferenceCollector& Collector);
@@ -1122,7 +1129,7 @@ private:
 	 * Contains the compiling id of this shader map when it is being compiled asynchronously. 
 	 * This can be used to access the shader map during async compiling, since GameThreadShaderMap will not have been set yet.
 	 */
-	int32 OutstandingCompileShaderMapId;
+	TArray<int32, TInlineAllocator<1> > OutstandingCompileShaderMapIds;
 
 	/** Quality level that this material is representing. */
 	EMaterialQualityLevel::Type QualityLevel;
diff --git a/Engine/Source/Runtime/Engine/Public/ShaderCompiler.h b/Engine/Source/Runtime/Engine/Public/ShaderCompiler.h
--- a/Engine/Source/Runtime/Engine/Public/ShaderCompiler.h
+++ b/Engine/Source/Runtime/Engine/Public/ShaderCompiler.h
@@ -308,7 +308,7 @@ public:
 	 * @param bBlockOnGlobalShaderCompletion - When enabled, ProcessAsyncResults will block until global shader maps are complete.
 	 *		This must be done before using global shaders for rendering.
 	 */
-	ENGINE_API void ProcessAsyncResults(bool bLimitExecutionTime, bool bBlockOnGlobalShaderCompletion);
+	ENGINE_API void ProcessAsyncResults(bool bLimitExecutionTime, bool bBlockOnGlobalShaderCompletion, bool bProcessShaderMapNotForRendering = false);
 
 	/**
 	 * Returns true if the given shader compile worker is still running.
