Extra check for null player controller in RequestEndPlayMap and removed pointless bool assignment

diff --git a/Engine/Source/Editor/UnrealEd/Private/PlayLevel.cpp b/Engine/Source/Editor/UnrealEd/Private/PlayLevel.cpp
--- a/Engine/Source/Editor/UnrealEd/Private/PlayLevel.cpp
+++ b/Engine/Source/Editor/UnrealEd/Private/PlayLevel.cpp
@@ -1547,14 +1547,13 @@ void UEditorEngine::RequestEndPlayMap()
 		// Cache the postion and rotation of the camera (the controller may be destroyed before we end the pie session and we need them to preserve the camera position)
 		if (bLastViewAndLocationValid == false)
 		{
-			bLastViewAndLocationValid = false;
 			for (int32 WorldIdx = WorldList.Num() - 1; WorldIdx >= 0; --WorldIdx)
 			{
 				FWorldContext &ThisContext = WorldList[WorldIdx];
 				if (ThisContext.WorldType == EWorldType::PIE)
 				{
 					FSlatePlayInEditorInfo* const SlatePlayInEditorSession = SlatePlayInEditorMap.Find(ThisContext.ContextHandle);
-					if (SlatePlayInEditorSession != nullptr && SlatePlayInEditorSession->EditorPlayer.IsValid())
+					if ((SlatePlayInEditorSession != nullptr) && (SlatePlayInEditorSession->EditorPlayer.IsValid() == true) && (SlatePlayInEditorSession->EditorPlayer.Get()->PlayerController != nullptr ) )
 					{
 						SlatePlayInEditorSession->EditorPlayer.Get()->PlayerController->GetPlayerViewPoint(LastViewLocation, LastViewRotation);
 						bLastViewAndLocationValid = true;
