FBlueprintEditorUtils::PostEditChangeBlueprintActors need not call InitializeActor directly as it is now done by RegisterComponent

diff --git a/Engine/Source/Editor/UnrealEd/Private/Kismet2/BlueprintEditorUtils.cpp b/Engine/Source/Editor/UnrealEd/Private/Kismet2/BlueprintEditorUtils.cpp
--- a/Engine/Source/Editor/UnrealEd/Private/Kismet2/BlueprintEditorUtils.cpp
+++ b/Engine/Source/Editor/UnrealEd/Private/Kismet2/BlueprintEditorUtils.cpp
@@ -5732,12 +5732,6 @@ void FBlueprintEditorUtils::PostEditChangeBlueprintActors(UBlueprint* Blueprint)
 			// We know the class was derived from AActor because we checked the Blueprint->GeneratedClass.
 			AActor* Actor = static_cast<AActor*>(ObjIt);
 			Actor->PostEditChange();
-
-			// Let components that got re-registered by PostEditChange also get begun play
-			if(Actor->bActorInitialized)
-			{
-				Actor->InitializeComponents();
-			}
 		}
 	}	
 
