Enabled json manifests for cook commandlet.

diff --git a/Engine/Source/Editor/UnrealEd/Private/Commandlets/CookCommandlet.cpp b/Engine/Source/Editor/UnrealEd/Private/Commandlets/CookCommandlet.cpp
--- a/Engine/Source/Editor/UnrealEd/Private/Commandlets/CookCommandlet.cpp
+++ b/Engine/Source/Editor/UnrealEd/Private/Commandlets/CookCommandlet.cpp
@@ -1431,6 +1431,11 @@ bool UCookCommandlet::Cook(const TArray<ITargetPlatform*>& Platforms, TArray<FSt
 		FString RegistryFilename = FPaths::GameDir() / TEXT("AssetRegistry.bin");
 		FString SandboxRegistryFilename = SandboxFile->ConvertToAbsolutePathForExternalAppForWrite(*RegistryFilename);
 		ManifestGenerator.SaveAssetRegistry(SandboxRegistryFilename);
+
+		FString CookedAssetRegistry = FPaths::GameDir() / TEXT("CookedAssetRegistry.json");
+		FString SandboxCookedAssetRegistryFilename = SandboxFile->ConvertToAbsolutePathForExternalAppForWrite(*CookedAssetRegistry);
+
+		ManifestGenerator.SaveCookedPackageAssetRegistry(SandboxCookedAssetRegistryFilename, true);
 	}
 
 	return true;
