[331442] Blueprints - More detailed & accurate compiler profiling.

- Modified the reported compile time to include time spent reinstancing and recompiling child Blueprint classes.
- Added an option to the Blueprint editor preferences for turning on additional details in the Compiler Results tab.
    - Turning this on will result in a detailed, hierarchical summary of event timings for the Blueprint that was just compiled.

diff --git a/Engine/Source/Editor/BlueprintGraph/Private/BlueprintEditorSettings.cpp b/Engine/Source/Editor/BlueprintGraph/Private/BlueprintEditorSettings.cpp
--- a/Engine/Source/Editor/BlueprintGraph/Private/BlueprintEditorSettings.cpp
+++ b/Engine/Source/Editor/BlueprintGraph/Private/BlueprintEditorSettings.cpp
@@ -15,6 +15,8 @@ UBlueprintEditorSettings::UBlueprintEditorSettings(const class FPostConstructIni
 	bShowContextualFavorites = false;
 	bFlattenFavoritesMenus = true;
 	bUseLegacyMenuingSystem = false;
+	bShowDetailedCompileResults = false;
+	CompileEventDisplayThresholdMs = 5;
 
 	// settings that were moved out of experimental...
 	UEditorExperimentalSettings const* ExperimentalSettings = GetDefault<UEditorExperimentalSettings>();
diff --git a/Engine/Source/Editor/BlueprintGraph/Public/BlueprintEditorSettings.h b/Engine/Source/Editor/BlueprintGraph/Public/BlueprintEditorSettings.h
--- a/Engine/Source/Editor/BlueprintGraph/Public/BlueprintEditorSettings.h
+++ b/Engine/Source/Editor/BlueprintGraph/Public/BlueprintEditorSettings.h
@@ -52,6 +52,14 @@ public:
 
 // Perf Settings
 public:
+	/** If enabled, additional details will be displayed in the Compiler Results tab after compiling a blueprint. */
+	UPROPERTY(EditAnywhere, config, Category=Performance)
+	bool bShowDetailedCompileResults;
+
+	/** Minimum event time threshold used as a filter when additional details are enabled for display in the Compiler Results tab. A value of zero means that all events will be included in the final summary. */
+	UPROPERTY(EditAnywhere, AdvancedDisplay, config, Category=Performance, DisplayName="Compile Event Results Threshold (ms)", meta=(ClampMin="0", UIMin="0"))
+	int32 CompileEventDisplayThresholdMs;
+
 	/** The node template cache is used to speed up blueprint menuing. This determines the peak data size for that cache. */
 	UPROPERTY(EditAnywhere, AdvancedDisplay, config, Category=Performance, DisplayName="Node-Template Cache Cap (MB)", meta=(ClampMin="0", UIMin="0"))
 	float NodeTemplateCacheCapMB;
diff --git a/Engine/Source/Editor/Kismet/Private/BlueprintEditor.cpp b/Engine/Source/Editor/Kismet/Private/BlueprintEditor.cpp
--- a/Engine/Source/Editor/Kismet/Private/BlueprintEditor.cpp
+++ b/Engine/Source/Editor/Kismet/Private/BlueprintEditor.cpp
@@ -2428,23 +2428,26 @@ void FBlueprintEditor::OnBlueprintChanged(UBlueprint* InBlueprint)
 
 void FBlueprintEditor::Compile()
 {
-	if (GetBlueprintObj())
+	UBlueprint* BlueprintObj = GetBlueprintObj();
+	if (BlueprintObj)
 	{
 		FMessageLog BlueprintLog("BlueprintLog");
 
 		FFormatNamedArguments Arguments;
-		Arguments.Add(TEXT("BlueprintName"), FText::FromString(GetBlueprintObj()->GetName()));
+		Arguments.Add(TEXT("BlueprintName"), FText::FromString(BlueprintObj->GetName()));
 		BlueprintLog.NewPage(FText::Format(LOCTEXT("CompilationPageLabel", "Compile {BlueprintName}"), Arguments));
 
 		FCompilerResultsLog LogResults;
-		FKismetEditorUtilities::CompileBlueprint(GetBlueprintObj(), false, false, bSaveIntermediateBuildProducts, &LogResults);
+		LogResults.bLogDetailedResults = GetDefault<UBlueprintEditorSettings>()->bShowDetailedCompileResults;
+		LogResults.EventDisplayThresholdMs = GetDefault<UBlueprintEditorSettings>()->CompileEventDisplayThresholdMs;
+		FKismetEditorUtilities::CompileBlueprint(BlueprintObj, false, false, bSaveIntermediateBuildProducts, &LogResults);
 
-		bool bForceMessageDisplay = ((LogResults.NumWarnings > 0) || (LogResults.NumErrors > 0)) && !GetBlueprintObj()->bIsRegeneratingOnLoad;
+		bool bForceMessageDisplay = ((LogResults.NumWarnings > 0) || (LogResults.NumErrors > 0)) && !BlueprintObj->bIsRegeneratingOnLoad;
 		DumpMessagesToCompilerLog(LogResults.Messages, bForceMessageDisplay);
 
 		// send record when player clicks compile and send the result
 		// this will make sure how the users activity is
-		AnalyticsTrackCompileEvent(GetBlueprintObj(), LogResults.NumErrors, LogResults.NumWarnings);
+		AnalyticsTrackCompileEvent(BlueprintObj, LogResults.NumErrors, LogResults.NumWarnings);
 	}
 }
 
diff --git a/Engine/Source/Editor/KismetCompiler/Private/KismetCompiler.cpp b/Engine/Source/Editor/KismetCompiler/Private/KismetCompiler.cpp
--- a/Engine/Source/Editor/KismetCompiler/Private/KismetCompiler.cpp
+++ b/Engine/Source/Editor/KismetCompiler/Private/KismetCompiler.cpp
@@ -33,7 +33,6 @@ static bool bDebugPropertyPropagation = false;
 
 DEFINE_STAT(EKismetCompilerStats_CompileTime);
 DEFINE_STAT(EKismetCompilerStats_CreateSchema);
-DEFINE_STAT(EKismetCompilerStats_ReplaceGraphRefsToGeneratedClass);
 DEFINE_STAT(EKismetCompilerStats_CreateFunctionList);
 DEFINE_STAT(EKismetCompilerStats_Expansion);
 DEFINE_STAT(EKismetCompilerStats_ProcessUbergraph);
@@ -43,8 +42,10 @@ DEFINE_STAT(EKismetCompilerStats_CompileFunction);
 DEFINE_STAT(EKismetCompilerStats_PostcompileFunction);
 DEFINE_STAT(EKismetCompilerStats_FinalizationWork);
 DEFINE_STAT(EKismetCompilerStats_CodeGenerationTime);
-DEFINE_STAT(EKismetCompilerStats_UpdateBlueprintGeneratedClass);
 DEFINE_STAT(EKismetCompilerStats_ChooseTerminalScope);
+DEFINE_STAT(EKismetCompilerStats_CleanAndSanitizeClass);
+DEFINE_STAT(EKismetCompilerStats_CreateClassVariables);
+DEFINE_STAT(EKismetCompilerStats_BindAndLinkClass);
 		
 //////////////////////////////////////////////////////////////////////////
 // FKismetCompilerContext
@@ -142,6 +143,8 @@ bool FKismetCompilerContext::FSubobjectCollection::operator()(const UObject* con
 
 void FKismetCompilerContext::CleanAndSanitizeClass(UBlueprintGeneratedClass* ClassToClean, UObject*& OldCDO)
 {
+	BP_SCOPED_COMPILER_EVENT_STAT(EKismetCompilerStats_CleanAndSanitizeClass);
+
 	const bool bRecompilingOnLoad = Blueprint->bIsRegeneratingOnLoad;
 	FString TransientClassString = FString::Printf(TEXT("TRASHCLASS_%s"), *Blueprint->GetName());
 	FName TransientClassName = MakeUniqueObjectName(GetTransientPackage(), UBlueprintGeneratedClass::StaticClass(), FName(*TransientClassString));
@@ -452,6 +455,8 @@ void FKismetCompilerContext::ValidateTimelineNames()
 
 void FKismetCompilerContext::CreateClassVariablesFromBlueprint()
 {
+	BP_SCOPED_COMPILER_EVENT_STAT(EKismetCompilerStats_CreateClassVariables);
+
 	// Ensure that member variable names are valid and that there are no collisions with a parent class
 	ValidateVariableNames();
 
@@ -1063,7 +1068,7 @@ void FKismetCompilerContext::ValidateSelfPinsInGraph(const UEdGraph* SourceGraph
  */
 void FKismetCompilerContext::PrecompileFunction(FKismetFunctionContext& Context)
 {
-	SCOPE_CYCLE_COUNTER(EKismetCompilerStats_PrecompileFunction);
+	BP_SCOPED_COMPILER_EVENT_STAT(EKismetCompilerStats_PrecompileFunction);
 
 	// Find the root node, which will drive everything else
 	check(Context.RootSet.Num() == 0);
@@ -1299,7 +1304,8 @@ void OrderedInsertIntoArray(TArray<DataType>& Array, const TMap<DataType, SortKe
  */
 void FKismetCompilerContext::CompileFunction(FKismetFunctionContext& Context)
 {
-	SCOPE_CYCLE_COUNTER(EKismetCompilerStats_CompileFunction);
+	BP_SCOPED_COMPILER_EVENT_STAT(EKismetCompilerStats_CompileFunction);
+
 	check(Context.IsValid());
 
 	// Generate statements for each node in the linear execution order (which should roughly correspond to the final execution order)
@@ -1431,7 +1437,7 @@ void FKismetCompilerContext::CompileFunction(FKismetFunctionContext& Context)
  */
 void FKismetCompilerContext::PostcompileFunction(FKismetFunctionContext& Context)
 {
-	SCOPE_CYCLE_COUNTER(EKismetCompilerStats_PostcompileFunction);
+	BP_SCOPED_COMPILER_EVENT_STAT(EKismetCompilerStats_PostcompileFunction);
 
 	// Sort the 'linear execution list' again by likely execution order.
 	Context.FinalSortLinearExecList();
@@ -2437,30 +2443,30 @@ void FKismetCompilerContext::ExpansionStep(UEdGraph* Graph, bool bAllowUbergraph
 {
 	if (bIsFullCompile)
 	{
-	SCOPE_CYCLE_COUNTER(EKismetCompilerStats_Expansion);
+		BP_SCOPED_COMPILER_EVENT_STAT(EKismetCompilerStats_Expansion);
 
-	// Collapse any remaining tunnels or macros
-	ExpandTunnelsAndMacros(Graph);
+		// Collapse any remaining tunnels or macros
+		ExpandTunnelsAndMacros(Graph);
 
-	for (int32 NodeIndex = 0; NodeIndex < Graph->Nodes.Num(); ++NodeIndex)
-	{
-		UK2Node* Node = Cast<UK2Node>(Graph->Nodes[NodeIndex]);
-		if (Node)
+		for (int32 NodeIndex = 0; NodeIndex < Graph->Nodes.Num(); ++NodeIndex)
 		{
-			Node->ExpandNode(*this, Graph);
+			UK2Node* Node = Cast<UK2Node>(Graph->Nodes[NodeIndex]);
+			if (Node)
+			{
+				Node->ExpandNode(*this, Graph);
+			}
 		}
-	}
 
-	if (bAllowUbergraphExpansions)
-	{
-		// Expand timeline nodes
-		ExpandTimelineNodes(Graph);
+		if (bAllowUbergraphExpansions)
+		{
+			// Expand timeline nodes
+			ExpandTimelineNodes(Graph);
 
-		// Expand PlayMovieScene nodes
-		ExpandPlayMovieSceneNodes(Graph);
+			// Expand PlayMovieScene nodes
+			ExpandPlayMovieSceneNodes(Graph);
+		}
 	}
 }
-}
 
 void FKismetCompilerContext::VerifyValidOverrideEvent(const UEdGraph* Graph)
 {
@@ -2543,7 +2549,7 @@ void FKismetCompilerContext::VerifyValidOverrideFunction(const UEdGraph* Graph)
 // Merges pages and creates function stubs, etc... from the ubergraph entry points
 void FKismetCompilerContext::CreateAndProcessUbergraph()
 {
-	SCOPE_CYCLE_COUNTER(EKismetCompilerStats_ProcessUbergraph);
+	BP_SCOPED_COMPILER_EVENT_STAT(EKismetCompilerStats_ProcessUbergraph);
 
 	ConsolidatedEventGraph = NewNamedObject<UEdGraph>(Blueprint, GetUbergraphCallName());
 	ConsolidatedEventGraph->Schema = UEdGraphSchema_K2::StaticClass();
@@ -2922,7 +2928,7 @@ void FKismetCompilerContext::ResetErrorFlags(UEdGraph* Graph) const
  */
 void FKismetCompilerContext::ProcessOneFunctionGraph(UEdGraph* SourceGraph)
 {
-	SCOPE_CYCLE_COUNTER(EKismetCompilerStats_ProcessFunctionGraph);
+	BP_SCOPED_COMPILER_EVENT_STAT(EKismetCompilerStats_ProcessFunctionGraph);
 
 	// Clone the source graph so we can modify it as needed; merging in the child graphs
 	UEdGraph* FunctionGraph = FEdGraphUtilities::CloneGraph(SourceGraph, Blueprint, &MessageLog, true); 
@@ -2994,6 +3000,8 @@ void FKismetCompilerContext::ValidateFunctionGraphNames()
 // Creates a copy of the graph to allow further transformations to occur
 void FKismetCompilerContext::CreateFunctionList()
 {
+	BP_SCOPED_COMPILER_EVENT_STAT(EKismetCompilerStats_CreateFunctionList);
+
 	// Process the ubergraph if one should be present
 	if (FBlueprintEditorUtils::DoesSupportEventGraphs(Blueprint))
 	{
@@ -3038,8 +3046,6 @@ FKismetFunctionContext* FKismetCompilerContext::CreateFunctionContext()
 /** Compile a blueprint into a class and a set of functions */
 void FKismetCompilerContext::Compile()
 {
-	SCOPE_CYCLE_COUNTER(EKismetCompilerStats_CompileTime);
-
 	// Interfaces only need function signatures, so we only need to perform the first phase of compilation for them
 	bIsFullCompile = CompileOptions.DoesRequireBytecodeGeneration() && (Blueprint->BlueprintType != BPTYPE_Interface);
 
@@ -3055,7 +3061,7 @@ void FKismetCompilerContext::Compile()
 
 	if (Schema == NULL)
 	{
-		SCOPE_CYCLE_COUNTER(EKismetCompilerStats_CreateSchema);
+		BP_SCOPED_COMPILER_EVENT_STAT(EKismetCompilerStats_CreateSchema);
 		Schema = CreateSchema();
 		PostCreateSchema();
 	}
@@ -3233,10 +3239,7 @@ void FKismetCompilerContext::Compile()
 	AddInterfacesFromBlueprint(NewClass);
 
 	// Construct a context for each function, doing validation and building the function interface
-	{	
-		SCOPE_CYCLE_COUNTER(EKismetCompilerStats_CreateFunctionList);
-		CreateFunctionList();
-	}
+	CreateFunctionList();
 
 	// Precompile the functions
 	// Handle delegates signatures first, because they are needed by other functions
@@ -3256,9 +3259,12 @@ void FKismetCompilerContext::Compile()
 		}
 	}
 
-	// Relink the class
-	NewClass->Bind();
-	NewClass->StaticLink(true);
+	{ BP_SCOPED_COMPILER_EVENT_STAT(EKismetCompilerStats_BindAndLinkClass);
+
+		// Relink the class
+		NewClass->Bind();
+		NewClass->StaticLink(true);
+	}
 
 	if (bIsFullCompile && !MessageLog.NumErrors)
 	{
@@ -3316,7 +3322,7 @@ void FKismetCompilerContext::Compile()
 			FKismetFunctionContext& Function = FunctionList[i];
 			if (Function.IsValid())
 			{
-				SCOPE_CYCLE_COUNTER(EKismetCompilerStats_PostcompileFunction);
+				BP_SCOPED_COMPILER_EVENT_STAT(EKismetCompilerStats_PostcompileFunction);
 				FinishCompilingFunction(Function);
 			}
 		}
@@ -3347,8 +3353,7 @@ void FKismetCompilerContext::Compile()
 		}
 	}
 
-	{
-		SCOPE_CYCLE_COUNTER(EKismetCompilerStats_FinalizationWork);
+	{ BP_SCOPED_COMPILER_EVENT_STAT(EKismetCompilerStats_FinalizationWork);
 
 		// Set any final flags and seal the class, build a CDO, etc...
 		FinishCompilingClass(NewClass);
@@ -3443,7 +3448,7 @@ void FKismetCompilerContext::Compile()
 
 		// Always run the VM backend, it's needed for more than just debug printing
 		{
-			SCOPE_CYCLE_COUNTER(EKismetCompilerStats_CodeGenerationTime);
+			BP_SCOPED_COMPILER_EVENT_STAT(EKismetCompilerStats_CodeGenerationTime);
 			const bool bGenerateStubsOnly = !bIsFullCompile || (0 != MessageLog.NumErrors);
 			Backend_VM.GenerateCodeFromClass(NewClass, FunctionList, bGenerateStubsOnly);
 		}
diff --git a/Engine/Source/Editor/KismetCompiler/Private/KismetCompilerMisc.cpp b/Engine/Source/Editor/KismetCompiler/Private/KismetCompilerMisc.cpp
--- a/Engine/Source/Editor/KismetCompiler/Private/KismetCompilerMisc.cpp
+++ b/Engine/Source/Editor/KismetCompiler/Private/KismetCompilerMisc.cpp
@@ -1193,7 +1193,7 @@ FBPTerminal* FKismetFunctionContext::CreateLocalTerminalFromPinAutoChooseScope(U
 	const bool OutputPin = EEdGraphPinDirection::EGPD_Output == Net->Direction;
 	if (bSharedTerm && bRegularBlueprint && UseLocalGraphVariables && OutputPin)
 	{
-		SCOPE_CYCLE_COUNTER(EKismetCompilerStats_ChooseTerminalScope);
+		BP_SCOPED_COMPILER_EVENT_STAT(EKismetCompilerStats_ChooseTerminalScope);
 
 		// Pin's connections are checked, to tell if created terminal is shared, or if it could be a local variable.
 		bSharedTerm = FEventGraphUtils::PinRepresentsSharedTerminal(*Net);
diff --git a/Engine/Source/Editor/KismetCompiler/Private/KismetCompilerModule.cpp b/Engine/Source/Editor/KismetCompiler/Private/KismetCompilerModule.cpp
--- a/Engine/Source/Editor/KismetCompiler/Private/KismetCompilerModule.cpp
+++ b/Engine/Source/Editor/KismetCompiler/Private/KismetCompilerModule.cpp
@@ -42,53 +42,13 @@ public:
 	}
 };
 
-struct FPrintCompilationSummaryHelper
-{
-	static void Print(double CompileStartTime, double FinishTime, const FString Name, bool bPrintResultSuccess, FCompilerResultsLog& Results)
-	{
-
-		FNumberFormattingOptions TimeFormat;
-	  TimeFormat.MaximumFractionalDigits = 2;
-	  TimeFormat.MinimumFractionalDigits = 2;
-	  TimeFormat.MaximumIntegralDigits = 4;
-	  TimeFormat.MinimumIntegralDigits = 4;
-	  TimeFormat.UseGrouping = false;
-  
-	  FFormatNamedArguments Args;
-	  Args.Add(TEXT("Time"), FText::AsNumber(FinishTime - GStartTime, &TimeFormat));
-	  Args.Add(TEXT("BlueprintName"), FText::FromString(Name));
-	  Args.Add(TEXT("NumWarnings"), Results.NumWarnings);
-	  Args.Add(TEXT("NumErrors"), Results.NumErrors);
-	  Args.Add(TEXT("TotalMilliseconds"), (int)((FinishTime - CompileStartTime) * 1000));
-
-		if (Results.NumErrors > 0)
-	  {
-		  FString FailMsg = FText::Format(LOCTEXT("CompileFailed", "[{Time}] Compile of {BlueprintName} failed. {NumErrors} Fatal Issue(s) {NumWarnings} Warning(s) [in {TotalMilliseconds} ms]"), Args).ToString();
-		  Results.Warning(*FailMsg);
-	  }
-	  else if(Results.NumWarnings > 0)
-	  {
-		  FString WarningMsg = FText::Format(LOCTEXT("CompileWarning", "[{Time}] Compile of {BlueprintName} successful, but with {NumWarnings} Warning(s) [in {TotalMilliseconds} ms]"), Args).ToString();
-		  Results.Warning(*WarningMsg);
-	  }
-	  else if(bPrintResultSuccess)
-	  {
-		  FString SuccessMsg = FText::Format(LOCTEXT("CompileSuccess", "[{Time}] Compile of {BlueprintName} successful! [in {TotalMilliseconds} ms]"), Args).ToString();
-		  Results.Note(*SuccessMsg);
-	  }
-	}
-};
-
-
 // Compiles a blueprint.
-void FKismet2CompilerModule::CompileBlueprintInner(class UBlueprint* Blueprint, bool bPrintResultSuccess, const FKismetCompilerOptions& CompileOptions, FCompilerResultsLog& Results, TArray<UObject*>* ObjLoaded)
+void FKismet2CompilerModule::CompileBlueprintInner(class UBlueprint* Blueprint, const FKismetCompilerOptions& CompileOptions, FCompilerResultsLog& Results, TArray<UObject*>* ObjLoaded)
 {
 	FBlueprintIsBeingCompiledHelper BeingCompiled(Blueprint);
 
 	Blueprint->CurrentMessageLog = &Results;
 
-	const double CompileStartTime = FPlatformTime::Seconds();
-
 	// Early out if blueprint parent is missing
 	if (Blueprint->ParentClass == NULL)
 	{
@@ -127,43 +87,46 @@ void FKismet2CompilerModule::CompileBlueprintInner(class UBlueprint* Blueprint,
 		}
 	}
 
-	const double FinishTime = FPlatformTime::Seconds();
-	FPrintCompilationSummaryHelper::Print(CompileStartTime, FinishTime, Blueprint->GetName(), bPrintResultSuccess, Results);
-
 	Blueprint->CurrentMessageLog = NULL;
 }
 
 
 void FKismet2CompilerModule::CompileStructure(class UUserDefinedStruct* Struct, FCompilerResultsLog& Results)
 {
-	const double CompileStartTime = FPlatformTime::Seconds();
+	Results.SetSourceName(Struct->GetName());
+	BP_SCOPED_COMPILER_EVENT_STAT(EKismetCompilerStats_CompileTime);
 	FUserDefinedStructureCompilerUtils::CompileStruct(Struct, Results, true);
-	const double FinishTime = FPlatformTime::Seconds();
-	FPrintCompilationSummaryHelper::Print(CompileStartTime, FinishTime, Struct->GetName(), true, Results);
 }
 
 // Compiles a blueprint.
 void FKismet2CompilerModule::CompileBlueprint(class UBlueprint* Blueprint, const FKismetCompilerOptions& CompileOptions, FCompilerResultsLog& Results, FBlueprintCompileReinstancer* ParentReinstancer, TArray<UObject*>* ObjLoaded)
 {
 	SCOPE_SECONDS_COUNTER(GBlueprintCompileTime);
+	BP_SCOPED_COMPILER_EVENT_STAT(EKismetCompilerStats_CompileTime);
+
+	Results.SetSourceName(Blueprint->GetName());
 
 	const bool bIsBrandNewBP = (Blueprint->SkeletonGeneratedClass == NULL) && (Blueprint->GeneratedClass == NULL) && (Blueprint->ParentClass != NULL) && !CompileOptions.bIsDuplicationInstigated;
 
 	{
+		BP_SCOPED_COMPILER_EVENT_NAME(TEXT("Compile Skeleton Class"));
+
 		FBlueprintCompileReinstancer SkeletonReinstancer(Blueprint->SkeletonGeneratedClass);
 
 		FCompilerResultsLog SkeletonResults;
 		SkeletonResults.bSilentMode = true;
 		FKismetCompilerOptions SkeletonCompileOptions;
 		SkeletonCompileOptions.CompileType = EKismetCompileType::SkeletonOnly;
-		CompileBlueprintInner(Blueprint, /*bPrintResultSuccess=*/ false, SkeletonCompileOptions, SkeletonResults, ObjLoaded);
+		CompileBlueprintInner(Blueprint, SkeletonCompileOptions, SkeletonResults, ObjLoaded);
 	}
 
 	// If this was a full compile, take appropriate actions depending on the success of failure of the compile
 	if( CompileOptions.IsGeneratedClassCompileType() )
 	{
+		BP_SCOPED_COMPILER_EVENT_NAME(TEXT("Compile Generated Class"));
+
 		// Perform the full compile
-		CompileBlueprintInner(Blueprint, /*bPrintResultSuccess=*/ true, CompileOptions, Results, ObjLoaded);
+		CompileBlueprintInner(Blueprint, CompileOptions, Results, ObjLoaded);
 
 		if (Results.NumErrors == 0)
 		{
@@ -206,7 +169,7 @@ void FKismet2CompilerModule::CompileBlueprint(class UBlueprint* Blueprint, const
 			FKismetCompilerOptions StubCompileOptions(CompileOptions);
 			StubCompileOptions.CompileType = EKismetCompileType::StubAfterFailure;
 
-			CompileBlueprintInner(Blueprint, /*bPrintResultSuccess=*/ false, StubCompileOptions, StubResults, ObjLoaded);
+			CompileBlueprintInner(Blueprint, StubCompileOptions, StubResults, ObjLoaded);
 
 			StubReinstancer.UpdateBytecodeReferences();
 			if( !Blueprint->bIsRegeneratingOnLoad )
diff --git a/Engine/Source/Editor/KismetCompiler/Public/EdGraphCompilerUtilities.h b/Engine/Source/Editor/KismetCompiler/Public/EdGraphCompilerUtilities.h
--- a/Engine/Source/Editor/KismetCompiler/Public/EdGraphCompilerUtilities.h
+++ b/Engine/Source/Editor/KismetCompiler/Public/EdGraphCompilerUtilities.h
@@ -14,7 +14,6 @@ DECLARE_STATS_GROUP(TEXT("Kismet Compiler"), STATGROUP_KismetCompiler, STATCAT_A
 
 DECLARE_CYCLE_STAT_EXTERN(TEXT("Compile Time"), EKismetCompilerStats_CompileTime, STATGROUP_KismetCompiler, );
 DECLARE_CYCLE_STAT_EXTERN(TEXT("Create Schema"), EKismetCompilerStats_CreateSchema, STATGROUP_KismetCompiler, );
-DECLARE_CYCLE_STAT_EXTERN(TEXT("Fixup GeneratedClass Refs"), EKismetCompilerStats_ReplaceGraphRefsToGeneratedClass, STATGROUP_KismetCompiler, );
 DECLARE_CYCLE_STAT_EXTERN(TEXT("Create Function List"), EKismetCompilerStats_CreateFunctionList, STATGROUP_KismetCompiler, );
 DECLARE_CYCLE_STAT_EXTERN(TEXT("Expansion"), EKismetCompilerStats_Expansion, STATGROUP_KismetCompiler, );
 DECLARE_CYCLE_STAT_EXTERN(TEXT("Process uber"), EKismetCompilerStats_ProcessUbergraph, STATGROUP_KismetCompiler, );
@@ -24,8 +23,10 @@ DECLARE_CYCLE_STAT_EXTERN(TEXT("Compile Function"), EKismetCompilerStats_Compile
 DECLARE_CYCLE_STAT_EXTERN(TEXT("Postcompile Function"), EKismetCompilerStats_PostcompileFunction, STATGROUP_KismetCompiler, );
 DECLARE_CYCLE_STAT_EXTERN(TEXT("Finalization"), EKismetCompilerStats_FinalizationWork, STATGROUP_KismetCompiler, );
 DECLARE_CYCLE_STAT_EXTERN(TEXT("Code Gen"), EKismetCompilerStats_CodeGenerationTime, STATGROUP_KismetCompiler, );
-DECLARE_CYCLE_STAT_EXTERN(TEXT("Update Instances"), EKismetCompilerStats_UpdateBlueprintGeneratedClass, STATGROUP_KismetCompiler, );
 DECLARE_CYCLE_STAT_EXTERN(TEXT("Choose Terminal Scope"), EKismetCompilerStats_ChooseTerminalScope, STATGROUP_KismetCompiler, );
+DECLARE_CYCLE_STAT_EXTERN(TEXT("Clean and Sanitize Class"), EKismetCompilerStats_CleanAndSanitizeClass, STATGROUP_KismetCompiler, );
+DECLARE_CYCLE_STAT_EXTERN(TEXT("Create Class Properties"), EKismetCompilerStats_CreateClassVariables, STATGROUP_KismetCompiler, );
+DECLARE_CYCLE_STAT_EXTERN(TEXT("Bind and Link Class"), EKismetCompilerStats_BindAndLinkClass, STATGROUP_KismetCompiler, );
 
 //////////////////////////////////////////////////////////////////////////
 
diff --git a/Engine/Source/Editor/KismetCompiler/Public/KismetCompilerModule.h b/Engine/Source/Editor/KismetCompiler/Public/KismetCompilerModule.h
--- a/Engine/Source/Editor/KismetCompiler/Public/KismetCompilerModule.h
+++ b/Engine/Source/Editor/KismetCompiler/Public/KismetCompilerModule.h
@@ -71,7 +71,7 @@ public:
 	virtual TArray<FBlueprintCompileDelegate>& GetCompilers() override { return Compilers; }
 	// End implementation
 private:
-	void CompileBlueprintInner(class UBlueprint* Blueprint, bool bPrintResultSuccess, const FKismetCompilerOptions& CompileOptions, FCompilerResultsLog& Results, TArray<UObject*>* ObjLoaded);
+	void CompileBlueprintInner(class UBlueprint* Blueprint, const FKismetCompilerOptions& CompileOptions, FCompilerResultsLog& Results, TArray<UObject*>* ObjLoaded);
 
 	TArray<FBlueprintCompileDelegate> Compilers;
 };
diff --git a/Engine/Source/Editor/UnrealEd/Private/Kismet2/BlueprintEditorUtils.cpp b/Engine/Source/Editor/UnrealEd/Private/Kismet2/BlueprintEditorUtils.cpp
--- a/Engine/Source/Editor/UnrealEd/Private/Kismet2/BlueprintEditorUtils.cpp
+++ b/Engine/Source/Editor/UnrealEd/Private/Kismet2/BlueprintEditorUtils.cpp
@@ -311,6 +311,8 @@ void FBlueprintEditorUtils::RefreshAllNodes(UBlueprint* Blueprint)
 
 void FBlueprintEditorUtils::RefreshExternalBlueprintDependencyNodes(UBlueprint* Blueprint, UStruct* RefreshOnlyChild)
 {
+	BP_SCOPED_COMPILER_EVENT_NAME(TEXT("Refresh External Dependency Nodes"));
+
 	if (!Blueprint || !Blueprint->HasAllFlags(RF_LoadCompleted))
 	{
 		UE_LOG(LogBlueprint, Warning,
@@ -1537,6 +1539,11 @@ void FBlueprintEditorUtils::MarkBlueprintAsStructurallyModified(UBlueprint* Blue
 	Blueprint->bCachedDependenciesUpToDate = false;
 	if (Blueprint->Status != BS_BeingCreated && !Blueprint->bBeingCompiled)
 	{
+		FCompilerResultsLog Results;
+		Results.bLogInfoOnly = Blueprint->bIsRegeneratingOnLoad;
+
+		BP_SCOPED_COMPILER_EVENT_NAME(TEXT("Mark Blueprint as Structurally Modified"));
+
 		TArray<UEdGraph*> AllGraphs;
 		Blueprint->GetAllGraphs(AllGraphs);
 		for (int32 i = 0; i < AllGraphs.Num(); i++)
@@ -1556,9 +1563,6 @@ void FBlueprintEditorUtils::MarkBlueprintAsStructurallyModified(UBlueprint* Blue
 			// Invoke the compiler to update the skeleton class definition
 			IKismetCompilerInterface& Compiler = FModuleManager::LoadModuleChecked<IKismetCompilerInterface>(KISMET_COMPILER_MODULENAME);
 
-			FCompilerResultsLog Results;
-			Results.bLogInfoOnly = Blueprint->bIsRegeneratingOnLoad;
-
 			FKismetCompilerOptions CompileOptions;
 			CompileOptions.CompileType = EKismetCompileType::SkeletonOnly;
 			Compiler.CompileBlueprint(Blueprint, CompileOptions, Results);
@@ -1566,8 +1570,11 @@ void FBlueprintEditorUtils::MarkBlueprintAsStructurallyModified(UBlueprint* Blue
 		}
 		UpdateDelegatesInBlueprint(Blueprint);
 
-		// Notify any interested parties that the blueprint has changed
-		Blueprint->BroadcastChanged();
+		{ BP_SCOPED_COMPILER_EVENT_NAME(TEXT("Notify Blueprint Changed"));
+
+			// Notify any interested parties that the blueprint has changed
+			Blueprint->BroadcastChanged();
+		}
 
 		Blueprint->MarkPackageDirty();
 	}
diff --git a/Engine/Source/Editor/UnrealEd/Private/Kismet2/CompilerResultsLog.cpp b/Engine/Source/Editor/UnrealEd/Private/Kismet2/CompilerResultsLog.cpp
--- a/Engine/Source/Editor/UnrealEd/Private/Kismet2/CompilerResultsLog.cpp
+++ b/Engine/Source/Editor/UnrealEd/Private/Kismet2/CompilerResultsLog.cpp
@@ -13,6 +13,7 @@
 #define LOCTEXT_NAMESPACE "Editor.Stats"
 
 const FName FCompilerResultsLog::Name(TEXT("CompilerResultsLog"));
+FCompilerResultsLog* FCompilerResultsLog::CurrentEventTarget = nullptr;
 
 //////////////////////////////////////////////////////////////////////////
 // FCompilerResultsLog
@@ -68,14 +69,24 @@ FCompilerResultsLog::FCompilerResultsLog()
 	, bSilentMode(false)
 	, bLogInfoOnly(false)
 	, bAnnotateMentionedNodes(true)
+	, bLogDetailedResults(false)
+	, EventDisplayThresholdMs(0)
 {
+	CurrentEventScope = nullptr;
+	if(CurrentEventTarget == nullptr)
+	{
+		CurrentEventTarget = this;
+	}
 }
 
 FCompilerResultsLog::~FCompilerResultsLog()
 {
+	if(CurrentEventTarget == this)
+	{
+		CurrentEventTarget = nullptr;
+	}
 }
 
-
 void FCompilerResultsLog::Register()
 {
 	FMessageLogModule& MessageLogModule = FModuleManager::LoadModuleChecked<FMessageLogModule>("MessageLog");
@@ -92,6 +103,81 @@ void FCompilerResultsLog::Unregister()
 	MessageLogModule.UnregisterLogListing(Name);
 }
 
+void FCompilerResultsLog::InternalLogEvent(const FCompilerEvent& InEvent, int32 InDepth)
+{
+	const int EventTimeMs = (int)((InEvent.FinishTime - InEvent.StartTime) * 1000);
+	if(EventTimeMs >= EventDisplayThresholdMs)
+	{
+		// Skip display of the top-most event since that time has already been logged
+		if(InDepth > 0)
+		{
+			FString EventString = FString::Printf(TEXT("- %s"), *InEvent.Name);
+			if(InEvent.Counter > 0)
+			{
+				EventString.Append(FString::Printf(TEXT(" (%d)"), InEvent.Counter + 1));
+			}
+
+			FFormatNamedArguments Args;
+			Args.Add(TEXT("EventTimeMs"), EventTimeMs);
+			EventString.Append(FText::Format(LOCTEXT("PerformanceSummaryEventTime", " [{EventTimeMs} ms]"), Args).ToString());
+
+			FString IndentString = FString::Printf(TEXT("%*s"), (InDepth - 1) << 1, TEXT(""));
+			Note(*FString::Printf(TEXT("%s%s"), *IndentString, *EventString));
+		}
+
+		const int32 NumChildEvents = InEvent.ChildEvents.Num();
+		for(int32 i = 0; i < NumChildEvents; ++i)
+		{
+			InternalLogEvent(InEvent.ChildEvents[i].Get(), InDepth + 1);
+		}
+	}
+}
+
+void FCompilerResultsLog::InternalLogSummary()
+{
+	if(CurrentEventScope.IsValid())
+	{
+		const double CompileStartTime = CurrentEventScope->StartTime;
+		const double CompileFinishTime = CurrentEventScope->FinishTime;
+
+		FNumberFormattingOptions TimeFormat;
+		TimeFormat.MaximumFractionalDigits = 2;
+		TimeFormat.MinimumFractionalDigits = 2;
+		TimeFormat.MaximumIntegralDigits = 4;
+		TimeFormat.MinimumIntegralDigits = 4;
+		TimeFormat.UseGrouping = false;
+
+		FFormatNamedArguments Args;
+		Args.Add(TEXT("Time"), FText::AsNumber(CompileFinishTime - GStartTime, &TimeFormat));
+		Args.Add(TEXT("SourceName"), FText::FromString(SourceName));
+		Args.Add(TEXT("NumWarnings"), NumWarnings);
+		Args.Add(TEXT("NumErrors"), NumErrors);
+		Args.Add(TEXT("TotalMilliseconds"), (int)((CompileFinishTime - CompileStartTime) * 1000));
+
+		if (NumErrors > 0)
+		{
+			FString FailMsg = FText::Format(LOCTEXT("CompileFailed", "[{Time}] Compile of {SourceName} failed. {NumErrors} Fatal Issue(s) {NumWarnings} Warning(s) [in {TotalMilliseconds} ms]"), Args).ToString();
+			Warning(*FailMsg);
+		}
+		else if(NumWarnings > 0)
+		{
+			FString WarningMsg = FText::Format(LOCTEXT("CompileWarning", "[{Time}] Compile of {SourceName} successful, but with {NumWarnings} Warning(s) [in {TotalMilliseconds} ms]"), Args).ToString();
+			Warning(*WarningMsg);
+		}
+		else
+		{
+			FString SuccessMsg = FText::Format(LOCTEXT("CompileSuccess", "[{Time}] Compile of {SourceName} successful! [in {TotalMilliseconds} ms]"), Args).ToString();
+			Note(*SuccessMsg);
+		}
+
+		if(bLogDetailedResults)
+		{
+			Note(*LOCTEXT("PerformanceSummaryHeading", "Performance summary:").ToString());
+			InternalLogEvent(*CurrentEventScope.Get());
+		}
+	}
+}
+
 /** Update the source backtrack map to note that NewObject was most closely generated/caused by the SourceObject */
 void FCompilerResultsLog::NotifyIntermediateObjectCreation(UObject* NewObject, UObject* SourceObject)
 {
@@ -434,6 +520,88 @@ void FCompilerResultsLog::Append(FCompilerResultsLog const& Other)
 	}
 }
 
+void FCompilerResultsLog::BeginEvent(const TCHAR* InName)
+{
+	// Create a new event
+	TSharedPtr<FCompilerEvent> ParentEventScope = CurrentEventScope;
+	CurrentEventScope = MakeShareable(new FCompilerEvent(ParentEventScope));
+
+	// Start event with given name
+	CurrentEventScope->Start(InName);
+}
+
+void FCompilerResultsLog::EndEvent()
+{
+	if(CurrentEventScope.IsValid())
+	{
+		// Mark finish time
+		CurrentEventScope->Finish();
+
+		// Get the parent event scope
+		TSharedPtr<FCompilerEvent> ParentEventScope = CurrentEventScope->ParentEventScope;
+		if(ParentEventScope.IsValid())
+		{
+			// Aggregate the current event into the parent event scope
+			TSharedRef<FCompilerEvent> ChildEvent = CurrentEventScope.ToSharedRef();
+			AddChildEvent(ParentEventScope, ChildEvent);
+
+			// Move current event scope back up to parent
+			CurrentEventScope = ParentEventScope;
+		}
+		else
+		{
+			// Log results summary once we've ended the top-level event
+			InternalLogSummary();
+
+			// Reset current event scope
+			CurrentEventScope = nullptr;
+		}
+	}
+}
+
+void FCompilerResultsLog::AddChildEvent(TSharedPtr<FCompilerEvent>& ParentEventScope, TSharedRef<FCompilerEvent>& ChildEventScope)
+{
+	if(ParentEventScope.IsValid())
+	{
+		// If we already have a matching parent scope, aggregate child events into the current parent scope
+		if(ParentEventScope->Name == ChildEventScope->Name)
+		{
+			for(int i = 0; i < ChildEventScope->ChildEvents.Num(); ++i)
+			{
+				AddChildEvent(ParentEventScope, ChildEventScope->ChildEvents[i]);
+			}
+		}
+		else
+		{
+			// Look for a matching sibling event under the current parent scope
+			bool bMatchFound = false;
+			for(int i = ParentEventScope->ChildEvents.Num() - 1; i >= 0 && !bMatchFound; --i)
+			{
+				// If we find a matching scope, combine this event with the existing one
+				TSharedPtr<FCompilerEvent> ExistingChildEvent = ParentEventScope->ChildEvents[i];
+				if(ExistingChildEvent->Name == ChildEventScope->Name)
+				{
+					bMatchFound = true;
+
+					// Append timing and child event data to the existing event to create an aggregate event
+					ExistingChildEvent->Counter += 1;
+					ExistingChildEvent->FinishTime += (ChildEventScope->FinishTime - ChildEventScope->StartTime);
+					for(int j = 0; j < ChildEventScope->ChildEvents.Num(); ++j)
+					{
+						AddChildEvent(ExistingChildEvent, ChildEventScope->ChildEvents[j]);
+					}
+				}
+			}
+
+			if(!bMatchFound)
+			{
+				// If we didn't find a matching event, append the current event to the list of child events under the current parent scope
+				ParentEventScope->ChildEvents.Add(ChildEventScope);
+			}
+		}
+	}
+}
+
 #undef LOCTEXT_NAMESPACE
 
 #endif	//#if !WITH_EDITOR
diff --git a/Engine/Source/Editor/UnrealEd/Private/Kismet2/Kismet2.cpp b/Engine/Source/Editor/UnrealEd/Private/Kismet2/Kismet2.cpp
--- a/Engine/Source/Editor/UnrealEd/Private/Kismet2/Kismet2.cpp
+++ b/Engine/Source/Editor/UnrealEd/Private/Kismet2/Kismet2.cpp
@@ -34,6 +34,7 @@
 #include "ComponentAssetBroker.h"
 
 #define LOCTEXT_NAMESPACE "UnrealEd.Editor"
+
 //////////////////////////////////////////////////////////////////////////
 // FArchiveInvalidateTransientRefs
 
@@ -369,6 +370,8 @@ void FKismetEditorUtilities::CompileBlueprint(UBlueprint* BlueprintObj, bool bIs
 	FCompilerResultsLog LocalResults;
 	FCompilerResultsLog& Results = (pResults != NULL) ? *pResults : LocalResults;
 
+	BP_SCOPED_COMPILER_EVENT_NAME(TEXT("Compile Blueprint"));
+
 	FBlueprintCompileReinstancer ReinstanceHelper(OldClass);
 
 	// Suppress errors/warnings in the log if we're recompiling on load on a build machine
@@ -409,6 +412,8 @@ void FKismetEditorUtilities::CompileBlueprint(UBlueprint* BlueprintObj, bool bIs
 
 		if (!bSkipGarbageCollection)
 		{
+			BP_SCOPED_COMPILER_EVENT_NAME(TEXT("Garbage Collection"));
+
 			// Garbage collect to make sure the old class and actors are disposed of
 			CollectGarbage(GARBAGE_COLLECTION_KEEPFLAGS);
 		}
@@ -417,12 +422,15 @@ void FKismetEditorUtilities::CompileBlueprint(UBlueprint* BlueprintObj, bool bIs
 		// ReinstanceHelper.VerifyReplacement();
 	}
 
-	// Blueprint has changed, broadcast notify
-	BlueprintObj->BroadcastChanged();
+	{ BP_SCOPED_COMPILER_EVENT_NAME(TEXT("Notify Blueprint Changed"));
 
-	if(GEditor)
-	{
-		GEditor->BroadcastBlueprintCompiled();	
+		// Blueprint has changed, broadcast notify
+		BlueprintObj->BroadcastChanged();
+
+		if(GEditor)
+		{
+			GEditor->BroadcastBlueprintCompiled();	
+		}
 	}
 
 	// Default Values are now set in CDO. And these copies could be soon obsolete, so better to reset them.
@@ -431,30 +439,34 @@ void FKismetEditorUtilities::CompileBlueprint(UBlueprint* BlueprintObj, bool bIs
 		BlueprintObj->NewVariables[VarIndex].DefaultValue.Empty();
 	}
 
-	TArray<UBlueprint*> DependentBPs;
-	FBlueprintEditorUtils::GetDependentBlueprints(BlueprintObj, DependentBPs);
+	{ BP_SCOPED_COMPILER_EVENT_NAME(TEXT("Refresh Dependent Blueprints"));
 
-	// refresh each dependent blueprint
-	for (UBlueprint* Dependent : DependentBPs)
-	{
-		// for interface changes, auto-refresh nodes on any dependent blueprints
-		// note: RefreshAllNodes() will internally send a change notification event to the dependent blueprint
-		if (FBlueprintEditorUtils::IsInterfaceBlueprint(BlueprintObj))
-		{
-			bool bPreviousRegenValue = Dependent->bIsRegeneratingOnLoad;
-			Dependent->bIsRegeneratingOnLoad = Dependent->bIsRegeneratingOnLoad || BlueprintObj->bIsRegeneratingOnLoad;
-			FBlueprintEditorUtils::RefreshAllNodes(Dependent);
-			Dependent->bIsRegeneratingOnLoad = bPreviousRegenValue;
-		}
-		else if(!BlueprintObj->bIsRegeneratingOnLoad)
+		TArray<UBlueprint*> DependentBPs;
+		FBlueprintEditorUtils::GetDependentBlueprints(BlueprintObj, DependentBPs);
+
+		// refresh each dependent blueprint
+		for (UBlueprint* Dependent : DependentBPs)
 		{
-			// for non-interface changes, nodes with an external dependency have already been refreshed, and it is now safe to send a change notification event
-			Dependent->BroadcastChanged();
+			// for interface changes, auto-refresh nodes on any dependent blueprints
+			// note: RefreshAllNodes() will internally send a change notification event to the dependent blueprint
+			if (FBlueprintEditorUtils::IsInterfaceBlueprint(BlueprintObj))
+			{
+				bool bPreviousRegenValue = Dependent->bIsRegeneratingOnLoad;
+				Dependent->bIsRegeneratingOnLoad = Dependent->bIsRegeneratingOnLoad || BlueprintObj->bIsRegeneratingOnLoad;
+				FBlueprintEditorUtils::RefreshAllNodes(Dependent);
+				Dependent->bIsRegeneratingOnLoad = bPreviousRegenValue;
+			}
+			else if(!BlueprintObj->bIsRegeneratingOnLoad)
+			{
+				// for non-interface changes, nodes with an external dependency have already been refreshed, and it is now safe to send a change notification event
+				Dependent->BroadcastChanged();
+			}
 		}
 	}
 
 	if(!bIsRegeneratingOnLoad && BlueprintObj->GeneratedClass)
 	{
+		BP_SCOPED_COMPILER_EVENT_NAME(TEXT("Validate Generated Class"));
 		UBlueprint::ValidateGeneratedClass(BlueprintObj->GeneratedClass);
 	}
 
diff --git a/Engine/Source/Editor/UnrealEd/Private/Kismet2/KismetReinstanceUtilities.cpp b/Engine/Source/Editor/UnrealEd/Private/Kismet2/KismetReinstanceUtilities.cpp
--- a/Engine/Source/Editor/UnrealEd/Private/Kismet2/KismetReinstanceUtilities.cpp
+++ b/Engine/Source/Editor/UnrealEd/Private/Kismet2/KismetReinstanceUtilities.cpp
@@ -1,12 +1,19 @@
 // Copyright 1998-2014 Epic Games, Inc. All Rights Reserved.
 
 #include "UnrealEd.h"
+#include "Kismet2/CompilerResultsLog.h"
 #include "Kismet2/KismetReinstanceUtilities.h"
 #include "Kismet2/KismetEditorUtilities.h"
 #include "BlueprintEditorUtils.h"
 #include "Layers/ILayers.h"
 #include "ComponentInstanceDataCache.h"
 
+DEFINE_STAT(EKismetReinstancerStats_ReplaceInstancesOfClass);
+DEFINE_STAT(EKismetReinstancerStats_FindReferencers);
+DEFINE_STAT(EKismetReinstancerStats_ReplaceReferences);
+DEFINE_STAT(EKismetReinstancerStats_UpdateBytecodeReferences);
+DEFINE_STAT(EKismetReinstancerStats_RecompileChildClasses);
+
 /////////////////////////////////////////////////////////////////////////////////
 // FBlueprintCompileReinstancer
 
@@ -151,6 +158,8 @@ FBlueprintCompileReinstancer::~FBlueprintCompileReinstancer()
 
 void FBlueprintCompileReinstancer::ReinstanceObjects()
 {
+	BP_SCOPED_COMPILER_EVENT_NAME(TEXT("Reinstance Objects"));
+
 	// Make sure we only reinstance classes once!
 	if( bHasReinstanced )
 	{
@@ -162,15 +171,18 @@ void FBlueprintCompileReinstancer::ReinstanceObjects()
 	{
 		ReplaceInstancesOfClass(DuplicatedClass, ClassToReinstance, OriginalCDO);
 	
-		// Reparent all dependent blueprints, and recompile to ensure that they get reinstanced with the new memory layout
-		for( auto ChildBP = Children.CreateIterator(); ChildBP; ++ChildBP)
-		{
-			UBlueprint* BP = *ChildBP;
-			if( BP->ParentClass == ClassToReinstance || BP->ParentClass == DuplicatedClass)
+		{ BP_SCOPED_COMPILER_EVENT_STAT(EKismetReinstancerStats_RecompileChildClasses);
+
+			// Reparent all dependent blueprints, and recompile to ensure that they get reinstanced with the new memory layout
+			for( auto ChildBP = Children.CreateIterator(); ChildBP; ++ChildBP)
 			{
-				ReparentChild(BP);
+				UBlueprint* BP = *ChildBP;
+				if( BP->ParentClass == ClassToReinstance || BP->ParentClass == DuplicatedClass)
+				{
+					ReparentChild(BP);
 
-				FKismetEditorUtilities::CompileBlueprint(BP, false, bSkipGarbageCollection);
+					FKismetEditorUtilities::CompileBlueprint(BP, false, bSkipGarbageCollection);
+				}
 			}
 		}
 	}
@@ -178,6 +190,8 @@ void FBlueprintCompileReinstancer::ReinstanceObjects()
 
 void FBlueprintCompileReinstancer::UpdateBytecodeReferences()
 {
+	BP_SCOPED_COMPILER_EVENT_STAT(EKismetReinstancerStats_UpdateBytecodeReferences);
+
 	if(ClassToReinstance != NULL)
 	{
 		TMap<UObject*, UObject*> FieldMappings;
@@ -211,6 +225,9 @@ void FBlueprintCompileReinstancer::UpdateBytecodeReferences()
 
 void FBlueprintCompileReinstancer::ReplaceInstancesOfClass(UClass* OldClass, UClass* NewClass, UObject*	OriginalCDO)
 {
+	USelection* SelectedActors;
+	bool bSelectionChanged = false;
+	TArray<UObject*> ObjectsToReplace;
 	const bool bLogConversions = false; // for debugging
 
 	// Map of old objects to new objects
@@ -221,253 +238,253 @@ void FBlueprintCompileReinstancer::ReplaceInstancesOfClass(UClass* OldClass, UCl
 	// Set global flag to let system know we are reconstructing blueprint instances
 	TGuardValue<bool> GuardTemplateNameFlag(GIsReconstructingBlueprintInstances, true);
 
-	TArray<UObject*> ObjectsToReplace;
-	const bool bIncludeDerivedClasses = false;
-	GetObjectsOfClass(OldClass, ObjectsToReplace, bIncludeDerivedClasses);
+	{ BP_SCOPED_COMPILER_EVENT_STAT(EKismetReinstancerStats_ReplaceInstancesOfClass);
 
-	bool bSelectionChanged = false;
+		const bool bIncludeDerivedClasses = false;
+		GetObjectsOfClass(OldClass, ObjectsToReplace, bIncludeDerivedClasses);
 	
-	USelection* SelectedActors = GEditor->GetSelectedActors();
-	SelectedActors->BeginBatchSelectOperation();
-	SelectedActors->Modify();
+		SelectedActors = GEditor->GetSelectedActors();
+		SelectedActors->BeginBatchSelectOperation();
+		SelectedActors->Modify();
 
-	// Then fix 'real' (non archetype) instances of the class
-	for(int32 i=0; i<ObjectsToReplace.Num(); i++)
-	{
-		UObject* OldObject = ObjectsToReplace[i];
-		if (!OldObject->IsTemplate() && !OldObject->IsPendingKill())
+		// Then fix 'real' (non archetype) instances of the class
+		for(int32 i=0; i<ObjectsToReplace.Num(); i++)
 		{
-			AActor* OldActor = Cast<AActor>(OldObject);
-			UObject* NewObject = NULL;
-			UObject* OldBlueprintDebugObject = NULL;
-			UBlueprint* CorrespondingBlueprint = NULL;
-			// If this object is being debugged, cache it off so we can preserve the 'object being debugged' association
-			CorrespondingBlueprint = Cast<UBlueprint>(OldObject->GetClass()->ClassGeneratedBy);
-			if ( CorrespondingBlueprint )
+			UObject* OldObject = ObjectsToReplace[i];
+			if (!OldObject->IsTemplate() && !OldObject->IsPendingKill())
 			{
-				if (CorrespondingBlueprint->GetObjectBeingDebugged() == OldObject)
+				AActor* OldActor = Cast<AActor>(OldObject);
+				UObject* NewObject = NULL;
+				UObject* OldBlueprintDebugObject = NULL;
+				UBlueprint* CorrespondingBlueprint = NULL;
+				// If this object is being debugged, cache it off so we can preserve the 'object being debugged' association
+				CorrespondingBlueprint = Cast<UBlueprint>(OldObject->GetClass()->ClassGeneratedBy);
+				if ( CorrespondingBlueprint )
 				{
-					OldBlueprintDebugObject = OldObject;
+					if (CorrespondingBlueprint->GetObjectBeingDebugged() == OldObject)
+					{
+						OldBlueprintDebugObject = OldObject;
+					}
 				}
-			}
-			if (OldActor != NULL)
-			{
-				FTransform WorldTransform = FTransform::Identity;
-				FVector* LocationPtr = NULL;
-				FVector Location = FVector::ZeroVector;
-				FRotator* RotationPtr = NULL;
-				FRotator Rotation = FRotator::ZeroRotator;
-				AActor* AttachParent = NULL;
-				USceneComponent* AttachParentComp = NULL;
-				FName   AttachSocketName;
-
-				// Struct to store info about attached actors
-				struct FAttachedActorInfo
+				if (OldActor != NULL)
 				{
-					AActor* AttachedActor;
-					FName AttachedToSocket;
-				};
+					FTransform WorldTransform = FTransform::Identity;
+					FVector* LocationPtr = NULL;
+					FVector Location = FVector::ZeroVector;
+					FRotator* RotationPtr = NULL;
+					FRotator Rotation = FRotator::ZeroRotator;
+					AActor* AttachParent = NULL;
+					USceneComponent* AttachParentComp = NULL;
+					FName   AttachSocketName;
+
+					// Struct to store info about attached actors
+					struct FAttachedActorInfo
+					{
+						AActor* AttachedActor;
+						FName AttachedToSocket;
+					};
 
-				// Save info about attached actors
-				TArray<FAttachedActorInfo> AttachedActorInfos;
+					// Save info about attached actors
+					TArray<FAttachedActorInfo> AttachedActorInfos;
 
-				// Create cache to store component data across rerunning construction scripts
-				FComponentInstanceDataCache InstanceDataCache(OldActor);
+					// Create cache to store component data across rerunning construction scripts
+					FComponentInstanceDataCache InstanceDataCache(OldActor);
 
-				// If there are attached objects detach them and store the socket names
-				TArray<AActor*> AttachedActors;
-				OldActor->GetAttachedActors(AttachedActors);
-				for (AActor* AttachedActor : AttachedActors)
-				{
-					USceneComponent* AttachedActorRoot = AttachedActor->GetRootComponent();
-					if (AttachedActorRoot && AttachedActorRoot->AttachParent)
+					// If there are attached objects detach them and store the socket names
+					TArray<AActor*> AttachedActors;
+					OldActor->GetAttachedActors(AttachedActors);
+					for (AActor* AttachedActor : AttachedActors)
 					{
-						// Save info about actor to reattach
-						FAttachedActorInfo Info;
-						Info.AttachedActor = AttachedActor;
-						Info.AttachedToSocket = AttachedActorRoot->AttachSocketName;
-						AttachedActorInfos.Add(Info);
-
-						// Now detach it
-						AttachedActorRoot->DetachFromParent(true);
+						USceneComponent* AttachedActorRoot = AttachedActor->GetRootComponent();
+						if (AttachedActorRoot && AttachedActorRoot->AttachParent)
+						{
+							// Save info about actor to reattach
+							FAttachedActorInfo Info;
+							Info.AttachedActor = AttachedActor;
+							Info.AttachedToSocket = AttachedActorRoot->AttachSocketName;
+							AttachedActorInfos.Add(Info);
+
+							// Now detach it
+							AttachedActorRoot->DetachFromParent(true);
+						}
 					}
-				}
 
-				if (USceneComponent* OldRootComponent = OldActor->GetRootComponent())
-				{
-					if(OldRootComponent->AttachParent != NULL)
+					if (USceneComponent* OldRootComponent = OldActor->GetRootComponent())
 					{
-						AttachParent = OldRootComponent->AttachParent->GetOwner();
-						// Root component should never be attached to another component in the same actor!
-						if (AttachParent == OldActor)
+						if(OldRootComponent->AttachParent != NULL)
 						{
-							UE_LOG(LogBlueprint, Warning, TEXT("ReplaceInstancesOfClass: RootComponent (%s) attached to another component in this Actor (%s)."), *OldRootComponent->GetPathName(), *AttachParent->GetPathName());
-							AttachParent = NULL;
-						}
+							AttachParent = OldRootComponent->AttachParent->GetOwner();
+							// Root component should never be attached to another component in the same actor!
+							if (AttachParent == OldActor)
+							{
+								UE_LOG(LogBlueprint, Warning, TEXT("ReplaceInstancesOfClass: RootComponent (%s) attached to another component in this Actor (%s)."), *OldRootComponent->GetPathName(), *AttachParent->GetPathName());
+								AttachParent = NULL;
+							}
 
-						AttachParentComp = OldRootComponent->AttachParent;
-						AttachSocketName = OldRootComponent->AttachSocketName;
+							AttachParentComp = OldRootComponent->AttachParent;
+							AttachSocketName = OldRootComponent->AttachSocketName;
 
-						//detach it to remove any scaling
-						OldRootComponent->DetachFromParent(true);
-					}
+							//detach it to remove any scaling
+							OldRootComponent->DetachFromParent(true);
+						}
 
-					// Save off transform
-					WorldTransform = OldRootComponent->ComponentToWorld;
-					WorldTransform.SetTranslation(OldRootComponent->GetComponentLocation()); // take into account any custom location
-					Location = OldActor->GetActorLocation();
-					LocationPtr = &Location;
-					Rotation = OldActor->GetActorRotation();
-					RotationPtr = &Rotation;
-				}
+						// Save off transform
+						WorldTransform = OldRootComponent->ComponentToWorld;
+						WorldTransform.SetTranslation(OldRootComponent->GetComponentLocation()); // take into account any custom location
+						Location = OldActor->GetActorLocation();
+						LocationPtr = &Location;
+						Rotation = OldActor->GetActorRotation();
+						RotationPtr = &Rotation;
+					}
 
-				// If this actor was spawned from an Archetype, we spawn the new actor from the new version of that archetype
-				UObject* OldArchetype = OldActor->GetArchetype();
-				UWorld* World = OldActor->GetWorld();
-				AActor* NewArchetype = Cast<AActor>(OldToNewInstanceMap.FindRef(OldArchetype));
+					// If this actor was spawned from an Archetype, we spawn the new actor from the new version of that archetype
+					UObject* OldArchetype = OldActor->GetArchetype();
+					UWorld* World = OldActor->GetWorld();
+					AActor* NewArchetype = Cast<AActor>(OldToNewInstanceMap.FindRef(OldArchetype));
 
-				// Check that either this was an instance of the class directly, or we found a new archetype for it
-				check(OldArchetype == OldClass->GetDefaultObject() || NewArchetype);
+					// Check that either this was an instance of the class directly, or we found a new archetype for it
+					check(OldArchetype == OldClass->GetDefaultObject() || NewArchetype);
 
-				// Spawn the new actor instance, in the same level as the original, but deferring running the construction script until we have transferred modified properties
-				ULevel* ActorLevel = OldActor->GetLevel();				
-				UClass** MappedClass = OldToNewClassMap.Find(OldActor->GetClass());
-				UClass* SpawnClass = MappedClass ? *MappedClass : NewClass;
+					// Spawn the new actor instance, in the same level as the original, but deferring running the construction script until we have transferred modified properties
+					ULevel* ActorLevel = OldActor->GetLevel();				
+					UClass** MappedClass = OldToNewClassMap.Find(OldActor->GetClass());
+					UClass* SpawnClass = MappedClass ? *MappedClass : NewClass;
 
-				FActorSpawnParameters SpawnInfo;
-				SpawnInfo.OverrideLevel = ActorLevel;
-				SpawnInfo.Template = NewArchetype;
-				SpawnInfo.bNoCollisionFail = true;
-				SpawnInfo.bDeferConstruction = true;
+					FActorSpawnParameters SpawnInfo;
+					SpawnInfo.OverrideLevel = ActorLevel;
+					SpawnInfo.Template = NewArchetype;
+					SpawnInfo.bNoCollisionFail = true;
+					SpawnInfo.bDeferConstruction = true;
 
-				// Temporarily remove the deprecated flag so we can respawn the Blueprint in the level
-				bool bIsClassDeprecated = SpawnClass->HasAnyClassFlags(CLASS_Deprecated);
-				SpawnClass->ClassFlags &= ~CLASS_Deprecated;
+					// Temporarily remove the deprecated flag so we can respawn the Blueprint in the level
+					bool bIsClassDeprecated = SpawnClass->HasAnyClassFlags(CLASS_Deprecated);
+					SpawnClass->ClassFlags &= ~CLASS_Deprecated;
 
-				AActor* NewActor = World->SpawnActor(SpawnClass, LocationPtr, RotationPtr, SpawnInfo );
+					AActor* NewActor = World->SpawnActor(SpawnClass, LocationPtr, RotationPtr, SpawnInfo );
 
-				// Reassign the deprecated flag if it was previously assigned
-				if(bIsClassDeprecated)
-				{
-					SpawnClass->ClassFlags |= CLASS_Deprecated;
-				}
+					// Reassign the deprecated flag if it was previously assigned
+					if(bIsClassDeprecated)
+					{
+						SpawnClass->ClassFlags |= CLASS_Deprecated;
+					}
 
-				check(NewActor);
-				NewObject = NewActor;
+					check(NewActor);
+					NewObject = NewActor;
 
-				OldActor->DestroyConstructedComponents(); // don't want to serialize components from the old actor
-				// Unregister native components so we don't copy any sub-components they generate for themselves (like UCameraComponent does)
-				OldActor->UnregisterAllComponents(); 
+					OldActor->DestroyConstructedComponents(); // don't want to serialize components from the old actor
+					// Unregister native components so we don't copy any sub-components they generate for themselves (like UCameraComponent does)
+					OldActor->UnregisterAllComponents(); 
 
-				NewActor->UnregisterAllComponents(); // Unregister any native components, might have cached state based on properties we are going to overwrite
+					NewActor->UnregisterAllComponents(); // Unregister any native components, might have cached state based on properties we are going to overwrite
 
-				UEditorEngine::CopyPropertiesForUnrelatedObjects(OldActor, NewActor);
+					UEditorEngine::CopyPropertiesForUnrelatedObjects(OldActor, NewActor);
 
-				NewActor->ResetPropertiesForConstruction(); // reset properties/streams
+					NewActor->ResetPropertiesForConstruction(); // reset properties/streams
 
-				NewActor->RegisterAllComponents(); // Register native components
+					NewActor->RegisterAllComponents(); // Register native components
 
-				// Because this is an editor context it's important to use this execution guard:
-				FEditorScriptExecutionGuard ScriptGuard;
+					// Because this is an editor context it's important to use this execution guard:
+					FEditorScriptExecutionGuard ScriptGuard;
 
-				// Run the construction script, which will use the properties we just copied over
-				NewActor->ExecuteConstruction(WorldTransform, &InstanceDataCache);
+					// Run the construction script, which will use the properties we just copied over
+					NewActor->ExecuteConstruction(WorldTransform, &InstanceDataCache);
 
-				if (USceneComponent* NewRootComponent = NewActor->GetRootComponent())
-				{
-					//Attach the new instance to original parent
-					if(AttachParent != nullptr)
+					if (USceneComponent* NewRootComponent = NewActor->GetRootComponent())
 					{
-						if(AttachParentComp == nullptr)
+						//Attach the new instance to original parent
+						if(AttachParent != nullptr)
 						{
-							AttachParentComp = AttachParent->GetRootComponent();
-						}
+							if(AttachParentComp == nullptr)
+							{
+								AttachParentComp = AttachParent->GetRootComponent();
+							}
 
-						if(AttachParentComp != nullptr)
-						{
-							NewRootComponent->AttachTo(AttachParentComp, AttachSocketName, EAttachLocation::KeepWorldPosition);
+							if(AttachParentComp != nullptr)
+							{
+								NewRootComponent->AttachTo(AttachParentComp, AttachSocketName, EAttachLocation::KeepWorldPosition);
+							}
 						}
-					}
 
-					// If we had attached children reattach them now - unless they are already attached
-					for(FAttachedActorInfo& Info : AttachedActorInfos)
-					{
-						// If this actor is no longer attached to anything, reattach
-						if (!Info.AttachedActor->IsPendingKill() && Info.AttachedActor->GetAttachParentActor() == nullptr)
+						// If we had attached children reattach them now - unless they are already attached
+						for(FAttachedActorInfo& Info : AttachedActorInfos)
 						{
-							USceneComponent* ChildRoot = Info.AttachedActor->GetRootComponent();
-							if (ChildRoot && ChildRoot->AttachParent != NewRootComponent)
+							// If this actor is no longer attached to anything, reattach
+							if (!Info.AttachedActor->IsPendingKill() && Info.AttachedActor->GetAttachParentActor() == nullptr)
 							{
-								ChildRoot->AttachTo(NewRootComponent, Info.AttachedToSocket, EAttachLocation::KeepWorldPosition);
-								ChildRoot->UpdateComponentToWorld();
+								USceneComponent* ChildRoot = Info.AttachedActor->GetRootComponent();
+								if (ChildRoot && ChildRoot->AttachParent != NewRootComponent)
+								{
+									ChildRoot->AttachTo(NewRootComponent, Info.AttachedToSocket, EAttachLocation::KeepWorldPosition);
+									ChildRoot->UpdateComponentToWorld();
+								}
 							}
 						}
 					}
-				}
 
-				// Determine if the actor used to be selected.
-				if (OldActor->IsSelected())
-				{
-					// Deselect the old actor and select the new actor.
-					const bool bNotify = false;
-					GEditor->SelectActor(OldActor, /*bInSelected=*/ false, /*bNotify=*/ bNotify);
-					GEditor->SelectActor(NewActor, /*bInSelected=*/ true, /*bNotify=*/ bNotify);
+					// Determine if the actor used to be selected.
+					if (OldActor->IsSelected())
+					{
+						// Deselect the old actor and select the new actor.
+						const bool bNotify = false;
+						GEditor->SelectActor(OldActor, /*bInSelected=*/ false, /*bNotify=*/ bNotify);
+						GEditor->SelectActor(NewActor, /*bInSelected=*/ true, /*bNotify=*/ bNotify);
 
-					bSelectionChanged = true;
-				}
+						bSelectionChanged = true;
+					}
 
-				// Destroy actor and clear references.
-				NewActor->Modify();
-				if(GEditor->Layers.IsValid()) // ensure(NULL != GEditor->Layers) ?? While cooking the Layers is NULL.
-				{
-					GEditor->Layers->InitializeNewActorLayers( NewActor );
+					// Destroy actor and clear references.
+					NewActor->Modify();
+					if(GEditor->Layers.IsValid()) // ensure(NULL != GEditor->Layers) ?? While cooking the Layers is NULL.
+					{
+						GEditor->Layers->InitializeNewActorLayers( NewActor );
 
-					GEditor->Layers->DisassociateActorFromLayers( OldActor );
-				}
+						GEditor->Layers->DisassociateActorFromLayers( OldActor );
+					}
 				
-				World->EditorDestroyActor( OldActor, true );
+					World->EditorDestroyActor( OldActor, true );
 
-				OldToNewInstanceMap.Add(OldActor, NewActor);
+					OldToNewInstanceMap.Add(OldActor, NewActor);
 
-				if(bLogConversions)
-				{
-					UE_LOG(LogBlueprint, Log, TEXT("Converted instance '%s' to '%s'"), *OldActor->GetPathName(), *NewActor->GetPathName() );
+					if(bLogConversions)
+					{
+						UE_LOG(LogBlueprint, Log, TEXT("Converted instance '%s' to '%s'"), *OldActor->GetPathName(), *NewActor->GetPathName() );
+					}
 				}
-			}
-			else
-			{
-				FName OldName(OldObject->GetFName());
-				OldObject->Rename(NULL, OldObject->GetOuter(), REN_DoNotDirty | REN_DontCreateRedirectors);
-				NewObject = ConstructObject<UObject>(NewClass, OldObject->GetOuter(), OldName);
-				check(NewObject);
+				else
+				{
+					FName OldName(OldObject->GetFName());
+					OldObject->Rename(NULL, OldObject->GetOuter(), REN_DoNotDirty | REN_DontCreateRedirectors);
+					NewObject = ConstructObject<UObject>(NewClass, OldObject->GetOuter(), OldName);
+					check(NewObject);
 
-				UEditorEngine::CopyPropertiesForUnrelatedObjects(OldObject, NewObject);
+					UEditorEngine::CopyPropertiesForUnrelatedObjects(OldObject, NewObject);
 
-				if (UAnimInstance* AnimTree = Cast<UAnimInstance>(NewObject))
-				{
-					// Initialising the anim instance isn't enough to correctly set up the skeletal mesh again in a
-					// paused world, need to initialise the skeletal mesh component that contains the anim instance.
-					if(USkeletalMeshComponent* SkelComponent = Cast<USkeletalMeshComponent>(AnimTree->GetOuter()))
+					if (UAnimInstance* AnimTree = Cast<UAnimInstance>(NewObject))
 					{
-						SkelComponent->InitAnim(true);
+						// Initialising the anim instance isn't enough to correctly set up the skeletal mesh again in a
+						// paused world, need to initialise the skeletal mesh component that contains the anim instance.
+						if(USkeletalMeshComponent* SkelComponent = Cast<USkeletalMeshComponent>(AnimTree->GetOuter()))
+						{
+							SkelComponent->InitAnim(true);
+						}
 					}
-				}
 
-				OldObject->RemoveFromRoot();
-				OldObject->MarkPendingKill();
+					OldObject->RemoveFromRoot();
+					OldObject->MarkPendingKill();
 
-				OldToNewInstanceMap.Add(OldObject, NewObject);
+					OldToNewInstanceMap.Add(OldObject, NewObject);
 
-				if(bLogConversions)
-				{
-					UE_LOG(LogBlueprint, Log, TEXT("Converted instance '%s' to '%s'"), *OldObject->GetPathName(), *NewObject->GetPathName() );
+					if(bLogConversions)
+					{
+						UE_LOG(LogBlueprint, Log, TEXT("Converted instance '%s' to '%s'"), *OldObject->GetPathName(), *NewObject->GetPathName() );
+					}
 				}
-			}
 
-			// If this original object came from a blueprint and it was in the selected debug set, change the debugging to the new object.
-			if( ( CorrespondingBlueprint) && ( OldBlueprintDebugObject ) && ( NewObject ) )
-			{
-				CorrespondingBlueprint->SetObjectBeingDebugged(NewObject);
+				// If this original object came from a blueprint and it was in the selected debug set, change the debugging to the new object.
+				if( ( CorrespondingBlueprint) && ( OldBlueprintDebugObject ) && ( NewObject ) )
+				{
+					CorrespondingBlueprint->SetObjectBeingDebugged(NewObject);
+				}
 			}
 		}
 	}
@@ -488,11 +505,14 @@ void FBlueprintCompileReinstancer::ReplaceInstancesOfClass(UClass* OldClass, UCl
 	}
 
 	// Find everything that references these objects
-	TFindObjectReferencers<UObject> Referencers(SourceObjects, NULL, false);
 	TSet<UObject *> Targets;
-	for (TFindObjectReferencers<UObject>::TIterator It(Referencers); It; ++It)
-	{
-		Targets.Add(It.Value());
+	{ BP_SCOPED_COMPILER_EVENT_STAT(EKismetReinstancerStats_FindReferencers);
+
+		TFindObjectReferencers<UObject> Referencers(SourceObjects, NULL, false);
+		for (TFindObjectReferencers<UObject>::TIterator It(Referencers); It; ++It)
+		{
+			Targets.Add(It.Value());
+		}
 	}
 
 	// Add the old->new class/CDO mapping into the fixup map
@@ -503,12 +523,15 @@ void FBlueprintCompileReinstancer::ReplaceInstancesOfClass(UClass* OldClass, UCl
 		OldToNewInstanceMap.Add(OriginalCDO, NewCDO);
 	}
 
-	for (TSet<UObject *>::TIterator It(Targets); It; ++It)
-	{
-		UObject* Obj = *It;
-		if (!ObjectsToReplace.Contains(Obj)) // Don't bother trying to fix old objects, this would break them
+	{ BP_SCOPED_COMPILER_EVENT_STAT(EKismetReinstancerStats_ReplaceReferences);
+
+		for (TSet<UObject *>::TIterator It(Targets); It; ++It)
 		{
-			FArchiveReplaceObjectRef<UObject> ReplaceAr(Obj, OldToNewInstanceMap, /*bNullPrivateRefs=*/ false, /*bIgnoreOuterRef=*/ false, /*bIgnoreArchetypeRef=*/ false);
+			UObject* Obj = *It;
+			if (!ObjectsToReplace.Contains(Obj)) // Don't bother trying to fix old objects, this would break them
+			{
+				FArchiveReplaceObjectRef<UObject> ReplaceAr(Obj, OldToNewInstanceMap, /*bNullPrivateRefs=*/ false, /*bIgnoreOuterRef=*/ false, /*bIgnoreArchetypeRef=*/ false);
+			}
 		}
 	}
 
diff --git a/Engine/Source/Editor/UnrealEd/Public/Kismet2/CompilerResultsLog.h b/Engine/Source/Editor/UnrealEd/Public/Kismet2/CompilerResultsLog.h
--- a/Engine/Source/Editor/UnrealEd/Public/Kismet2/CompilerResultsLog.h
+++ b/Engine/Source/Editor/UnrealEd/Public/Kismet2/CompilerResultsLog.h
@@ -29,6 +29,40 @@ public:
 /** This class represents a log of compiler output lines (errors, warnings, and information notes), each of which can be a rich tokenized message */
 class UNREALED_API FCompilerResultsLog
 {
+	// Compiler event
+	struct FCompilerEvent
+	{
+		FString Name;
+		uint32 Counter;
+		double StartTime;
+		double FinishTime;
+		TSharedPtr<FCompilerEvent> ParentEventScope;
+		TArray< TSharedRef<FCompilerEvent> > ChildEvents;
+
+		FCompilerEvent(TSharedPtr<FCompilerEvent> InParentEventScope = nullptr)
+			:Name(TEXT(""))
+			,Counter(0)
+			,StartTime(0.0)
+			,FinishTime(0.0)
+			,ParentEventScope(InParentEventScope)
+		{
+		}
+
+		void Start(const TCHAR* InName)
+		{
+			Name = InName;
+			StartTime = FPlatformTime::Seconds();
+		}
+
+		void Finish()
+		{
+			FinishTime = FPlatformTime::Seconds();
+		}
+	};
+
+	// Current compiler event scope
+	TSharedPtr<FCompilerEvent> CurrentEventScope;
+
 public:
 	// List of all tokenized messages
 	TArray< TSharedRef<class FTokenizedMessage> > Messages;
@@ -47,9 +81,20 @@ public:
 
 	// Should nodes mentioned in messages be annotated for display with that message?
 	bool bAnnotateMentionedNodes;
+
+	// Should detailed results be appended to the final summary log?
+	bool bLogDetailedResults;
+
+	// Minimum event time (ms) for inclusion into the final summary log
+	int EventDisplayThresholdMs;
+
 protected:
 	// Maps from transient object created during compiling to original 'source code' object
 	FBacktrackMap SourceBacktrackMap;
+
+	// Name of the source object being compiled
+	FString SourceName;
+
 public:
 	FCompilerResultsLog();
 	virtual ~FCompilerResultsLog();
@@ -63,6 +108,12 @@ public:
 	/** Accessor for the LogName, so it can be opened elsewhere */
 	static FName GetLogName(){ return Name; }
 
+	/** Set the source name for the final log summary */
+	void SetSourceName(const FString& InSourceName)
+	{
+		SourceName = InSourceName;
+	}
+
 	// Note: Message is not a fprintf string!  It should be preformatted, but can contain @@ to indicate object references, which are the varargs
 	void Error(const TCHAR* Message, ...);
 	void Warning(const TCHAR* Message, ...);
@@ -92,13 +143,35 @@ public:
 	}
 	
 	void Append(FCompilerResultsLog const& Other);
+
+	/** Begin a new compiler event */
+	void BeginEvent(const TCHAR* InName);
+
+	/** End the current compiler event */
+	void EndEvent();
+
+	/** Access the current event target log */
+	static FCompilerResultsLog* GetEventTarget()
+	{
+		return CurrentEventTarget;
+	}
+
 protected:
+	/** Helper method to add a child event to the given parent event scope */
+	void AddChildEvent(TSharedPtr<FCompilerEvent>& ParentEventScope, TSharedRef<FCompilerEvent>& ChildEventScope);
+
 	/** Create a tokenized message record from a message containing @@ indicating where each UObject* in the ArgPtr list goes and place it in the MessageLog. */
 	void InternalLogMessage(const EMessageSeverity::Type& Severity, const TCHAR* Message, va_list ArgPtr);
 	
 	/** */
 	void AnnotateNode(class UEdGraphNode* Node, TSharedRef<FTokenizedMessage> LogLine);
 
+	/** Internal method to append the final compiler results summary to the MessageLog */
+	void InternalLogSummary();
+
+	/** Internal helper method to recursively append event details into the MessageLog */
+	void InternalLogEvent(const FCompilerEvent& InEvent, int32 InDepth = 0);
+
 private:
 	/** Parses a compiler log dump to generate tokenized output */
 	static TArray< TSharedRef<class FTokenizedMessage> > ParseCompilerLogDump(const FString& LogDump);
@@ -111,5 +184,45 @@ private:
 
 	/** The log's name, for easy re-use */
 	static const FName Name;
+	
+	/** The log target for compile events */
+	static FCompilerResultsLog* CurrentEventTarget;
 };
+
+/** This class will begin a new compile event on construction, and automatically end it when the instance goes out of scope */
+class UNREALED_API FScopedCompilerEvent
+{
+public:
+	/** Constructor; automatically begins a new event */
+	FScopedCompilerEvent(const TCHAR* InName)
+	{
+		FCompilerResultsLog* ResultsLog = FCompilerResultsLog::GetEventTarget();
+		if(ResultsLog != nullptr)
+		{
+			ResultsLog->BeginEvent(InName);
+		}
+	}
+
+	/** Destructor; automatically ends the event */
+	~FScopedCompilerEvent()
+	{
+		FCompilerResultsLog* ResultsLog = FCompilerResultsLog::GetEventTarget();
+		if(ResultsLog != nullptr)
+		{
+			ResultsLog->EndEvent();
+		}
+	}
+};
+
+#define BP_SCOPED_COMPILER_EVENT_NAME(Name) \
+	FScopedCompilerEvent ScopedCompilerEvent(Name)
+
+#if STATS
+#define BP_SCOPED_COMPILER_EVENT_STAT(Stat) \
+	SCOPE_CYCLE_COUNTER(Stat); \
+	FScopedCompilerEvent ScopedCompilerEvent(GET_STATDESCRIPTION(Stat))
+#else
+#define BP_SCOPED_COMPILER_EVENT_STAT(Stat) \
+	FScopedCompilerEvent ScopedCompilerEvent(ANSI_TO_TCHAR(#Stat))
+#endif
 #endif	//#if WITH_EDITOR
diff --git a/Engine/Source/Editor/UnrealEd/Public/Kismet2/KismetReinstanceUtilities.h b/Engine/Source/Editor/UnrealEd/Public/Kismet2/KismetReinstanceUtilities.h
--- a/Engine/Source/Editor/UnrealEd/Public/Kismet2/KismetReinstanceUtilities.h
+++ b/Engine/Source/Editor/UnrealEd/Public/Kismet2/KismetReinstanceUtilities.h
@@ -8,6 +8,14 @@
 
 #include "Engine.h"
 
+DECLARE_STATS_GROUP(TEXT("Kismet Reinstancer"), STATGROUP_KismetReinstancer, STATCAT_Advanced);
+
+DECLARE_CYCLE_STAT_EXTERN(TEXT("Replace Instances"), EKismetReinstancerStats_ReplaceInstancesOfClass, STATGROUP_KismetReinstancer, );
+DECLARE_CYCLE_STAT_EXTERN(TEXT("Find Referencers"), EKismetReinstancerStats_FindReferencers, STATGROUP_KismetReinstancer, );
+DECLARE_CYCLE_STAT_EXTERN(TEXT("Replace References"), EKismetReinstancerStats_ReplaceReferences, STATGROUP_KismetReinstancer, );
+DECLARE_CYCLE_STAT_EXTERN(TEXT("Update Bytecode References"), EKismetReinstancerStats_UpdateBytecodeReferences, STATGROUP_KismetReinstancer, );
+DECLARE_CYCLE_STAT_EXTERN(TEXT("Recompile Child Classes"), EKismetReinstancerStats_RecompileChildClasses, STATGROUP_KismetReinstancer, );
+
 class UNREALED_API FBlueprintCompileReinstancer
 {
 protected:
