Moving EdGraphNode_Comment into Editor, out of Runtime/Engine

diff --git a/Engine/Source/Editor/BlueprintGraph/Private/BlueprintActionDatabase.cpp b/Engine/Source/Editor/BlueprintGraph/Private/BlueprintActionDatabase.cpp
--- a/Engine/Source/Editor/BlueprintGraph/Private/BlueprintActionDatabase.cpp
+++ b/Engine/Source/Editor/BlueprintGraph/Private/BlueprintActionDatabase.cpp
@@ -39,7 +39,7 @@
 #include "K2Node_DynamicCast.h"
 #include "K2Node_ClassDynamicCast.h"
 // used below in BlueprintActionDatabaseImpl::GetNodeSpectificActions()
-#include "EdGraph/EdGraphNode_Comment.h"
+#include "EdGraphNode_Comment.h"
 #include "EdGraph/EdGraphNode_Documentation.h"
 
 #define LOCTEXT_NAMESPACE "BlueprintActionDatabase"
diff --git a/Engine/Source/Editor/GraphEditor/Private/MaterialNodes/SGraphNodeMaterialComment.cpp b/Engine/Source/Editor/GraphEditor/Private/MaterialNodes/SGraphNodeMaterialComment.cpp
--- a/Engine/Source/Editor/GraphEditor/Private/MaterialNodes/SGraphNodeMaterialComment.cpp
+++ b/Engine/Source/Editor/GraphEditor/Private/MaterialNodes/SGraphNodeMaterialComment.cpp
@@ -4,6 +4,7 @@
 #include "Materials/MaterialExpressionComment.h"
 #include "SGraphNodeComment.h"
 #include "SGraphNodeMaterialComment.h"
+#include "MaterialGraph/MaterialGraphNode_Comment.h"
 
 void SGraphNodeMaterialComment::Construct(const FArguments& InArgs, class UMaterialGraphNode_Comment* InNode)
 {
diff --git a/Engine/Source/Editor/MaterialEditor/Private/MaterialEditor.cpp b/Engine/Source/Editor/MaterialEditor/Private/MaterialEditor.cpp
--- a/Engine/Source/Editor/MaterialEditor/Private/MaterialEditor.cpp
+++ b/Engine/Source/Editor/MaterialEditor/Private/MaterialEditor.cpp
@@ -2,6 +2,8 @@
 
 #include "MaterialEditorModule.h"
 
+#include "MaterialGraph/MaterialGraphNode_Comment.h"
+
 #include "Materials/MaterialExpressionBreakMaterialAttributes.h"
 #include "Materials/MaterialExpressionCollectionParameter.h"
 #include "Materials/MaterialExpressionComment.h"
diff --git a/Engine/Source/Editor/UnrealEd/Classes/MaterialGraph/MaterialGraph.h b/Engine/Source/Editor/UnrealEd/Classes/MaterialGraph/MaterialGraph.h
--- a/Engine/Source/Editor/UnrealEd/Classes/MaterialGraph/MaterialGraph.h
+++ b/Engine/Source/Editor/UnrealEd/Classes/MaterialGraph/MaterialGraph.h
@@ -119,7 +119,7 @@ public:
 	 *
 	 * @param	GraphPin	Pin representing the material input
 	 */
-	UNREALED_API bool IsInputActive(UEdGraphPin* GraphPin) const;
+	UNREALED_API bool IsInputActive(class UEdGraphPin* GraphPin) const;
 
 	/**
 	 * Get a list of nodes representing expressions that are not used in the Material
diff --git a/Engine/Source/Editor/UnrealEd/Classes/MaterialGraph/MaterialGraphNode_Comment.h b/Engine/Source/Editor/UnrealEd/Classes/MaterialGraph/MaterialGraphNode_Comment.h
--- a/Engine/Source/Editor/UnrealEd/Classes/MaterialGraph/MaterialGraphNode_Comment.h
+++ b/Engine/Source/Editor/UnrealEd/Classes/MaterialGraph/MaterialGraphNode_Comment.h
@@ -1,7 +1,9 @@
 // Copyright 1998-2014 Epic Games, Inc. All Rights Reserved.
 
 #pragma once
-#include "EdGraph/EdGraphNode_Comment.h"
+#include "EdGraphNode_Comment.h"
+#include "MaterialGraph/MaterialGraph.h"
+
 #include "MaterialGraphNode_Comment.generated.h"
 
 UCLASS(MinimalAPI)
diff --git a/Engine/Source/Editor/UnrealEd/Private/EdGraphNode_Comment.cpp b/Engine/Source/Editor/UnrealEd/Private/EdGraphNode_Comment.cpp
new file mode 100644
--- /dev/null
+++ b/Engine/Source/Editor/UnrealEd/Private/EdGraphNode_Comment.cpp
@@ -0,0 +1,159 @@
+// Copyright 1998-2014 Epic Games, Inc. All Rights Reserved.
+
+#include "UnrealEd.h"
+#include "EdGraphNode_Comment.h"
+
+#include "BlueprintUtilities.h"
+#if WITH_EDITOR
+#include "Editor/UnrealEd/Public/Kismet2/BlueprintEditorUtils.h"
+#include "Slate.h"
+#include "ScopedTransaction.h"
+#include "Editor/UnrealEd/Public/Kismet2/Kismet2NameValidators.h"
+#endif
+
+#define LOCTEXT_NAMESPACE "EdGraph"
+
+/////////////////////////////////////////////////////
+// UEdGraphNode_Comment
+
+UEdGraphNode_Comment::UEdGraphNode_Comment(const class FPostConstructInitializeProperties& PCIP)
+	: Super(PCIP)
+{
+	NodeWidth = 400.0f;
+	NodeHeight = 100.0f;
+	CommentColor = FLinearColor::White;
+	bColorCommentBubble = false;
+	MoveMode = ECommentBoxMode::GroupMovement;
+
+#if WITH_EDITORONLY_DATA
+	bCanResizeNode = true;
+	bCanRenameNode = true;
+#endif // WITH_EDITORONLY_DATA
+}
+
+void UEdGraphNode_Comment::AddReferencedObjects(UObject* InThis, FReferenceCollector& Collector) 
+{
+#if WITH_EDITOR
+	UEdGraphNode_Comment* This = CastChecked<UEdGraphNode_Comment>(InThis);
+	for (auto It = This->NodesUnderComment.CreateIterator(); It; ++It)
+	{
+		Collector.AddReferencedObject(*It, This);
+	}
+#endif // WITH_EDITOR
+
+	Super::AddReferencedObjects(InThis, Collector);
+}
+
+#if WITH_EDITOR
+void UEdGraphNode_Comment::PostPlacedNewNode()
+{
+	//@TODO: Consider making the default value we use here a preference
+	// This is done here instead of in the constructor so we can later change the default for newly placed
+	// instances without changing all of the existing ones (due to delta serialization)
+	MoveMode = ECommentBoxMode::GroupMovement;
+
+	NodeComment = NSLOCTEXT("K2Node", "CommentBlock_NewEmptyComment", "Comment").ToString();
+}
+
+FText UEdGraphNode_Comment::GetTooltipText() const
+{
+	if (CachedTooltip.IsOutOfDate())
+	{
+		CachedTooltip = FText::Format(NSLOCTEXT("K2Node", "CommentBlock_Tooltip", "Comment:\n{0}"), FText::FromString(NodeComment));
+	}
+	return CachedTooltip;
+}
+
+FString UEdGraphNode_Comment::GetDocumentationLink() const
+{
+	return TEXT("Shared/GraphNodes/Common");
+}
+
+FString UEdGraphNode_Comment::GetDocumentationExcerptName() const
+{
+	return TEXT("UEdGraphNode_Comment");
+}
+
+FName UEdGraphNode_Comment::GetPaletteIcon(FLinearColor& OutColor) const
+{
+	OutColor = FLinearColor::White;
+	return TEXT("GraphEditor.Comment_16x");
+}
+
+FText UEdGraphNode_Comment::GetNodeTitle(ENodeTitleType::Type TitleType) const
+{
+	if(TitleType == ENodeTitleType::MenuTitle)
+	{
+		return NSLOCTEXT("K2Node", "NoComment_ListTitle", "Add Comment...");
+	}
+	else if(TitleType == ENodeTitleType::ListView)
+	{
+		return NSLOCTEXT("K2Node", "CommentBlock_ListTitle", "Comment");	
+	}
+
+	return FText::FromString(NodeComment);
+}
+
+FString UEdGraphNode_Comment::GetPinNameOverride(const UEdGraphPin& Pin) const
+{
+	return GetNodeTitle(ENodeTitleType::ListView).ToString();
+}
+
+FLinearColor UEdGraphNode_Comment::GetNodeCommentColor() const
+{
+	// Only affects the 'zoomed out' comment bubble color, not the box itself
+	return (bColorCommentBubble)
+		? CommentColor 
+		: FLinearColor::White;
+}
+
+void UEdGraphNode_Comment::ResizeNode(const FVector2D& NewSize)
+{
+	if (bCanResizeNode) 
+	{
+		NodeHeight = NewSize.Y;
+		NodeWidth = NewSize.X;
+	}
+}
+
+void UEdGraphNode_Comment::AddNodeUnderComment(UObject* Object)
+{
+	NodesUnderComment.Add(Object);
+}
+
+void UEdGraphNode_Comment::ClearNodesUnderComment()
+{
+	NodesUnderComment.Empty();
+}
+
+void UEdGraphNode_Comment::SetBounds(const class FSlateRect& Rect)
+{
+	NodePosX = Rect.Left;
+	NodePosY = Rect.Top;
+
+	FVector2D Size = Rect.GetSize();
+	NodeWidth = Size.X;
+	NodeHeight = Size.Y;
+}
+
+const FCommentNodeSet& UEdGraphNode_Comment::GetNodesUnderComment() const
+{
+	return NodesUnderComment;
+}
+
+void UEdGraphNode_Comment::OnRenameNode(const FString& NewName)
+{
+	NodeComment = NewName;
+	CachedTooltip.MarkDirty();
+}
+
+TSharedPtr<class INameValidatorInterface> UEdGraphNode_Comment::MakeNameValidator() const
+{
+	// Comments can be duplicated, etc...
+	return MakeShareable(new FDummyNameValidator(EValidatorResult::Ok));
+}
+#endif // WITH_EDITOR
+
+/////////////////////////////////////////////////////
+
+#undef LOCTEXT_NAMESPACE
diff --git a/Engine/Source/Editor/UnrealEd/Public/EdGraphNode_Comment.h b/Engine/Source/Editor/UnrealEd/Public/EdGraphNode_Comment.h
new file mode 100644
--- /dev/null
+++ b/Engine/Source/Editor/UnrealEd/Public/EdGraphNode_Comment.h
@@ -0,0 +1,84 @@
+// Copyright 1998-2014 Epic Games, Inc. All Rights Reserved.
+
+#pragma once
+#include "EdGraph/EdGraphNode.h"
+#include "EdGraph/EdGraphNodeUtils.h" // for FNodeTextCache
+#include "EdGraphNode_Comment.generated.h"
+
+typedef TArray<class UObject*> FCommentNodeSet;
+
+UENUM()
+namespace ECommentBoxMode
+{
+	enum Type
+	{
+		// This comment box will move any fully contained nodes when it moves
+		GroupMovement UMETA(DisplayName="Group Movement"),
+
+		// This comment box has no effect on nodes contained inside it
+		NoGroupMovement UMETA(DisplayName="Comment")
+	};
+}
+
+UCLASS(MinimalAPI)
+class UEdGraphNode_Comment : public UEdGraphNode
+{
+	GENERATED_UCLASS_BODY()
+
+	/** Color to style comment with */
+	UPROPERTY(EditAnywhere, Category=Comment)
+	FLinearColor CommentColor;
+
+	/** Whether to use Comment Color to color the background of the comment bubble shown when zoomed out. */
+	UPROPERTY(EditAnywhere, Category=Comment, meta=(DisplayName="Color Bubble"))
+	uint32 bColorCommentBubble:1;
+
+	/** Whether the comment should move any fully enclosed nodes around when it is moved */
+	UPROPERTY(EditAnywhere, Category=Comment)
+	TEnumAsByte<ECommentBoxMode::Type> MoveMode;
+
+public:
+
+	// Begin UObject Interface
+	UNREALED_API static void AddReferencedObjects(UObject* InThis, FReferenceCollector& Collector);
+	// End UObject Interface
+
+#if WITH_EDITOR
+	// Begin UEdGraphNode interface
+	virtual void AllocateDefaultPins() override {}
+	UNREALED_API virtual FText GetTooltipText() const override;
+	UNREALED_API virtual FLinearColor GetNodeCommentColor() const override;
+	UNREALED_API virtual FText GetNodeTitle(ENodeTitleType::Type TitleType) const override;
+	virtual bool ShouldOverridePinNames() const override { return true; }
+	UNREALED_API virtual FString GetPinNameOverride(const UEdGraphPin& Pin) const override;
+	UNREALED_API virtual void ResizeNode(const FVector2D& NewSize) override;
+	UNREALED_API virtual void PostPlacedNewNode() override;
+	UNREALED_API virtual void OnRenameNode(const FString& NewName) override;
+	UNREALED_API virtual TSharedPtr<class INameValidatorInterface> MakeNameValidator() const override;
+	virtual bool ShouldDrawNodeAsComment() const override { return true; }
+	UNREALED_API virtual FString GetDocumentationLink() const override;
+	UNREALED_API virtual FString GetDocumentationExcerptName() const override;
+	UNREALED_API virtual FName GetPaletteIcon(FLinearColor& OutColor) const override;
+	// End UEdGraphNode interface
+
+	/** Add a node that will be dragged when this comment is dragged */
+	UNREALED_API void	AddNodeUnderComment(UObject* Object);
+
+	/** Clear all of the nodes which are currently being dragged with this comment */
+	UNREALED_API void	ClearNodesUnderComment();
+
+	/** Set the Bounds for the comment node */
+	UNREALED_API void SetBounds(const class FSlateRect& Rect);
+
+	/** Return the set of nodes underneath the comment */
+	UNREALED_API const FCommentNodeSet& GetNodesUnderComment() const;
+
+private:
+	/** Nodes currently within the region of the comment */
+	FCommentNodeSet	NodesUnderComment;
+
+	/** Constructing FText strings can be costly, so we cache the node's tooltip */
+	FNodeTextCache CachedTooltip;
+#endif
+};
+
diff --git a/Engine/Source/Runtime/Engine/Classes/EdGraph/EdGraph.h b/Engine/Source/Runtime/Engine/Classes/EdGraph/EdGraph.h
--- a/Engine/Source/Runtime/Engine/Classes/EdGraph/EdGraph.h
+++ b/Engine/Source/Runtime/Engine/Classes/EdGraph/EdGraph.h
@@ -59,7 +59,7 @@ class ENGINE_API UEdGraph : public UObject
 
 	/** The schema that this graph obeys */
 	UPROPERTY()
-	TSubclassOf<UEdGraphSchema>  Schema;
+	TSubclassOf<class UEdGraphSchema>  Schema;
 
 	/** Set of all nodes in this graph */
 	UPROPERTY()
@@ -95,7 +95,7 @@ public:
 	template <typename NodeType> friend struct FGraphNodeCreator;
 
 	/** Get the schema associated with this graph */
-	const UEdGraphSchema* GetSchema() const;
+	const class UEdGraphSchema* GetSchema() const;
 
 	/** Add a listener for OnGraphChanged events */
 	void AddOnGraphChangedHandler( const FOnGraphChanged::FDelegate& InHandler );
diff --git a/Engine/Source/Runtime/Engine/Classes/EdGraph/EdGraphNode_Comment.h b/Engine/Source/Runtime/Engine/Classes/EdGraph/EdGraphNode_Comment.h
deleted file mode 100644
diff --git a/Engine/Source/Runtime/Engine/Private/EdGraph/EdGraphNode_Comment.cpp b/Engine/Source/Runtime/Engine/Private/EdGraph/EdGraphNode_Comment.cpp
deleted file mode 100644
