Default Comment Node Color now pulled from GraphEditorSettings, rather than hard coded...

diff --git a/Engine/Source/Editor/GraphEditor/Private/GraphEditorSettings.cpp b/Engine/Source/Editor/GraphEditor/Private/GraphEditorSettings.cpp
--- a/Engine/Source/Editor/GraphEditor/Private/GraphEditorSettings.cpp
+++ b/Engine/Source/Editor/GraphEditor/Private/GraphEditorSettings.cpp
@@ -72,6 +72,7 @@ UGraphEditorSettings::UGraphEditorSettings( const class FPostConstructInitialize
 	ExecBranchNodeTitleColor = FLinearColor(1.0f, 1.0f, 1.0f, 1.0f);
 	ExecSequenceNodeTitleColor = FLinearColor(0.8f, 0.4f, 0.4f, 1.0f);
 	ResultNodeTitleColor = FLinearColor(1.0f, 0.65f, 0.4f, 1.0f);
+	DefaultCommentNodeTitleColor = FLinearColor::White;
 }
 
 #if WITH_EDITOR
diff --git a/Engine/Source/Editor/GraphEditor/Public/GraphEditorSettings.h b/Engine/Source/Editor/GraphEditor/Public/GraphEditorSettings.h
--- a/Engine/Source/Editor/GraphEditor/Public/GraphEditorSettings.h
+++ b/Engine/Source/Editor/GraphEditor/Public/GraphEditorSettings.h
@@ -210,6 +210,10 @@ public:
 	UPROPERTY(EditAnywhere,  config, Category=NodeTitleColors)
 	FLinearColor ResultNodeTitleColor;
 
+	/** Default Comment node title color */
+	UPROPERTY(EditAnywhere, config, Category=NodeTitleColors)
+	FLinearColor DefaultCommentNodeTitleColor;
+
 public:
 
 	UPROPERTY(EditAnywhere,  config, Category=Tracing)
diff --git a/Engine/Source/Editor/UnrealEd/Private/EdGraphNode_Comment.cpp b/Engine/Source/Editor/UnrealEd/Private/EdGraphNode_Comment.cpp
--- a/Engine/Source/Editor/UnrealEd/Private/EdGraphNode_Comment.cpp
+++ b/Engine/Source/Editor/UnrealEd/Private/EdGraphNode_Comment.cpp
@@ -4,12 +4,11 @@
 #include "EdGraphNode_Comment.h"
 
 #include "BlueprintUtilities.h"
-#if WITH_EDITOR
+#include "GraphEditorSettings.h"
 #include "Editor/UnrealEd/Public/Kismet2/BlueprintEditorUtils.h"
 #include "Slate.h"
 #include "ScopedTransaction.h"
 #include "Editor/UnrealEd/Public/Kismet2/Kismet2NameValidators.h"
-#endif
 
 #define LOCTEXT_NAMESPACE "EdGraph"
 
@@ -21,30 +20,25 @@ UEdGraphNode_Comment::UEdGraphNode_Comment(const class FPostConstructInitializeP
 {
 	NodeWidth = 400.0f;
 	NodeHeight = 100.0f;
-	CommentColor = FLinearColor::White;
+	CommentColor = GetDefault<UGraphEditorSettings>()->DefaultCommentNodeTitleColor;
 	bColorCommentBubble = false;
 	MoveMode = ECommentBoxMode::GroupMovement;
 
-#if WITH_EDITORONLY_DATA
 	bCanResizeNode = true;
 	bCanRenameNode = true;
-#endif // WITH_EDITORONLY_DATA
 }
 
 void UEdGraphNode_Comment::AddReferencedObjects(UObject* InThis, FReferenceCollector& Collector) 
 {
-#if WITH_EDITOR
 	UEdGraphNode_Comment* This = CastChecked<UEdGraphNode_Comment>(InThis);
 	for (auto It = This->NodesUnderComment.CreateIterator(); It; ++It)
 	{
 		Collector.AddReferencedObject(*It, This);
 	}
-#endif // WITH_EDITOR
 
 	Super::AddReferencedObjects(InThis, Collector);
 }
 
-#if WITH_EDITOR
 void UEdGraphNode_Comment::PostPlacedNewNode()
 {
 	//@TODO: Consider making the default value we use here a preference
@@ -152,7 +146,6 @@ TSharedPtr<class INameValidatorInterface> UEdGraphNode_Comment::MakeNameValidato
 	// Comments can be duplicated, etc...
 	return MakeShareable(new FDummyNameValidator(EValidatorResult::Ok));
 }
-#endif // WITH_EDITOR
 
 /////////////////////////////////////////////////////
 
diff --git a/Engine/Source/Editor/UnrealEd/Public/EdGraphNode_Comment.h b/Engine/Source/Editor/UnrealEd/Public/EdGraphNode_Comment.h
--- a/Engine/Source/Editor/UnrealEd/Public/EdGraphNode_Comment.h
+++ b/Engine/Source/Editor/UnrealEd/Public/EdGraphNode_Comment.h
@@ -43,7 +43,6 @@ public:
 	UNREALED_API static void AddReferencedObjects(UObject* InThis, FReferenceCollector& Collector);
 	// End UObject Interface
 
-#if WITH_EDITOR
 	// Begin UEdGraphNode interface
 	virtual void AllocateDefaultPins() override {}
 	UNREALED_API virtual FText GetTooltipText() const override;
@@ -79,6 +78,5 @@ private:
 
 	/** Constructing FText strings can be costly, so we cache the node's tooltip */
 	FNodeTextCache CachedTooltip;
-#endif
 };
 
