Moved project packaging settings defaults to BaseGame.ini from the class constructor. This allows the defaults to be properly used by automation scripts  and inherited by new projects.

diff --git a/Engine/Config/BaseGame.ini b/Engine/Config/BaseGame.ini
--- a/Engine/Config/BaseGame.ini
+++ b/Engine/Config/BaseGame.ini
@@ -32,6 +32,14 @@ ProjectVersion=
 Homepage=
 SupportContact=
 
+[/Script/UnrealEd.ProjectPackagingSettings]
+BuildConfiguration=PPBC_Development
+FullRebuild=True
+UsePakFile=True
+IncludePrerequisites=True
+UseOBB_InAPK=False
++CulturesToStage=en
+
 [/Script/Engine.HUD]
 ConsoleMessageCount=4
 ConsoleFontSize=5
diff --git a/Engine/Source/Editor/UnrealEd/Private/Settings/SettingsClasses.cpp b/Engine/Source/Editor/UnrealEd/Private/Settings/SettingsClasses.cpp
--- a/Engine/Source/Editor/UnrealEd/Private/Settings/SettingsClasses.cpp
+++ b/Engine/Source/Editor/UnrealEd/Private/Settings/SettingsClasses.cpp
@@ -273,12 +273,6 @@ void ULevelEditorViewportSettings::PostEditChangeProperty( struct FPropertyChang
 UProjectPackagingSettings::UProjectPackagingSettings( const class FPostConstructInitializeProperties& PCIP )
 	: Super(PCIP)
 {
-	BuildConfiguration = PPBC_Development;
-	FullRebuild = true;
-	UsePakFile = true;
-	IncludePrerequisites = true;
-	UseOBB_InAPK = false;
-	CulturesToStage.Add("en");
 }
 
 
