Fixed crash when exiting multi-client PIE after localhost connection error

diff --git a/Engine/Source/Editor/UnrealEd/Private/PlayLevel.cpp b/Engine/Source/Editor/UnrealEd/Private/PlayLevel.cpp
--- a/Engine/Source/Editor/UnrealEd/Private/PlayLevel.cpp
+++ b/Engine/Source/Editor/UnrealEd/Private/PlayLevel.cpp
@@ -1555,9 +1555,12 @@ void UEditorEngine::RequestEndPlayMap()
 					FSlatePlayInEditorInfo* const SlatePlayInEditorSession = SlatePlayInEditorMap.Find(ThisContext.ContextHandle);
 					if ((SlatePlayInEditorSession != nullptr) && (SlatePlayInEditorSession->EditorPlayer.IsValid() == true) && (SlatePlayInEditorSession->EditorPlayer.Get()->PlayerController != nullptr ) )
 					{
-						SlatePlayInEditorSession->EditorPlayer.Get()->PlayerController->GetPlayerViewPoint(LastViewLocation, LastViewRotation);
-						bLastViewAndLocationValid = true;
-						break;
+						if( SlatePlayInEditorSession->EditorPlayer.Get()->PlayerController != nullptr )
+						{
+							SlatePlayInEditorSession->EditorPlayer.Get()->PlayerController->GetPlayerViewPoint( LastViewLocation, LastViewRotation );
+							bLastViewAndLocationValid = true;
+							break;
+						}
 					}
 				}
 			}
