Improved progress reporting in the editor

FScopedSlowTask has been refactored to better allow for nesting of slow operations. This allows us to cascade nested scopes and provide accurate feedback on slow tasks. FScopedSlowTasks now work together when nested inside sub functions. Break up long functions that contain calls to multiple nested FScopedSlowTasks with FScopedSlowTask::EnterProgressFrame().

Example Usage:
void DoSlowWork()
{
    FScopedSlowTask Progress(2.f, LOCTEXT("DoingSlowWork", "Doing Slow Work..."));
    // Optionally make this show a dialog if not already shown
    Progress.MakeDialog();

    // Indicate that we are entering a frame representing 1 unit of work
    Progress.EnterProgressFrame(1.f);

    // DoFirstThing() can follow a similar pattern of creating a scope divided into frames. These contribute to their parent's progress frame proportionately.
    DoFirstThing();

    Progress.EnterProgressFrame(1.f);
    DoSecondThing();
}

This addresses TTP#338602 - NEEDS REVIEW: Editor progress bars nearly always just show 100%, don't offer useful indication of progress

diff --git a/Engine/Plugins/2D/Paper2D/Source/Paper2DEditor/Private/ContentBrowserExtensions/ContentBrowserExtensions.cpp b/Engine/Plugins/2D/Paper2D/Source/Paper2DEditor/Private/ContentBrowserExtensions/ContentBrowserExtensions.cpp
--- a/Engine/Plugins/2D/Paper2D/Source/Paper2DEditor/Private/ContentBrowserExtensions/ContentBrowserExtensions.cpp
+++ b/Engine/Plugins/2D/Paper2D/Source/Paper2DEditor/Private/ContentBrowserExtensions/ContentBrowserExtensions.cpp
@@ -87,13 +87,18 @@ struct FCreateSpriteFromTextureExtension : public FContentBrowserSelectedAssetEx
 			}
 			else
 			{
-				GWarn->BeginSlowTask(NSLOCTEXT("Paper2D", "Paper2D_ExtractSpritesFromTexture", "Extracting Sprites From Texture"), true, true);
+				FScopedSlowTask Feedback(1, NSLOCTEXT("Paper2D", "Paper2D_ExtractSpritesFromTexture", "Extracting Sprites From Texture"));
+				Feedback.MakeDialog(true);
 
+				// First extract the rects from the texture
 				TArray<FIntRect> ExtractedRects;
 				UPaperSprite::ExtractRectsFromTexture(Texture, /*out*/ ExtractedRects);
+
+				Feedback.TotalAmountOfWork = ExtractedRects.Num();
+
 				for (int ExtractedRectIndex = 0; ExtractedRectIndex < ExtractedRects.Num(); ++ExtractedRectIndex)
 				{
-					GWarn->StatusUpdate(ExtractedRectIndex, ExtractedRects.Num(), NSLOCTEXT("Paper2D", "Paper2D_ExtractSpritesFromTexture", "Extracting Sprites From Texture"));
+					Feedback.EnterProgressFrame(1, NSLOCTEXT("Paper2D", "Paper2D_ExtractSpritesFromTexture", "Extracting Sprites From Texture"));
 
 					FIntRect& ExtractedRect = ExtractedRects[ExtractedRectIndex];
 					SpriteFactory->bUseSourceRegion = true;
@@ -114,8 +119,6 @@ struct FCreateSpriteFromTextureExtension : public FContentBrowserSelectedAssetEx
 						break;
 					}
 				}
-
-				GWarn->EndSlowTask();
 			}
 		}
 
diff --git a/Engine/Source/Developer/AssetTools/Private/AssetFixUpRedirectors.cpp b/Engine/Source/Developer/AssetTools/Private/AssetFixUpRedirectors.cpp
--- a/Engine/Source/Developer/AssetTools/Private/AssetFixUpRedirectors.cpp
+++ b/Engine/Source/Developer/AssetTools/Private/AssetFixUpRedirectors.cpp
@@ -140,13 +140,16 @@ bool FAssetFixUpRedirectors::UpdatePackageStatus(const TArray<FRedirectorRefs>&
 
 void FAssetFixUpRedirectors::LoadReferencingPackages(TArray<FRedirectorRefs>& RedirectorsToFix, TArray<UPackage*>& OutReferencingPackagesToSave) const
 {
-	const bool bAllowNewSlowTask = true;
-	FScopedSlowTask SlowTaskMessage( LOCTEXT( "LoadingReferencingPackages", "Loading Referencing Packages..." ), bAllowNewSlowTask );
+	FScopedSlowTask SlowTask( RedirectorsToFix.Num(), LOCTEXT( "LoadingReferencingPackages", "Loading Referencing Packages..." ) );
+	SlowTask.MakeDialog();
+
 	ISourceControlProvider& SourceControlProvider = ISourceControlModule::Get().GetProvider();
 
 	// Load all packages that reference each redirector, if possible
 	for ( auto RedirectorRefsIt = RedirectorsToFix.CreateIterator(); RedirectorRefsIt; ++RedirectorRefsIt )
 	{
+		SlowTask.EnterProgressFrame(1);
+
 		FRedirectorRefs& RedirectorRefs = *RedirectorRefsIt;
 		if ( ISourceControlModule::Get().IsEnabled() )
 		{
diff --git a/Engine/Source/Developer/AssetTools/Private/AssetTools.cpp b/Engine/Source/Developer/AssetTools/Private/AssetTools.cpp
--- a/Engine/Source/Developer/AssetTools/Private/AssetTools.cpp
+++ b/Engine/Source/Developer/AssetTools/Private/AssetTools.cpp
@@ -429,15 +429,22 @@ TArray<UObject*> FAssetTools::ImportAssets(const TArray<FString>& Files, const F
 	TArray<UObject*> ReturnObjects;
 	TMap< FString, TArray<UFactory*> > ExtensionToFactoriesMap;
 
-	GWarn->BeginSlowTask(LOCTEXT("ImportSlowTask", "Importing"), true);
+	FScopedSlowTask SlowTask(Files.Num() + 3, LOCTEXT("ImportSlowTask", "Importing"));
+	SlowTask.MakeDialog();
 
 	// Reset the 'Do you want to overwrite the existing object?' Yes to All / No to All prompt, to make sure the
 	// user gets a chance to select something
 	UFactory::ResetState();
 
+	SlowTask.EnterProgressFrame();
+
 	TArray<TPair<FString, FString>> FilesAndDestinations;
 	ExpandDirectories(Files, RootDestinationPath, FilesAndDestinations);
 
+
+
+	SlowTask.EnterProgressFrame(1, LOCTEXT("Import_DeterminingImportTypes", "Determining asset types"));
+
 	if (SpecifiedFactory == NULL)
 	{	
 		// First instantiate one factory for each file extension encountered that supports the extension
@@ -546,6 +553,8 @@ TArray<UObject*> FAssetTools::ImportAssets(const TArray<FString>& Files, const F
 		const FString& Filename = FilesAndDestinations[FileIdx].Key;
 		const FString& DestinationPath = FilesAndDestinations[FileIdx].Value;
 
+		SlowTask.EnterProgressFrame(1, FText::Format(LOCTEXT("Import_ImportingFile", "Importing \"{0}\"..."), FText::FromString(FPaths::GetBaseFilename(Filename))));
+
 		FString FileExtension = FPaths::GetExtension(Filename);
 
 		const TArray<UFactory*>* FactoriesPtr = ExtensionToFactoriesMap.Find(FileExtension);
@@ -766,6 +775,8 @@ TArray<UObject*> FAssetTools::ImportAssets(const TArray<FString>& Files, const F
 		}
 	}
 
+	SlowTask.EnterProgressFrame(1);
+
 	// Clean up and remove the factories we created from the root set
 	for ( auto ExtensionIt = ExtensionToFactoriesMap.CreateConstIterator(); ExtensionIt; ++ExtensionIt)
 	{
@@ -776,8 +787,6 @@ TArray<UObject*> FAssetTools::ImportAssets(const TArray<FString>& Files, const F
 		}
 	}
 
-	GWarn->EndSlowTask();
-
 	// Sync content browser to the newly created assets
 	if ( ReturnObjects.Num() )
 	{
@@ -1262,11 +1271,13 @@ void FAssetTools::PerformMigratePackages(TArray<FName> PackageNamesToMigrate) co
 	// Form a full list of packages to move by including the dependencies of the supplied packages
 	TSet<FName> AllPackageNamesToMove;
 	{
-		const bool bAllowNewSlowTask = true;
-		FScopedSlowTask SlowTaskGatheringDependenciesMessage( LOCTEXT( "MigratePackages_GatheringDependencies", "Gathering Dependencies..." ) );
+		FScopedSlowTask SlowTask( PackageNamesToMigrate.Num(), LOCTEXT( "MigratePackages_GatheringDependencies", "Gathering Dependencies..." ) );
+		SlowTask.MakeDialog();
 
 		for ( auto PackageIt = PackageNamesToMigrate.CreateConstIterator(); PackageIt; ++PackageIt )
 		{
+			SlowTask.EnterProgressFrame();
+
 			if ( !AllPackageNamesToMove.Contains(*PackageIt) )
 			{
 				AllPackageNamesToMove.Add(*PackageIt);
@@ -1359,79 +1370,88 @@ void FAssetTools::MigratePackages_ReportConfirmed(TArray<FString> ConfirmedPacka
 	bool bUserCanceled = false;
 
 	// Copy all specified assets and their dependencies to the destination folder
-	const bool bAllowNewSlowTask = true;
-	FScopedSlowTask SlowTaskMessage( LOCTEXT( "MigratePackages_CopyingFiles", "Copying Files..." ), bAllowNewSlowTask );
+	FScopedSlowTask SlowTask( 2, LOCTEXT( "MigratePackages_CopyingFiles", "Copying Files..." ) );
+	SlowTask.MakeDialog();
 
 	EAppReturnType::Type LastResponse = EAppReturnType::Yes;
 	TArray<FString> SuccessfullyCopiedFiles;
 	TArray<FString> SuccessfullyCopiedPackages;
 	FString CopyErrors;
-	for ( auto PackageNameIt = ConfirmedPackageNamesToMigrate.CreateConstIterator(); PackageNameIt; ++PackageNameIt )
+
+	SlowTask.EnterProgressFrame();
 	{
-		const FString& PackageName = *PackageNameIt;
-		FString SrcFilename;
-		if ( !FPackageName::DoesPackageExist(PackageName, NULL, &SrcFilename) )
-		{
-			const FText ErrorMessage = FText::Format(LOCTEXT("MigratePackages_PackageMissing", "{0} does not exist on disk."), FText::FromString(PackageName));
-			UE_LOG(LogAssetTools, Warning, TEXT("%s"), *ErrorMessage.ToString());
-			CopyErrors += ErrorMessage.ToString() + LINE_TERMINATOR;
-		}
-		else if (SrcFilename.Contains(FPaths::EngineContentDir()))
-		{
-			const FString LeafName = SrcFilename.Replace(*FPaths::EngineContentDir(), TEXT("Engine/"));
-			CopyErrors += FText::Format(LOCTEXT("MigratePackages_EngineContent", "Unable to migrate Engine asset {0}. Engine assets cannot be migrated."), FText::FromString(LeafName)).ToString() + LINE_TERMINATOR;
-		}
-		else
+		FScopedSlowTask LoopProgress(ConfirmedPackageNamesToMigrate.Num());
+		for ( auto PackageNameIt = ConfirmedPackageNamesToMigrate.CreateConstIterator(); PackageNameIt; ++PackageNameIt )
 		{
-			const FString DestFilename = SrcFilename.Replace(*FPaths::GameContentDir(), *DestinationFolder);
+			LoopProgress.EnterProgressFrame();
 
-			bool bFileOKToCopy = true;
-			if ( IFileManager::Get().FileSize(*DestFilename) > 0 )
+			const FString& PackageName = *PackageNameIt;
+			FString SrcFilename;
+			if ( !FPackageName::DoesPackageExist(PackageName, NULL, &SrcFilename) )
 			{
-				// The destination file already exists! Ask the user what to do.
-				EAppReturnType::Type Response;
-				if ( LastResponse == EAppReturnType::YesAll || LastResponse == EAppReturnType::NoAll )
-				{
-					Response = LastResponse;
-				}
-				else
+				const FText ErrorMessage = FText::Format(LOCTEXT("MigratePackages_PackageMissing", "{0} does not exist on disk."), FText::FromString(PackageName));
+				UE_LOG(LogAssetTools, Warning, TEXT("%s"), *ErrorMessage.ToString());
+				CopyErrors += ErrorMessage.ToString() + LINE_TERMINATOR;
+			}
+			else if (SrcFilename.Contains(FPaths::EngineContentDir()))
+			{
+				const FString LeafName = SrcFilename.Replace(*FPaths::EngineContentDir(), TEXT("Engine/"));
+				CopyErrors += FText::Format(LOCTEXT("MigratePackages_EngineContent", "Unable to migrate Engine asset {0}. Engine assets cannot be migrated."), FText::FromString(LeafName)).ToString() + LINE_TERMINATOR;
+			}
+			else
+			{
+				const FString DestFilename = SrcFilename.Replace(*FPaths::GameContentDir(), *DestinationFolder);
+
+				bool bFileOKToCopy = true;
+				if ( IFileManager::Get().FileSize(*DestFilename) > 0 )
 				{
-					const FText Message = FText::Format( LOCTEXT("MigratePackages_AlreadyExists", "An asset already exists at location {0} would you like to overwrite it?"), FText::FromString(DestFilename) );
-					Response = FMessageDialog::Open( EAppMsgType::YesNoYesAllNoAllCancel, Message );
-					if ( Response == EAppReturnType::Cancel )
+					// The destination file already exists! Ask the user what to do.
+					EAppReturnType::Type Response;
+					if ( LastResponse == EAppReturnType::YesAll || LastResponse == EAppReturnType::NoAll )
 					{
-						// The user chose to cancel mid-operation. Break out.
-						bUserCanceled = true;
-						break;
+						Response = LastResponse;
+					}
+					else
+					{
+						const FText Message = FText::Format( LOCTEXT("MigratePackages_AlreadyExists", "An asset already exists at location {0} would you like to overwrite it?"), FText::FromString(DestFilename) );
+						Response = FMessageDialog::Open( EAppMsgType::YesNoYesAllNoAllCancel, Message );
+						if ( Response == EAppReturnType::Cancel )
+						{
+							// The user chose to cancel mid-operation. Break out.
+							bUserCanceled = true;
+							break;
+						}
+						LastResponse = Response;
 					}
-					LastResponse = Response;
-				}
 
-				const bool bWantOverwrite = Response == EAppReturnType::Yes || Response == EAppReturnType::YesAll;
-				if( !bWantOverwrite )
-				{
-					// User chose not to replace the package
-					bFileOKToCopy = false;
+					const bool bWantOverwrite = Response == EAppReturnType::Yes || Response == EAppReturnType::YesAll;
+					if( !bWantOverwrite )
+					{
+						// User chose not to replace the package
+						bFileOKToCopy = false;
+					}
 				}
-			}
 
-			if ( bFileOKToCopy )
-			{
-				if ( IFileManager::Get().Copy(*DestFilename, *SrcFilename) == COPY_OK )
-				{
-					SuccessfullyCopiedPackages.Add(PackageName);
-					SuccessfullyCopiedFiles.Add(DestFilename);
-				}
-				else
+				if ( bFileOKToCopy )
 				{
-					UE_LOG(LogAssetTools, Warning, TEXT("Failed to copy %s to %s while migrating assets"), *SrcFilename, *DestFilename);
-					CopyErrors += SrcFilename + LINE_TERMINATOR;
+					if ( IFileManager::Get().Copy(*DestFilename, *SrcFilename) == COPY_OK )
+					{
+						SuccessfullyCopiedPackages.Add(PackageName);
+						SuccessfullyCopiedFiles.Add(DestFilename);
+					}
+					else
+					{
+						UE_LOG(LogAssetTools, Warning, TEXT("Failed to copy %s to %s while migrating assets"), *SrcFilename, *DestFilename);
+						CopyErrors += SrcFilename + LINE_TERMINATOR;
+					}
 				}
 			}
 		}
 	}
 
 	FString SourceControlErrors;
+	SlowTask.EnterProgressFrame();
+
 	if ( !bUserCanceled && SuccessfullyCopiedFiles.Num() > 0 )
 	{
 		// attempt to add files to source control (this can quite easily fail, but if it works it is very useful)
@@ -1442,8 +1462,11 @@ void FAssetTools::MigratePackages_ReportConfirmed(TArray<FString> ConfirmedPacka
 				ISourceControlProvider& SourceControlProvider = ISourceControlModule::Get().GetProvider();
 				if(SourceControlProvider.Execute(ISourceControlOperation::Create<FMarkForAdd>(), SuccessfullyCopiedFiles) == ECommandResult::Failed)
 				{
+					FScopedSlowTask LoopProgress(SuccessfullyCopiedFiles.Num());
+
 					for(auto FileIt(SuccessfullyCopiedFiles.CreateConstIterator()); FileIt; FileIt++)
 					{
+						LoopProgress.EnterProgressFrame();
 						if(!SourceControlProvider.GetState(*FileIt, EStateCacheUsage::Use)->IsAdded())
 						{
 							SourceControlErrors += FText::Format(LOCTEXT("MigratePackages_SourceControlError", "{0} could not be added to source control"), FText::FromString(*FileIt)).ToString();
diff --git a/Engine/Source/Developer/DesktopPlatform/Private/Mac/MacNativeFeedbackContext.cpp b/Engine/Source/Developer/DesktopPlatform/Private/Mac/MacNativeFeedbackContext.cpp
--- a/Engine/Source/Developer/DesktopPlatform/Private/Mac/MacNativeFeedbackContext.cpp
+++ b/Engine/Source/Developer/DesktopPlatform/Private/Mac/MacNativeFeedbackContext.cpp
@@ -398,18 +398,18 @@ bool FMacNativeFeedbackContext::ReceivedUserCancel()
 	return bReceivedUserCancel;
 }
 
-void FMacNativeFeedbackContext::BeginSlowTask( const FText& Task, bool bShowProgressDialog, bool bInShowCancelButton )
+void FMacNativeFeedbackContext::StartSlowTask(const FText& Task, bool bInShowCancelButton)
 {
-	GIsSlowTask = ++SlowTaskCount>0;
+	FFeedbackContext::StartSlowTask(Task, bInShowCancelButton);
 
-	if(SlowTaskCount > 0 && bShowProgressDialog && WindowController != NULL && !bShowingConsoleForSlowTask)
+	if(WindowController != NULL && !bShowingConsoleForSlowTask)
 	{
 		MainThreadCall(^{
 			[WindowController setTitleText:Task.ToString().GetNSString()];
 			[WindowController setStatusText:@"Progress:"];
 			[WindowController setShowCancelButton:bInShowCancelButton];
 			[WindowController setShowProgress:true];
-			[WindowController setProgress:0 total:100];
+			[WindowController setProgress:0 total:1];
 			[WindowController setShowProgress:false];
 			
 			SetDefaultTextColor();
@@ -421,12 +421,11 @@ void FMacNativeFeedbackContext::BeginSlowTask( const FText& Task, bool bShowProg
 	}
 }
 
-void FMacNativeFeedbackContext::EndSlowTask()
+void FMacNativeFeedbackContext::FinalizeSlowTask()
 {
-	check(SlowTaskCount>0);
-	GIsSlowTask = --SlowTaskCount>0;
+	FFeedbackContext::FinalizeSlowTask();
 
-	if(SlowTaskCount == 0 && bShowingConsoleForSlowTask)
+	if(bShowingConsoleForSlowTask)
 	{
 		MainThreadCall(^{
 			if(WindowController != NULL)
@@ -438,32 +437,16 @@ void FMacNativeFeedbackContext::EndSlowTask()
 	}
 }
 
-bool FMacNativeFeedbackContext::StatusUpdate( int32 Numerator, int32 Denominator, const FText& NewStatus )
+void FMacNativeFeedbackContext::ProgressReported( const float TotalProgressInterp, FText DisplayMessage )
 {
 	if(WindowController != NULL && bShowingConsoleForSlowTask)
 	{
 		MainThreadCall(^{
-			[WindowController setStatusText:NewStatus.ToString().GetNSString()];
+			[WindowController setStatusText:DisplayMessage.ToString().GetNSString()];
+			[WindowController setShowProgress:true];
+			[WindowController setProgress:TotalProgressInterp total:1];
 		});
-		UpdateProgress(Numerator, Denominator);
 	}
-	return true;
-}
-
-bool FMacNativeFeedbackContext::StatusForceUpdate( int32 Numerator, int32 Denominator, const FText& StatusText )
-{
-	return StatusUpdate(Numerator, Denominator, StatusText);
-}
-
-void FMacNativeFeedbackContext::UpdateProgress(int32 Numerator, int32 Denominator)
-{
-	MainThreadCall(^{
-		if(WindowController != NULL && bShowingConsoleForSlowTask)
-		{
-			[WindowController setShowProgress:true];
-			[WindowController setProgress:Numerator total:Denominator];
-		}
-	});
 }
 
 FContextSupplier* FMacNativeFeedbackContext::GetContext() const
diff --git a/Engine/Source/Developer/DesktopPlatform/Private/Mac/MacNativeFeedbackContext.h b/Engine/Source/Developer/DesktopPlatform/Private/Mac/MacNativeFeedbackContext.h
--- a/Engine/Source/Developer/DesktopPlatform/Private/Mac/MacNativeFeedbackContext.h
+++ b/Engine/Source/Developer/DesktopPlatform/Private/Mac/MacNativeFeedbackContext.h
@@ -42,12 +42,10 @@ public:
 	virtual bool YesNof(const FText& Text) override;
 
 	virtual bool ReceivedUserCancel() override;
-	virtual void BeginSlowTask( const FText& Task, bool bShowProgressDialog, bool bInShowCancelButton=false ) override;
-	virtual void EndSlowTask() override;
 
-	virtual bool StatusUpdate( int32 Numerator, int32 Denominator, const FText& NewStatus ) override;
-	virtual bool StatusForceUpdate( int32 Numerator, int32 Denominator, const FText& StatusText ) override;
-	virtual void UpdateProgress(int32 Numerator, int32 Denominator) override;
+	virtual void StartSlowTask( const FText& Task, bool bShowCancelButton=false ) override;
+	virtual void FinalizeSlowTask( ) override;
+	virtual void ProgressReported( const float TotalProgressInterp, FText DisplayMessage ) override;
 
 	FContextSupplier* GetContext() const;
 	void SetContext( FContextSupplier* InSupplier );
diff --git a/Engine/Source/Developer/DesktopPlatform/Private/Windows/WindowsNativeFeedbackContext.cpp b/Engine/Source/Developer/DesktopPlatform/Private/Windows/WindowsNativeFeedbackContext.cpp
--- a/Engine/Source/Developer/DesktopPlatform/Private/Windows/WindowsNativeFeedbackContext.cpp
+++ b/Engine/Source/Developer/DesktopPlatform/Private/Windows/WindowsNativeFeedbackContext.cpp
@@ -8,7 +8,6 @@
 FWindowsNativeFeedbackContext::FWindowsNativeFeedbackContext()
 	: FFeedbackContext()
 	, Context( NULL )
-	, SlowTaskCount( 0 )
 	, hThread( NULL )
 	, hCloseEvent( NULL )
 	, hUpdateEvent( NULL )
@@ -69,7 +68,7 @@ void FWindowsNativeFeedbackContext::Serialize( const TCHAR* V, ELogVerbosity::Ty
 	}
 
 	// Buffer up the output during a slow task so that we can dump it all to the log console if the show log button is clicked
-	if(SlowTaskCount > 0)
+	if(GIsSlowTask)
 	{
 		FScopeLock Lock(&CriticalSection);
 		if(hThread != NULL)
@@ -100,58 +99,27 @@ bool FWindowsNativeFeedbackContext::ReceivedUserCancel()
 	return bReceivedUserCancel;
 }
 
-void FWindowsNativeFeedbackContext::BeginSlowTask( const FText& Task, bool bShowProgressDialog, bool bInShowCancelButton )
+void FWindowsNativeFeedbackContext::StartSlowTask( const FText& Task, bool bShowCancelButton )
 {
-	GIsSlowTask = ++SlowTaskCount>0;
-
-	if(SlowTaskCount > 0 && bShowProgressDialog)
-	{
-		CreateSlowTaskWindow(Task, bInShowCancelButton);
-	}
+	FFeedbackContext::StartSlowTask( Task, bShowCancelButton );
+	CreateSlowTaskWindow(Task, bShowCancelButton);
 }
-
-void FWindowsNativeFeedbackContext::EndSlowTask()
+void FWindowsNativeFeedbackContext::FinalizeSlowTask( )
 {
-	check(SlowTaskCount>0);
-	GIsSlowTask = --SlowTaskCount>0;
-
-	if(SlowTaskCount == 0)
-	{
-		DestroySlowTaskWindow();
-	}
+	FFeedbackContext::FinalizeSlowTask( );
+	DestroySlowTaskWindow();
 }
 
-bool FWindowsNativeFeedbackContext::StatusUpdate( int32 Numerator, int32 Denominator, const FText& NewStatus )
+void FWindowsNativeFeedbackContext::ProgressReported( const float TotalProgressInterp, FText DisplayMessage )
 {
 	FScopeLock Lock(&CriticalSection);
 	if(hThread != NULL)
 	{
-		UpdateProgress(Numerator, Denominator);
-
-		Status = NewStatus.ToString();
+		Progress = TotalProgressInterp;
+		Status = DisplayMessage.ToString();
 
 		SetEvent(hUpdateEvent);
 	}
-	return true;
-}
-
-bool FWindowsNativeFeedbackContext::StatusForceUpdate( int32 Numerator, int32 Denominator, const FText& StatusText )
-{
-	return StatusUpdate(Numerator, Denominator, StatusText);
-}
-
-void FWindowsNativeFeedbackContext::UpdateProgress(int32 Numerator, int32 Denominator)
-{
-	FScopeLock Lock(&CriticalSection);
-	if(hThread != NULL)
-	{
-		float NewProgress = (float)Numerator / (float)Denominator;
-		if(Progress != NewProgress)
-		{
-			Progress = NewProgress;
-			SetEvent(hUpdateEvent);
-		}
-	}
 }
 
 FContextSupplier* FWindowsNativeFeedbackContext::GetContext() const
diff --git a/Engine/Source/Developer/DesktopPlatform/Private/Windows/WindowsNativeFeedbackContext.h b/Engine/Source/Developer/DesktopPlatform/Private/Windows/WindowsNativeFeedbackContext.h
--- a/Engine/Source/Developer/DesktopPlatform/Private/Windows/WindowsNativeFeedbackContext.h
+++ b/Engine/Source/Developer/DesktopPlatform/Private/Windows/WindowsNativeFeedbackContext.h
@@ -21,12 +21,9 @@ public:
 	VARARG_BODY( bool, YesNof, const TCHAR*, VARARG_NONE );
 
 	virtual bool ReceivedUserCancel() override;
-	virtual void BeginSlowTask( const FText& Task, bool bShowProgressDialog, bool bInShowCancelButton=false ) override;
-	virtual void EndSlowTask() override;
-
-	virtual bool StatusUpdate( int32 Numerator, int32 Denominator, const FText& NewStatus ) override;
-	virtual bool StatusForceUpdate( int32 Numerator, int32 Denominator, const FText& StatusText ) override;
-	virtual void UpdateProgress(int32 Numerator, int32 Denominator) override;
+	virtual void StartSlowTask( const FText& Task, bool bShowCancelButton=false ) override;
+	virtual void FinalizeSlowTask( ) override;
+	virtual void ProgressReported( const float TotalProgressInterp, FText DisplayMessage ) override;
 
 	FContextSupplier* GetContext() const;
 	void SetContext( FContextSupplier* InSupplier );
@@ -55,7 +52,6 @@ private:
 	static const uint16 LogOutputCtlId = 203;
 
 	FContextSupplier* Context;
-	int32 SlowTaskCount;
 	HANDLE hThread;
 	HANDLE hCloseEvent;
 	HANDLE hUpdateEvent;
diff --git a/Engine/Source/Developer/HotReload/Private/HotReload.cpp b/Engine/Source/Developer/HotReload/Private/HotReload.cpp
--- a/Engine/Source/Developer/HotReload/Private/HotReload.cpp
+++ b/Engine/Source/Developer/HotReload/Private/HotReload.cpp
@@ -220,7 +220,7 @@ private:
 	bool InvokeUnrealBuildToolForCompile(const FString& InCmdLineParams, FOutputDevice &Ar);
 
 	/** Checks to see if a pending compilation action has completed and optionally waits for it to finish.  If completed, fires any appropriate callbacks and reports status provided bFireEvents is true. */
-	void CheckForFinishedModuleDLLCompile(const bool bWaitForCompletion, bool& bCompileStillInProgress, bool& bCompileSucceeded, FOutputDevice& Ar, const FText& SlowTaskOverrideText = FText::GetEmpty(), bool bFireEvents = true);
+	void CheckForFinishedModuleDLLCompile(const bool bWaitForCompletion, bool& bCompileStillInProgress, bool& bCompileSucceeded, FOutputDevice& Ar, bool bFireEvents = true);
 
 	/** Called when the compile data for a module need to be update in memory and written to config */
 	void UpdateModuleCompileData(FName ModuleName);
@@ -402,14 +402,13 @@ FString FHotReloadModule::GetModuleCompileMethod(FName InModuleName)
 bool FHotReloadModule::RecompileModule( const FName InModuleName, const bool bReloadAfterRecompile, FOutputDevice &Ar )
 {
 #if WITH_HOT_RELOAD
-	const bool bShowProgressDialog = true;
-	const bool bShowCancelButton = false;
 
 	FFormatNamedArguments Args;
 	Args.Add( TEXT("CodeModuleName"), FText::FromName( InModuleName ) );
 	const FText StatusUpdate = FText::Format( NSLOCTEXT("ModuleManager", "Recompile_SlowTaskName", "Compiling {CodeModuleName}..."), Args );
 
-	GWarn->BeginSlowTask( StatusUpdate, bShowProgressDialog, bShowCancelButton );
+	FScopedSlowTask SlowTask(1, StatusUpdate);
+	SlowTask.MakeDialog();
 
 	ModuleCompilerStartedEvent.Broadcast();
 
@@ -468,7 +467,6 @@ bool FHotReloadModule::RecompileModule( const FName InModuleName, const bool bRe
 		}
 	}
 
-	GWarn->EndSlowTask();
 	return bWasSuccessful;
 #else
 	return false;
@@ -891,26 +889,29 @@ void FHotReloadModule::DoHotReloadFromIDE()
 
 		UE_LOG(LogHotReload, Log, TEXT("Starting Hot-Reload from IDE"));
 
-		GWarn->BeginSlowTask(LOCTEXT("CompilingGameCode", "Compiling Game Code"), true);
+		FScopedSlowTask SlowTask(100.f, LOCTEXT("CompilingGameCode", "Compiling Game Code"));
+		SlowTask.MakeDialog();
 
 		// Update compile data before we start compiling
 		for (auto& NewModule : NewModules)
 		{
+			// Move on 10% / num items
+			SlowTask.EnterProgressFrame(10.f/NewModules.Num());
+
 			UpdateModuleCompileData(*NewModule.Name);
 			OnModuleCompileSucceeded(*NewModule.Name, NewModule.NewFilename);
 		}
 
-
+		SlowTask.EnterProgressFrame(10);
 		GetPackagesToRebindAndDependentModules(GameModuleNames, PackagesToRebind, DependentModules);
 
+		SlowTask.EnterProgressFrame(80);
 		check(PackagesToRebind.Num() || DependentModules.Num())
 		{
 			const bool bRecompileFinished = true;
 			const ECompilationResult::Type RecompileResult = ECompilationResult::Succeeded;
 			Result = DoHotReloadInternal(bRecompileFinished, RecompileResult, PackagesToRebind, DependentModules, *GLog);
 		}
-
-		GWarn->EndSlowTask();
 	}
 
 	RecordAnalyticsEvent(TEXT("IDE"), Result, Duration, PackagesToRebind.Num(), DependentModules.Num());
@@ -1197,7 +1198,7 @@ bool FHotReloadModule::InvokeUnrealBuildToolForCompile(const FString& InCmdLineP
 #endif // WITH_HOT_RELOAD
 }
 
-void FHotReloadModule::CheckForFinishedModuleDLLCompile(const bool bWaitForCompletion, bool& bCompileStillInProgress, bool& bCompileSucceeded, FOutputDevice& Ar, const FText& SlowTaskOverrideText, bool bFireEvents)
+void FHotReloadModule::CheckForFinishedModuleDLLCompile(const bool bWaitForCompletion, bool& bCompileStillInProgress, bool& bCompileSucceeded, FOutputDevice& Ar, bool bFireEvents)
 {
 #if WITH_HOT_RELOAD
 	bCompileStillInProgress = false;
@@ -1208,34 +1209,21 @@ void FHotReloadModule::CheckForFinishedModuleDLLCompile(const bool bWaitForCompl
 	{
 		bCompileStillInProgress = true;
 
-		// Ensure slow task messages are seen.
-		GWarn->PushStatus();
-
-		// Update the slow task dialog if we were summoned from a synchronous recompile path
-		if (GIsSlowTask)
+		FText StatusUpdate;
+		if ( ModulesBeingCompiled.Num() > 0 )
 		{
-			if ( !SlowTaskOverrideText.IsEmpty() )
-			{
-				GWarn->StatusUpdate(-1, -1, SlowTaskOverrideText);
-			}
-			else
-			{
-				FText StatusUpdate;
-				if ( ModulesBeingCompiled.Num() > 0 )
-				{
-					FFormatNamedArguments Args;
-					Args.Add( TEXT("CodeModuleName"), FText::FromString( ModulesBeingCompiled[0].ModuleName ) );
-					StatusUpdate = FText::Format( NSLOCTEXT("FModuleManager", "CompileSpecificModuleStatusMessage", "{CodeModuleName}: Compiling modules..."), Args );
-				}
-				else
-				{
-					StatusUpdate = NSLOCTEXT("FModuleManager", "CompileStatusMessage", "Compiling modules...");
-				}
-
-				GWarn->StatusUpdate(-1, -1, StatusUpdate);
-			}
+			FFormatNamedArguments Args;
+			Args.Add( TEXT("CodeModuleName"), FText::FromString( ModulesBeingCompiled[0].ModuleName ) );
+			StatusUpdate = FText::Format( NSLOCTEXT("FModuleManager", "CompileSpecificModuleStatusMessage", "{CodeModuleName}: Compiling modules..."), Args );
+		}
+		else
+		{
+			StatusUpdate = NSLOCTEXT("FModuleManager", "CompileStatusMessage", "Compiling modules...");
 		}
 
+		FScopedSlowTask SlowTask(0, StatusUpdate, GIsSlowTask);
+		SlowTask.MakeDialog();
+
 		// Check to see if the compile has finished yet
 		int32 ReturnCode = -1;
 		while (bCompileStillInProgress)
@@ -1273,9 +1261,6 @@ void FHotReloadModule::CheckForFinishedModuleDLLCompile(const bool bWaitForCompl
 		
 		bRequestCancelCompilation = false;
 
-		// Restore any status from before the loop - see PushStatus() above.
-		GWarn->PopStatus();
-
 		if( !bCompileStillInProgress )		
 		{
 			// Compilation finished, now we need to grab all of the text from the output pipe
diff --git a/Engine/Source/Editor/ContentBrowser/Private/AssetContextMenu.cpp b/Engine/Source/Editor/ContentBrowser/Private/AssetContextMenu.cpp
--- a/Engine/Source/Editor/ContentBrowser/Private/AssetContextMenu.cpp
+++ b/Engine/Source/Editor/ContentBrowser/Private/AssetContextMenu.cpp
@@ -781,21 +781,26 @@ void FAssetContextMenu::ExecuteFindAssetInWorld()
 
 	if (AssetsToFind.Num() > 0)
 	{
-		const bool ShowProgressDialog = true;
-		GWarn->BeginSlowTask(NSLOCTEXT("AssetContextMenu", "FindAssetInWorld", "Finding actors that use this asset..."), ShowProgressDialog);
+		FScopedSlowTask SlowTask(2 + AssetsToFind.Num(), NSLOCTEXT("AssetContextMenu", "FindAssetInWorld", "Finding actors that use this asset..."));
+		SlowTask.MakeDialog();
 
 		CollectGarbage(GARBAGE_COLLECTION_KEEPFLAGS);
 
 		TArray< TWeakObjectPtr<UObject> > OutObjects;
 		WorldReferenceGenerator ObjRefGenerator;
+
+		SlowTask.EnterProgressFrame();
 		ObjRefGenerator.BuildReferencingData();
 
 		for (int32 AssetIdx = 0; AssetIdx < AssetsToFind.Num(); ++AssetIdx)
 		{
+			SlowTask.EnterProgressFrame();
 			ObjRefGenerator.MarkAllObjects();
 			ObjRefGenerator.Generate(AssetsToFind[AssetIdx], OutObjects);
 		}
 
+		SlowTask.EnterProgressFrame();
+
 		if (OutObjects.Num() > 0)
 		{
 			const bool InSelected = true;
@@ -815,8 +820,6 @@ void FAssetContextMenu::ExecuteFindAssetInWorld()
 			Info.ExpireDuration = 3.0f;
 			FSlateNotificationManager::Get().AddNotification(Info);
 		}
-
-		GWarn->EndSlowTask();
 	}
 }
 
diff --git a/Engine/Source/Editor/ContentBrowser/Private/SContentBrowser.cpp b/Engine/Source/Editor/ContentBrowser/Private/SContentBrowser.cpp
--- a/Engine/Source/Editor/ContentBrowser/Private/SContentBrowser.cpp
+++ b/Engine/Source/Editor/ContentBrowser/Private/SContentBrowser.cpp
@@ -1329,17 +1329,23 @@ void SContentBrowser::OnAssetsActivated(const TArray<FAssetData>& ActivatedAsset
 	TMap< TSharedRef<IAssetTypeActions>, TArray<UObject*> > TypeActionsToObjects;
 	TArray<UObject*> ObjectsWithoutTypeActions;
 
+	const FText LoadingTemplate = LOCTEXT("LoadingAssetName", "Loading {0}...");
+	const FText DefaultText = ActivatedAssets.Num() == 1 ? FText::Format(LoadingTemplate, FText::FromName(ActivatedAssets[0].AssetName)) : LOCTEXT("LoadingObjects", "Loading Objects...");
+	FScopedSlowTask SlowTask(100, DefaultText);
+
 	// Iterate over all activated assets to map them to AssetTypeActions.
 	// This way individual asset type actions will get a batched list of assets to operate on
 	for ( auto AssetIt = ActivatedAssets.CreateConstIterator(); AssetIt; ++AssetIt )
 	{
 		const FAssetData& AssetData = *AssetIt;
-		const bool bShowProgressDialog = (!AssetData.IsAssetLoaded() && FEditorFileUtils::IsMapPackageAsset(AssetData.ObjectPath.ToString()));
-		GWarn->BeginSlowTask(LOCTEXT("LoadingObjects", "Loading Objects..."), bShowProgressDialog);
+		if (!AssetData.IsAssetLoaded() && FEditorFileUtils::IsMapPackageAsset(AssetData.ObjectPath.ToString()))
+		{
+			SlowTask.MakeDialog();
+		}
 
-		UObject* Asset = (*AssetIt).GetAsset();
+		SlowTask.EnterProgressFrame(75.f/ActivatedAssets.Num(), FText::Format(LoadingTemplate, FText::FromName(AssetData.AssetName)));
 
-		GWarn->EndSlowTask();
+		UObject* Asset = (*AssetIt).GetAsset();
 
 		if ( Asset != NULL )
 		{
@@ -1361,6 +1367,8 @@ void SContentBrowser::OnAssetsActivated(const TArray<FAssetData>& ActivatedAsset
 	// Now that we have created our map, activate all the lists of objects for each asset type action.
 	for ( auto TypeActionsIt = TypeActionsToObjects.CreateConstIterator(); TypeActionsIt; ++TypeActionsIt )
 	{
+		SlowTask.EnterProgressFrame(25.f/TypeActionsToObjects.Num());
+
 		const TSharedRef<IAssetTypeActions>& TypeActions = TypeActionsIt.Key();
 		const TArray<UObject*>& ObjList = TypeActionsIt.Value();
 
diff --git a/Engine/Source/Editor/EditorStyle/Private/SlateEditorStyle.cpp b/Engine/Source/Editor/EditorStyle/Private/SlateEditorStyle.cpp
--- a/Engine/Source/Editor/EditorStyle/Private/SlateEditorStyle.cpp
+++ b/Engine/Source/Editor/EditorStyle/Private/SlateEditorStyle.cpp
@@ -490,6 +490,9 @@ void FSlateEditorStyle::FStyle::SetupGeneralStyles()
 			.SetMarqueeImage( IMAGE_BRUSH( "Common/ProgressBar_Marquee", FVector2D(20,12), FLinearColor::White, ESlateBrushTileType::Horizontal ) )
 			);
 
+		Set( "ProgressBar.ThinBackground", new BOX_BRUSH( "Common/ProgressBar_Thin_Background", FMargin(5.f/12.f) ) );
+		Set( "ProgressBar.ThinFill", new BOX_BRUSH( "Common/ProgressBar_Thin_Fill", FMargin(5.f/12.f) ) );
+
 		// Legacy ProgressBar styles; kept here because other FEditorStyle controls (mis)use them
 		// todo: jdale - Widgets using these styles should be updated to use SlateStyle types once FEditorStyle has been obliterated from the Slate core
 		Set( "ProgressBar.Background", new BOX_BRUSH( "Common/ProgressBar_Background", FMargin(5.f/12.f) ) );
diff --git a/Engine/Source/Editor/GameProjectGeneration/Private/GameProjectGenerationModule.cpp b/Engine/Source/Editor/GameProjectGeneration/Private/GameProjectGenerationModule.cpp
--- a/Engine/Source/Editor/GameProjectGeneration/Private/GameProjectGenerationModule.cpp
+++ b/Engine/Source/Editor/GameProjectGeneration/Private/GameProjectGenerationModule.cpp
@@ -80,8 +80,8 @@ bool FGameProjectGenerationModule::UpdateGameProject(const FString& ProjectFile,
 
 bool FGameProjectGenerationModule::UpdateCodeProject(FText& OutFailReason)
 {
-	const bool bAllowNewSlowTask = true;
-	FScopedSlowTask SlowTaskMessage( LOCTEXT( "UpdatingCodeProject", "Updating code project..." ), bAllowNewSlowTask );
+	FScopedSlowTask SlowTask(0, LOCTEXT( "UpdatingCodeProject", "Updating code project..." ) );
+	SlowTask.MakeDialog();
 
 	return GameProjectUtils::GenerateCodeProjectFiles(FPaths::GetProjectFilePath(), OutFailReason);
 }
diff --git a/Engine/Source/Editor/GameProjectGeneration/Private/GameProjectUtils.cpp b/Engine/Source/Editor/GameProjectGeneration/Private/GameProjectUtils.cpp
--- a/Engine/Source/Editor/GameProjectGeneration/Private/GameProjectUtils.cpp
+++ b/Engine/Source/Editor/GameProjectGeneration/Private/GameProjectUtils.cpp
@@ -515,14 +515,16 @@ bool GameProjectUtils::CopyStarterContent(const FString& DestProjectFolder, FTex
 	TArray<FString> FilesToCopy;
 	GetStarterContentFiles(FilesToCopy);
 
+	FScopedSlowTask SlowTask(FilesToCopy.Num(), LOCTEXT("CreatingProjectStatus_CopyingFiles", "Copying Files {SrcFilename}..."));
+	SlowTask.MakeDialog();
+
 	TArray<FString> CreatedFiles;
 	for (FString SrcFilename : FilesToCopy)
 	{
 		// Update the slow task dialog
-		const bool bAllowNewSlowTask = false;
 		FFormatNamedArguments Args;
 		Args.Add(TEXT("SrcFilename"), FText::FromString(FPaths::GetCleanFilename(SrcFilename)));
-		FScopedSlowTask SlowTaskMessage(FText::Format(LOCTEXT("CreatingProjectStatus_CopyingFile", "Copying File {SrcFilename}..."), Args), bAllowNewSlowTask);
+		SlowTask.EnterProgressFrame(1, FText::Format(LOCTEXT("CreatingProjectStatus_CopyingFile", "Copying File {SrcFilename}..."), Args));
 
 		FString FileRelPath = FPaths::GetPath(SrcFilename);
 		FPaths::MakePathRelativeTo(FileRelPath, *SrcFolder);
@@ -558,8 +560,8 @@ bool GameProjectUtils::CreateProject(const FProjectInformation& InProjectInfo, F
 		return false;
 	}
 
-	const bool bAllowNewSlowTask = true;
-	FScopedSlowTask SlowTaskMessage( LOCTEXT( "CreatingProjectStatus", "Creating project..." ), bAllowNewSlowTask );
+	FScopedSlowTask SlowTask(0, LOCTEXT( "CreatingProjectStatus", "Creating project..." ));
+	SlowTask.MakeDialog();
 
 	bool bProjectCreationSuccessful = false;
 	FString TemplateName;
@@ -819,10 +821,14 @@ UTemplateProjectDefs* GameProjectUtils::LoadTemplateDefs(const FString& ProjectD
 
 bool GameProjectUtils::GenerateProjectFromScratch(const FProjectInformation& InProjectInfo, FText& OutFailReason)
 {
+	FScopedSlowTask SlowTask(5);
+
 	const FString NewProjectFolder = FPaths::GetPath(InProjectInfo.ProjectFilename);
 	const FString NewProjectName = FPaths::GetBaseFilename(InProjectInfo.ProjectFilename);
 	TArray<FString> CreatedFiles;
 
+	SlowTask.EnterProgressFrame();
+
 	// Generate config files
 	if (!GenerateConfigFiles(InProjectInfo, CreatedFiles, OutFailReason))
 	{
@@ -841,9 +847,14 @@ bool GameProjectUtils::GenerateProjectFromScratch(const FProjectInformation& InP
 		return false;
 	}
 
+	SlowTask.EnterProgressFrame();
+
 	TArray<FString> StartupModuleNames;
 	if ( InProjectInfo.bShouldGenerateCode )
 	{
+		FScopedSlowTask LocalScope(2);
+
+		LocalScope.EnterProgressFrame();
 		// Generate basic source code files
 		if ( !GenerateBasicSourceCode(NewProjectFolder / TEXT("Source"), NewProjectName, NewProjectFolder, StartupModuleNames, CreatedFiles, OutFailReason) )
 		{
@@ -851,6 +862,7 @@ bool GameProjectUtils::GenerateProjectFromScratch(const FProjectInformation& InP
 			return false;
 		}
 
+		LocalScope.EnterProgressFrame();
 		// Generate game framework source code files
 		if ( !GenerateGameFrameworkSourceCode(NewProjectFolder / TEXT("Source"), NewProjectName, CreatedFiles, OutFailReason) )
 		{
@@ -859,6 +871,8 @@ bool GameProjectUtils::GenerateProjectFromScratch(const FProjectInformation& InP
 		}
 	}
 
+	SlowTask.EnterProgressFrame();
+
 	// Generate the project file
 	{
 		FText LocalFailReason;
@@ -874,6 +888,8 @@ bool GameProjectUtils::GenerateProjectFromScratch(const FProjectInformation& InP
 		}
 	}
 
+	SlowTask.EnterProgressFrame();
+
 	if ( InProjectInfo.bShouldGenerateCode )
 	{
 		// Generate project files
@@ -885,6 +901,8 @@ bool GameProjectUtils::GenerateProjectFromScratch(const FProjectInformation& InP
 		}
 	}
 
+	SlowTask.EnterProgressFrame();
+
 	if (InProjectInfo.bCopyStarterContent)
 	{
 		// Copy the starter content
@@ -902,6 +920,8 @@ bool GameProjectUtils::GenerateProjectFromScratch(const FProjectInformation& InP
 
 bool GameProjectUtils::CreateProjectFromTemplate(const FProjectInformation& InProjectInfo, FText& OutFailReason)
 {
+	FScopedSlowTask SlowTask(10);
+
 	const FString ProjectName = FPaths::GetBaseFilename(InProjectInfo.ProjectFilename);
 	const FString TemplateName = FPaths::GetBaseFilename(InProjectInfo.TemplateFile);
 	const FString SrcFolder = FPaths::GetPath(InProjectInfo.TemplateFile);
@@ -915,6 +935,8 @@ bool GameProjectUtils::CreateProjectFromTemplate(const FProjectInformation& InPr
 		return false;
 	}
 
+	SlowTask.EnterProgressFrame();
+
 	UTemplateProjectDefs* TemplateDefs = LoadTemplateDefs(SrcFolder);
 	if ( TemplateDefs == NULL )
 	{
@@ -925,6 +947,8 @@ bool GameProjectUtils::CreateProjectFromTemplate(const FProjectInformation& InPr
 		return false;
 	}
 
+	SlowTask.EnterProgressFrame();
+
 	// Fix up the replacement strings using the specified project name
 	TemplateDefs->FixupStrings(TemplateName, ProjectName);
 
@@ -938,149 +962,166 @@ bool GameProjectUtils::CreateProjectFromTemplate(const FProjectInformation& InPr
 	// Keep a list of created files so we can delete them if project creation fails
 	TArray<FString> CreatedFiles;
 
+	SlowTask.EnterProgressFrame();
+
 	// Discover and copy all files in the src folder to the destination, excluding a few files and folders
 	TArray<FString> FilesToCopy;
 	TArray<FString> FilesThatNeedContentsReplaced;
 	TMap<FString, FString> ClassRenames;
 	IFileManager::Get().FindFilesRecursive(FilesToCopy, *SrcFolder, TEXT("*"), /*Files=*/true, /*Directories=*/false);
-	for ( auto FileIt = FilesToCopy.CreateConstIterator(); FileIt; ++FileIt )
+
+	SlowTask.EnterProgressFrame();
 	{
-		const FString SrcFilename = (*FileIt);
+		// Open a new feedback scope for the loop so we can report how far through the copy we are
+		FScopedSlowTask InnerSlowTask(FilesToCopy.Num());
+		for ( auto FileIt = FilesToCopy.CreateConstIterator(); FileIt; ++FileIt )
+		{
+			const FString SrcFilename = (*FileIt);
+
+			// Update the progress
+			FFormatNamedArguments Args;
+			Args.Add( TEXT("SrcFilename"), FText::FromString( FPaths::GetCleanFilename(SrcFilename) ) );
+			InnerSlowTask.EnterProgressFrame(1, FText::Format( LOCTEXT( "CreatingProjectStatus_CopyingFile", "Copying File {SrcFilename}..." ), Args ));
 
 		// Get the file path, relative to the src folder
-		const FString SrcFileSubpath = SrcFilename.RightChop(SrcFolder.Len() + 1);
+			const FString SrcFileSubpath = SrcFilename.RightChop(SrcFolder.Len() + 1);
 
-		// Skip any files that were configured to be ignored
-		bool bThisFileIsIgnored = false;
-		for ( auto IgnoreIt = TemplateDefs->FilesToIgnore.CreateConstIterator(); IgnoreIt; ++IgnoreIt )
-		{
-			if ( SrcFileSubpath == *IgnoreIt )
+			// Skip any files that were configured to be ignored
+			bool bThisFileIsIgnored = false;
+			for ( auto IgnoreIt = TemplateDefs->FilesToIgnore.CreateConstIterator(); IgnoreIt; ++IgnoreIt )
+			{
+				if ( SrcFileSubpath == *IgnoreIt )
+				{
+					// This file was marked as "ignored"
+					bThisFileIsIgnored = true;
+					break;
+				}
+			}
+
+			if ( bThisFileIsIgnored )
 			{
 				// This file was marked as "ignored"
-				bThisFileIsIgnored = true;
-				break;
+				continue;
 			}
-		}
 
-		if ( bThisFileIsIgnored )
-		{
-			// This file was marked as "ignored"
-			continue;
-		}
+			// Skip any folders that were configured to be ignored
+			bool bThisFolderIsIgnored = false;
+			for ( auto IgnoreIt = TemplateDefs->FoldersToIgnore.CreateConstIterator(); IgnoreIt; ++IgnoreIt )
+			{
+				if ( SrcFileSubpath.StartsWith((*IgnoreIt) + TEXT("/") ) )
+				{
+					// This folder was marked as "ignored"
+					bThisFolderIsIgnored = true;
+					break;
+				}
+			}
 
-		// Skip any folders that were configured to be ignored
-		bool bThisFolderIsIgnored = false;
-		for ( auto IgnoreIt = TemplateDefs->FoldersToIgnore.CreateConstIterator(); IgnoreIt; ++IgnoreIt )
-		{
-			if ( SrcFileSubpath.StartsWith((*IgnoreIt) + TEXT("/") ) )
+			if ( bThisFolderIsIgnored )
 			{
 				// This folder was marked as "ignored"
-				bThisFolderIsIgnored = true;
-				break;
+				continue;
 			}
-		}
-
-		if ( bThisFolderIsIgnored )
-		{
-			// This folder was marked as "ignored"
-			continue;
-		}
-
-		// Update the slow task dialog
-		const bool bAllowNewSlowTask = false;
-		FFormatNamedArguments Args;
-		Args.Add( TEXT("SrcFilename"), FText::FromString( FPaths::GetCleanFilename(SrcFilename) ) );
-		FScopedSlowTask SlowTaskMessage( FText::Format( LOCTEXT( "CreatingProjectStatus_CopyingFile", "Copying File {SrcFilename}..." ), Args ), bAllowNewSlowTask );
 
-		// Retarget any folders that were chosen to be renamed by choosing a new destination subpath now
-		FString DestFileSubpathWithoutFilename = FPaths::GetPath(SrcFileSubpath) + TEXT("/");
-		for ( auto RenameIt = TemplateDefs->FolderRenames.CreateConstIterator(); RenameIt; ++RenameIt )
-		{
-			const FTemplateFolderRename& FolderRename = *RenameIt;
-			if ( SrcFileSubpath.StartsWith(FolderRename.From + TEXT("/")) )
+			// Retarget any folders that were chosen to be renamed by choosing a new destination subpath now
+			FString DestFileSubpathWithoutFilename = FPaths::GetPath(SrcFileSubpath) + TEXT("/");
+			for ( auto RenameIt = TemplateDefs->FolderRenames.CreateConstIterator(); RenameIt; ++RenameIt )
 			{
-				// This was a file in a renamed folder. Retarget to the new location
-				DestFileSubpathWithoutFilename = FolderRename.To / DestFileSubpathWithoutFilename.RightChop( FolderRename.From.Len() );
+				const FTemplateFolderRename& FolderRename = *RenameIt;
+				if ( SrcFileSubpath.StartsWith(FolderRename.From + TEXT("/")) )
+				{
+					// This was a file in a renamed folder. Retarget to the new location
+					DestFileSubpathWithoutFilename = FolderRename.To / DestFileSubpathWithoutFilename.RightChop( FolderRename.From.Len() );
+				}
 			}
-		}
 
-		// Retarget any files that were chosen to have parts of their names replaced here
-		FString DestBaseFilename = FPaths::GetBaseFilename(SrcFileSubpath);
-		const FString FileExtension = FPaths::GetExtension(SrcFileSubpath);
-		for ( auto ReplacementIt = TemplateDefs->FilenameReplacements.CreateConstIterator(); ReplacementIt; ++ReplacementIt )
-		{
-			const FTemplateReplacement& Replacement = *ReplacementIt;
-			if ( Replacement.Extensions.Contains( FileExtension ) )
+			// Retarget any files that were chosen to have parts of their names replaced here
+			FString DestBaseFilename = FPaths::GetBaseFilename(SrcFileSubpath);
+			const FString FileExtension = FPaths::GetExtension(SrcFileSubpath);
+			for ( auto ReplacementIt = TemplateDefs->FilenameReplacements.CreateConstIterator(); ReplacementIt; ++ReplacementIt )
 			{
-				// This file matched a filename replacement extension, apply it now
-				DestBaseFilename = DestBaseFilename.Replace(*Replacement.From, *Replacement.To, Replacement.bCaseSensitive ? ESearchCase::CaseSensitive : ESearchCase::IgnoreCase);
+				const FTemplateReplacement& Replacement = *ReplacementIt;
+				if ( Replacement.Extensions.Contains( FileExtension ) )
+				{
+					// This file matched a filename replacement extension, apply it now
+					DestBaseFilename = DestBaseFilename.Replace(*Replacement.From, *Replacement.To, Replacement.bCaseSensitive ? ESearchCase::CaseSensitive : ESearchCase::IgnoreCase);
+				}
 			}
-		}
 
-		// Perform the copy
-		const FString DestFilename = DestFolder / DestFileSubpathWithoutFilename + DestBaseFilename + TEXT(".") + FileExtension;
-		if ( IFileManager::Get().Copy(*DestFilename, *SrcFilename) == COPY_OK )
-		{
-			CreatedFiles.Add(DestFilename);
-
-			if ( ReplacementsInFilesExtensions.Contains(FileExtension) )
+			// Perform the copy
+			const FString DestFilename = DestFolder / DestFileSubpathWithoutFilename + DestBaseFilename + TEXT(".") + FileExtension;
+			if ( IFileManager::Get().Copy(*DestFilename, *SrcFilename) == COPY_OK )
 			{
-				FilesThatNeedContentsReplaced.Add(DestFilename);
-			}
+				CreatedFiles.Add(DestFilename);
+
+				if ( ReplacementsInFilesExtensions.Contains(FileExtension) )
+				{
+					FilesThatNeedContentsReplaced.Add(DestFilename);
+				}
 
-			// Allow project template to extract class renames from this file copy
-			if (FPaths::GetBaseFilename(SrcFilename) != FPaths::GetBaseFilename(DestFilename)
-				&& TemplateDefs->IsClassRename(DestFilename, SrcFilename, FileExtension))
+				// Allow project template to extract class renames from this file copy
+				if (FPaths::GetBaseFilename(SrcFilename) != FPaths::GetBaseFilename(DestFilename)
+					&& TemplateDefs->IsClassRename(DestFilename, SrcFilename, FileExtension))
+				{
+					// Looks like a UObject file!
+					ClassRenames.Add(FPaths::GetBaseFilename(SrcFilename), FPaths::GetBaseFilename(DestFilename));
+				}
+			}
+			else
 			{
-				// Looks like a UObject file!
-				ClassRenames.Add(FPaths::GetBaseFilename(SrcFilename), FPaths::GetBaseFilename(DestFilename));
+				FFormatNamedArguments FailArgs;
+				FailArgs.Add(TEXT("SrcFilename"), FText::FromString(SrcFilename));
+				FailArgs.Add(TEXT("DestFilename"), FText::FromString(DestFilename));
+				OutFailReason = FText::Format(LOCTEXT("FailedToCopyFile", "Failed to copy \"{SrcFilename}\" to \"{DestFilename}\"."), FailArgs);
+				DeleteCreatedFiles(DestFolder, CreatedFiles);
+				return false;
 			}
 		}
-		else
-		{
-			FFormatNamedArguments FailArgs;
-			FailArgs.Add(TEXT("SrcFilename"), FText::FromString(SrcFilename));
-			FailArgs.Add(TEXT("DestFilename"), FText::FromString(DestFilename));
-			OutFailReason = FText::Format(LOCTEXT("FailedToCopyFile", "Failed to copy \"{SrcFilename}\" to \"{DestFilename}\"."), FailArgs);
-			DeleteCreatedFiles(DestFolder, CreatedFiles);
-			return false;
-		}
 	}
 
-	// Open all files with the specified extensions and replace text
-	for ( auto FileIt = FilesThatNeedContentsReplaced.CreateConstIterator(); FileIt; ++FileIt )
+	SlowTask.EnterProgressFrame();
 	{
-		const FString FileToFix = *FileIt;
-		bool bSuccessfullyProcessed = false;
+		// Open a new feedback scope for the loop so we can report how far through the process we are
+		FScopedSlowTask InnerSlowTask(FilesThatNeedContentsReplaced.Num());
 
-		FString FileContents;
-		if ( FFileHelper::LoadFileToString(FileContents, *FileToFix) )
+		// Open all files with the specified extensions and replace text
+		for ( auto FileIt = FilesThatNeedContentsReplaced.CreateConstIterator(); FileIt; ++FileIt )
 		{
-			for ( auto ReplacementIt = TemplateDefs->ReplacementsInFiles.CreateConstIterator(); ReplacementIt; ++ReplacementIt )
+			InnerSlowTask.EnterProgressFrame();
+
+			const FString FileToFix = *FileIt;
+			bool bSuccessfullyProcessed = false;
+
+			FString FileContents;
+			if ( FFileHelper::LoadFileToString(FileContents, *FileToFix) )
 			{
-				const FTemplateReplacement& Replacement = *ReplacementIt;
-				if ( Replacement.Extensions.Contains( FPaths::GetExtension(FileToFix) ) )
+				for ( auto ReplacementIt = TemplateDefs->ReplacementsInFiles.CreateConstIterator(); ReplacementIt; ++ReplacementIt )
+				{
+					const FTemplateReplacement& Replacement = *ReplacementIt;
+					if ( Replacement.Extensions.Contains( FPaths::GetExtension(FileToFix) ) )
+					{
+						FileContents = FileContents.Replace(*Replacement.From, *Replacement.To, Replacement.bCaseSensitive ? ESearchCase::CaseSensitive : ESearchCase::IgnoreCase);
+					}
+				}
+
+				if ( FFileHelper::SaveStringToFile(FileContents, *FileToFix) )
 				{
-					FileContents = FileContents.Replace(*Replacement.From, *Replacement.To, Replacement.bCaseSensitive ? ESearchCase::CaseSensitive : ESearchCase::IgnoreCase);
+					bSuccessfullyProcessed = true;
 				}
 			}
 
-			if ( FFileHelper::SaveStringToFile(FileContents, *FileToFix) )
+			if ( !bSuccessfullyProcessed )
 			{
-				bSuccessfullyProcessed = true;
+				FFormatNamedArguments Args;
+				Args.Add( TEXT("FileToFix"), FText::FromString( FileToFix ) );
+				OutFailReason = FText::Format( LOCTEXT("FailedToFixUpFile", "Failed to process file \"{FileToFix}\"."), Args );
+				DeleteCreatedFiles(DestFolder, CreatedFiles);
+				return false;
 			}
 		}
-
-		if ( !bSuccessfullyProcessed )
-		{
-			FFormatNamedArguments Args;
-			Args.Add( TEXT("FileToFix"), FText::FromString( FileToFix ) );
-			OutFailReason = FText::Format( LOCTEXT("FailedToFixUpFile", "Failed to process file \"{FileToFix}\"."), Args );
-			DeleteCreatedFiles(DestFolder, CreatedFiles);
-			return false;
-		}
 	}
 
+	SlowTask.EnterProgressFrame();
+
 	const FString ProjectConfigPath = DestFolder / TEXT("Config");
 
 	// Write out the hardware class target settings chosen for this project
@@ -1220,6 +1261,8 @@ bool GameProjectUtils::CreateProjectFromTemplate(const FProjectInformation& InPr
 		}
 	}
 
+	SlowTask.EnterProgressFrame();
+
 	// Generate the project file
 	{
 		// Load the source project
@@ -1254,6 +1297,8 @@ bool GameProjectUtils::CreateProjectFromTemplate(const FProjectInformation& InPr
 		CreatedFiles.Add(InProjectInfo.ProjectFilename);
 	}
 
+	SlowTask.EnterProgressFrame();
+
 	// Copy resources
 	const FString GameModuleSourcePath = DestFolder / TEXT("Source") / ProjectName;
 	if (GenerateGameResourceFiles(GameModuleSourcePath, ProjectName, DestFolder, InProjectInfo.bShouldGenerateCode, CreatedFiles, OutFailReason) == false)
@@ -1262,6 +1307,7 @@ bool GameProjectUtils::CreateProjectFromTemplate(const FProjectInformation& InPr
 		return false;
 	}
 
+	SlowTask.EnterProgressFrame();
 	if ( InProjectInfo.bShouldGenerateCode )
 	{
 		// Generate project files
@@ -1273,6 +1319,8 @@ bool GameProjectUtils::CreateProjectFromTemplate(const FProjectInformation& InPr
 		}
 	}
 
+	SlowTask.EnterProgressFrame();
+
 	if (InProjectInfo.bCopyStarterContent)
 	{
 		// Copy the starter content
@@ -2658,8 +2706,10 @@ bool GameProjectUtils::AddCodeToProject_Internal(const FString& NewClassName, co
 		return false;
 	}
 
-	const bool bAllowNewSlowTask = true;
-	FScopedSlowTask SlowTaskMessage( LOCTEXT( "AddingCodeToProject", "Adding code to project..." ), bAllowNewSlowTask );
+	FScopedSlowTask SlowTask( 5,  LOCTEXT( "AddingCodeToProject", "Adding code to project..." ) );
+	SlowTask.MakeDialog();
+
+	SlowTask.EnterProgressFrame();
 
 	// If the project does not already contain code, add the primary game module
 	TArray<FString> CreatedFiles;
@@ -2683,6 +2733,8 @@ bool GameProjectUtils::AddCodeToProject_Internal(const FString& NewClassName, co
 		}
 	}
 
+	SlowTask.EnterProgressFrame();
+
 	// Class Header File
 	FString SyncLocation;
 	const FString NewHeaderFilename = NewHeaderPath / ParentClassInfo.GetHeaderFilename(NewClassName);
@@ -2698,6 +2750,8 @@ bool GameProjectUtils::AddCodeToProject_Internal(const FString& NewClassName, co
 		}
 	}
 
+	SlowTask.EnterProgressFrame();
+
 	// Class CPP file
 	const FString NewCppFilename = NewCppPath / ParentClassInfo.GetSourceFilename(NewClassName);
 	{
@@ -2712,6 +2766,8 @@ bool GameProjectUtils::AddCodeToProject_Internal(const FString& NewClassName, co
 		}
 	}
 
+	SlowTask.EnterProgressFrame();
+
 	// Generate project files if we happen to be using a project file.
 	if ( !FDesktopPlatformModule::Get()->GenerateProjectFiles(FPaths::RootDir(), FPaths::GetProjectFilePath(), GWarn) )
 	{
@@ -2719,6 +2775,8 @@ bool GameProjectUtils::AddCodeToProject_Internal(const FString& NewClassName, co
 		return false;
 	}
 
+	SlowTask.EnterProgressFrame();
+
 	// Mark the files for add in SCC
 	ISourceControlProvider& SourceControlProvider = ISourceControlModule::Get().GetProvider();
 	if ( ISourceControlModule::Get().IsEnabled() && SourceControlProvider.IsAvailable() )
diff --git a/Engine/Source/Editor/UnrealEd/Private/Editor.cpp b/Engine/Source/Editor/UnrealEd/Private/Editor.cpp
--- a/Engine/Source/Editor/UnrealEd/Private/Editor.cpp
+++ b/Engine/Source/Editor/UnrealEd/Private/Editor.cpp
@@ -4329,6 +4329,8 @@ bool UEditorEngine::SavePackage( UPackage* InOuter, UObject* InBase, EObjectFlag
 				 FOutputDevice* Error, ULinkerLoad* Conform, bool bForceByteSwapping, bool bWarnOfLongFilename, 
 				 uint32 SaveFlags, const class ITargetPlatform* TargetPlatform, const FDateTime& FinalTimeStamp )
 {
+	FScopedSlowTask SlowTask(100);
+
 	UObject* Base = InBase;
 	if ( !Base && InOuter && InOuter->PackageFlags & PKG_ContainsMap )
 	{
@@ -4342,6 +4344,8 @@ bool UEditorEngine::SavePackage( UPackage* InOuter, UObject* InBase, EObjectFlag
 		OriginalPackageFlags = InOuter->PackageFlags;
 	}
 
+	SlowTask.EnterProgressFrame(10);
+
 	UWorld* World = Cast<UWorld>(Base);
 	if ( World )
 	{
@@ -4356,8 +4360,12 @@ bool UEditorEngine::SavePackage( UPackage* InOuter, UObject* InBase, EObjectFlag
 		bAutoAddPkgToSCC = IsPackageValidForAutoAdding( InOuter, Filename );
 	}
 
+	SlowTask.EnterProgressFrame(70);
+
 	bool bSuccess = UPackage::SavePackage(InOuter, Base, TopLevelFlags, Filename, Error, Conform, bForceByteSwapping, bWarnOfLongFilename, SaveFlags, TargetPlatform, FinalTimeStamp);
 
+	SlowTask.EnterProgressFrame(10);
+
 	// If the package is a valid candidate for being automatically-added to source control, go ahead and add it
 	// to the default changelist
 	if( bSuccess && bAutoAddPkgToSCC )
@@ -4378,6 +4386,8 @@ bool UEditorEngine::SavePackage( UPackage* InOuter, UObject* InBase, EObjectFlag
 		}
 	}
 
+	SlowTask.EnterProgressFrame(10);
+
 	if ( World )
 	{
 		OnPostSaveWorld(SaveFlags, World, OriginalPackageFlags, bSuccess);
@@ -4436,14 +4446,12 @@ void UEditorEngine::OnPreSaveWorld(uint32 SaveFlags, UWorld* World)
 			// Don't do this when autosaving or PIE saving so that actor adds can still undo.
 			World->ShrinkLevel();
 
-			GWarn->PushStatus();
-			GWarn->StatusUpdate( -1, -1, FText::Format(NSLOCTEXT("UnrealEd", "SavingMapStatus_CollectingGarbage", "Saving map: {0}... (Collecting garbage)"), FText::FromString(World->GetName())) );
-		
-			// NULL empty or "invalid" entries (e.g. IsPendingKill()) in actors array.
-			CollectGarbage( GARBAGE_COLLECTION_KEEPFLAGS );
-
-			GWarn->PopStatus();
-
+			{
+				FScopedSlowTask SlowTask(0, FText::Format(NSLOCTEXT("UnrealEd", "SavingMapStatus_CollectingGarbage", "Saving map: {0}... (Collecting garbage)"), FText::FromString(World->GetName())));
+				// NULL empty or "invalid" entries (e.g. IsPendingKill()) in actors array.
+				CollectGarbage( GARBAGE_COLLECTION_KEEPFLAGS );
+			}
+			
 			// Compact and sort actors array to remove empty entries.
 			// Don't do this when autosaving or PIE saving so that actor adds can still undo.
 			World->PersistentLevel->SortActorList();
diff --git a/Engine/Source/Editor/UnrealEd/Private/Editor/EditorGameAgnosticSettings.cpp b/Engine/Source/Editor/UnrealEd/Private/Editor/EditorGameAgnosticSettings.cpp
--- a/Engine/Source/Editor/UnrealEd/Private/Editor/EditorGameAgnosticSettings.cpp
+++ b/Engine/Source/Editor/UnrealEd/Private/Editor/EditorGameAgnosticSettings.cpp
@@ -57,8 +57,8 @@ void UEditorGameAgnosticSettings::AutoApplyScalabilityBenchmark()
 {
 	const TCHAR* Section = TEXT("EngineBenchmarkResult");
 
-	GWarn->StatusUpdate(0, 1, NSLOCTEXT("UnrealEd", "RunningEngineBenchmark", "Running engine benchmark..."));
-	GWarn->PushStatus();
+	FScopedSlowTask SlowTask(0, NSLOCTEXT("UnrealEd", "RunningEngineBenchmark", "Running engine benchmark..."));
+	SlowTask.MakeDialog();
 
 
 	Scalability::FQualityLevels Temporary = Scalability::BenchmarkQualityLevels(AutoScalabilityWorkScaleAmount);
@@ -74,8 +74,6 @@ void UEditorGameAgnosticSettings::AutoApplyScalabilityBenchmark()
 
 	Scalability::SetQualityLevels(Temporary);
 	Scalability::SaveState(GEditorGameAgnosticIni);
-
-	GWarn->PopStatus();
 }
 
 bool UEditorGameAgnosticSettings::IsScalabilityBenchmarkValid() const
diff --git a/Engine/Source/Editor/UnrealEd/Private/EditorCsg.cpp b/Engine/Source/Editor/UnrealEd/Private/EditorCsg.cpp
--- a/Engine/Source/Editor/UnrealEd/Private/EditorCsg.cpp
+++ b/Engine/Source/Editor/UnrealEd/Private/EditorCsg.cpp
@@ -275,11 +275,11 @@ void UEditorEngine::csgRebuild( UWorld* InWorld )
 		}
 	}
 
-	GWarn->PushStatus();
-		GWarn->StatusUpdate( 0, DynamicBrushes.Num(), NSLOCTEXT("UnrealEd", "RebuildCSGRebuildingDynamicBrushBSPs", "Rebuild CSG: Rebuilding Dynamic Brush BSPs") );
+	{
+		FScopedSlowTask SlowTask(DynamicBrushes.Num(), NSLOCTEXT("UnrealEd", "RebuildCSGRebuildingDynamicBrushBSPs", "Rebuild CSG: Rebuilding Dynamic Brush BSPs") );
 		for ( int32 BrushIndex = 0; BrushIndex < DynamicBrushes.Num(); BrushIndex++ )
 		{
-			GWarn->UpdateProgress(BrushIndex, DynamicBrushes.Num());
+			SlowTask.EnterProgressFrame();
 
 			ABrush* B = DynamicBrushes[BrushIndex];
 			FBSPOps::csgPrepMovingBrush(B);
@@ -289,7 +289,7 @@ void UEditorEngine::csgRebuild( UWorld* InWorld )
 				break;
 			}
 		}
-	GWarn->PopStatus();
+	}
 
 	GWarn->UpdateProgress( 4, 4 );
 
diff --git a/Engine/Source/Editor/UnrealEd/Private/EditorLevelUtils.cpp b/Engine/Source/Editor/UnrealEd/Private/EditorLevelUtils.cpp
--- a/Engine/Source/Editor/UnrealEd/Private/EditorLevelUtils.cpp
+++ b/Engine/Source/Editor/UnrealEd/Private/EditorLevelUtils.cpp
@@ -129,7 +129,8 @@ namespace EditorLevelUtils
 			return nullptr;
 		}
 
-		FScopedSlowTask SlowTask(LOCTEXT("AddLevelsToWorldTask", "Adding Levels to World"), true);
+		FScopedSlowTask SlowTask(LevelPackageNames.Num(), LOCTEXT("AddLevelsToWorldTask", "Adding Levels to World"));
+		SlowTask.MakeDialog();
 
 		TArray<FString> PackageNames = LevelPackageNames;
 
@@ -143,6 +144,8 @@ namespace EditorLevelUtils
 		ULevel* NewLevel = nullptr;
 		for (const auto& PackageName : PackageNames)
 		{
+			SlowTask.EnterProgressFrame();
+
 			NewLevel = AddLevelToWorld(InWorld, *PackageName, LevelStreamingClass);
 			if (NewLevel)
 			{
diff --git a/Engine/Source/Editor/UnrealEd/Private/EditorServer.cpp b/Engine/Source/Editor/UnrealEd/Private/EditorServer.cpp
--- a/Engine/Source/Editor/UnrealEd/Private/EditorServer.cpp
+++ b/Engine/Source/Editor/UnrealEd/Private/EditorServer.cpp
@@ -2038,7 +2038,27 @@ bool UEditorEngine::Map_Load(const TCHAR* Str, FOutputDevice& Ar)
 				// Should we display progress while loading?
 				int32 bShowProgress = 1;
 				FParse::Value(Str, TEXT("SHOWPROGRESS="), bShowProgress);
-			
+
+				FString MapFileName = FPaths::GetCleanFilename(TempFname);
+
+				// Detect whether the map we are loading is a template map and alter the undo
+				// readout accordingly.
+				FText LocalizedLoadingMap;
+				if (!bIsLoadingMapTemplate)
+				{
+					LocalizedLoadingMap = FText::Format( NSLOCTEXT("UnrealEd", "LoadingMap_F", "Loading map: {0}..."), FText::FromString( MapFileName ) );
+				}
+				else
+				{
+					LocalizedLoadingMap = NSLOCTEXT("UnrealEd", "LoadingMap_Template", "New Map");
+				}
+				
+				// Don't show progress dialogs when loading one of our startup maps. They should load rather quickly.
+				FScopedSlowTask SlowTask(100, FText::Format( NSLOCTEXT("UnrealEd", "LoadingMapStatus_Loading", "Loading {0}"), LocalizedLoadingMap ), bShowProgress != 0);
+				SlowTask.MakeDialog();
+
+				SlowTask.EnterProgressFrame(10, FText::Format( NSLOCTEXT("UnrealEd", "LoadingMapStatus_CleaningUp", "{0} (Clearing existing world)"), LocalizedLoadingMap ));
+
 				UObject* OldOuter = NULL;
 
 				{
@@ -2052,26 +2072,8 @@ bool UEditorEngine::Map_Load(const TCHAR* Str, FOutputDevice& Ar)
 
 					OldOuter = Context.World()->GetOuter();
 
-					FString MapFileName = FPaths::GetCleanFilename(TempFname);
-
-					// Detect whether the map we are loading is a template map and alter the undo
-					// readout accordingly.
-					FText LocalizedLoadingMap;
-					if (!bIsLoadingMapTemplate)
-					{
-						LocalizedLoadingMap = FText::Format( NSLOCTEXT("UnrealEd", "LoadingMap_F", "Loading map: {0}..."), FText::FromString( MapFileName ) );
-					}
-					else
-					{
-						LocalizedLoadingMap = NSLOCTEXT("UnrealEd", "LoadingMap_Template", "New Map");
-					}
-					
 					ResetTransaction( LocalizedLoadingMap );
 
-					// Don't show progress dialogs when loading one of our startup maps.  They should load rather quickly.
-					const bool bShowProgressDialog = (bShowProgress != 0);
-					GWarn->BeginSlowTask( FText::Format( NSLOCTEXT("UnrealEd", "LoadingMapStatus_CleaningUp", "{0} (Clearing existing world)"), LocalizedLoadingMap ), bShowProgressDialog );
-
 					// Don't clear errors if we are loading a startup map so we can see all startup load errors
 					if (!FEditorFileUtils::IsLoadingStartupMap())
 					{
@@ -2100,7 +2102,7 @@ bool UEditorEngine::Map_Load(const TCHAR* Str, FOutputDevice& Ar)
 						ExistingWorld = NULL;
 					}
 
-					GWarn->StatusUpdate( -1, -1, LocalizedLoadingMap );
+					SlowTask.EnterProgressFrame( 70, LocalizedLoadingMap );
 				}
 
 				// Record the name of this file to make sure we load objects in this package on top of in-memory objects in this package.
@@ -2125,14 +2127,20 @@ bool UEditorEngine::Map_Load(const TCHAR* Str, FOutputDevice& Ar)
 					}
 				}
 
-			
+				
 				UPackage* WorldPackage;
 				// Load startup maps and templates into new outermost packages so that the Save function in the editor won't overwrite the original
 				if (bIsLoadingMapTemplate)
 				{
+					FScopedSlowTask LoadScope(2);
+
+					LoadScope.EnterProgressFrame();
+
 					//create a package with the proper name
 					WorldPackage = CreatePackage(NULL, *(MakeUniqueObjectName(NULL, UPackage::StaticClass()).ToString()));
 
+					LoadScope.EnterProgressFrame();
+
 					//now load the map into the package created above
 					const FName WorldPackageFName = WorldPackage->GetFName();
 					UWorld::WorldTypePreLoadMap.FindOrAdd(WorldPackageFName) = EWorldType::Editor;
@@ -2159,7 +2167,6 @@ bool UEditorEngine::Map_Load(const TCHAR* Str, FOutputDevice& Ar)
 				{
 					FMessageDialog::Open( EAppMsgType::Ok, NSLOCTEXT("UnrealEd", "MapPackageLoadFailed", "Failed to open map file. This is most likely because the map was saved with a newer version of the engine."));
 					GIsEditorLoadingPackage = false;
-					GWarn->EndSlowTask();
 					return false;
 				}
 
@@ -2204,11 +2211,13 @@ bool UEditorEngine::Map_Load(const TCHAR* Str, FOutputDevice& Ar)
 					World->InitWorld( GetEditorWorldInitializationValues() );
 				}
 
+				SlowTask.EnterProgressFrame(20, FText::Format( LOCTEXT( "LoadingMapStatus_Initializing", "Loading map: {0}... (Initializing world)" ), FText::FromString(MapFileName) ));
 				{
 					FBSPOps::bspValidateBrush(Context.World()->GetDefaultBrush()->Brush, 0, 1);
 
-					FString MapFileName = FPaths::GetCleanFilename(TempFname);
-					GWarn->StatusUpdate( -1, -1, FText::Format( LOCTEXT( "LoadingMapStatus_Initializing", "Loading map: {0}... (Initializing world)" ), FText::FromString(MapFileName) ) );
+					// This is a relatively long process, so break it up a bit
+					FScopedSlowTask InitializingFeedback(5);
+					InitializingFeedback.EnterProgressFrame();
 
 					Context.World()->AddToRoot();
 
@@ -2226,6 +2235,8 @@ bool UEditorEngine::Map_Load(const TCHAR* Str, FOutputDevice& Ar)
 					// Update any actors that can be affected by CullDistanceVolumes
 					Context.World()->UpdateCullDistanceVolumes();
 
+					InitializingFeedback.EnterProgressFrame();
+
 					// A new level was loaded into the editor, so we need to let other systems know about the new
 					// actors in the scene
 					FEditorDelegates::MapChange.Broadcast(MapChangeEventFlags::NewMap );
@@ -2240,7 +2251,8 @@ bool UEditorEngine::Map_Load(const TCHAR* Str, FOutputDevice& Ar)
 
 					NoteSelectionChange();
 
-					GWarn->EndSlowTask();
+					InitializingFeedback.EnterProgressFrame();
+
 					// Look for 'orphan' actors - that is, actors which are in the Package of the level we just loaded, but not in the Actors array.
 					// If we find any, set IsPendingKill() to 'true', so that PendingKill will return 'true' for them. We can NOT use FActorIterator here
 					// as it just traverses the Actors list.
@@ -2267,6 +2279,8 @@ bool UEditorEngine::Map_Load(const TCHAR* Str, FOutputDevice& Ar)
 						Actor->SetFlags( RF_Transactional );
 					}
 
+					InitializingFeedback.EnterProgressFrame();
+
 					UNavigationSystem::InitializeForWorld(Context.World(), FNavigationSystem::EditorMode);
 					Context.World()->CreateAISystem();
 
@@ -2314,6 +2328,8 @@ bool UEditorEngine::Map_Load(const TCHAR* Str, FOutputDevice& Ar)
 						GEditor->Layers->SetLayersVisibility( LayersToHide, bIsVisible );
 					}
 
+					InitializingFeedback.EnterProgressFrame();
+
 					FEditorDelegates::DisplayLoadErrors.Broadcast();
 
 					if( FModuleManager::Get().IsModuleLoaded("LevelEditor") )
@@ -4028,18 +4044,12 @@ bool UEditorEngine::Exec_Obj( const TCHAR* Str, FOutputDevice& Ar )
 			FParse::Bool( Str, TEXT( "KEEPDIRTY=" ), bKeepDirty );
 
 			// Save the package.
-			if( !bSilent )
-			{
-				const UWorld* const AssociatedWorld = UWorld::FindWorldInPackage(Pkg);
-				const bool bIsMapPackage = AssociatedWorld != nullptr;
+			const bool bIsMapPackage = UWorld::FindWorldInPackage(Pkg) != nullptr;
+			const FText SavingPackageText = (bIsMapPackage) 
+				? FText::Format(NSLOCTEXT("UnrealEd", "SavingMapf", "Saving map {0}"), FText::FromString(Pkg->GetName()))
+				: FText::Format(NSLOCTEXT("UnrealEd", "SavingAssetf", "Saving asset {0}"), FText::FromString(Pkg->GetName()));
 
-				const FText SavingPackageText = (bIsMapPackage) 
-					? FText::Format(NSLOCTEXT("UnrealEd", "SavingMapf", "Saving map {0}"), FText::FromString(Pkg->GetName()))
-					: FText::Format(NSLOCTEXT("UnrealEd", "SavingAssetf", "Saving asset {0}"), FText::FromString(Pkg->GetName()));
-
-				GWarn->BeginSlowTask( NSLOCTEXT("UnrealEd", "SavingPackage", "Saving package" ), true );
-				GWarn->StatusUpdate( 1, 1, SavingPackageText );
-			}
+			FScopedSlowTask SlowTask(100, SavingPackageText, !bSilent);
 
 			uint32 SaveFlags = bAutosaving ? SAVE_FromAutosave : SAVE_None;
 			if ( bKeepDirty )
@@ -4049,10 +4059,6 @@ bool UEditorEngine::Exec_Obj( const TCHAR* Str, FOutputDevice& Ar )
 
 			const bool bWarnOfLongFilename = !bAutosaving;
 			bWasSuccessful = this->SavePackage( Pkg, NULL, RF_Standalone, TempFname, &Ar, NULL, false, bWarnOfLongFilename, SaveFlags );
-			if( !bSilent )
-			{
-				GWarn->EndSlowTask();
-			}
 		}
 		else
 		{
diff --git a/Engine/Source/Editor/UnrealEd/Private/FeedbackContextEditor.cpp b/Engine/Source/Editor/UnrealEd/Private/FeedbackContextEditor.cpp
--- a/Engine/Source/Editor/UnrealEd/Private/FeedbackContextEditor.cpp
+++ b/Engine/Source/Editor/UnrealEd/Private/FeedbackContextEditor.cpp
@@ -9,163 +9,314 @@
 /** Called to cancel the slow task activity */
 DECLARE_DELEGATE( FOnCancelClickedDelegate );
 
-
 /**
  * Simple "slow task" widget
  */
 class SSlowTaskWidget : public SBorder
 {
+	/** The maximum number of secondary bars to show on the widget */
+	static const int32 MaxNumSecondaryBars = 3;
+
+	/** The width of the dialog, and horizontal padding */
+	static const int32 FixedWidth = 600, FixedPaddingH = 24;
+
+	/** The heights of the progress bars on this widget */
+	static const int32 MainBarHeight = 12, SecondaryBarHeight = 3;
 public:
 	SLATE_BEGIN_ARGS( SSlowTaskWidget )	{ }
+
+		/** Called to when an asset is clicked */
+		SLATE_EVENT( FOnCancelClickedDelegate, OnCancelClickedDelegate )
+
+		/** The feedback scope stack that we are presenting to the user */
+		SLATE_ARGUMENT( const FScopedSlowTaskStack*, ScopeStack )
+
 	SLATE_END_ARGS()
 
+	/** Construct this widget */
 	void Construct( const FArguments& InArgs )
 	{
-		SBorder::Construct(SBorder::FArguments()
-			.BorderImage(FEditorStyle::GetBrush("Menu.Background"))
-			.VAlign(VAlign_Center)
+		OnCancelClickedDelegate = InArgs._OnCancelClickedDelegate;
+		ScopeStack = InArgs._ScopeStack;
+
+		TSharedRef<SVerticalBox> VerticalBox = SNew(SVerticalBox)
+
+			// Construct the main progress bar and text
+			+ SVerticalBox::Slot()
+			.AutoHeight()
 			[
 				SNew(SVerticalBox)
-				+SVerticalBox::Slot()
+
+				+ SVerticalBox::Slot()
 				.AutoHeight()
-				.Padding( 14.0f, 4.0f, 14.0f, 10.0f )
+				.Padding(FMargin(0, 0, 0, 5.f))
+				.VAlign(VAlign_Center)
 				[
-					SNew( STextBlock )
-						.Text( this, &SSlowTaskWidget::OnGetProgressText )
-						.ShadowOffset( FVector2D( 1.0f, 1.0f ) )
+					SNew(SBox)
+					.HeightOverride(24.f)
+					[
+						SNew(SHorizontalBox)
+
+						+ SHorizontalBox::Slot()
+						[
+							SNew( STextBlock )
+							.AutoWrapText(true)
+							.Text( this, &SSlowTaskWidget::GetProgressText, 0 )
+							// The main font size dynamically changes depending on the content
+							.Font( this, &SSlowTaskWidget::GetMainTextFont )
+						]
+
+						+ SHorizontalBox::Slot()
+						.Padding(FMargin(5.f, 0, 0, 0))
+						.AutoWidth()
+						[
+							SNew( STextBlock )
+							.Text( this, &SSlowTaskWidget::GetPercentageText )
+							// The main font size dynamically changes depending on the content
+							.Font( FSlateFontInfo( FPaths::EngineContentDir() / TEXT("Slate/Fonts/Roboto-Light.ttf"), 14, EFontHinting::AutoLight ) )
+						]
+					]
 				]
-				+SVerticalBox::Slot()
+
+				+ SVerticalBox::Slot()
 				.AutoHeight()
-				.Padding(10.0f, 7.0f)
 				[
-					SNew(SHorizontalBox)
-					+SHorizontalBox::Slot()
-					.Padding(5,0,0,0)
-					.FillWidth(1.0f)
+					SNew(SBox)
+					.HeightOverride(MainBarHeight)
 					[
 						SNew(SProgressBar)
-						.Percent( this, &SSlowTaskWidget::GetProgressFraction )
-					]
-					+SHorizontalBox::Slot()
-					.Padding(5,0,0,0)
-					.FillWidth(0.2f)
-					[
-						SNew(SButton)
-						.Text( NSLOCTEXT("FeedbackContextProgress", "Cancel", "Cancel") )
-						.HAlign(EHorizontalAlignment::HAlign_Center)
-						.OnClicked(this, &SSlowTaskWidget::OnCancel)
-						.Visibility(this, &SSlowTaskWidget::GetCancelButtonVisibility)
+						.BorderPadding(FVector2D::ZeroVector)
+						.Percent( this, &SSlowTaskWidget::GetProgressFraction, 0 )
+						.BackgroundImage( FEditorStyle::GetBrush("ProgressBar.ThinBackground") )
+						.FillImage( FEditorStyle::GetBrush("ProgressBar.ThinFill") )
 					]
 				]
 			]
+			
+			// Secondary progress bars
+			+ SVerticalBox::Slot()
+			.Padding(FMargin(0.f, 8.f, 0.f, 0.f))
+			[
+				SAssignNew(SecondaryBars, SVerticalBox)
+			];
+		
+
+		if ( OnCancelClickedDelegate.IsBound() )
+		{
+			VerticalBox->AddSlot()
+				.AutoHeight()
+				.HAlign(HAlign_Center)
+				.Padding(10.0f, 7.0f)
+				[
+					SNew(SButton)
+					.Text( NSLOCTEXT("FeedbackContextProgress", "Cancel", "Cancel") )
+					.HAlign(EHorizontalAlignment::HAlign_Center)
+					.OnClicked(this, &SSlowTaskWidget::OnCancel)
+				];
+		}
+
+		SBorder::Construct( SBorder::FArguments()
+			.BorderImage(FEditorStyle::GetBrush("Menu.Background"))
+			.VAlign(VAlign_Center)
+			.Padding(FMargin(FixedPaddingH))
+			[
+				SNew(SBox).WidthOverride(FixedWidth) [ VerticalBox ]
+			]
 		);
 	}
 
-	SSlowTaskWidget()
-		: ProgressNumerator(0)
-		, ProgressDenominator(1)
+	virtual void Tick(const FGeometry& AllottedGeometry, const double InCurrentTime, const float InDeltaTime) override
 	{
-	}
+		static const double VisibleScopeThreshold = 0.5;
 
-	void SetProgressPercent( int32 InProgressNumerator, int32 InProgressDenominator )
-	{
-		// interlocked exchange for threaded version
-		FPlatformAtomics::InterlockedExchange( &ProgressNumerator, InProgressNumerator );
-		FPlatformAtomics::InterlockedExchange( &ProgressDenominator, InProgressDenominator );
-		UpdateProgressText();
-	}
+		DynamicProgressIndices.Empty();
+		
+		// Always show the first one
+		DynamicProgressIndices.Add(0);
 
-	void SetProgressText( const FText& InProgressText )
-	{
-		ProgressText = InProgressText;
-		UpdateProgressText();
-	}
-	
-	void BindOnCancelDelegate(const FOnCancelClickedDelegate& InCancelClickedDelegate)
-	{
-		OnCancelClickedDelegate = InCancelClickedDelegate;
-	}
+		for (int32 Index = 1; Index < ScopeStack->Num(); ++Index)
+		{
+			const auto* Scope = (*ScopeStack)[Index];
 
-	void UnbindOnCancelDelegate()
-	{
-		OnCancelClickedDelegate.Unbind();
+			if (Scope->bVisibleOnUI && !Scope->DefaultMessage.IsEmpty())
+			{
+				const auto TimeOpen = FPlatformTime::Seconds() - Scope->StartTime;
+				const auto WorkDone = ScopeStack->GetProgressFraction(Index);
+
+				// We only show visible scopes if they have been opened a while, and have a reasonable amount of work left
+				if (WorkDone * TimeOpen > VisibleScopeThreshold)
+				{
+					DynamicProgressIndices.Add(Index);
+				}
+
+				if (DynamicProgressIndices.Num() == MaxNumSecondaryBars)
+				{
+					break;
+				}
+			}
+		}
+
+		// Create progress bars for anything that we haven't cached yet
+		// We don't destroy old widgets, they just remain ghosted until shown again
+		for (int32 Index = SecondaryBars->GetChildren()->Num() + 1; Index < DynamicProgressIndices.Num(); ++Index)
+		{
+			CreateSecondaryBar(Index);
+		}
 	}
 
 private:
 
-	FReply OnCancel()
+	/** Create a progress bar for the specified index */
+	void CreateSecondaryBar(int32 Index) 
 	{
-		OnCancelClickedDelegate.ExecuteIfBound();
-		return FReply::Handled();
+		SecondaryBars->AddSlot()
+		.Padding( 0.f, 16.f, 0.f, 0.f )
+		[
+			SNew(SVerticalBox)
+
+			+ SVerticalBox::Slot()
+			.Padding( 0.f, 0.f, 0.f, 4.f )
+			.AutoHeight()
+			[
+				SNew( STextBlock )
+				.Text( this, &SSlowTaskWidget::GetProgressText, Index )
+				.Font( FSlateFontInfo( FPaths::EngineContentDir() / TEXT("Slate/Fonts/Roboto-Regular.ttf"), 9, EFontHinting::AutoLight ) )
+				.ColorAndOpacity( FSlateColor::UseSubduedForeground() )
+			]
+
+			+ SVerticalBox::Slot()
+			.AutoHeight()
+			[
+				SNew(SBox)
+				.HeightOverride(SecondaryBarHeight)
+				[
+					SNew(SBorder)
+					.Padding(0)
+					.BorderImage(FEditorStyle::GetBrush("NoBorder"))
+					.ColorAndOpacity( this, &SSlowTaskWidget::GetSecondaryProgressBarTint, Index )
+					[
+						SNew(SProgressBar)
+						.BorderPadding(FVector2D::ZeroVector)
+						.Percent( this, &SSlowTaskWidget::GetProgressFraction, Index )
+						.BackgroundImage( FEditorStyle::GetBrush("ProgressBar.ThinBackground") )
+						.FillImage( FEditorStyle::GetBrush("ProgressBar.ThinFill") )
+					]
+				]
+			]
+		];
 	}
 
-	EVisibility GetCancelButtonVisibility() const
+private:
+
+	/** The main text that we will display in the window */
+	FText GetPercentageText() const
 	{
-		return OnCancelClickedDelegate.IsBound() ? EVisibility::Visible : EVisibility::Collapsed;
+		const float ProgressInterp = ScopeStack->GetProgressFraction(0);
+
+		FFormatOrderedArguments Args;
+		Args.Add(int(ProgressInterp * 100));
+
+		return FText::Format(NSLOCTEXT("FeedbackContextEditor", "MainProgressDisplayText", "{0}%"), Args);
 	}
 
-	TOptional<float> GetProgressFraction() const
+	/** Calculate the best font to display the main text with */
+	FSlateFontInfo GetMainTextFont() const
 	{
-		// Only show a percentage if there is something interesting to report
-		if( ProgressNumerator > 0 && ProgressDenominator > 0 )
-		{
-			return (float)ProgressNumerator/ProgressDenominator;
-		}
-		else
+		TSharedRef<FSlateFontMeasure> MeasureService = FSlateApplication::Get().GetRenderer()->GetFontMeasureService();
+
+		const int32 MaxFontSize = 14;
+		FSlateFontInfo FontInfo( FPaths::EngineContentDir() / TEXT("Slate/Fonts/Roboto-Light.ttf"), MaxFontSize, EFontHinting::AutoLight );
+
+		const FText MainText = GetProgressText(0);
+		const int32 MaxTextWidth = FixedWidth - FixedPaddingH*2;
+		while( FontInfo.Size > 9 && MeasureService->Measure(MainText, FontInfo).X > MaxTextWidth )
 		{
-			// Return non-value to indicate marquee mode
-			// for the progress bar.
-			return TOptional<float>();
+			FontInfo.Size -= 4;
 		}
+
+		return FontInfo;
 	}
 
-	FText OnGetProgressText() const
+	/** Get the tint for a secondary progress bar */
+	FLinearColor GetSecondaryProgressBarTint(int32 Index) const
 	{
-		return ProgressStatusText;
+		if (!DynamicProgressIndices.IsValidIndex(Index) || !ScopeStack->IsValidIndex(DynamicProgressIndices[Index]))
+		{
+			return FLinearColor::White.CopyWithNewOpacity(0.25f);
+		}
+		return FLinearColor::White;
 	}
 
-	void UpdateProgressText()
+	/** Get the fractional percentage of completion for a progress bar */
+	TOptional<float> GetProgressFraction(int32 Index) const
 	{
-		if( ProgressNumerator > 0 && ProgressDenominator > 0 )
+		if (DynamicProgressIndices.IsValidIndex(Index) && ScopeStack->IsValidIndex(DynamicProgressIndices[Index]))
 		{
-			FFormatNamedArguments Args;
-			Args.Add( TEXT("StatusText"), ProgressText );
-			Args.Add( TEXT("ProgressCompletePercentage"), FText::AsPercent( (float)ProgressNumerator/ProgressDenominator) );
-			ProgressStatusText = FText::Format( NSLOCTEXT("FeedbackContextProgress", "ProgressStatusFormat", "{StatusText} ({ProgressCompletePercentage})"), Args );
+			return ScopeStack->GetProgressFraction(DynamicProgressIndices[Index]);
 		}
-		else
+		return TOptional<float>();
+	}
+
+	/** Get the text to display for a progress bar */
+	FText GetProgressText(int32 Index) const
+	{
+		if (DynamicProgressIndices.IsValidIndex(Index) && ScopeStack->IsValidIndex(DynamicProgressIndices[Index]))
 		{
-			ProgressStatusText = ProgressText;
+			return (*ScopeStack)[DynamicProgressIndices[Index]]->GetCurrentMessage();
 		}
+
+		return FText();
+	}
+
+	/** Called when the cancel button is clicked */
+	FReply OnCancel()
+	{
+		OnCancelClickedDelegate.ExecuteIfBound();
+		return FReply::Handled();
 	}
 
 private:
-	FText ProgressText;
-	FText ProgressStatusText;
-	int32 ProgressNumerator;
-	int32 ProgressDenominator;
+
+	/** Delegate to invoke if the user clicks cancel */
 	FOnCancelClickedDelegate OnCancelClickedDelegate;
+
+	/** The scope stack that we are reflecting */
+	const FScopedSlowTaskStack* ScopeStack;
+
+	/** The vertical box containing the secondary progress bars */
+	TSharedPtr<SVerticalBox> SecondaryBars;
+
+	/** Array mapping progress bar index -> scope stack index. Updated every tick. */
+	TArray<int32> DynamicProgressIndices;
 };
 
+/** Static integer definitions required on some builds where the linker needs access to these */
+const int32 SSlowTaskWidget::MaxNumSecondaryBars;
+const int32 SSlowTaskWidget::FixedWidth;
+const int32 SSlowTaskWidget::FixedPaddingH;;
+const int32 SSlowTaskWidget::MainBarHeight;
+const int32 SSlowTaskWidget::SecondaryBarHeight;
+
 static void TickSlate()
 {
-	// Tick Slate application
-	FSlateApplication::Get().Tick();
-
-	// Sync the game thread and the render thread. This is needed if many StatusUpdate are called
-	FSlateApplication::Get().GetRenderer()->Sync();
-}
-
+	static double Seconds = FPlatformTime::Seconds();
+	static double MinFrameTime = 0.05;		// Only update at 20fps so as not to slow down the actual task
+	if (FPlatformTime::Seconds() - Seconds > MinFrameTime)
+	{
+		Seconds = FPlatformTime::Seconds();
 
-const float FFeedbackContextEditor::UIUpdateGatingTime = 1.0f / 60.0f;
+		// Tick Slate application
+		FSlateApplication::Get().Tick();
 
+		// Sync the game thread and the render thread. This is needed if many StatusUpdate are called
+		FSlateApplication::Get().GetRenderer()->Sync();
+	}
+}
 
-FFeedbackContextEditor::FFeedbackContextEditor() : 
-	DialogRequestCount(0),
-	HasTaskBeenCancelled(false),
-	SlowTaskCount(0)
+FFeedbackContextEditor::FFeedbackContextEditor()
+	: HasTaskBeenCancelled(false)
 {
-	DialogRequestStack.Reserve(15);
+	
 }
 
 void FFeedbackContextEditor::Serialize( const TCHAR* V, ELogVerbosity::Type Verbosity, const class FName& Category )
@@ -176,14 +327,10 @@ void FFeedbackContextEditor::Serialize( const TCHAR* V, ELogVerbosity::Type Verb
 	}
 }
 
-/**
- * Tells the editor that a slow task is beginning
- * 
- * @param Task					The description of the task beginning.
- * @param ShowProgressDialog	Whether to show the progress dialog.
- */
-void FFeedbackContextEditor::BeginSlowTask( const FText& Task, bool bShowProgressDialog, bool bShowCancelButton )
+void FFeedbackContextEditor::StartSlowTask( const FText& Task, bool bShowCancelButton )
 {
+	FFeedbackContext::StartSlowTask( Task, bShowCancelButton );
+
 	// Attempt to parent the slow task window to the slate main frame
 	TSharedPtr<SWindow> ParentWindow;
 	if( FModuleManager::Get().IsModuleLoaded( "MainFrame" ) )
@@ -194,24 +341,15 @@ void FFeedbackContextEditor::BeginSlowTask( const FText& Task, bool bShowProgres
 
 	if( GIsEditor && ParentWindow.IsValid())
 	{
-		GIsSlowTask = ++SlowTaskCount>0;
 		GSlowTaskOccurred = GIsSlowTask;
 
-		// Show a wait cursor for long running tasks
-		if (SlowTaskCount == 1)
-		{
-			// NOTE: Any slow tasks that start after the first slow task is running won't display
-			//   status text unless 'StatusUpdate' is called
-			StatusMessage.StatusText = Task;
-		}
-
 		// Don't show the progress dialog if the Build Progress dialog is already visible
 		bool bProgressWindowShown = BuildProgressWidget.IsValid();
 
 		// Don't show the progress dialog if a Slate menu is currently open
 		const bool bHaveOpenMenu = FSlateApplication::Get().AnyMenusVisible();
 
-		if( bShowProgressDialog && bHaveOpenMenu )
+		if( bHaveOpenMenu )
 		{
 			UE_LOG(LogSlate, Warning, TEXT("Prevented a slow task dialog from being summoned while a context menu was open") );
 		}
@@ -219,71 +357,42 @@ void FFeedbackContextEditor::BeginSlowTask( const FText& Task, bool bShowProgres
 		// reset the cancellation flag
 		HasTaskBeenCancelled = false;
 
-		bShowProgressDialog &= !bHaveOpenMenu;
-
-		DialogRequestStack.Push( ( bProgressWindowShown ? 0 : bShowProgressDialog) );
-
-		if(!bProgressWindowShown && bShowProgressDialog && ++DialogRequestCount == 1 )
+		if (!bProgressWindowShown && !bHaveOpenMenu && FSlateApplication::Get().CanDisplayWindows())
 		{
-			if( !bProgressWindowShown && FSlateApplication::Get().CanDisplayWindows() )
+			FOnCancelClickedDelegate OnCancelClicked;
+			if ( bShowCancelButton )
 			{
-				TSharedRef<SWindow> SlowTaskWindowRef = SNew(SWindow)
-					.ClientSize(FVector2D(500,100))
-					.IsPopupWindow(true)
-					.ActivateWhenFirstShown(true);
-
-				SlowTaskWidget = SNew(SSlowTaskWidget);
-				SlowTaskWindowRef->SetContent( SlowTaskWidget.ToSharedRef() );
+				// The cancel button is only displayed if a delegate is bound to it.
+				OnCancelClicked = FOnCancelClickedDelegate::CreateRaw(this, &FFeedbackContextEditor::OnUserCancel);
+			}
 
-				EnableUserCancel(bShowCancelButton);
+			TSharedRef<SWindow> SlowTaskWindowRef = SNew(SWindow)
+				.SizingRule(ESizingRule::Autosized)
+				.AutoCenter(EAutoCenter::PreferredWorkArea)
+				.IsPopupWindow(true)
+				.CreateTitleBar(true)
+				.ActivateWhenFirstShown(true);
 
-				SlowTaskWidget->SetProgressText( StatusMessage.StatusText );
+			SlowTaskWidget = SNew(SSlowTaskWidget)
+				.ScopeStack(&ScopeStack)
+				.OnCancelClickedDelegate( OnCancelClicked );
+			SlowTaskWindowRef->SetContent( SlowTaskWidget.ToSharedRef() );
 
-				SlowTaskWindow = SlowTaskWindowRef;
+			SlowTaskWindow = SlowTaskWindowRef;
 
-				const bool bSlowTask = true;
-				FSlateApplication::Get().AddModalWindow( SlowTaskWindowRef, ParentWindow, bSlowTask );
+			const bool bSlowTask = true;
+			FSlateApplication::Get().AddModalWindow( SlowTaskWindowRef, ParentWindow, bSlowTask );
 
-				SlowTaskWindowRef->ShowWindow();
+			SlowTaskWindowRef->ShowWindow();
 
-				TickSlate();
-			}
+			TickSlate();
 		}
-	}
-}
-
-/** Whether or not the user has canceled out of this dialog */
-bool FFeedbackContextEditor::ReceivedUserCancel( void )
-{
-	const bool res = HasTaskBeenCancelled;
-	HasTaskBeenCancelled = false;
-	return res;
-}
 
-void FFeedbackContextEditor::EnableUserCancel( bool bUserCancel )
-{
-	if (SlowTaskWidget.IsValid())
-	{
-		if (!bUserCancel)
-		{
-			SlowTaskWidget->UnbindOnCancelDelegate();
-		}
-		else
-		{
-			SlowTaskWidget->BindOnCancelDelegate(FOnCancelClickedDelegate::CreateRaw(this, &FFeedbackContextEditor::OnUserCancel));
-		}
+		FPlatformSplash::SetSplashText( SplashTextType::StartupProgress, *Task.ToString() );
 	}
 }
 
-void FFeedbackContextEditor::OnUserCancel()
-{
-	HasTaskBeenCancelled = true;
-}
-
-/**
- * Tells the editor that the slow task is done.
- */
-void FFeedbackContextEditor::EndSlowTask()
+void FFeedbackContextEditor::FinalizeSlowTask()
 {
 	TSharedPtr<SWindow> ParentWindow;
 	if( FModuleManager::Get().IsModuleLoaded( "MainFrame" ) )
@@ -294,312 +403,59 @@ void FFeedbackContextEditor::EndSlowTask()
 
 	if( GIsEditor && ParentWindow.IsValid())
 	{
-		check(SlowTaskCount>0);
-		GIsSlowTask = --SlowTaskCount>0;
-
-		checkSlow(DialogRequestStack.Num()>0);
-		const bool bWasDialogRequest = DialogRequestStack.Pop();
-
-		// Restore the cursor now that the long running task is done
-		if (SlowTaskCount == 0)
-		{
-			// Reset cached message
-			StatusMessage.StatusText = FText::GetEmpty();
-			StatusMessage.ProgressNumerator = StatusMessage.SavedNumerator = 0;
-			StatusMessage.ProgressDenominator = StatusMessage.SavedDenominator = 0;
-			StatusMessage.LastUpdateTime  = FPlatformTime::Seconds();
-		}
-
-		if( bWasDialogRequest )
-		{
-			checkSlow(DialogRequestCount>0);
-			if ( --DialogRequestCount == 0 )
-			{
-				if( SlowTaskWindow.IsValid() )
-				{
-					SlowTaskWindow.Pin()->RequestDestroyWindow();
-					SlowTaskWindow.Reset();
-					SlowTaskWidget.Reset();
-
-				}
-			}
-		}
-	}
-}
-
-void FFeedbackContextEditor::SetContext( FContextSupplier* InSupplier )
-{
-}
-
-bool FFeedbackContextEditor::StatusUpdate( int32 Numerator, int32 Denominator, const FText& StatusText )
-{
-	// don't update too frequently
-	if( GIsSlowTask )
-	{
-		// limit to update to once every 'y' seconds
-		double Now = FPlatformTime::Seconds();
-		double Diff = Now - StatusMessage.LastUpdateTime;
-		if ( Diff >= UIUpdateGatingTime )
-		{
-			StatusMessage.LastUpdateTime = Now;
-		}
-		else
-		{
-			return true;
-		}
-	}
-
-	return ApplyStatusUpdate( StatusText, Numerator, Denominator );
-}
-
-bool FFeedbackContextEditor::StatusForceUpdate( int32 Numerator, int32 Denominator, const FText& StatusText )
-{
-	return ApplyStatusUpdate( StatusText, Numerator, Denominator );
-}
-
-bool FFeedbackContextEditor::ApplyStatusUpdate( const FText& StatusText, int32 ProgressNumerator, int32 ProgressDenominator )
-{
-	// Cache the new status message and progress
-	StatusMessage.StatusText = StatusText;
-
-	if( ProgressNumerator >= 0 ) // Ignore if set to -1
-	{
-		StatusMessage.ProgressNumerator = StatusMessage.SavedNumerator = ProgressNumerator;
-	}
-	if( ProgressDenominator >= 0 ) // Ignore if set to -1
-	{
-		StatusMessage.ProgressDenominator = StatusMessage.SavedDenominator = ProgressDenominator;
-	}
-
-	if( GIsSlowTask )
-	{
-		bool bCanCancelTask = (BuildProgressWidget.IsValid());
-
-		// Don't bother refreshing progress for sub-tasks (tasks greater than 2 levels deep.)  The reason we
-		// allow two levels is because some subsystems call PushStatus() before reporting any status for
-		// first-tier tasks, and we wouldn't be able to display any progress for those tasks otherwise.
-		if( StatusMessageStack.Num() <= 1 || bCanCancelTask)
+		if( SlowTaskWindow.IsValid() )
 		{
-			StatusUpdateProgress(
-				StatusMessage.StatusText,
-				StatusMessage.ProgressNumerator,
-				StatusMessage.ProgressDenominator,
-				StatusMessageStack.Num() == 0 || DialogRequestCount == 0
-				);
-
-			if ( DialogRequestCount > 0 )
-			{
-				if( SlowTaskWindow.IsValid() )
-				{
-					SlowTaskWidget->SetProgressText( StatusMessage.StatusText );
-					SlowTaskWidget->SetProgressPercent( StatusMessage.ProgressNumerator, StatusMessage.ProgressDenominator );
-
-					if (FSlateApplication::Get().CanDisplayWindows())
-					{
-						TickSlate();
-					}
-				}
-			}
+			SlowTaskWindow.Pin()->RequestDestroyWindow();
+			SlowTaskWindow.Reset();
+			SlowTaskWidget.Reset();
 		}
 	}
 
-	// Also update the splash screen text (in case we're in the middle of starting up)
-	if ( !StatusMessage.StatusText.IsEmpty() )
-	{
-		FPlatformSplash::SetSplashText( SplashTextType::StartupProgress, *StatusMessage.StatusText.ToString() );
-	}
-	return true;
+	FFeedbackContext::FinalizeSlowTask( );
 }
 
-/**
- * Updates the progress amount without changing the status message text
- *
- * @param Numerator		New progress numerator
- * @param Denominator	New progress denominator
- */
-void FFeedbackContextEditor::UpdateProgress( int32 Numerator, int32 Denominator )
+void FFeedbackContextEditor::ProgressReported( const float TotalProgressInterp, FText DisplayMessage )
 {
-	// Cache the new progress
-	if( Numerator >= 0 ) // Ignore if set to -1
+	if (SlowTaskWidget.IsValid() && FSlateApplication::Get().CanDisplayWindows())
 	{
-		StatusMessage.ProgressNumerator = Numerator;
-	}
-	if( Denominator >= 0 ) // Ignore if set to -1
-	{
-		StatusMessage.ProgressDenominator = Denominator;
+		TickSlate();
 	}
-
-	if( GIsSlowTask )
+	else
 	{
-		bool bCanCancelTask = (BuildProgressWidget.IsValid());
-
-		// Don't bother refreshing progress for sub-tasks (tasks greater than 2 levels deep.)  The reason we
-		// allow two levels is because some subsystems call PushStatus() before reporting any status for
-		// first-tier tasks, and we wouldn't be able to display any progress for those tasks otherwise.
-		if( StatusMessageStack.Num() <= 1 || bCanCancelTask)
+		// Always show the top-most message
+		for (auto& Scope : ScopeStack)
 		{
-			// calculate our previous percentage and our new one
-			float SavedRatio = 0.0f;
-			float NewRatio = 0.0f;
-
-			// Check divisors to guard against divide by zeros
-			if ( StatusMessage.SavedDenominator > 0 )
-			{
-				SavedRatio = (float)StatusMessage.SavedNumerator / (float)StatusMessage.SavedDenominator;
-			}
-
-			if ( ensure(StatusMessage.ProgressDenominator > 0) )
-			{
-				NewRatio = (float)StatusMessage.ProgressNumerator / (float)StatusMessage.ProgressDenominator;
-			}
-
-			double Now = FPlatformTime::Seconds();
-
-			// update the progress bar if we've moved enough since last time, or we are going to start or end of bar,
-			// or if a second has passed
-			//
-			// or if the build progress is running...
-			bool bUpdate =
-				(FMath::Abs<float>(SavedRatio - NewRatio) > 0.1f || Numerator == 0 || Numerator >= (Denominator - 1) ||
-				Now >= StatusMessage.LastUpdateTime + UIUpdateGatingTime);
-				
-			if (bUpdate || bCanCancelTask)
+			const FText ThisMessage = Scope->GetCurrentMessage();
+			if (!ThisMessage.IsEmpty())
 			{
-				StatusMessage.SavedNumerator = StatusMessage.ProgressNumerator;
-				StatusMessage.SavedDenominator = StatusMessage.ProgressDenominator;
-				StatusMessage.LastUpdateTime = Now;
-
-				StatusUpdateProgress(
-					FText::GetEmpty(),
-					StatusMessage.ProgressNumerator,
-					StatusMessage.ProgressDenominator,
-					StatusMessageStack.Num() == 0 || DialogRequestCount == 0
-					);
-			
-				if ( DialogRequestCount > 0 )
-				{
-					if( SlowTaskWindow.IsValid()  )
-					{
-						SlowTaskWidget->SetProgressPercent( StatusMessage.ProgressNumerator, StatusMessage.ProgressDenominator );
-
-						if (FSlateApplication::Get().CanDisplayWindows())
-						{
-							TickSlate();
-						}
-					}
-				}
+				DisplayMessage = ThisMessage;
+				break;
 			}
 		}
-	}
-}
-
-/**
- *  Updates text and value for various progress meters.
- *
- *	@param StatusText				New status text
- *	@param ProgressNumerator		Numerator for the progress meter (its current value).
- *	@param ProgressDenominitator	Denominiator for the progress meter (its range).
- */
-void FFeedbackContextEditor::StatusUpdateProgress( const FText& StatusText, int32 ProgressNumerator, int32 ProgressDenominator, bool bUpdateBuildDialog/*=true*/ )
-{
-	// Clean up deferred cleanup objects from rendering thread every once in a while.
-	static double LastTimePendingCleanupObjectsWhereDeleted;
-	if( FPlatformTime::Seconds() - LastTimePendingCleanupObjectsWhereDeleted > 1 )
-	{
-		// Get list of objects that are pending cleanup.
-		FPendingCleanupObjects* PendingCleanupObjects = GetPendingCleanupObjects();
-		// Flush rendering commands in the queue.
-		FlushRenderingCommands();
-		// It is now save to delete the pending clean objects.
-		delete PendingCleanupObjects;
-		// Keep track of time this operation was performed so we don't do it too often.
-		LastTimePendingCleanupObjectsWhereDeleted = FPlatformTime::Seconds();
-	}
-
-	// Update build progress dialog if it is visible.
-	const bool bBuildProgressDialogVisible = BuildProgressWidget.IsValid();
 
-	if( bBuildProgressDialogVisible && bUpdateBuildDialog )
-	{
-		if( !StatusText.IsEmpty() )
+		if (!DisplayMessage.IsEmpty() && TotalProgressInterp > 0)
 		{
-			BuildProgressWidget->SetBuildStatusText( StatusText );
-		}
+            FFormatOrderedArguments Args;
+            Args.Add(DisplayMessage);
+            Args.Add(int(TotalProgressInterp * 100.f));
 
-		BuildProgressWidget->SetBuildProgressPercent( ProgressNumerator, ProgressDenominator );
-
-		if (FSlateApplication::Get().CanDisplayWindows())
-		{
-			TickSlate();
+            DisplayMessage = FText::Format(NSLOCTEXT("FeedbackContextEditor", "ProgressDisplayText", "{0} ({1}%)"), Args);
 		}
+		FPlatformSplash::SetSplashText( SplashTextType::StartupProgress, *DisplayMessage.ToString() );
 	}
 }
 
-/** Pushes the current status message/progress onto the stack so it can be restored later */
-void FFeedbackContextEditor::PushStatus()
+/** Whether or not the user has canceled out of this dialog */
+bool FFeedbackContextEditor::ReceivedUserCancel( void )
 {
-	StatusMessage.LastUpdateTime = 0.0f; // Gurantees the next call to update this window will not fail due to the 1 second requirement between updates. When we are popped, our time will also be 0.0f.
-
-	// Push the current message onto the stack.  This doesn't change the current message though.  You should
-	// call StatusUpdate after calling PushStatus to update the message.
-	StatusMessageStack.Add( StatusMessage );
+	const bool res = HasTaskBeenCancelled;
+	HasTaskBeenCancelled = false;
+	return res;
 }
 
-
-
-/** Restores the previously pushed status message/progress */
-void FFeedbackContextEditor::PopStatus()
+void FFeedbackContextEditor::OnUserCancel()
 {
-	if( StatusMessageStack.Num() > 0 )
-	{
-		// Pop from stack
-		StatusMessageStackItem PoppedStatusMessage = StatusMessageStack.Pop();
-
-		// Update the message text
-		StatusMessage.StatusText = PoppedStatusMessage.StatusText;
-
-		// Only overwrite progress if the item on the stack actually had some progress set
-		if( PoppedStatusMessage.ProgressDenominator > 0 )
-		{
-			StatusMessage.ProgressNumerator = PoppedStatusMessage.ProgressNumerator;
-			StatusMessage.ProgressDenominator = PoppedStatusMessage.ProgressDenominator;
-			StatusMessage.SavedNumerator = PoppedStatusMessage.SavedNumerator;
-			StatusMessage.SavedDenominator = PoppedStatusMessage.SavedDenominator;
-			StatusMessage.LastUpdateTime = PoppedStatusMessage.LastUpdateTime;
-		}
-
-		// Update the GUI!
-		if( GIsSlowTask )
-		{
-			StatusMessage.SavedNumerator = StatusMessage.ProgressNumerator;
-			StatusMessage.SavedDenominator = StatusMessage.ProgressDenominator;
-
-			// Don't bother refreshing progress for sub-tasks (tasks greater than 2 levels deep.)  The reason we
-			// allow two levels is because some subsystems call PushStatus() before reporting any status for
-			// first-tier tasks, and we wouldn't be able to display any progress for those tasks otherwise.
-			if( StatusMessageStack.Num() <= 1 )
-			{
-				StatusUpdateProgress(
-					StatusMessage.StatusText,
-					StatusMessage.ProgressNumerator,
-					StatusMessage.ProgressDenominator);
-
-				if ( DialogRequestCount > 0 )
-				{
-					if( SlowTaskWindow.IsValid()  )
-					{
-						SlowTaskWidget->SetProgressText( StatusMessage.StatusText );
-						SlowTaskWidget->SetProgressPercent( StatusMessage.ProgressNumerator, StatusMessage.ProgressDenominator );
-
-						if (FSlateApplication::Get().CanDisplayWindows())
-						{
-							TickSlate();
-						}
-					}
-				}
-			}
-		}
-	}
+	HasTaskBeenCancelled = true;
 }
 
 /** 
diff --git a/Engine/Source/Editor/UnrealEd/Private/FileHelpers.cpp b/Engine/Source/Editor/UnrealEd/Private/FileHelpers.cpp
--- a/Engine/Source/Editor/UnrealEd/Private/FileHelpers.cpp
+++ b/Engine/Source/Editor/UnrealEd/Private/FileHelpers.cpp
@@ -485,7 +485,11 @@ static bool SaveWorld(UWorld* World,
 
 		FFormatNamedArguments Args;
 		Args.Add( TEXT("MapFilename"), FText::FromString( FPaths::GetCleanFilename(FinalFilename) ) );
-		GWarn->BeginSlowTask( FText::Format( NSLOCTEXT("UnrealEd", "SavingMap_F", "Saving map: {MapFilename}..." ), Args ), true, true );
+		
+		FScopedSlowTask SlowTask(100, FText::Format( NSLOCTEXT("UnrealEd", "SavingMap_F", "Saving map: {MapFilename}..." ), Args ));
+		SlowTask.MakeDialog(true);
+
+		SlowTask.EnterProgressFrame(25);
 
 		// Rename the package and the object, as necessary
 		if ( bRenamePackageToFile )
@@ -501,6 +505,8 @@ static bool SaveWorld(UWorld* World,
 			}
 		}
 
+		SlowTask.EnterProgressFrame(50);
+
 		// Save package.
 		{
 			const FString AutoSavingString = (bAutosaving || bPIESaving) ? TEXT("true") : TEXT("false");
@@ -511,6 +517,8 @@ static bool SaveWorld(UWorld* World,
 			SaveErrors.Flush();
 		}
 
+		SlowTask.EnterProgressFrame(25);
+
 		// If the package save was not successful. Rename anything we changed back to the original name.
 		if( bRenamePackageToFile && !bSuccess )
 		{
@@ -524,8 +532,6 @@ static bool SaveWorld(UWorld* World,
 				World->Rename(*OriginalWorldName, NULL, REN_NonTransactional);
 			}
 		}
-
-		GWarn->EndSlowTask();
 	}
 
 	return bSuccess;
diff --git a/Engine/Source/Editor/UnrealEd/Private/ObjectTools.cpp b/Engine/Source/Editor/UnrealEd/Private/ObjectTools.cpp
--- a/Engine/Source/Editor/UnrealEd/Private/ObjectTools.cpp
+++ b/Engine/Source/Editor/UnrealEd/Private/ObjectTools.cpp
@@ -3513,8 +3513,8 @@ namespace ThumbnailTools
 			// When generating a material thumbnail to save in a package, make sure we finish compilation on the material first
 			if ( UMaterial* InMaterial = Cast<UMaterial>(InObject) )
 			{
-				const bool bAllowNewSlowTask = true;
-				FScopedSlowTask SlowTaskMessage( NSLOCTEXT( "ObjectTools", "FinishingCompilationStatus", "Finishing Shader Compilation..." ), bAllowNewSlowTask );
+				FScopedSlowTask SlowTask(0, NSLOCTEXT( "ObjectTools", "FinishingCompilationStatus", "Finishing Shader Compilation..." ) );
+				SlowTask.MakeDialog();
 
 				// Block until the shader maps that we will save have finished being compiled
 				InMaterial->GetMaterialResource(GMaxRHIFeatureLevel)->FinishCompilation();
diff --git a/Engine/Source/Editor/UnrealEd/Private/StaticLightingSystem/StaticLightingSystem.cpp b/Engine/Source/Editor/UnrealEd/Private/StaticLightingSystem/StaticLightingSystem.cpp
--- a/Engine/Source/Editor/UnrealEd/Private/StaticLightingSystem/StaticLightingSystem.cpp
+++ b/Engine/Source/Editor/UnrealEd/Private/StaticLightingSystem/StaticLightingSystem.cpp
@@ -1012,11 +1012,13 @@ void FStaticLightingSystem::EncodeTextures(bool bLightingSuccessful)
 {
 	FLightmassStatistics::FScopedGather EncodeStatScope(LightmassStatistics.EncodingTime);
 
+	FScopedSlowTask SlowTask(2);
+
 	// Flush pending shadow-map and light-map encoding.
-	GWarn->StatusUpdate( -1, -1, LOCTEXT("EncodingImportedStaticLightMapsStatusMessage", "Encoding imported static light maps.") );
+	SlowTask.EnterProgressFrame(1, LOCTEXT("EncodingImportedStaticLightMapsStatusMessage", "Encoding imported static light maps.") );
 	FLightMap2D::EncodeTextures(World, bLightingSuccessful, true);
 
-	GWarn->StatusUpdate( -1, -1, LOCTEXT("EncodingImportedStaticShadowMapsStatusMessage", "Encoding imported static shadow maps.") );
+	SlowTask.EnterProgressFrame(1, LOCTEXT("EncodingImportedStaticShadowMapsStatusMessage", "Encoding imported static shadow maps.") );
 	FShadowMap2D::EncodeTextures(World, bLightingSuccessful);
 }
 
@@ -1886,70 +1888,82 @@ bool FStaticLightingSystem::FinishLightmassProcess()
 
 	double TimeWaitingOnUserToAccept = FPlatformTime::Seconds() - WaitForUserAcceptStartTime;
 	
-	GWarn->BeginSlowTask( LOCTEXT("InvalidatingPreviousLightingStatus", "Invalidating previous lighting."), true );
+	{
+		FScopedSlowTask SlowTask(7);
+		SlowTask.MakeDialog();
 
-	InvalidateStaticLighting();
+		SlowTask.EnterProgressFrame(1, LOCTEXT("InvalidatingPreviousLightingStatus", "Invalidating previous lighting"));
+		InvalidateStaticLighting();
 	
-	GWarn->StatusUpdate( -1, -1, LOCTEXT("ImportingBuiltStaticLightingStatus", "Importing built static lighting.") );
 
-	bSuccessful = LightmassProcessor->CompleteRun();
+		SlowTask.EnterProgressFrame(1, LOCTEXT("ImportingBuiltStaticLightingStatus", "Importing built static lighting"));
+		bSuccessful = LightmassProcessor->CompleteRun();
 
-	if (bSuccessful)
-	{
-		CompleteDeterministicMappings(LightmassProcessor);
-		
-		if (!Options.bOnlyBuildVisibility)
+
+		SlowTask.EnterProgressFrame();
+		if (bSuccessful)
 		{
-			FLightmassStatistics::FScopedGather FinishStatScope(LightmassStatistics.FinishingTime);
-			ULightComponent::ReassignStationaryLightChannels(GWorld, true);
+			CompleteDeterministicMappings(LightmassProcessor);
+		
+			if (!Options.bOnlyBuildVisibility)
+			{
+				FLightmassStatistics::FScopedGather FinishStatScope(LightmassStatistics.FinishingTime);
+				ULightComponent::ReassignStationaryLightChannels(GWorld, true);
+			}
 		}
-	}
 	
-	EncodeTextures(bSuccessful);
-
-	{
-		FLightmassStatistics::FScopedGather CloseJobStatScope(LightmassProcessStatistics.SwarmJobCloseTime);
-		bSuccessful = LightmassProcessor->CloseJob() && bSuccessful;
-	}
-
-	{
-		FLightmassStatistics::FScopedGather FinishStatScope(LightmassStatistics.FinishingTime);
-		// Add in the time measurements from the LightmassProcessor
-		LightmassStatistics += LightmassProcessor->GetStatistics();
 
-		// A final update on the lighting build warnings and errors dialog, now that everything is finished
-		FMessageLog("LightingResults").Open();
+		SlowTask.EnterProgressFrame(1, LOCTEXT("EncodingTexturesStaticLightingStatis", "Encoding textures"));
+		EncodeTextures(bSuccessful);
 
-		// Check the for build cancellation.
-		bBuildCanceled = bBuildCanceled || GEditor->GetMapBuildCancelled();
-		bSuccessful = bSuccessful && !bBuildCanceled;
 
-		FStatsViewerModule& StatsViewerModule = FModuleManager::Get().LoadModuleChecked<FStatsViewerModule>(TEXT("StatsViewer"));
-		if (bSuccessful)
+		SlowTask.EnterProgressFrame();
 		{
-			StatsViewerModule.GetPage(EStatsPage::LightingBuildInfo)->Refresh();
+			FLightmassStatistics::FScopedGather CloseJobStatScope(LightmassProcessStatistics.SwarmJobCloseTime);
+			bSuccessful = LightmassProcessor->CloseJob() && bSuccessful;
 		}
-		
-		bool bShowLightingBuildInfo = false;
-		GConfig->GetBool( TEXT("LightingBuildOptions"), TEXT("ShowLightingBuildInfo"), bShowLightingBuildInfo, GEditorUserSettingsIni );
-		if( bShowLightingBuildInfo )
+
 		{
-			StatsViewerModule.GetPage(EStatsPage::LightingBuildInfo)->Show();
+			FLightmassStatistics::FScopedGather FinishStatScope(LightmassStatistics.FinishingTime);
+			// Add in the time measurements from the LightmassProcessor
+			LightmassStatistics += LightmassProcessor->GetStatistics();
+
+			// A final update on the lighting build warnings and errors dialog, now that everything is finished
+			FMessageLog("LightingResults").Open();
+
+			// Check the for build cancellation.
+			bBuildCanceled = bBuildCanceled || GEditor->GetMapBuildCancelled();
+			bSuccessful = bSuccessful && !bBuildCanceled;
+
+			FStatsViewerModule& StatsViewerModule = FModuleManager::Get().LoadModuleChecked<FStatsViewerModule>(TEXT("StatsViewer"));
+			if (bSuccessful)
+			{
+				StatsViewerModule.GetPage(EStatsPage::LightingBuildInfo)->Refresh();
+			}
+		
+			bool bShowLightingBuildInfo = false;
+			GConfig->GetBool( TEXT("LightingBuildOptions"), TEXT("ShowLightingBuildInfo"), bShowLightingBuildInfo, GEditorUserSettingsIni );
+			if( bShowLightingBuildInfo )
+			{
+				StatsViewerModule.GetPage(EStatsPage::LightingBuildInfo)->Show();
+			}
 		}
-	}
 
-	ApplyNewLightingData(bSuccessful);
 
-	PostInvalidateStaticLighting();
+		SlowTask.EnterProgressFrame();
+		ApplyNewLightingData(bSuccessful);
 
-	// Finish up timing statistics
-	LightmassStatistics += LightmassProcessStatistics;
-	LightmassStatistics.TotalTime += FPlatformTime::Seconds() - StartTime - TimeWaitingOnUserToAccept;
 
-	GetRendererModule().UpdateMapNeedsLightingFullyRebuiltState(World);
-	GEngine->DeferredCommands.AddUnique(TEXT("MAP CHECK NOTIFYRESULTS"));
+		SlowTask.EnterProgressFrame();
+		PostInvalidateStaticLighting();
 
-	GWarn->EndSlowTask();
+		// Finish up timing statistics
+		LightmassStatistics += LightmassProcessStatistics;
+		LightmassStatistics.TotalTime += FPlatformTime::Seconds() - StartTime - TimeWaitingOnUserToAccept;
+
+		GetRendererModule().UpdateMapNeedsLightingFullyRebuiltState(World);
+		GEngine->DeferredCommands.AddUnique(TEXT("MAP CHECK NOTIFYRESULTS"));
+	}
 
 	ReportStatistics();
 	
diff --git a/Engine/Source/Editor/UnrealEd/Private/UnrealEd.cpp b/Engine/Source/Editor/UnrealEd/Private/UnrealEd.cpp
--- a/Engine/Source/Editor/UnrealEd/Private/UnrealEd.cpp
+++ b/Engine/Source/Editor/UnrealEd/Private/UnrealEd.cpp
@@ -53,6 +53,10 @@ int32 EditorInit( IEngineLoop& EngineLoop )
 	GDebugToolExec = new FDebugToolExec;
 
 	DECLARE_SCOPE_CYCLE_COUNTER(TEXT("Editor Initialized"), STAT_EditorStartup, STATGROUP_LoadTime);
+	
+	FScopedSlowTask SlowTask(100, NSLOCTEXT("EngineLoop", "EngineLoop_Loading", "Loading..."));
+	
+	SlowTask.EnterProgressFrame(50);
 
 	int32 ErrorLevel = EngineLoop.Init();
 	if( ErrorLevel != 0 )
@@ -73,8 +77,12 @@ int32 EditorInit( IEngineLoop& EngineLoop )
 		FEngineAnalytics::GetProvider().RecordEvent(TEXT("Editor.ProgramStarted"), EventAttributes);
 	}
 
+	SlowTask.EnterProgressFrame(40);
+
 	// Initialize the misc editor
 	FUnrealEdMisc::Get().OnInit();
+	
+	SlowTask.EnterProgressFrame(10);
 
 	// Prime our array of default directories for loading and saving content files to
 	FEditorDirectories::Get().LoadLastDirectories();
diff --git a/Engine/Source/Editor/UnrealEd/Private/UnrealEdMisc.cpp b/Engine/Source/Editor/UnrealEd/Private/UnrealEdMisc.cpp
--- a/Engine/Source/Editor/UnrealEd/Private/UnrealEdMisc.cpp
+++ b/Engine/Source/Editor/UnrealEd/Private/UnrealEdMisc.cpp
@@ -169,7 +169,10 @@ void FUnrealEdMisc::OnInit()
 		return;
 	}
 	bInitialized = true;
-	
+
+	FScopedSlowTask SlowTask(100);
+	SlowTask.EnterProgressFrame(10);
+
 	// Register all callback notifications
 	FEditorDelegates::SelectedProps.AddRaw(this, &FUnrealEdMisc::CB_SelectedProps);
 	FEditorDelegates::DisplayLoadErrors.AddRaw(this, &FUnrealEdMisc::CB_DisplayLoadErrors);
@@ -224,6 +227,8 @@ void FUnrealEdMisc::OnInit()
 	const TCHAR* ParsedCmdLine = FCommandLine::Get();
 	const bool bIsImmersive = FParse::Param( ParsedCmdLine, TEXT( "immersive" ) );
 
+	SlowTask.EnterProgressFrame(10);
+
 	ISourceControlModule::Get().GetProvider().Init();
 
 	// Init the editor tools.
@@ -247,6 +252,7 @@ void FUnrealEdMisc::OnInit()
 	const bool bDoAutomatedMapBuild = FParse::Param( ParsedCmdLine, TEXT("AutomatedMapBuild") );
 
 	// Load startup map (conditionally)
+	SlowTask.EnterProgressFrame(60);
 	{
 		bool bMapLoaded = false;
 
@@ -360,6 +366,8 @@ void FUnrealEdMisc::OnInit()
 		FEditorBuildUtils::EditorAutomatedBuildAndSubmit( AutomatedBuildSettings, ErrorText );
 	}
 
+	SlowTask.EnterProgressFrame(10);
+
 	LoadFBxLibraries();
 
 	// Register message log UIs
@@ -417,6 +425,8 @@ void FUnrealEdMisc::OnInit()
 	// Register to receive notification of new key bindings
 	FInputBindingManager::Get().RegisterUserDefinedGestureChanged(FOnUserDefinedGestureChanged::FDelegate::CreateRaw( this, &FUnrealEdMisc::OnUserDefinedGestureChanged ));
 
+	SlowTask.EnterProgressFrame(10);
+
 	// Send Project Analytics
 	InitEngineAnalytics();
 	
diff --git a/Engine/Source/Editor/UnrealEd/Public/FeedbackContextEditor.h b/Engine/Source/Editor/UnrealEd/Public/FeedbackContextEditor.h
--- a/Engine/Source/Editor/UnrealEd/Public/FeedbackContextEditor.h
+++ b/Engine/Source/Editor/UnrealEd/Public/FeedbackContextEditor.h
@@ -7,103 +7,35 @@
 
 #pragma once
 
-
-
 /**
  * A FFeedbackContext implementation for use in UnrealEd.
  */
 class FFeedbackContextEditor : public FFeedbackContext
 {
-	/** minimum time between progress updates displayed - more frequent updates are ignored */
-	static const float UIUpdateGatingTime;
-
-	/** the number of times a caller requested the progress dialog be shown */
-	int32				DialogRequestCount;
-	/** keeps track of the order in which calls were made requesting the progress dialog */
-	TArray<bool>	DialogRequestStack;
-
 	/** Slate slow task widget */
 	TWeakPtr<class SWindow> SlowTaskWindow;
 	TSharedPtr<class SSlowTaskWidget> SlowTaskWidget;
 
-	/**
-	 * StatusMessageStackItem
-	 */
-	struct StatusMessageStackItem
-	{
-		/** Status message text */
-		FText StatusText;
-
-		/** Progress numerator */
-		int32 ProgressNumerator;
-
-		/** Progress denominator */
-		int32 ProgressDenominator;
-
-		/** Cached numerator so we can update less frequently */
-		int32 SavedNumerator;
-		
-		/** Cached denominator so we can update less frequently */
-		int32 SavedDenominator;
-
-		/** The list time we updated the progress, so we can make sure to update at least once a second */
-		double LastUpdateTime;
-
-		StatusMessageStackItem()
-		: ProgressNumerator(0), ProgressDenominator(0), SavedNumerator(0), SavedDenominator(0), LastUpdateTime(0.0)
-		{
-		}
-	};
-
-	/** Current status message and progress */
-	StatusMessageStackItem StatusMessage;
-
-	/** Stack of status messages and progress values */
-	TArray< StatusMessageStackItem > StatusMessageStack;
-
 	/** Special Windows/Widget popup for building */
 	TWeakPtr<class SWindow> BuildProgressWindow;
 	TSharedPtr<class SBuildProgressWidget> BuildProgressWidget;
 
 	bool HasTaskBeenCancelled;
 
-
-	/**
-	 *  Updates text and value for various progress meters.
-	 *
-	 *	@param StatusText				New status text
-	 *	@param ProgressNumerator		Numerator for the progress meter (its current value).
-	 *	@param ProgressDenominitator	Denominiator for the progress meter (its range).
-	 */
-	void StatusUpdateProgress( const FText& StatusText, int32 ProgressNumerator, int32 ProgressDenominator, bool bUpdateBuildDialog =true );
-
-
-	/** 
-	 *  Update dialog window text and progress value
-	 *	@param StatusText				New status text
-	 *	@param ProgressNumerator		Numerator for the progress meter (its current value).
-	 *	@param ProgressDenominitator	Denominator for the progress meter (its range).
-	 *
-	 *	@return true by default and false if status update can't be done. 
-	 */
-	bool ApplyStatusUpdate( const FText& StatusText, int32 ProgressNumerator, int32 ProgressDenominator );
-
 public:
-	int32 SlowTaskCount;
 
 	UNREALED_API FFeedbackContextEditor();
 
 	virtual void Serialize( const TCHAR* V, ELogVerbosity::Type Verbosity, const class FName& Category ) override;
 
-	void BeginSlowTask( const FText& Task, bool bShouldShowProgressDialog, bool bShowCancelButton=false );
-	void EndSlowTask();
-	void SetContext( FContextSupplier* InSupplier );
+	virtual void StartSlowTask( const FText& Task, bool bShowCancelButton=false ) override;
+	virtual void FinalizeSlowTask( ) override;
+	virtual void ProgressReported( const float TotalProgressInterp, FText DisplayMessage ) override;
 
-	/** Whether or not the user has canceled out of this dialog */
-	bool ReceivedUserCancel();
+	void SetContext( FContextSupplier* InSupplier ) {}
 
-	/** Enable/Disable the ability to cancel the current task */
-	void EnableUserCancel(bool bUserCancel) override;
+	/** Whether or not the user has canceled out of this dialog */
+	virtual bool ReceivedUserCancel() override;
 
 	void OnUserCancel();
 
@@ -112,40 +44,6 @@ public:
 		return EAppReturnType::Yes == FMessageDialog::Open( EAppMsgType::YesNo, Question );
 	}
 
-	/** 
-	 * Slow task update status
-	 * @param Numerator		New progress numerator
-	 * @param Denominator	New progress denominator
-	 * @param StatusText	New status message
-	 * 
-	 * @return true by default and false if status update is not possible
-	 */
-	virtual bool StatusUpdate( int32 Numerator, int32 Denominator, const FText& StatusText ) override;
-
-	/** Force updating feedback status
-	 *
-	 * @param Numerator		New progress numerator
-	 * @param Denominator	New progress denominator
-	 * @param StatusText	New status message
-	 * @return true by default and false if status update is not possible
-	 *
-	 */
-	virtual bool StatusForceUpdate( int32 Numerator, int32 Denominator, const FText& StatusText ) override;
-
-	/**
-	 * Updates the progress amount without changing the status message text
-	 *
-	 * @param Numerator		New progress numerator
-	 * @param Denominator	New progress denominator
-	 */
-	virtual void UpdateProgress( int32 Numerator, int32 Denominator );
-
-	/** Pushes the current status message/progress onto the stack so it can be restored later */
-	virtual void PushStatus();
-
-	/** Restores the previously pushed status message/progress */
-	virtual void PopStatus();
-
 	/** 
 	 * Show the Build Progress Window 
 	 * @return Handle to the Build Progress Widget created
diff --git a/Engine/Source/Runtime/Core/Private/Mac/MacPlatformFeedbackContextPrivate.h b/Engine/Source/Runtime/Core/Private/Mac/MacPlatformFeedbackContextPrivate.h
--- a/Engine/Source/Runtime/Core/Private/Mac/MacPlatformFeedbackContextPrivate.h
+++ b/Engine/Source/Runtime/Core/Private/Mac/MacPlatformFeedbackContextPrivate.h
@@ -15,14 +15,10 @@ class FFeedbackContextMac : public FFeedbackContext
 	FContextSupplier*	Context;
 
 public:
-	// Variables.
-	int32					SlowTaskCount;
-
 	// Constructor.
 	FFeedbackContextMac()
 		: FFeedbackContext()
 		, Context( NULL )
-		, SlowTaskCount( 0 )
 	{}
 
 	void Serialize( const TCHAR* V, ELogVerbosity::Type Verbosity, const class FName& Category )
@@ -79,20 +75,6 @@ public:
 	{
 		return FPlatformMisc::MessageBoxExt(EAppMsgType::YesNo, *Question.ToString(), TEXT("")) == EAppReturnType::Yes;
 	}
-
-	void BeginSlowTask( const FText& Task, bool ShowProgressDialog, bool bShowCancelButton=false )
-	{
-		GIsSlowTask = ++SlowTaskCount>0;
-	}
-	void EndSlowTask()
-	{
-		check(SlowTaskCount>0);
-		GIsSlowTask = --SlowTaskCount>0;
-	}
-	virtual bool StatusUpdate( int32 Numerator, int32 Denominator, const FText& StatusText )
-	{
-		return true;
-	}
 	FContextSupplier* GetContext() const
 	{
 		return Context;
diff --git a/Engine/Source/Runtime/Core/Private/Misc/FeedbackContext.cpp b/Engine/Source/Runtime/Core/Private/Misc/FeedbackContext.cpp
new file mode 100644
--- /dev/null
+++ b/Engine/Source/Runtime/Core/Private/Misc/FeedbackContext.cpp
@@ -0,0 +1,108 @@
+// Copyright 1998-2014 Epic Games, Inc. All Rights Reserved.
+
+#include "Core.h"
+#include "FeedbackContext.h"
+
+void FFeedbackContext::RequestUpdateUI(bool bForceUpdate)
+{
+	// Only update a maximum of 5 times a second
+	static double MinUpdateTimeS = 0.2;
+
+	static double LastUIUpdateTime = FPlatformTime::Seconds();
+	const double CurrentTime = FPlatformTime::Seconds();
+
+	if (bForceUpdate || CurrentTime - LastUIUpdateTime > MinUpdateTimeS)
+	{
+		LastUIUpdateTime = CurrentTime;
+		UpdateUI();
+	}
+}
+
+void FFeedbackContext::UpdateUI()
+{
+	if (ScopeStack.Num() != 0)
+	{
+		ProgressReported(ScopeStack.GetProgressFraction(0), ScopeStack[0]->GetCurrentMessage());
+	}
+}
+
+/**** Begin legacy API ****/
+void FFeedbackContext::BeginSlowTask( const FText& Task, bool ShowProgressDialog, bool bShowCancelButton )
+{
+	TUniquePtr<FScopedSlowTask> NewScope(new FScopedSlowTask(0, Task, true, *this));
+	if (ShowProgressDialog)
+	{
+		NewScope->MakeDialog(bShowCancelButton);
+	}
+	LegacyAPIScopes.Add(MoveTemp(NewScope));
+}
+
+void FFeedbackContext::UpdateProgress( int32 Numerator, int32 Denominator )
+{
+	if (LegacyAPIScopes.Num() != 0)
+	{
+		LegacyAPIScopes.Last()->TotalAmountOfWork = Denominator;
+		LegacyAPIScopes.Last()->CompletedWork = Numerator;
+		LegacyAPIScopes.Last()->CurrentFrameScope = Denominator - Numerator;
+		RequestUpdateUI();
+	}
+}
+
+void FFeedbackContext::StatusUpdate( int32 Numerator, int32 Denominator, const FText& StatusText )
+{
+	if (LegacyAPIScopes.Num() != 0)
+	{
+		if (Numerator > 0 && Denominator > 0)
+		{
+			UpdateProgress(Numerator, Denominator);
+		}
+		LegacyAPIScopes.Last()->FrameMessage = StatusText;
+		RequestUpdateUI();
+	}
+}
+
+void FFeedbackContext::StatusForceUpdate( int32 Numerator, int32 Denominator, const FText& StatusText )
+{
+	if (LegacyAPIScopes.Num() != 0)
+	{
+		UpdateProgress(Numerator, Denominator);
+		LegacyAPIScopes.Last()->FrameMessage = StatusText;
+		UpdateUI();
+	}
+}
+
+void FFeedbackContext::EndSlowTask()
+{
+	check(LegacyAPIScopes.Num() != 0);
+	LegacyAPIScopes.Pop();
+}
+/**** End legacy API ****/
+
+void FScopedSlowTask::MakeDialog(bool bShowCancelButton)
+{
+	const bool bIsDialogAllowed = bEnabled && !GIsSilent && !GIsPlayInEditorWorld && !IsRunningCommandlet() && IsInGameThread();
+	if (!GIsSlowTask && bIsDialogAllowed)
+	{
+		Context.StartSlowTask(GetCurrentMessage(), bShowCancelButton);
+		bCreatedDialog = true;
+	}
+}
+
+float FScopedSlowTaskStack::GetProgressFraction(int32 Index) const
+{
+	const int32 StartIndex = Num() - 1;
+	const int32 EndIndex = Index;
+
+	float Progress = 0.f;
+	for (int32 CurrentIndex = StartIndex; CurrentIndex >= EndIndex; --CurrentIndex)
+	{
+		const FScopedSlowTask* Scope = (*this)[CurrentIndex];
+		
+		const float ThisScopeCompleted = float(Scope->CompletedWork) / Scope->TotalAmountOfWork;
+		const float ThisScopeCurrentFrame = float(Scope->CurrentFrameScope) / Scope->TotalAmountOfWork;
+
+		Progress = ThisScopeCompleted + ThisScopeCurrentFrame*Progress;
+	}
+
+	return Progress;
+}
\ No newline at end of file
diff --git a/Engine/Source/Runtime/Core/Private/WinRT/WinRTFeedbackContextPrivate.h b/Engine/Source/Runtime/Core/Private/WinRT/WinRTFeedbackContextPrivate.h
--- a/Engine/Source/Runtime/Core/Private/WinRT/WinRTFeedbackContextPrivate.h
+++ b/Engine/Source/Runtime/Core/Private/WinRT/WinRTFeedbackContextPrivate.h
@@ -15,13 +15,10 @@ class FFeedbackContextWinRT : public FFeedbackContext
 	FContextSupplier*	Context;
 
 public:
-	// Variables.
-	int32 SlowTaskCount;
 
 	// Constructor.
 	FFeedbackContextWinRT()
 		: FFeedbackContext()
-		, SlowTaskCount( 0 )
 		, Context( NULL )
 	{}
 
@@ -94,17 +91,6 @@ public:
 		}
 	}
 
-	void BeginSlowTask( const TCHAR* Task, bool ShowProgressDialog, bool bShowCancelButton=false )
-	{
-		GIsSlowTask = ++SlowTaskCount>0;
-	}
-
-	void EndSlowTask()
-	{
-		check(SlowTaskCount>0);
-		GIsSlowTask = --SlowTaskCount>0;
-	}
-
 	VARARG_BODY( bool VARARGS, StatusUpdatef, const TCHAR*, VARARG_EXTRA(int32 Numerator) VARARG_EXTRA(int32 Denominator) )
 	{
 		return true;
diff --git a/Engine/Source/Runtime/Core/Private/Windows/WindowsPlatformFeedbackContextPrivate.h b/Engine/Source/Runtime/Core/Private/Windows/WindowsPlatformFeedbackContextPrivate.h
--- a/Engine/Source/Runtime/Core/Private/Windows/WindowsPlatformFeedbackContextPrivate.h
+++ b/Engine/Source/Runtime/Core/Private/Windows/WindowsPlatformFeedbackContextPrivate.h
@@ -16,13 +16,11 @@ class FFeedbackContextWindows : public FFeedbackContext
 
 public:
 	// Variables.
-	int32					SlowTaskCount;
 
 	// Constructor.
 	FFeedbackContextWindows()
 	: FFeedbackContext()
 	, Context( NULL )
-	, SlowTaskCount( 0 )
 	{}
 	void Serialize( const TCHAR* V, ELogVerbosity::Type Verbosity, const class FName& Category )
 	{
@@ -82,20 +80,6 @@ public:
 			return false;
 		}
 	}
-
-	void BeginSlowTask( const FText& Task, bool ShowProgressDialog, bool bShowCancelButton=false )
-	{
-		GIsSlowTask = ++SlowTaskCount>0;
-	}
-	void EndSlowTask()
-	{
-		check(SlowTaskCount>0);
-		GIsSlowTask = --SlowTaskCount>0;
-	}
-	virtual bool StatusUpdate( int32 Numerator, int32 Denominator, const FText& StatusText )
-	{
-		return true;
-	}
 	FContextSupplier* GetContext() const
 	{
 		return Context;
diff --git a/Engine/Source/Runtime/Core/Public/HAL/FeedbackContextAnsi.h b/Engine/Source/Runtime/Core/Public/HAL/FeedbackContextAnsi.h
--- a/Engine/Source/Runtime/Core/Public/HAL/FeedbackContextAnsi.h
+++ b/Engine/Source/Runtime/Core/Public/HAL/FeedbackContextAnsi.h
@@ -14,7 +14,6 @@ class FFeedbackContextAnsi : public FFeedbackContext
 {
 public:
 	// Variables.
-	int32					SlowTaskCount;
 	FContextSupplier*	Context;
 	FOutputDevice*		AuxOut;
 
@@ -33,7 +32,6 @@ public:
 	// Constructor.
 	FFeedbackContextAnsi()
 	: FFeedbackContext()
-	, SlowTaskCount(0)
 	, Context(nullptr)
 	, AuxOut(nullptr)
 	{}
@@ -101,20 +99,6 @@ public:
 		return true;
 	}
 
-	void BeginSlowTask( const FText& Task, bool ShowProgressDialog, bool bShowCancelButton=false )
-	{
-		GIsSlowTask = ++SlowTaskCount>0;
-	}
-	void EndSlowTask()
-	{
-		check(SlowTaskCount>0);
-		GIsSlowTask = --SlowTaskCount>0;
-	}
-	virtual bool StatusUpdate(int32 Numerator, int32 Denominator, const FText& StatusText)
-	{
-		return true;
-	}
-
 	void SetContext( FContextSupplier* InSupplier )
 	{
 		Context = InSupplier;
diff --git a/Engine/Source/Runtime/Core/Public/Misc/AutomationTest.h b/Engine/Source/Runtime/Core/Public/Misc/AutomationTest.h
--- a/Engine/Source/Runtime/Core/Public/Misc/AutomationTest.h
+++ b/Engine/Source/Runtime/Core/Public/Misc/AutomationTest.h
@@ -536,7 +536,7 @@ private:
 
 		/** Constructor */
 		FAutomationTestFeedbackContext() 
-			: SlowTaskCount( 0 ), CurTest( NULL ) {}
+			: CurTest( NULL ) {}
 
 		/** Destructor */
 		~FAutomationTestFeedbackContext()
@@ -544,25 +544,6 @@ private:
 			CurTest = NULL;
 		}
 
-		/** Signify the a slow task is beginning (parameters unused in this implementation) */
-		void BeginSlowTask( const FText& Task, bool ShowProgressDialog, bool bShowCancelButton=false )
-		{
-			GIsSlowTask = ++SlowTaskCount > 0;
-		}
-
-		/** Signify that a slow task is ending (parameters unused in this implementation) */
-		void EndSlowTask()
-		{
-			check( SlowTaskCount > 0 );
-			GIsSlowTask = --SlowTaskCount > 0;
-		}
-
-		/** Unneeded by the context, just return true */
-		virtual bool StatusUpdate(int32 Numerator, int32 Denominator, const FText& StatusText )
-		{
-			return true;
-		}
-
 		/**
 		 * FOutputDevice interface
 		 *
@@ -583,8 +564,6 @@ private:
 		}
 
 	private:
-		/** Number of slow tasks currently occurring */
-		int32 SlowTaskCount;
 
 		/** Associated automation test; all warnings, errors, etc. are routed to the automation test to track */
 		class FAutomationTestBase* CurTest;
diff --git a/Engine/Source/Runtime/Core/Public/Misc/FeedbackContext.h b/Engine/Source/Runtime/Core/Public/Misc/FeedbackContext.h
--- a/Engine/Source/Runtime/Core/Public/Misc/FeedbackContext.h
+++ b/Engine/Source/Runtime/Core/Public/Misc/FeedbackContext.h
@@ -6,57 +6,68 @@
 
 #pragma once
 
+#include "UniquePtr.h"
+
+struct FScopedSlowTask;
+
+/** A stack of feedback scopes */
+struct FScopedSlowTaskStack : TArray<FScopedSlowTask*>
+{
+	/** Get the fraction of work that has been completed for the specified index in the stack (0=total progress) */
+	CORE_API float GetProgressFraction(int32 Index) const;
+};
+
 /** A context for displaying modal warning messages. */
 class CORE_API FFeedbackContext : public FOutputDevice
 {
 public:
-	
-	virtual bool YesNof( const FText& Question ) { return false; }
 
-	virtual void BeginSlowTask( const FText& Task, bool ShowProgressDialog, bool bShowCancelButton=false )=0;
-	virtual void EndSlowTask()=0;
+	/** Ask the user a binary question, returning their answer */
+	virtual bool YesNof( const FText& Question ) { return false; }
+	
+	/** Return whether the user has requested to cancel the current slow task */
 	virtual bool ReceivedUserCancel() { return false; };
-	virtual void EnableUserCancel(bool bUserCancel){}
-
-
-	/** Slow task update status
-	 * @param Numerator		New progress numerator
-	 * @param Denominator	New progress denominator
-	 * @param StatusText	New status message
-	 * @return true by default and false if status update is not possible
-	 */
-	virtual bool StatusUpdate( int32 Numerator, int32 Denominator, const FText& StatusText ) = 0;
+	
+	/** Public const access to the current state of the scope stack */
+	FORCEINLINE const FScopedSlowTaskStack& GetScopeStack() const
+	{
+		return ScopeStack;
+	}
 
-	/** Force updating feedback status
-	 * @param Numerator		New progress numerator
-	 * @param Denominator	New progress denominator
-	 * @param StatusText	New status message
-	 * @return true by default and false if status update is not possible
-	 */
-	virtual bool StatusForceUpdate( int32 Numerator, int32 Denominator, const FText& StatusText ) { return true; } 
+	/**** Legacy API - not deprecated as it's still in heavy use, but superceded by FScopedSlowTask ****/
+	void BeginSlowTask( const FText& Task, bool ShowProgressDialog, bool bShowCancelButton=false );
+	void UpdateProgress( int32 Numerator, int32 Denominator );
+	void StatusUpdate( int32 Numerator, int32 Denominator, const FText& StatusText );
+	void StatusForceUpdate( int32 Numerator, int32 Denominator, const FText& StatusText );
+	void EndSlowTask();
+	/**** end legacy API ****/
 
+protected:
 
 	/**
-	 * Updates the progress amount without changing the status message text.  Override this in derived classes.
-	 *
-	 * @param Numerator		New progress numerator
-	 * @param Denominator	New progress denominator
+	 * Called to create a slow task
 	 */
-	virtual void UpdateProgress( int32 Numerator, int32 Denominator )
+	virtual void StartSlowTask( const FText& Task, bool bShowCancelButton=false )
 	{
+		GIsSlowTask = true;
 	}
 
-	/** Pushes the current status message/progress onto the stack so it can be restored later.  Override this
-	    in derived classes. */
-	virtual void PushStatus()
+	/**
+	 * Called to destroy a slow task
+	 */
+	virtual void FinalizeSlowTask( )
 	{
+		GIsSlowTask = false;
 	}
 
-	/** Restores the previously pushed status message/progress.  Override this in derived classes. */
-	virtual void PopStatus()
-	{
-	}
+	/**
+	 * Called when some progress has occurred
+	 * @param	TotalProgressInterp		[0..1] Value indicating the total progress of the slow task
+	 * @param	DisplayMessage			The message to display on the slow task
+	 */
+	virtual void ProgressReported( const float TotalProgressInterp, FText DisplayMessage ) {}
 
+public:
 	virtual FContextSupplier* GetContext() const { return NULL; }
 	virtual void SetContext( FContextSupplier* InSupplier ) {}
 
@@ -67,64 +78,188 @@ public:
 	TArray<FString> Warnings;
 	TArray<FString> Errors;
 
-	bool	TreatWarningsAsErrors;		
+	bool	TreatWarningsAsErrors;
 
 	FFeedbackContext() :
 		 TreatWarningsAsErrors( 0 )
 	{}
+
+private:
+	FFeedbackContext(const FFeedbackContext&);
+	FFeedbackContext& operator=(const FFeedbackContext&);
+
+protected:
+	
+	friend FScopedSlowTask;
+
+	/** Stack of pointers to feedback scopes that are currently open */
+	FScopedSlowTaskStack ScopeStack;
+	TArray<TUniquePtr<FScopedSlowTask>> LegacyAPIScopes;
+
+	/** Ask that the UI be updated as a result of the scope stack changing */
+	void RequestUpdateUI(bool bForceUpdate = false);
+
+	/** Update the UI as a result of the scope stack changing */
+	void UpdateUI();
+
 };
 
+
+
 /**
- * Helper class to display a status update message in the editor.
+ * A scope block representing an amount of work divided up into sections.
+ * Use one scope at the top of each function to give accurate feedback to the user of a slow operation's progress.
+ *
+ * Example Usage:
+ *	void DoSlowWork()
+ *	{
+ *		FScopedSlowTask Progress(2.f, LOCTEXT("DoingSlowWork", "Doing Slow Work..."));
+ *		// Optionally make this show a dialog if not already shown
+ *		Progress.MakeDialog();
+ *
+ *		// Indicate that we are entering a frame representing 1 unit of work
+ *		Progress.EnterProgressFrame(1.f);
+ *		
+ *		// DoFirstThing() can follow a similar pattern of creating a scope divided into frames. These contribute to their parent's progress frame proportionately.
+ *		DoFirstThing();
+ *		
+ *		Progress.EnterProgressFrame(1.f);
+ *		DoSecondThing();
+ *	}
+ *
  */
-class FScopedSlowTask
+struct CORE_API FScopedSlowTask
 {
+	/** Default message to display to the user when not overridden by a frame */
+	FText DefaultMessage;
+
+	/** Message pertaining to the current frame's work */
+	FText FrameMessage;
+
+	/** The amount of work to do in this scope */
+	float TotalAmountOfWork;
+
+	/** The amount of work we have already completed in this scope */
+	float CompletedWork;
+
+	/** The amount of work the current frame is responsible for */
+	float CurrentFrameScope;
+
+	/** true if this scope should be presented to the user */
+	bool bVisibleOnUI;
+
+	/** The time that this scope was created */
+	double StartTime;
+
+private:
+
+	/** Boolean flag to control whether this scope actually does anything (unset for quiet operations etc) */
+	bool bEnabled;
+
+	/** Flag that specifies whether this feedback scope created a new slow task dialog */
+	bool bCreatedDialog;
+	
+	/** The feedback context that we belong to */
+	FFeedbackContext& Context;
+
+	/** Prevent copying */
+	FScopedSlowTask(const FScopedSlowTask&);
+
 public:
 
 	/**
-	 * Updates the status message displayed to the user.
+	 * Construct this scope from an amount of work to do, and a message to display
+	 * @param		InAmountOfWork			Arbitrary number of work units to perform (can be a percentage or number of steps).
+	 *										0 indicates that no progress frames are to be entered in this scope (automatically enters a frame encompassing the entire scope)
+	 * @param		InDefaultMessage		A message to display to the user to describe the purpose of the scope
+	 * @param		bInVisible				When false, this scope will have no effect. Allows for proper scoped objects that are conditionally hidden.
 	 */
-	explicit FScopedSlowTask( const FText& InMessage, bool bAllowNewSlowTask = false )
+	FORCEINLINE FScopedSlowTask(float InAmountOfWork, const FText& InDefaultMessage = FText(), bool bInEnabled = true, FFeedbackContext& InContext = *GWarn)
+		: DefaultMessage(InDefaultMessage)
+		, FrameMessage()
+		, TotalAmountOfWork(InAmountOfWork)
+		, CompletedWork(0)
+		, CurrentFrameScope(0)
+		, bVisibleOnUI(true)
+		, StartTime(FPlatformTime::Seconds())
+		, bEnabled(bInEnabled)
+		, bCreatedDialog(false)		// only set to true if we create a dialog
+		, Context(InContext)
 	{
-		bIsBeginSlowTask = false;
+		// If we have no work to do ourselves, create an arbitrary scope so that any actions performed underneath this still contribute to this one.
+		if (TotalAmountOfWork == 0.f)
+		{
+			TotalAmountOfWork = CurrentFrameScope = 1.f;
+		}
 
-		if ( GIsEditor && !IsRunningCommandlet() && IsInGameThread() )
+		if (bEnabled)
 		{
-			if ( GIsSlowTask )
-			{
-				GWarn->PushStatus();
-				GWarn->StatusUpdate(-1, -1, InMessage);
-			}
-			else if ( bAllowNewSlowTask )
-			{
-				bIsBeginSlowTask = true;
-				const bool bShowProgressDialog = true;
-				const bool bShowCancelButton = false;
-				GWarn->BeginSlowTask( InMessage, bShowProgressDialog, bShowCancelButton );
-			}
+			Context.ScopeStack.Push(this);
 		}
 	}
 
-	/**
-	 * Ensures the status context is popped off the stack.
-	 */
-	~FScopedSlowTask()
+	/** Destructor that finishes any remaining work and removes itself from the global scope stack */
+	FORCEINLINE ~FScopedSlowTask()
 	{
-		if ( GIsEditor && !IsRunningCommandlet() && IsInGameThread() )
+		if (bEnabled)
 		{
-			if ( bIsBeginSlowTask )
+			if (bCreatedDialog)
 			{
-				GWarn->EndSlowTask();
+				checkSlow(GIsSlowTask);
+				Context.FinalizeSlowTask();
 			}
-			else
+
+			FScopedSlowTaskStack& Stack = Context.ScopeStack;
+			checkSlow(Stack.Num() != 0 && Stack.Last() == this);
+
+			Stack.Pop(false);
+			if (Stack.Num() != 0)
 			{
-				GWarn->PopStatus();
+				// Stop anything else contributing to the parent frame
+				auto* Parent = Stack.Last();
+				Parent->EnterProgressFrame(0, Parent->FrameMessage);
 			}
+
 		}
 	}
 
-private:
-	bool bIsBeginSlowTask;
-};
+	/**
+	 * Creates a new dialog for this slow task, if there is currently not one open
+	 * @param		bShowCancelButton		Whether to show a cancel button on the dialog or not
+	 */
+	void MakeDialog(bool bShowCancelButton = false);
+
+	/**
+	 * Indicate that we are to enter a frame that will take up the specified amount of work. Completes any previous frames (potentially contributing to parent scopes' progress).
+	 * @param		ExpectedWorkThisFrame	The amount of work that will happen between now and the next frame, as a numerator of TotalAmountOfWork.
+	 * @param		Text					Optional text to describe this frame's purpose.
+	 */
+	FORCEINLINE void EnterProgressFrame(float ExpectedWorkThisFrame = 1.f, FText Text = FText())
+	{
+		FrameMessage = Text;
+		CompletedWork += CurrentFrameScope;
 
+		const float WorkRemaining = TotalAmountOfWork - CompletedWork;
+		verifySlow(ExpectedWorkThisFrame <= WorkRemaining);
+		CurrentFrameScope = FMath::Min(WorkRemaining, ExpectedWorkThisFrame);
 
+		Context.RequestUpdateUI(bCreatedDialog || Context.ScopeStack[0] == this);
+	}
+
+	/**
+	 * Get the frame message or default message if empty
+	 */
+	FText GetCurrentMessage() const
+	{
+		return FrameMessage.IsEmpty() ? DefaultMessage : FrameMessage;
+	}
+
+	/**
+	 * Return whether or not this scope should be presented to the user on slow task dialogs
+	 */
+	bool IsPresentableToUser() const
+	{
+		// For now we only show scopes that attempted to open a slow task dialog
+		return bVisibleOnUI;
+	}
+};
diff --git a/Engine/Source/Runtime/CoreUObject/Private/UObject/Class.cpp b/Engine/Source/Runtime/CoreUObject/Private/UObject/Class.cpp
--- a/Engine/Source/Runtime/CoreUObject/Private/UObject/Class.cpp
+++ b/Engine/Source/Runtime/CoreUObject/Private/UObject/Class.cpp
@@ -2568,17 +2568,6 @@ public:
 			GLog->Serialize( V, Verbosity, Category );
 	}
 
-	void BeginSlowTask( const FText& Task, bool ShowProgressDialog, bool bShowCancelButton=false )
-	{
-	}
-	void EndSlowTask()
-	{
-	}
-	virtual bool StatusUpdate(int32 Numerator, int32 Denominator, const FText& StatusText )
-	{
-		return true;
-	}
-
 	FContextSupplier* GetContext() const
 	{
 		return Context;
diff --git a/Engine/Source/Runtime/CoreUObject/Private/UObject/LinkerLoad.cpp b/Engine/Source/Runtime/CoreUObject/Private/UObject/LinkerLoad.cpp
--- a/Engine/Source/Runtime/CoreUObject/Private/UObject/LinkerLoad.cpp
+++ b/Engine/Source/Runtime/CoreUObject/Private/UObject/LinkerLoad.cpp
@@ -592,10 +592,21 @@ ULinkerLoad::ULinkerLoad( const class FPostConstructInitializeProperties& PCIP,
 :	ULinker( PCIP, InParent, InFilename )
 ,	LoadFlags( InLoadFlags )
 ,	bHaveImportsBeenVerified( false )
+#if WITH_EDITOR
+,	LoadProgressScope( nullptr )
+#endif
 {
 	check(!HasAnyFlags(RF_ClassDefaultObject));
 }
 
+ULinkerLoad::~ULinkerLoad()
+{
+#if WITH_EDITOR
+	// Make sure this is deleted if it's still allocated
+	delete LoadProgressScope;
+#endif
+}
+
 /**
  * Returns whether the time limit allotted has been exceeded, if enabled.
  *
@@ -627,82 +638,31 @@ bool ULinkerLoad::IsTimeLimitExceeded( const TCHAR* CurrentTask, int32 Granulari
 	return bTimeLimitExceeded;
 }
 
-void UpdateObjectLoadingStatusMessage()
+/**
+ * Creates loader used to serialize content.
+ */
+ULinkerLoad::ELinkerStatus ULinkerLoad::CreateLoader()
 {
 #if WITH_EDITOR
-	// Used to control animation of the load progress status updates.
-	static int32 ProgressIterator = 3;
-
-	// Time that progress was last updated
-	static  double LastProgressUpdateTime = 0.0;
-
-	const double UpdateDelta = 0.25;
-	const double SlowLoadDelta = 2.0;
-	FText StatusUpdate;
 
-	// This can be a long operation so we will output some progress feedback to the 
-	//  user in the form of 3 dots that animate between "." ".." "..."
-	const double CurTime = FPlatformTime::Seconds();
-	if ( CurTime - LastProgressUpdateTime > UpdateDelta )
+	if (!LoadProgressScope)
 	{
-		if ( ProgressIterator == 1 )
-		{
-			StatusUpdate = NSLOCTEXT("Core", "LoadingRefObjectsMessageState1", "Loading.");
-		}
-		else if ( ProgressIterator == 2 )
-		{
-			StatusUpdate = NSLOCTEXT("Core", "LoadingRefObjectsMessageState2", "Loading..");
-		}
-		else if ( ProgressIterator == 3 )
-		{
-			StatusUpdate = NSLOCTEXT("Core", "LoadingRefObjectsMessageState3", "Loading...");
-		}
-		else
-		{
-			StatusUpdate = NSLOCTEXT("Core", "LoadingRefObjectsMessageState0", "Loading");
-		}
-
-		LastProgressUpdateTime = CurTime;
-		ProgressIterator = (ProgressIterator + 1) % 4;
-	}		
+		const bool bShowProgress = ( LoadFlags & ( LOAD_Quiet | LOAD_SeekFree ) ) == 0;
+		LoadProgressScope = new FScopedSlowTask(ULinkerDefs::TotalProgressSteps, NSLOCTEXT("Core", "GenericLoading", "Loading..."), bShowProgress);
+	}
 
-	GWarn->StatusUpdate( -1, -1, StatusUpdate );
 #endif
-}
 
-/**
- * Creates loader used to serialize content.
- */
-ULinkerLoad::ELinkerStatus ULinkerLoad::CreateLoader()
-{
 	CreateActiveRedirectsMap( GEngineIni );
-
 	if( !Loader )
 	{
 		bool bIsSeekFree = LoadFlags & LOAD_SeekFree;
 
 #if WITH_EDITOR
-		if ((LoadFlags & ( LOAD_Quiet | LOAD_SeekFree ) ) == 0)
-		{
-			FString CleanFilename = FPaths::GetCleanFilename( *Filename );
-
-			// We currently only allow status updates during the editor load splash screen.
-			const bool bAllowStatusUpdate = GIsEditor && !IsRunningCommandlet() && !GIsSlowTask;
-			if( bAllowStatusUpdate )
-			{
-				UpdateObjectLoadingStatusMessage();
-			}
-			else
-			{
-				if ( GIsSlowTask )
-				{
-					FFormatNamedArguments Args;
-					Args.Add( TEXT("CleanFilename"), FText::FromString( CleanFilename ) );
-					GWarn->StatusUpdate( 0, ULinkerDefs::TotalProgressSteps, FText::Format( NSLOCTEXT("Core", "Loading", "Loading file: {CleanFilename}..."), Args ) );
-				}
-			}
-			GWarn->PushStatus();
-		}
+		FFormatNamedArguments FeedbackArgs;
+		FeedbackArgs.Add( TEXT("CleanFilename"), FText::FromString( FPaths::GetCleanFilename( *Filename ) ) );
+		LoadProgressScope->DefaultMessage = FText::Format( NSLOCTEXT("Core", "LoadingFileWithFilename", "Loading file: {CleanFilename}..."), FeedbackArgs );
+		LoadProgressScope->EnterProgressFrame();
 #endif
 
 		// NOTE: Precached memory read gets highest priority, then memory reader, then seek free, then normal
@@ -801,11 +761,6 @@ ULinkerLoad::ELinkerStatus ULinkerLoad::CreateLoader()
 
 		// Reset all custom versions
 		ResetCustomVersions();
-
-		if ((LoadFlags & ( LOAD_Quiet | LOAD_SeekFree ) ) == 0)
-		{
-			GWarn->UpdateProgress( 1, ULinkerDefs::TotalProgressSteps );
-		}
 	}
 
 	bool bExecuteNextStep = true;
@@ -833,6 +788,9 @@ ULinkerLoad::ELinkerStatus ULinkerLoad::SerializePackageFileSummary()
 {
 	if( bHasSerializedPackageFileSummary == false )
 	{
+#if WITH_EDITOR
+		LoadProgressScope->EnterProgressFrame(1);
+#endif
 		// Read summary from file.
 		*this << Summary;
 
@@ -1025,11 +983,6 @@ ULinkerLoad::ELinkerStatus ULinkerLoad::SerializePackageFileSummary()
 
 		// Avoid serializing it again.
 		bHasSerializedPackageFileSummary = true;
-
-		if ((LoadFlags & ( LOAD_Quiet | LOAD_SeekFree ) ) == 0)
-		{
-			GWarn->UpdateProgress( 2, ULinkerDefs::TotalProgressSteps );
-		}
 	}
 
 	return !IsTimeLimitExceeded( TEXT("serializing package file summary") ) ? LINKER_Loaded : LINKER_TimedOut;
@@ -1113,6 +1066,10 @@ ULinkerLoad::ELinkerStatus ULinkerLoad::FixupImportMap()
 {
 	if( bHasFixedUpImportMap == false )
 	{
+#if WITH_EDITOR
+		LoadProgressScope->EnterProgressFrame(1);
+#endif
+
 		// Fix up imports, not required if everything is cooked.
 		if (!FPlatformProperties::RequiresCookedData())
 		{
@@ -1295,11 +1252,6 @@ ULinkerLoad::ELinkerStatus ULinkerLoad::FixupImportMap()
 		}
 		// Avoid duplicate work in async case.
 		bHasFixedUpImportMap = true;
-
-		if ((LoadFlags & ( LOAD_Quiet | LOAD_SeekFree ) ) == 0)
-		{
-			GWarn->UpdateProgress( 3, ULinkerDefs::TotalProgressSteps );
-		}
 	}
 	return IsTimeLimitExceeded( TEXT("fixing up import map") ) ? LINKER_TimedOut : LINKER_Loaded;
 }
@@ -1595,6 +1547,7 @@ ULinkerLoad::ELinkerStatus ULinkerLoad::FindExistingExports()
 	{
 		// only look for existing exports in the editor after it has started up
 #if WITH_EDITOR
+		LoadProgressScope->EnterProgressFrame(1);
 		if( GIsEditor && GIsRunning )
 		{
 			// Hunt down any existing objects and hook them up to this linker unless the user is either currently opening this
@@ -1615,11 +1568,6 @@ ULinkerLoad::ELinkerStatus ULinkerLoad::FindExistingExports()
 
 		// Avoid duplicate work in the case of async linker creation.
 		bHasFoundExistingExports = true;
-
-		if ((LoadFlags & ( LOAD_Quiet | LOAD_SeekFree ) ) == 0)
-		{
-			GWarn->UpdateProgress( 4, ULinkerDefs::TotalProgressSteps );
-		}
 	}
 	return IsTimeLimitExceeded( TEXT("finding existing exports") ) ? LINKER_TimedOut : LINKER_Loaded;
 }
@@ -1631,6 +1579,10 @@ ULinkerLoad::ELinkerStatus ULinkerLoad::FinalizeCreation()
 {
 	if( bHasFinishedInitialization == false )
 	{
+#if WITH_EDITOR
+		LoadProgressScope->EnterProgressFrame(1);
+#endif
+
 		// Add this linker to the object manager's linker array.
 		GObjLoaders.Add(LinkerRoot, this);
 		// And remove it from the pending loaders list.
@@ -1653,12 +1605,12 @@ ULinkerLoad::ELinkerStatus ULinkerLoad::FinalizeCreation()
 		// Avoid duplicate work in the case of async linker creation.
 		bHasFinishedInitialization = true;
 
-		if ((LoadFlags & ( LOAD_Quiet | LOAD_SeekFree ) ) == 0)
-		{
-			GWarn->UpdateProgress( 5, ULinkerDefs::TotalProgressSteps );
-			GWarn->PopStatus();
-		}
+#if WITH_EDITOR
+		delete LoadProgressScope;
+		LoadProgressScope = nullptr;
+#endif
 	}
+
 	return IsTimeLimitExceeded( TEXT("finalizing creation") ) ? LINKER_TimedOut : LINKER_Loaded;
 }
 
@@ -1738,10 +1690,15 @@ void ULinkerLoad::Verify()
 	{
 		if (!bHaveImportsBeenVerified)
 		{
+			FScopedSlowTask SlowTask(Summary.ImportCount, NSLOCTEXT("Core", "LinkerLoad_Imports", "Loading Imports"));
+
 			// Validate all imports and map them to their remote linkers.
 			for (int32 ImportIndex = 0; ImportIndex < Summary.ImportCount; ImportIndex++)
 			{
-				VerifyImport(ImportIndex);
+				FObjectImport& Import = ImportMap[ImportIndex];
+				
+				SlowTask.EnterProgressFrame(1, FText::Format(NSLOCTEXT("Core", "LinkerLoad_LoadingImportName", "Loading Import '{0}'"), FText::FromString(Import.ObjectName.ToString())));
+				VerifyImport( ImportIndex );
 			}
 		}
 	}
@@ -2112,6 +2069,8 @@ bool ULinkerLoad::VerifyImportInner(const int32 ImportIndex, FString& WarningSuf
 
 	FObjectImport& Import = ImportMap[ImportIndex];
 
+	FScopedSlowTask SlowTask(100, FText::Format(NSLOCTEXT("Core", "VerifyPackage_Scope", "Verifying '{0}'"), FText::FromString(Import.ObjectName.ToString())));
+
 	if
 	(	(Import.SourceLinker && Import.SourceIndex != INDEX_NONE)
 	||	Import.ClassPackage	== NAME_None
@@ -2147,6 +2106,9 @@ bool ULinkerLoad::VerifyImportInner(const int32 ImportIndex, FString& WarningSuf
 			TmpPkg = FindObjectFast<UPackage>(NULL, Import.ObjectName);
 			bWasFullyLoaded = TmpPkg && TmpPkg->IsFullyLoaded();
 		}
+
+		SlowTask.EnterProgressFrame(30);
+
 		if (!bWasFullyLoaded)
 		{
 			// we now fully load the package that we need a single export from - however, we still use CreatePackage below as it handles all cases when the package
@@ -2154,6 +2116,8 @@ bool ULinkerLoad::VerifyImportInner(const int32 ImportIndex, FString& WarningSuf
 			TmpPkg = LoadPackage(NULL, *Import.ObjectName.ToString(), InternalLoadFlags);
 		}
 
+		SlowTask.EnterProgressFrame(30);
+
 		// following is the original VerifyImport code
 		// @todo linkers: This could quite possibly be cleaned up
 		if (TmpPkg == NULL)
@@ -2174,6 +2138,8 @@ bool ULinkerLoad::VerifyImportInner(const int32 ImportIndex, FString& WarningSuf
 			InternalLoadFlags |= LOAD_NoVerify;
 		}
 
+		SlowTask.EnterProgressFrame(40);
+
 		// Get the linker if the package hasn't been fully loaded already.
 		if (!bWasFullyLoaded)
 		{
@@ -2185,6 +2151,8 @@ bool ULinkerLoad::VerifyImportInner(const int32 ImportIndex, FString& WarningSuf
 		// this resource's Outer is not a UPackage
 		checkf(Import.OuterIndex.IsImport(),TEXT("Outer for Import %s (%i) is not an import - OuterIndex:%i"), *GetImportFullName(ImportIndex), ImportIndex, Import.OuterIndex.ForDebugging());
 
+		SlowTask.EnterProgressFrame(50);
+
 		VerifyImport( Import.OuterIndex.ToImport() );
 
 		FObjectImport& OuterImport = Imp(Import.OuterIndex);
@@ -2213,6 +2181,7 @@ bool ULinkerLoad::VerifyImportInner(const int32 ImportIndex, FString& WarningSuf
 			Import.SourceLinker = OuterImport.SourceLinker;
 		}
 
+		SlowTask.EnterProgressFrame(50);
 
 		//check(Import.SourceLinker);
 		//@todo what does it mean if we don't have a SourceLinker here?
@@ -2494,12 +2463,14 @@ int32 ULinkerLoad::LoadMetaDataFromExportMap(bool bForcePreload/* = false */)
  */
 void ULinkerLoad::LoadAllObjects( bool bForcePreload )
 {
+	FScopedSlowTask SlowTask(ExportMap.Num(), NSLOCTEXT("Core", "LinkerLoad_LoadingObjects", "Loading Objects"), (LoadFlags & ( LOAD_Quiet | LOAD_SeekFree ) ) == 0);
+	SlowTask.bVisibleOnUI = false;
+	
 	if ( (LoadFlags&LOAD_SeekFree) != 0 )
 	{
 		bForcePreload = true;
 	}
 
-	bool bAllowedToShowStatusUpdate = (LoadFlags & ( LOAD_Quiet | LOAD_SeekFree ) ) == 0;
 	double StartTime = FPlatformTime::Seconds();
 
 	// MetaData object index in this package.
@@ -2509,9 +2480,10 @@ void ULinkerLoad::LoadAllObjects( bool bForcePreload )
 	{
 		MetaDataIndex = LoadMetaDataFromExportMap(bForcePreload);
 	}
-
+	
 	for(int32 ExportIndex = 0; ExportIndex < ExportMap.Num(); ++ExportIndex)
 	{
+		SlowTask.EnterProgressFrame(1);
 		if(ExportIndex == MetaDataIndex)
 		{
 			continue;
diff --git a/Engine/Source/Runtime/CoreUObject/Private/UObject/SavePackage.cpp b/Engine/Source/Runtime/CoreUObject/Private/UObject/SavePackage.cpp
--- a/Engine/Source/Runtime/CoreUObject/Private/UObject/SavePackage.cpp
+++ b/Engine/Source/Runtime/CoreUObject/Private/UObject/SavePackage.cpp
@@ -2490,10 +2490,6 @@ bool UPackage::SavePackage( UPackage* InOuter, UObject* Base, EObjectFlags TopLe
 		TempFilename = FPaths::CreateTempFilename(*FPaths::GameSavedDir(), *BaseFilename);
 
 		// Init.
-		const int32 SimpleSaveSteps = 29;
-		const int32 ExportSaveSteps = 101;
-		const int32 TotalSaveSteps = SimpleSaveSteps + ExportSaveSteps;
-		int32 CurSaveStep = 0;
 		FString CleanFilename = FPaths::GetCleanFilename( Filename );
 
 		FFormatNamedArguments Args;
@@ -2501,16 +2497,17 @@ bool UPackage::SavePackage( UPackage* InOuter, UObject* Base, EObjectFlags TopLe
 
 		FText StatusMessage = FText::Format( NSLOCTEXT("Core", "SavingFile", "Saving file: {CleanFilename}..."), Args );
 
-		// Only allow the user to cancel only if we're autosaving
-		GWarn->EnableUserCancel((SaveFlags & SAVE_FromAutosave ? true : false));
+		const int32 TotalSaveSteps = 31;
+		FScopedSlowTask SlowTask(TotalSaveSteps, StatusMessage);
+		SlowTask.MakeDialog(SaveFlags & SAVE_FromAutosave ? true : false);
+
+		SlowTask.EnterProgressFrame();
 
-		//Force status update based on current file name and progress percent
-		GWarn->StatusForceUpdate( CurSaveStep, TotalSaveSteps, StatusMessage );
 		bool Success = true;
 
 		ResetLoadersForSave(InOuter,Filename);
 
-		GWarn->UpdateProgress( ++CurSaveStep, TotalSaveSteps );
+		SlowTask.EnterProgressFrame();
 
 	
 		// Untag all objects and names.
@@ -2589,7 +2586,7 @@ bool UPackage::SavePackage( UPackage* InOuter, UObject* Base, EObjectFlags TopLe
 				UE_LOG_COOK_TIME(TEXT("BeginCacheForCookedPlatformData"));
 
 
-				GWarn->UpdateProgress( ++CurSaveStep, TotalSaveSteps );
+				SlowTask.EnterProgressFrame();
 
 
 				// structure to track what ever export needs to import
@@ -2627,7 +2624,7 @@ bool UPackage::SavePackage( UPackage* InOuter, UObject* Base, EObjectFlags TopLe
 				Linker->SetCookingTarget(TargetPlatform);
 
 				if ( EndSavingIfCancelled( Linker, TempFilename ) ) { return false; }
-				GWarn->UpdateProgress( ++CurSaveStep, TotalSaveSteps );
+				SlowTask.EnterProgressFrame();
 			
 				// keep a list of objects that would normally have gone into the dependency map, but since they are from cross-level dependencies, won't be found in the import map
 				TArray<UObject*> DependenciesToIgnore;
@@ -2725,7 +2722,7 @@ bool UPackage::SavePackage( UPackage* InOuter, UObject* Base, EObjectFlags TopLe
 				UE_LOG_COOK_TIME(TEXT("Serialize Imports"));
 				
 				if ( EndSavingIfCancelled( Linker, TempFilename ) ) { return false; }
-				GWarn->UpdateProgress( ++CurSaveStep, TotalSaveSteps );
+				SlowTask.EnterProgressFrame();
 
 				TArray<UObject*> PrivateObjects;
 				TArray<UObject*> ObjectsInOtherMaps;
@@ -2802,7 +2799,7 @@ bool UPackage::SavePackage( UPackage* InOuter, UObject* Base, EObjectFlags TopLe
 				UE_LOG_COOK_TIME(TEXT("Mark Names"));
 
 				if ( EndSavingIfCancelled( Linker, TempFilename ) ) { return false; }
-				GWarn->UpdateProgress( ++CurSaveStep, TotalSaveSteps );
+				SlowTask.EnterProgressFrame();
 
 				if ( LevelObjects.Num() > 0 && ObjectsInOtherMaps.Num() == 0 )
 				{
@@ -2960,7 +2957,7 @@ bool UPackage::SavePackage( UPackage* InOuter, UObject* Base, EObjectFlags TopLe
 				int32 OffsetAfterPackageFileSummary = Linker->Tell();
 		
 				if ( EndSavingIfCancelled( Linker, TempFilename ) ) { return false; }
-				GWarn->UpdateProgress( ++CurSaveStep, TotalSaveSteps );
+				SlowTask.EnterProgressFrame();
 
 
 				// Build NameMap.
@@ -2971,14 +2968,14 @@ bool UPackage::SavePackage( UPackage* InOuter, UObject* Base, EObjectFlags TopLe
 				UE_LOG_COOK_TIME(TEXT("Serialize Summary"));
 
 				if ( EndSavingIfCancelled( Linker, TempFilename ) ) { return false; }
-				GWarn->UpdateProgress( ++CurSaveStep, TotalSaveSteps );
+				SlowTask.EnterProgressFrame();
 
 				// Sort names.
 				FObjectNameSortHelper NameSortHelper;
 				NameSortHelper.SortNames( Linker, Conform );
 
 				if ( EndSavingIfCancelled( Linker, TempFilename ) ) { return false; }
-				GWarn->UpdateProgress( ++CurSaveStep, TotalSaveSteps );
+				SlowTask.EnterProgressFrame();
 
 				// Save names.
 				Linker->Summary.NameCount = Linker->NameMap.Num();
@@ -2992,7 +2989,7 @@ bool UPackage::SavePackage( UPackage* InOuter, UObject* Base, EObjectFlags TopLe
 				UE_LOG_COOK_TIME(TEXT("Serialize Names"));
 
 				if ( EndSavingIfCancelled( Linker, TempFilename ) ) { return false; }
-				GWarn->UpdateProgress( ++CurSaveStep, TotalSaveSteps );
+				SlowTask.EnterProgressFrame();
 
 				// Build ImportMap.
 				{
@@ -3008,7 +3005,7 @@ bool UPackage::SavePackage( UPackage* InOuter, UObject* Base, EObjectFlags TopLe
 
 
 				if ( EndSavingIfCancelled( Linker, TempFilename ) ) { return false; }
-				GWarn->UpdateProgress( ++CurSaveStep, TotalSaveSteps );
+				SlowTask.EnterProgressFrame();
 
 				// sort and conform imports
 				FObjectImportSortHelper ImportSortHelper;
@@ -3016,7 +3013,7 @@ bool UPackage::SavePackage( UPackage* InOuter, UObject* Base, EObjectFlags TopLe
 				Linker->Summary.ImportCount = Linker->ImportMap.Num();
 
 				if ( EndSavingIfCancelled( Linker, TempFilename ) ) { return false; }
-				GWarn->UpdateProgress( ++CurSaveStep, TotalSaveSteps );
+				SlowTask.EnterProgressFrame();
 
 				
 				// Build ExportMap.
@@ -3034,14 +3031,14 @@ bool UPackage::SavePackage( UPackage* InOuter, UObject* Base, EObjectFlags TopLe
 				UE_LOG_COOK_TIME(TEXT("Build Export Map"));
 
 				if ( EndSavingIfCancelled( Linker, TempFilename ) ) { return false; }
-				GWarn->UpdateProgress( ++CurSaveStep, TotalSaveSteps );
+				SlowTask.EnterProgressFrame();
 
 				// Sort exports alphabetically and conform the export table (if necessary)
 				FObjectExportSortHelper ExportSortHelper;
 				ExportSortHelper.SortExports( Linker, Conform );
 				
 				if ( EndSavingIfCancelled( Linker, TempFilename ) ) { return false; }
-				GWarn->UpdateProgress( ++CurSaveStep, TotalSaveSteps );
+				SlowTask.EnterProgressFrame();
 
 				// Sort exports for seek-free loading.
 				FObjectExportSeekFreeSorter SeekFreeSorter;
@@ -3051,7 +3048,7 @@ bool UPackage::SavePackage( UPackage* InOuter, UObject* Base, EObjectFlags TopLe
 				UE_LOG_COOK_TIME(TEXT("Sort Exports"));
 
 				if ( EndSavingIfCancelled( Linker, TempFilename ) ) { return false; }
-				GWarn->UpdateProgress( ++CurSaveStep, TotalSaveSteps );
+				SlowTask.EnterProgressFrame();
 
 				// Pre-size depends map.
 				Linker->DependsMap.AddZeroed( Linker->ExportMap.Num() );
@@ -3125,7 +3122,7 @@ bool UPackage::SavePackage( UPackage* InOuter, UObject* Base, EObjectFlags TopLe
 
 
 				if ( EndSavingIfCancelled( Linker, TempFilename ) ) { return false; }
-				GWarn->UpdateProgress( ++CurSaveStep, TotalSaveSteps );
+				SlowTask.EnterProgressFrame();
 
 				// Set linker reverse mappings.
 				// also set netplay required data for any UPackages in the export map
@@ -3149,7 +3146,7 @@ bool UPackage::SavePackage( UPackage* InOuter, UObject* Base, EObjectFlags TopLe
 				}
 
 				if ( EndSavingIfCancelled( Linker, TempFilename ) ) { return false; }
-				GWarn->UpdateProgress( ++CurSaveStep, TotalSaveSteps );
+				SlowTask.EnterProgressFrame();
 
 				// If this is a map package, make sure there is a world or level in the export map.
 				if ( InOuter->ContainsMap() )
@@ -3185,7 +3182,7 @@ bool UPackage::SavePackage( UPackage* InOuter, UObject* Base, EObjectFlags TopLe
 				UE_LOG_COOK_TIME(TEXT("Build Dependency Map"));
 
 				if ( EndSavingIfCancelled( Linker, TempFilename ) ) { return false; }
-				GWarn->UpdateProgress( ++CurSaveStep, TotalSaveSteps );
+				SlowTask.EnterProgressFrame();
 
 				for( int32 i=0; i<Linker->ImportMap.Num(); i++ )
 				{
@@ -3203,12 +3200,12 @@ bool UPackage::SavePackage( UPackage* InOuter, UObject* Base, EObjectFlags TopLe
 				}
 
 
-				GWarn->UpdateProgress( ++CurSaveStep, TotalSaveSteps );
+				SlowTask.EnterProgressFrame();
 
 				// Find components referenced by exports.
 
 				if ( EndSavingIfCancelled( Linker, TempFilename ) ) { return false; }
-				GWarn->UpdateProgress( ++CurSaveStep, TotalSaveSteps );
+				SlowTask.EnterProgressFrame();
 
 				// Save dummy import map, overwritten later.
 				Linker->Summary.ImportOffset = Linker->Tell();
@@ -3224,7 +3221,7 @@ bool UPackage::SavePackage( UPackage* InOuter, UObject* Base, EObjectFlags TopLe
 				UE_LOG_COOK_TIME(TEXT("Serialize Import Map"));
 
 				if ( EndSavingIfCancelled( Linker, TempFilename ) ) { return false; }
-				GWarn->UpdateProgress( ++CurSaveStep, TotalSaveSteps );
+				SlowTask.EnterProgressFrame();
 
 				// Save dummy export map, overwritten later.
 				Linker->Summary.ExportOffset = Linker->Tell();
@@ -3239,7 +3236,7 @@ bool UPackage::SavePackage( UPackage* InOuter, UObject* Base, EObjectFlags TopLe
 				UE_LOG_COOK_TIME(TEXT("Serialize Export Map"));
 
 				if ( EndSavingIfCancelled( Linker, TempFilename ) ) { return false; }
-				GWarn->UpdateProgress( ++CurSaveStep, TotalSaveSteps );
+				SlowTask.EnterProgressFrame();
 
 				// save depends map (no need for later patching)
 				check(Linker->DependsMap.Num()==Linker->ExportMap.Num());
@@ -3254,7 +3251,7 @@ bool UPackage::SavePackage( UPackage* InOuter, UObject* Base, EObjectFlags TopLe
 				UE_LOG_COOK_TIME(TEXT("Serialize Dependency Map"));
 
 				if (EndSavingIfCancelled(Linker, TempFilename)) { return false; }
-				GWarn->UpdateProgress(++CurSaveStep, TotalSaveSteps);
+				SlowTask.EnterProgressFrame();
 
 				// Save string asset reference map
 				Linker->Summary.StringAssetReferencesOffset = Linker->Tell();
@@ -3279,7 +3276,7 @@ bool UPackage::SavePackage( UPackage* InOuter, UObject* Base, EObjectFlags TopLe
 				Linker->Summary.TotalHeaderSize	= Linker->Tell();
 
 				if ( EndSavingIfCancelled( Linker, TempFilename ) ) { return false; }
-				GWarn->UpdateProgress( ++CurSaveStep, TotalSaveSteps );
+				SlowTask.EnterProgressFrame(1, NSLOCTEXT("Core", "ProcessingExports", "ProcessingExports..."));
 
 				// look for this package in the list of packages to generate script SHA for 
 				TArray<uint8>* ScriptSHABytes = ULinkerSave::PackagesToScriptSHAMap.Find(*FPaths::GetBaseFilename(Filename));
@@ -3291,95 +3288,90 @@ bool UPackage::SavePackage( UPackage* InOuter, UObject* Base, EObjectFlags TopLe
 				}
 
 				UE_LOG_COOK_TIME(TEXT("SaveThumbNails Save AssetRegistryData SaveWorldLevelInfo"));
-
-				// Save exports.
-				int32 LastExportSaveStep = 0;
-				for( int32 i=0; i<Linker->ExportMap.Num(); i++ )
+				
 				{
-					FObjectExport& Export = Linker->ExportMap[i];
-					if( Export.Object )
+					FScopedSlowTask ExportScope(Linker->ExportMap.Num());
+
+					// Save exports.
+					int32 LastExportSaveStep = 0;
+					for( int32 i=0; i<Linker->ExportMap.Num(); i++ )
 					{
-						// Set class index.
-						// If this is *exactly* a UClass, store null instead; for anything else, including UClass-derived classes, map it
-						if( Export.Object->GetClass() != UClass::StaticClass() )
-						{
-							Export.ClassIndex = Linker->MapObject(Export.Object->GetClass());
-							check(!Export.ClassIndex.IsNull());
-						}
-						else
-						{
-							Export.ClassIndex = FPackageIndex();
-						}
+						if ( EndSavingIfCancelled( Linker, TempFilename ) ) { return false; }
+						ExportScope.EnterProgressFrame();
 
-						// Set the parent index, if this export represents a UStruct-derived object
-						if( UStruct* Struct = dynamic_cast<UStruct*>(Export.Object) )
+						FObjectExport& Export = Linker->ExportMap[i];
+						if( Export.Object )
 						{
-							if (Struct->GetSuperStruct() != NULL)
+							// Set class index.
+							// If this is *exactly* a UClass, store null instead; for anything else, including UClass-derived classes, map it
+							if( Export.Object->GetClass() != UClass::StaticClass() )
 							{
-								Export.SuperIndex = Linker->MapObject(Struct->GetSuperStruct());
-								check(!Export.SuperIndex.IsNull());
+								Export.ClassIndex = Linker->MapObject(Export.Object->GetClass());
+								check(!Export.ClassIndex.IsNull());
 							}
 							else
 							{
-								Export.SuperIndex = FPackageIndex();
+								Export.ClassIndex = FPackageIndex();
 							}
-						}
-						else
-						{
-							Export.SuperIndex = FPackageIndex();
-						}
-
-						// Set FPackageIndex for this export's Outer. If the export's Outer
-						// is the UPackage corresponding to this package's LinkerRoot, the
-						if( Export.Object->GetOuter() != InOuter )
-						{
-							check( Export.Object->GetOuter() );
-							checkf(Export.Object->GetOuter()->IsIn(InOuter),
-								TEXT("Export Object (%s) Outer (%s) mismatch."),
-								*(Export.Object->GetPathName()),
-								*(Export.Object->GetOuter()->GetPathName()));
-							Export.OuterIndex = Linker->MapObject(Export.Object->GetOuter());
-							checkf(!Export.OuterIndex.IsImport(), 
-								TEXT("Export Object (%s) Outer (%s) is an Import."),
-								*(Export.Object->GetPathName()),
-								*(Export.Object->GetOuter()->GetPathName()));
-						}
-						else
-						{
-							// this export's Outer is the LinkerRoot for this package
-							Export.OuterIndex = FPackageIndex();
-						}
 
-						// Save the object data.
-						Export.SerialOffset = Linker->Tell();
-						// UE_LOG(LogSavePackage, Log, TEXT("export %s for %s"), *Export.Object->GetFullName(), *Linker->CookingTarget()->PlatformName());
-						if ( Export.Object->HasAnyFlags(RF_ClassDefaultObject) )
-						{
-							Export.Object->GetClass()->SerializeDefaultObject(Export.Object, *Linker);
-						}
-						else
-						{
-							Export.Object->Serialize( *Linker );
-						}
-						Export.SerialSize = Linker->Tell() - Export.SerialOffset;
+							// Set the parent index, if this export represents a UStruct-derived object
+							if( UStruct* Struct = dynamic_cast<UStruct*>(Export.Object) )
+							{
+								if (Struct->GetSuperStruct() != NULL)
+								{
+									Export.SuperIndex = Linker->MapObject(Struct->GetSuperStruct());
+									check(!Export.SuperIndex.IsNull());
+								}
+								else
+								{
+									Export.SuperIndex = FPackageIndex();
+								}
+							}
+							else
+							{
+								Export.SuperIndex = FPackageIndex();
+							}
 
-						// Mark object as having been saved.
-						Export.Object->Mark(OBJECTMARK_Saved);
-					}
+							// Set FPackageIndex for this export's Outer. If the export's Outer
+							// is the UPackage corresponding to this package's LinkerRoot, the
+							if( Export.Object->GetOuter() != InOuter )
+							{
+								check( Export.Object->GetOuter() );
+								checkf(Export.Object->GetOuter()->IsIn(InOuter),
+									TEXT("Export Object (%s) Outer (%s) mismatch."),
+									*(Export.Object->GetPathName()),
+									*(Export.Object->GetOuter()->GetPathName()));
+								Export.OuterIndex = Linker->MapObject(Export.Object->GetOuter());
+								checkf(!Export.OuterIndex.IsImport(), 
+									TEXT("Export Object (%s) Outer (%s) is an Import."),
+									*(Export.Object->GetPathName()),
+									*(Export.Object->GetOuter()->GetPathName()));
+							}
+							else
+							{
+								// this export's Outer is the LinkerRoot for this package
+								Export.OuterIndex = FPackageIndex();
+							}
 
+							// Save the object data.
+							Export.SerialOffset = Linker->Tell();
+							// UE_LOG(LogSavePackage, Log, TEXT("export %s for %s"), *Export.Object->GetFullName(), *Linker->CookingTarget()->PlatformName());
+							if ( Export.Object->HasAnyFlags(RF_ClassDefaultObject) )
+							{
+								Export.Object->GetClass()->SerializeDefaultObject(Export.Object, *Linker);
+							}
+							else
+							{
+								Export.Object->Serialize( *Linker );
+							}
+							Export.SerialSize = Linker->Tell() - Export.SerialOffset;
 
-					// Update progress for exports
-					int32 CurExportSaveStep = ( i + 1 ) * ExportSaveSteps / ( Linker->ExportMap.Num() );
-					if( CurExportSaveStep > LastExportSaveStep )
-					{
-						CurSaveStep += CurExportSaveStep - LastExportSaveStep;
-						if ( EndSavingIfCancelled( Linker, TempFilename ) ) { return false; }
-						GWarn->UpdateProgress( CurSaveStep, TotalSaveSteps );
+							// Mark object as having been saved.
+							Export.Object->Mark(OBJECTMARK_Saved);
+						}
 					}
-					LastExportSaveStep = CurExportSaveStep;
 				}
 
-
 				UE_LOG_COOK_TIME(TEXT("Serialize Exports"));
 
 				// if we want to generate the SHA key, get it out now that the package has finished saving
@@ -3395,34 +3387,41 @@ bool UPackage::SavePackage( UPackage* InOuter, UObject* Base, EObjectFlags TopLe
 
 				
 				if ( EndSavingIfCancelled( Linker, TempFilename ) ) { return false; }
-				GWarn->UpdateProgress( ++CurSaveStep, TotalSaveSteps );
+				SlowTask.EnterProgressFrame(1, NSLOCTEXT("Core", "SerializingBulkData", "Serializing bulk data"));
 
 				// now we write all the bulkdata that is supposed to be at the end of the package
 				// and fix up the offset
 				int64 StartOfBulkDataArea = Linker->Tell();
 				Linker->Summary.BulkDataStartOffset = StartOfBulkDataArea;
-			
-				for (int32 i=0; i < Linker->BulkDataToAppend.Num(); ++i)
+
 				{
-					ULinkerSave::FBulkDataStorageInfo& BulkData = Linker->BulkDataToAppend[i];
+					FScopedSlowTask BulkDataFeedback(Linker->BulkDataToAppend.Num());
+
+					for (int32 i=0; i < Linker->BulkDataToAppend.Num(); ++i)
+					{
+						BulkDataFeedback.EnterProgressFrame();
+
+						ULinkerSave::FBulkDataStorageInfo& BulkData = Linker->BulkDataToAppend[i];
 
-					int64 BulkStartOffset = Linker->Tell();
-					int64 StoredBulkStartOffset = BulkStartOffset - StartOfBulkDataArea;
+						int64 BulkStartOffset = Linker->Tell();
+						int64 StoredBulkStartOffset = BulkStartOffset - StartOfBulkDataArea;
 
-					BulkData.BulkData->SerializeBulkData(*Linker, BulkData.BulkData->Lock(LOCK_READ_ONLY));
-					BulkData.BulkData->Unlock();
+						BulkData.BulkData->SerializeBulkData(*Linker, BulkData.BulkData->Lock(LOCK_READ_ONLY));
+						BulkData.BulkData->Unlock();
 
-					int64 BulkEndOffset = Linker->Tell();
-					int32 SizeOnDisk = (int32)(BulkEndOffset - BulkStartOffset);
+						int64 BulkEndOffset = Linker->Tell();
+						int32 SizeOnDisk = (int32)(BulkEndOffset - BulkStartOffset);
 				
-					Linker->Seek(BulkData.BulkDataOffsetInFilePos);
-					*Linker << StoredBulkStartOffset;
+						Linker->Seek(BulkData.BulkDataOffsetInFilePos);
+						*Linker << StoredBulkStartOffset;
 
-					Linker->Seek(BulkData.BulkDataSizeOnDiskPos);
-					*Linker << SizeOnDisk;
+						Linker->Seek(BulkData.BulkDataSizeOnDiskPos);
+						*Linker << SizeOnDisk;
 
-					Linker->Seek(BulkEndOffset);
+						Linker->Seek(BulkEndOffset);
+					}
 				}
+
 				Linker->BulkDataToAppend.Empty();
 			
 				// write the package post tag
@@ -3489,14 +3488,11 @@ bool UPackage::SavePackage( UPackage* InOuter, UObject* Base, EObjectFlags TopLe
 				UE_LOG_COOK_TIME(TEXT("Serialize Export Map"));
 
 				if ( EndSavingIfCancelled( Linker, TempFilename ) ) { return false; }
-				GWarn->UpdateProgress( ++CurSaveStep, TotalSaveSteps );
+				SlowTask.EnterProgressFrame();
 
 				FFormatNamedArguments NamedArgs;
 				NamedArgs.Add( TEXT("CleanFilename"), FText::FromString( CleanFilename ) );
-				StatusMessage = FText::Format( NSLOCTEXT("Core", "Finalizing", "Finalizing: {CleanFilename}..."), NamedArgs );
-
-				// Rewrite updated file summary.
-				GWarn->StatusForceUpdate( CurSaveStep, TotalSaveSteps, StatusMessage );
+				SlowTask.DefaultMessage = FText::Format( NSLOCTEXT("Core", "Finalizing", "Finalizing: {CleanFilename}..."), NamedArgs );
 
 				//@todo: remove ExportCount and NameCount - no longer used
 				Linker->Summary.Generations.Last().ExportCount = Linker->Summary.ExportCount;
@@ -3520,7 +3516,7 @@ bool UPackage::SavePackage( UPackage* InOuter, UObject* Base, EObjectFlags TopLe
 				check( Linker->Tell() == OffsetAfterPackageFileSummary );
 
 				if ( EndSavingIfCancelled( Linker, TempFilename ) ) { return false; }
-				GWarn->UpdateProgress( ++CurSaveStep, TotalSaveSteps );
+				SlowTask.EnterProgressFrame();
 
 				// Detach archive used for saving, closing file handle.
 				if( Linker && !bCompressFromMemory && !bSaveAsync)
@@ -3531,7 +3527,7 @@ bool UPackage::SavePackage( UPackage* InOuter, UObject* Base, EObjectFlags TopLe
 				UE_LOG(LogSavePackage, Log,  TEXT("Save=%fms"), FPlatformTime::ToMilliseconds(Time) );
 		
 				if ( EndSavingIfCancelled( Linker, TempFilename ) ) { return false; }
-				GWarn->UpdateProgress( ++CurSaveStep, TotalSaveSteps );
+				SlowTask.EnterProgressFrame();
 
 				if( Success == true )
 				{
@@ -3663,7 +3659,7 @@ bool UPackage::SavePackage( UPackage* InOuter, UObject* Base, EObjectFlags TopLe
 				}
 
 
-				GWarn->UpdateProgress( ++CurSaveStep, TotalSaveSteps );
+				SlowTask.EnterProgressFrame();
 
 				// Route PostSaveRoot to allow e.g. the world to detach components for the persistent level that where
 				// attached in PreSaveRoot.
@@ -3672,7 +3668,7 @@ bool UPackage::SavePackage( UPackage* InOuter, UObject* Base, EObjectFlags TopLe
 					Base->PostSaveRoot( bCleanupIsRequired );
 				}
 
-				GWarn->UpdateProgress( ++CurSaveStep, TotalSaveSteps );
+				SlowTask.EnterProgressFrame();
 			
 
 				for ( int CachedObjectIndex = 0; CachedObjectIndex < CachedObjects.Num(); ++CachedObjectIndex )
@@ -3689,7 +3685,7 @@ bool UPackage::SavePackage( UPackage* InOuter, UObject* Base, EObjectFlags TopLe
 		}
 
 		// We're done!
-		GWarn->UpdateProgress( TotalSaveSteps, TotalSaveSteps );
+		SlowTask.EnterProgressFrame();
 
 		UE_FINISH_LOG_COOK_TIME();
 
diff --git a/Engine/Source/Runtime/CoreUObject/Private/UObject/UObjectGlobals.cpp b/Engine/Source/Runtime/CoreUObject/Private/UObject/UObjectGlobals.cpp
--- a/Engine/Source/Runtime/CoreUObject/Private/UObject/UObjectGlobals.cpp
+++ b/Engine/Source/Runtime/CoreUObject/Private/UObject/UObjectGlobals.cpp
@@ -798,9 +798,15 @@ UPackage* LoadPackage( UPackage* InOuter, const TCHAR* InLongPackageName, uint32
 	TGuardValue<bool> IsEditorLoadingPackage(GIsEditorLoadingPackage, GIsEditor || GIsEditorLoadingPackage);
 #endif
 
+	FScopedSlowTask SlowTask(100, FText::Format(NSLOCTEXT("Core", "LoadingPackage_Scope", "Loading Package '{0}'"), FText::FromString(FileToLoad)));
+	SlowTask.bVisibleOnUI = false;
+	
+	SlowTask.EnterProgressFrame(10);
+
 	// Try to load.
 	BeginLoad();
 
+	SlowTask.EnterProgressFrame(30);
 	{
 		// Keep track of start time.
 		const double StartTime = FPlatformTime::Seconds();
@@ -843,17 +849,19 @@ UPackage* LoadPackage( UPackage* InOuter, const TCHAR* InLongPackageName, uint32
 			Linker->StartScriptSHAGeneration();
 		}
 
+		SlowTask.EnterProgressFrame(30);
+
 		if( !(LoadFlags & LOAD_Verify) )
 		{
 			Linker->LoadAllObjects();
 		}
-#if WITH_EDITOR
-		// Add a LoadContext string to the endload function in the form of: "<FileToLoad> Package"
-		EndLoad( GIsEditor ? *FPaths::GetBaseFilename(FileToLoad) : NULL );
 
-		GIsEditorLoadingPackage = *IsEditorLoadingPackage;
-#else
+		SlowTask.EnterProgressFrame(30);
+
 		EndLoad();
+
+#if WITH_EDITOR
+		GIsEditorLoadingPackage = *IsEditorLoadingPackage;
 #endif
 
 		// Set package-requires-localization flags from archive after loading. This reinforces flagging of packages that haven't yet been resaved.
@@ -987,78 +995,10 @@ struct FCompareUObjectByLinkerAndOffset
 	}
 };
 
-void UpdateObjectLoadingStatusMessage( const TCHAR* LoadContext, const double StartTime )
-{
-#if WITH_EDITOR
-// Used to control animation of the load progress status updates.
-	static int32 ProgressIterator = 3;
-
-// Time that progress was last updated
-	static  double LastProgressUpdateTime = 0.0;
-
-	const double UpdateDelta = 0.25;
-	const double SlowLoadDelta = 2.0;
-	FText StatusUpdate;
-
-	// This can be a long operation so we will output some progress feedback to the 
-	//  user in the form of 3 dots that animate between "." ".." "..."
-	const double CurTime = FPlatformTime::Seconds();
-	if ( CurTime - LastProgressUpdateTime > UpdateDelta )
-	{
-		if( ( CurTime - StartTime ) > SlowLoadDelta )
-		{
-			FFormatNamedArguments Args;
-			Args.Add( TEXT("LoadContext"), FText::FromString( LoadContext ) );
-
-			if ( ProgressIterator == 1 )
-			{
-				StatusUpdate = FText::Format( NSLOCTEXT("Core", "LoadingRefObjectsMessageStateContext1", "Loading {LoadContext}."), Args );
-			}
-			else if ( ProgressIterator == 2 )
-			{
-				StatusUpdate = FText::Format( NSLOCTEXT("Core", "LoadingRefObjectsMessageStateContext2", "Loading {LoadContext}.."), Args );
-			}
-			else if ( ProgressIterator == 3 )
-			{
-				StatusUpdate = FText::Format( NSLOCTEXT("Core", "LoadingRefObjectsMessageStateContext3", "Loading {LoadContext}..."), Args );
-			}
-			else
-			{
-				StatusUpdate = FText::Format( NSLOCTEXT("Core", "LoadingRefObjectsMessageStateContext0", "Loading {LoadContext}"), Args );
-			}
-		}
-		else
-		{
-			if ( ProgressIterator == 1 )
-			{
-				StatusUpdate = NSLOCTEXT("Core", "LoadingRefObjectsMessageState1", "Loading.");
-			}
-			else if ( ProgressIterator == 2 )
-			{
-				StatusUpdate = NSLOCTEXT("Core", "LoadingRefObjectsMessageState2", "Loading..");
-			}
-			else if ( ProgressIterator == 3 )
-			{
-				StatusUpdate = NSLOCTEXT("Core", "LoadingRefObjectsMessageState3", "Loading...");
-			}
-			else
-			{
-				StatusUpdate = NSLOCTEXT("Core", "LoadingRefObjectsMessageState0", "Loading");
-			}
-		}
-
-		LastProgressUpdateTime = CurTime;
-		ProgressIterator = (ProgressIterator + 1) % 4;
-	}		
-
-	GWarn->StatusUpdate( -1, -1, StatusUpdate );
-#endif
-}
-
 //
 // End loading packages.
 //
-void EndLoad(const TCHAR* LoadContext)
+void EndLoad()
 {
 	check(GObjBeginLoadCount>0);
 	if (GIsAsyncLoading)
@@ -1067,138 +1007,137 @@ void EndLoad(const TCHAR* LoadContext)
 		return;
 	}
 
+#if WITH_EDITOR
+	const bool bReportProgress = GIsEditor && !IsRunningCommandlet();
+	
+	FScopedSlowTask SlowTask(0, NSLOCTEXT("Core", "PerformingPostLoad", "Performing post-load..."), bReportProgress);
+
+	int32 NumObjectsLoaded = 0, NumObjectsFound = 0;
+#endif
+
 	while( --GObjBeginLoadCount == 0 && (GObjLoaded.Num() || GImportCount || GForcedExportCount) )
 	{
 		// Make sure we're not recursively calling EndLoad as e.g. loading a config file could cause
 		// BeginLoad/EndLoad to be called.
 		GObjBeginLoadCount++;
 
-			// Temporary list of loaded objects as GObjLoaded might expand during iteration.
-			TArray<UObject*> ObjLoaded;
-			TSet<ULinkerLoad*> LoadedLinkers;
+		// Temporary list of loaded objects as GObjLoaded might expand during iteration.
+		TArray<UObject*> ObjLoaded;
+		TSet<ULinkerLoad*> LoadedLinkers;
+		while( GObjLoaded.Num() )
+		{
+			// Accumulate till GObjLoaded no longer increases.
+			ObjLoaded += GObjLoaded;
+			GObjLoaded.Empty();
 
-#if WITH_EDITOR
-			// Stores the progress symbols that we will animate through during long operations
-			const double StartTime = FPlatformTime::Seconds();
-			const bool bIsLoadContextValid = LoadContext != NULL && *LoadContext != '\0';
-			// We currently only allow status updates during the editor load splash screen.
-			const bool bAllowStatusUpdate = GIsEditor && !IsRunningCommandlet() && !GIsSlowTask && bIsLoadContextValid;
-#endif
-			while( GObjLoaded.Num() )
+			// Sort by Filename and Offset.
+			ObjLoaded.Sort( FCompareUObjectByLinkerAndOffset() );
+
+			// Finish loading everything.
+			for( int32 i=0; i<ObjLoaded.Num(); i++ )
+			{
+				// Preload.
+				UObject* Obj = ObjLoaded[i];
+				if( Obj->HasAnyFlags(RF_NeedLoad) )
+				{
+					check(Obj->GetLinker());
+					Obj->GetLinker()->Preload( Obj );
+				}
+			}
+
+			// Start over again as new objects have been loaded that need to have "Preload" called on them before
+			// we can safely PostLoad them.
+			if(GObjLoaded.Num())
 			{
-				// Accumulate till GObjLoaded no longer increases.
-				ObjLoaded += GObjLoaded;
-				GObjLoaded.Empty();
+				continue;
+			}
 
-				// Sort by Filename and Offset.
-				ObjLoaded.Sort( FCompareUObjectByLinkerAndOffset() );
+#if WITH_EDITOR
+			SlowTask.CompletedWork = SlowTask.TotalAmountOfWork;
+			SlowTask.TotalAmountOfWork += ObjLoaded.Num();
+			SlowTask.CurrentFrameScope = 0;
+#endif
 
-				// Finish loading everything.
+			if ( GIsEditor )
+			{
 				for( int32 i=0; i<ObjLoaded.Num(); i++ )
 				{
-#if WITH_EDITOR
-					if ( bAllowStatusUpdate )
-					{
-						UpdateObjectLoadingStatusMessage( LoadContext, StartTime );
-					}		
-#endif				
-					// Preload.
 					UObject* Obj = ObjLoaded[i];
-					if( Obj->HasAnyFlags(RF_NeedLoad) )
+					if ( Obj->GetLinker() )
 					{
-						check(Obj->GetLinker());
-						Obj->GetLinker()->Preload( Obj );
+						LoadedLinkers.Add(Obj->GetLinker());
 					}
 				}
+			}
 
-				// Start over again as new objects have been loaded that need to have "Preload" called on them before
-				// we can safely PostLoad them.
-				if(GObjLoaded.Num())
-				{
-					continue;
-				}
-
-				if ( GIsEditor )
-				{
-					for( int32 i=0; i<ObjLoaded.Num(); i++ )
-					{
-						UObject* Obj = ObjLoaded[i];
-						if ( Obj->GetLinker() )
-						{
-							LoadedLinkers.Add(Obj->GetLinker());
-						}
-					}
-				}
+			{
+				// set this so that we can perform certain operations in which are only safe once all objects have been de-serialized.
+				TGuardValue<bool> GuardIsRoutingPostLoad(GIsRoutingPostLoad, true);
 
+				// Postload objects.
+				for(int32 i = 0; i < ObjLoaded.Num(); i++)
 				{
-					// set this so that we can perform certain operations in which are only safe once all objects have been de-serialized.
-					TGuardValue<bool> GuardIsRoutingPostLoad(GIsRoutingPostLoad, true);
-
-					// Postload objects.
-					for(int32 i = 0; i < ObjLoaded.Num(); i++)
-					{
+				
+					UObject* Obj = ObjLoaded[i];
+					check(Obj);
 #if WITH_EDITOR
-						if(bAllowStatusUpdate)
-						{
-							UpdateObjectLoadingStatusMessage(LoadContext, StartTime);
-						}
+					SlowTask.EnterProgressFrame(1, FText::Format(NSLOCTEXT("Core", "FinalizingUObject", "Finalizing load of {0}"), FText::FromString(Obj->GetName())));
 #endif
-						UObject* Obj = ObjLoaded[i];
-						check(Obj);
-						Obj->ConditionalPostLoad();
-					}
+					Obj->ConditionalPostLoad();
 				}
+			}
+
 #if WITH_EDITOR
-				// Send global notification for each object that was loaded.
-				// Useful for updating UI such as ContentBrowser's loaded status.
+			// Send global notification for each object that was loaded.
+			// Useful for updating UI such as ContentBrowser's loaded status.
+			{
+				for( int32 CurObjIndex=0; CurObjIndex<ObjLoaded.Num(); CurObjIndex++ )
 				{
-					for( int32 CurObjIndex=0; CurObjIndex<ObjLoaded.Num(); CurObjIndex++ )
+					UObject* Obj = ObjLoaded[CurObjIndex];
+					check(Obj);
+					if ( Obj->IsAsset() )
 					{
-						UObject* Obj = ObjLoaded[CurObjIndex];
-						check(Obj);
-						if ( Obj->IsAsset() )
-						{
-							FCoreUObjectDelegates::OnAssetLoaded.Broadcast(Obj);
-						}
+						FCoreUObjectDelegates::OnAssetLoaded.Broadcast(Obj);
 					}
 				}
+			}
 #endif	// WITH_EDITOR
 
-				// Empty array before next iteration as we finished postloading all objects.
-				ObjLoaded.Empty( GObjLoaded.Num() );
-			}
+			// Empty array before next iteration as we finished postloading all objects.
+			ObjLoaded.Empty( GObjLoaded.Num() );
+		}
 
-			if ( GIsEditor && LoadedLinkers.Num() > 0 )
+		if ( GIsEditor && LoadedLinkers.Num() > 0 )
+		{
+			for ( TSet<ULinkerLoad*>::TIterator It(LoadedLinkers); It; ++It )
 			{
-				for ( TSet<ULinkerLoad*>::TIterator It(LoadedLinkers); It; ++It )
-				{
-					ULinkerLoad* LoadedLinker = *It;
-					check(LoadedLinker);
+				ULinkerLoad* LoadedLinker = *It;
+				check(LoadedLinker);
 
-					if ( LoadedLinker->LinkerRoot != NULL && !LoadedLinker->LinkerRoot->IsFullyLoaded() )
+				if ( LoadedLinker->LinkerRoot != NULL && !LoadedLinker->LinkerRoot->IsFullyLoaded() )
+				{
+					bool bAllExportsCreated = true;
+					for ( int32 ExportIndex = 0; ExportIndex < LoadedLinker->ExportMap.Num(); ExportIndex++ )
 					{
-						bool bAllExportsCreated = true;
-						for ( int32 ExportIndex = 0; ExportIndex < LoadedLinker->ExportMap.Num(); ExportIndex++ )
+						FObjectExport& Export = LoadedLinker->ExportMap[ExportIndex];
+						if ( !Export.bForcedExport && Export.Object == NULL )
 						{
-							FObjectExport& Export = LoadedLinker->ExportMap[ExportIndex];
-							if ( !Export.bForcedExport && Export.Object == NULL )
-							{
-								bAllExportsCreated = false;
-								break;
-							}
+							bAllExportsCreated = false;
+							break;
 						}
+					}
 
-						if ( bAllExportsCreated )
-						{
-							LoadedLinker->LinkerRoot->MarkAsFullyLoaded();
-						}
+					if ( bAllExportsCreated )
+					{
+						LoadedLinker->LinkerRoot->MarkAsFullyLoaded();
 					}
 				}
 			}
+		}
 
-			// Dissociate all linker import and forced export object references, since they
-			// may be destroyed, causing their pointers to become invalid.
-			DissociateImportsAndForcedExports();
+		// Dissociate all linker import and forced export object references, since they
+		// may be destroyed, causing their pointers to become invalid.
+		DissociateImportsAndForcedExports();
 
 		// close any linkers' loaders that were requested to be closed once GObjBeginLoadCount goes to 0
 		for (int32 Index = 0; Index < GDelayedLinkerClosePackages.Num(); Index++)
diff --git a/Engine/Source/Runtime/CoreUObject/Public/UObject/Linker.h b/Engine/Source/Runtime/CoreUObject/Public/UObject/Linker.h
--- a/Engine/Source/Runtime/CoreUObject/Public/UObject/Linker.h
+++ b/Engine/Source/Runtime/CoreUObject/Public/UObject/Linker.h
@@ -1375,7 +1375,15 @@ private:
 
 	static FName NAME_LoadErrors;
 
+#if WITH_EDITOR
+	/** Feedback scope that is created to house the slow task of an asynchronous linker load. Raw ptr so we don't pull in TUniquePtr for everything. */
+	FScopedSlowTask* LoadProgressScope;
+#endif 
+
 public:
+	
+	/** Destructor */
+	~ULinkerLoad();
 
 	/**
 	 * Initialize the static variables
diff --git a/Engine/Source/Runtime/CoreUObject/Public/UObject/UObjectGlobals.h b/Engine/Source/Runtime/CoreUObject/Public/UObject/UObjectGlobals.h
--- a/Engine/Source/Runtime/CoreUObject/Public/UObject/UObjectGlobals.h
+++ b/Engine/Source/Runtime/CoreUObject/Public/UObject/UObjectGlobals.h
@@ -400,7 +400,7 @@ namespace EAsyncPackageState
  */
 COREUOBJECT_API EAsyncPackageState::Type ProcessAsyncLoading( bool bUseTimeLimit, bool bUseFullTimeLimit, float TimeLimit, FName ExcludeType = NAME_None );
 COREUOBJECT_API void BeginLoad();
-COREUOBJECT_API void EndLoad(const TCHAR* LoadContext = NULL);
+COREUOBJECT_API void EndLoad();
 
 /**
  * Find an existing package by name
diff --git a/Engine/Source/Runtime/Engine/Private/AudioDerivedData.cpp b/Engine/Source/Runtime/Engine/Private/AudioDerivedData.cpp
--- a/Engine/Source/Runtime/Engine/Private/AudioDerivedData.cpp
+++ b/Engine/Source/Runtime/Engine/Private/AudioDerivedData.cpp
@@ -619,7 +619,7 @@ void FStreamedAudioPlatformData::Serialize(FArchive& Ar, USoundWave* Owner)
 /**
  * Helper class to display a status update message in the editor.
  */
-class FAudioStatusMessageContext
+class FAudioStatusMessageContext : FScopedSlowTask
 {
 public:
 
@@ -627,26 +627,11 @@ public:
 	 * Updates the status message displayed to the user.
 	 */
 	explicit FAudioStatusMessageContext( const FText& InMessage )
+	 : FScopedSlowTask(1, InMessage, GIsEditor && !IsRunningCommandlet())
 	{
-		if ( GIsEditor && !IsRunningCommandlet() )
-		{
-			GWarn->PushStatus();
-			GWarn->StatusUpdate(-1, -1, InMessage);
-		}
 		DEFINE_LOG_CATEGORY_STATIC(LogAudioDerivedData, Log, All);
 		UE_LOG(LogAudioDerivedData, Display, TEXT("%s"), *InMessage.ToString());
 	}
-
-	/**
-	 * Ensures the status context is popped off the stack.
-	 */
-	~FAudioStatusMessageContext()
-	{
-		if ( GIsEditor && !IsRunningCommandlet() )
-		{
-			GWarn->PopStatus();
-		}
-	}
 };
 
 /**
diff --git a/Engine/Source/Runtime/Engine/Private/ShaderCompiler/ShaderCompiler.cpp b/Engine/Source/Runtime/Engine/Private/ShaderCompiler/ShaderCompiler.cpp
--- a/Engine/Source/Runtime/Engine/Private/ShaderCompiler/ShaderCompiler.cpp
+++ b/Engine/Source/Runtime/Engine/Private/ShaderCompiler/ShaderCompiler.cpp
@@ -1841,24 +1841,19 @@ void FShaderCompilingManager::FinishCompilation(const TCHAR* MaterialName, const
 	check(!FPlatformProperties::RequiresCookedData());
 	const double StartTime = FPlatformTime::Seconds();
 
-	if (GIsEditor && !IsRunningCommandlet())
+	FText StatusUpdate;
+	if ( MaterialName != NULL )
 	{
-		GWarn->PushStatus();
-
-		FText StatusUpdate;
-		if ( MaterialName != NULL )
-		{
-			FFormatNamedArguments Args;
-			Args.Add( TEXT("MaterialName"), FText::FromString( MaterialName ) );
-			StatusUpdate = FText::Format( NSLOCTEXT("ShaderCompilingManager", "CompilingShadersForMaterialStatus", "Compiling shaders: {MaterialName}..."), Args );
-		}
-		else
-		{
-			StatusUpdate = NSLOCTEXT("ShaderCompilingManager", "CompilingShadersStatus", "Compiling shaders...");
-		}
-
-		GWarn->StatusUpdate(-1, -1, StatusUpdate);
+		FFormatNamedArguments Args;
+		Args.Add( TEXT("MaterialName"), FText::FromString( MaterialName ) );
+		StatusUpdate = FText::Format( NSLOCTEXT("ShaderCompilingManager", "CompilingShadersForMaterialStatus", "Compiling shaders: {MaterialName}..."), Args );
 	}
+	else
+	{
+		StatusUpdate = NSLOCTEXT("ShaderCompilingManager", "CompilingShadersStatus", "Compiling shaders...");
+	}
+
+	FScopedSlowTask SlowTask(0, StatusUpdate, GIsEditor && !IsRunningCommandlet());
 
 	TMap<int32, FShaderMapFinalizeResults> CompiledShaderMaps;
 	BlockOnShaderMapCompletion(ShaderMapIdsToFinishCompiling, CompiledShaderMaps);
@@ -1873,14 +1868,6 @@ void FShaderCompilingManager::FinishCompilation(const TCHAR* MaterialName, const
 	ProcessCompiledShaderMaps(CompiledShaderMaps, FLT_MAX);
 	check(CompiledShaderMaps.Num() == 0);
 
-	if (GIsEditor && !IsRunningCommandlet())
-	{
-		// Clear the status so it doesn't look like we're compiling shaders for the entire load, in the case that no other status is specified
-		GWarn->StatusUpdate(-1, -1, FText::GetEmpty());
-		// We called StatusUpdatef earlier, so we want to restore whatever the previous status was now
-		GWarn->PopStatus();
-	}
-
 	const double EndTime = FPlatformTime::Seconds();
 
 	UE_LOG(LogShaders, Log, TEXT("FinishCompilation %s %.3fs"), MaterialName ? MaterialName : TEXT(""), (float)(EndTime - StartTime));
diff --git a/Engine/Source/Runtime/Engine/Private/StaticMesh.cpp b/Engine/Source/Runtime/Engine/Private/StaticMesh.cpp
--- a/Engine/Source/Runtime/Engine/Private/StaticMesh.cpp
+++ b/Engine/Source/Runtime/Engine/Private/StaticMesh.cpp
@@ -1038,9 +1038,10 @@ class FStaticMeshStatusMessageContext : public FScopedSlowTask
 {
 public:
 	explicit FStaticMeshStatusMessageContext(const FText& InMessage)
-		: FScopedSlowTask(InMessage)
+		: FScopedSlowTask(0, InMessage)
 	{
 		UE_LOG(LogStaticMesh,Log,TEXT("%s"),*InMessage.ToString());
+		MakeDialog();
 	}
 };
 
diff --git a/Engine/Source/Runtime/Engine/Private/TextureDerivedData.cpp b/Engine/Source/Runtime/Engine/Private/TextureDerivedData.cpp
--- a/Engine/Source/Runtime/Engine/Private/TextureDerivedData.cpp
+++ b/Engine/Source/Runtime/Engine/Private/TextureDerivedData.cpp
@@ -468,9 +468,10 @@ class FTextureStatusMessageContext : public FScopedSlowTask
 {
 public:
 	explicit FTextureStatusMessageContext(const FText& InMessage)
-		: FScopedSlowTask(InMessage)
+		: FScopedSlowTask(0, InMessage)
 	{
 		UE_LOG(LogTexture,Display,TEXT("%s"),*InMessage.ToString());
+		MakeDialog();
 	}
 };
 
diff --git a/Engine/Source/Runtime/Launch/Private/LaunchEngineLoop.cpp b/Engine/Source/Runtime/Launch/Private/LaunchEngineLoop.cpp
--- a/Engine/Source/Runtime/Launch/Private/LaunchEngineLoop.cpp
+++ b/Engine/Source/Runtime/Launch/Private/LaunchEngineLoop.cpp
@@ -1196,6 +1196,11 @@ int32 FEngineLoop::PreInit( const TCHAR* CmdLine )
 		FCoreStyle::ResetToDefault();
 	}
 
+	
+	FScopedSlowTask SlowTask(100, NSLOCTEXT("EngineLoop", "EngineLoop_Initializing", "Initializing..."));
+
+	SlowTask.EnterProgressFrame(10);
+
 	{
 		DECLARE_SCOPE_CYCLE_COUNTER(TEXT("Initial UObject load"), STAT_InitialUObjectLoad, STATGROUP_LoadTime);
 
@@ -1224,6 +1229,8 @@ int32 FEngineLoop::PreInit( const TCHAR* CmdLine )
 		UMaterialInterface::AssertDefaultMaterialsPostLoaded();
 	}
 
+	SlowTask.EnterProgressFrame(10);
+
 	// Tell the module manager is may now process newly-loaded UObjects when new C++ modules are loaded
 	FModuleManager::Get().StartProcessingNewlyLoadedObjects();
 
@@ -1240,6 +1247,8 @@ int32 FEngineLoop::PreInit( const TCHAR* CmdLine )
 		GUObjectArray.DisableDisregardForGC();
 	}
 	
+	SlowTask.EnterProgressFrame(10);
+
 	if ( !LoadStartupCoreModules() )
 	{
 		// At least one startup module failed to load, return 1 to indicate an error
@@ -1259,6 +1268,8 @@ int32 FEngineLoop::PreInit( const TCHAR* CmdLine )
 	}
 #endif
 
+	SlowTask.EnterProgressFrame(10);
+	
 	// Load up all modules that need to hook into the loading screen
 	if (!IProjectManager::Get().LoadModulesForProject(ELoadingPhase::PreLoadingScreen) || !IPluginManager::Get().LoadModulesForEnabledPlugins(ELoadingPhase::PreLoadingScreen))
 	{
@@ -1286,6 +1297,8 @@ int32 FEngineLoop::PreInit( const TCHAR* CmdLine )
         // do any post appInit processing, before the render thread is started.
         FPlatformMisc::PlatformPostInit(true);
     }
+	SlowTask.EnterProgressFrame(10);
+
 	if (GUseThreadedRendering)
 	{
 #if PLATFORM_SUPPORTS_RHI_THREAD
@@ -1316,6 +1329,8 @@ int32 FEngineLoop::PreInit( const TCHAR* CmdLine )
 	SetIsServerForOnlineSubsystemsDelegate(FQueryIsRunningServer::CreateStatic(&IsServerDelegateForOSS));
 #endif
 
+	SlowTask.EnterProgressFrame(50);
+
 #if WITH_EDITOR
 	if (!bHasEditorToken)
 	{
@@ -1816,6 +1831,9 @@ void GameLoopIsStarved()
 
 int32 FEngineLoop::Init()
 {
+	FScopedSlowTask SlowTask(100);
+	SlowTask.EnterProgressFrame(10);
+
 	// Figure out which UEngine variant to use.
 	UClass* EngineClass = NULL;
 	if( !GIsEditor )
@@ -1843,8 +1861,12 @@ int32 FEngineLoop::Init()
 
 	InitTime();
 
+	SlowTask.EnterProgressFrame(60);
+
 	GEngine->Init(this);
 
+	SlowTask.EnterProgressFrame(30);
+
 	// Load all the post-engine init modules
 	if (!IProjectManager::Get().LoadModulesForProject(ELoadingPhase::PostEngineInit) || !IPluginManager::Get().LoadModulesForEnabledPlugins(ELoadingPhase::PostEngineInit))
 	{
