BP: fix for graph editor crash

diff --git a/Engine/Source/Editor/UnrealEd/Private/EdGraphUtilities.cpp b/Engine/Source/Editor/UnrealEd/Private/EdGraphUtilities.cpp
--- a/Engine/Source/Editor/UnrealEd/Private/EdGraphUtilities.cpp
+++ b/Engine/Source/Editor/UnrealEd/Private/EdGraphUtilities.cpp
@@ -323,12 +323,14 @@ FIntRect FEdGraphUtilities::CalculateApproximateNodeBoundaries(const TArray<UEdG
 	for (auto NodeIt(Nodes.CreateConstIterator()); NodeIt; ++NodeIt)
 	{
 		UEdGraphNode* Node = *NodeIt;
-
-		// Update stats
-		MinNodeX = FMath::Min<int32>(MinNodeX, Node->NodePosX);
-		MinNodeY = FMath::Min<int32>(MinNodeY, Node->NodePosY);
-		MaxNodeX = FMath::Max<int32>(MaxNodeX, Node->NodePosX + Node->NodeWidth);
-		MaxNodeY = FMath::Max<int32>(MaxNodeY, Node->NodePosY + Node->NodeHeight);
+		if (Node)
+		{
+			// Update stats
+			MinNodeX = FMath::Min<int32>(MinNodeX, Node->NodePosX);
+			MinNodeY = FMath::Min<int32>(MinNodeY, Node->NodePosY);
+			MaxNodeX = FMath::Max<int32>(MaxNodeX, Node->NodePosX + Node->NodeWidth);
+			MaxNodeY = FMath::Max<int32>(MaxNodeY, Node->NodePosY + Node->NodeHeight);
+		}
 	}
 
 	const int32 AverageNodeWidth = 200;
