Worked on the Editor Performance tool.
Made it so it doesn't collect information as often as it does.

diff --git a/Engine/Source/Editor/UnrealEd/Private/Tests/AutomationEditorCommon.cpp b/Engine/Source/Editor/UnrealEd/Private/Tests/AutomationEditorCommon.cpp
--- a/Engine/Source/Editor/UnrealEd/Private/Tests/AutomationEditorCommon.cpp
+++ b/Engine/Source/Editor/UnrealEd/Private/Tests/AutomationEditorCommon.cpp
@@ -10,7 +10,7 @@
 #include "ShaderCompiler.h"
 #include "BehaviorTree/BehaviorTree.h"
 #include "Engine/DestructibleMesh.h"
-
+#include "FileManagerGeneric.h"
 
 
 DEFINE_LOG_CATEGORY_STATIC(LogAutomationEditorCommon, Log, All);
@@ -333,11 +333,11 @@ float ValueFromStringArrayofNumbers(const TArray<FString>& NumberArray, bool bAv
 	//Otherwise return the total value.
 	if (bAverageInstead)
 	{
-		UE_LOG(LogEditorAutomationTests, Log, TEXT("Average value of the Array is %f"), (TotalValue / NumberCounter));
+		UE_LOG(LogEditorAutomationTests, VeryVerbose, TEXT("Average value of the Array is %f"), (TotalValue / NumberCounter));
 		return (TotalValue / NumberCounter);
 	}
 	
-	UE_LOG(LogEditorAutomationTests, Log, TEXT("Total Value of the Array is %f"), TotalValue);
+	UE_LOG(LogEditorAutomationTests, VeryVerbose, TEXT("Total Value of the Array is %f"), TotalValue);
 	return TotalValue;
 
 }
@@ -360,7 +360,7 @@ float LargestValueFromStringArrayofNumbers(const TArray<FString>& NumberArray)
 			}
 		}
 	}
-	UE_LOG(LogEditorAutomationTests, Log, TEXT("The Largest value of the array is %f"), LargestValue);
+	UE_LOG(LogEditorAutomationTests, VeryVerbose, TEXT("The Largest value of the array is %f"), LargestValue);
 	return LargestValue;
 }
 
@@ -370,13 +370,22 @@ TArray<FString> CreateArrayFromFile(const FString& FullFileLocation)
 	FString RawData;
 	TArray<FString> DataArray;
 
-	UE_LOG(LogEditorAutomationTests, Log, TEXT("Loading and parsing the data from '%s' into an array."), *FullFileLocation);
-	FFileHelper::LoadFileToString(RawData, *FullFileLocation);
-	RawData.ParseIntoArray(&DataArray, TEXT(","), false);
-	
+	if (FPaths::FileExists(*FullFileLocation))
+	{
+
+		UE_LOG(LogEditorAutomationTests, VeryVerbose, TEXT("Loading and parsing the data from '%s' into an array."), *FullFileLocation);
+		FFileHelper::LoadFileToString(RawData, *FullFileLocation);
+		RawData.ParseIntoArray(&DataArray, TEXT(","), false);
+
+		return DataArray;
+	}
+
+	RawData = TEXT("0");
+	DataArray.Add(RawData);
 	return DataArray;
 }
 
+
 ///////////////////////////////////////////////////////////////////////
 // Common Latent commands
 
@@ -470,7 +479,7 @@ bool FEditorPerformanceCommand::Update()
 	//Gets the main frame module to get the name of our current level.
 	const IMainFrameModule& MainFrameModule = FModuleManager::GetModuleChecked< IMainFrameModule >("MainFrame");
 	FString ShortMapName = MainFrameModule.GetLoadedLevelName();
-	
+
 	//Variables that hold the content from the text files
 	FString SavedAverageFPS;
 	FString SavedMemory;
@@ -483,14 +492,35 @@ bool FEditorPerformanceCommand::Update()
 	static SIZE_T UsedPeakVirtual = 0;
 
 	//Grab the FPS and working memory number.
-	float CurrentTime = FPlatformTime::Seconds();
-	if ((CurrentTime - StartTime) <= Duration)
+	//float CurrentTime = FPlatformTime::Seconds();
+	//float ElapsedTime = CurrentTime - StartTime;
+	float CurrentTime = 0;
+	float ElapsedTime = 0;
+	FDateTime CurrentDateAndTime = FDateTime::UtcNow();
+
+	//If the last entry and the current time are the same then return false until it is. We only want to grab data for every actual second.
+	FString TestDurationFileLocation = FPaths::Combine(*FPaths::AutomationLogDir(), TEXT("Performance"), *ShortMapName, TEXT("RAWDuration.txt"));
+	if (FPaths::FileExists(TestDurationFileLocation))
+	{
+		FDateTime PerformanceRunTime = IFileManager::Get().GetTimeStamp(*TestDurationFileLocation);
+		FTimespan TimeDifferenceBetweenRuns = CurrentDateAndTime.GetTicks() - PerformanceRunTime.GetTicks();
+		if (TimeDifferenceBetweenRuns.GetTicks() < ETimespan::TicksPerSecond)
+		{
+			return false;
+		}
+
+		CurrentTime = FMath::RoundToInt(TimeDifferenceBetweenRuns.GetTotalSeconds());
+		TArray<FString> TestRunTimeRawArray = CreateArrayFromFile(TestDurationFileLocation);
+		ElapsedTime = (ValueFromStringArrayofNumbers(TestRunTimeRawArray, false)) + CurrentTime;
+	}
+
+	if (ElapsedTime < Duration || (FPlatformTime::Seconds() - StartTime) < Duration)
 	{
 		//Find the Average FPS
 		//Clamp to avoid huge averages at startup or after hitches
 		const float CurrentFPS = 1.0f / FSlateApplication::Get().GetAverageDeltaTime();
 		const float ClampedFPS = (CurrentFPS < 0.0f || CurrentFPS > 4000.0f) ? 0.0f : CurrentFPS;
-		
+
 		//Find the Frame Time in ms.
 		//Clamp to avoid huge averages at startup or after hitches
 		const float AverageMS = FSlateApplication::Get().GetAverageDeltaTime() * 1000.0f;
@@ -536,24 +566,24 @@ bool FEditorPerformanceCommand::Update()
 
 		//Log out to a text file the current time stamp.
 		FString OldCurrentTimeData;
-		FDateTime CurrentDateAndTime = FDateTime::Now();
 		FString FormatedTimeStamp = FString::Printf(TEXT("%04i-%02i-%02i %02i:%02i:%02i.%03i"), CurrentDateAndTime.GetYear(), CurrentDateAndTime.GetMonth(), CurrentDateAndTime.GetDay(), CurrentDateAndTime.GetHour(), CurrentDateAndTime.GetMinute(), CurrentDateAndTime.GetSecond(), CurrentDateAndTime.GetMillisecond());
 		FString OldFileLocation = FPaths::AutomationLogDir() + TEXT("Performance/") + ShortMapName + TEXT("/RAWTimeStamp.txt");
-		FFileHelper::LoadFileToString( OldCurrentTimeData, *OldFileLocation );
+		FFileHelper::LoadFileToString(OldCurrentTimeData, *OldFileLocation);
 		FString FileSetup = OldCurrentTimeData + FormatedTimeStamp + TEXT(",");
 		FFileHelper::SaveStringToFile(FileSetup, *OldFileLocation);
-		
+
 		return false;
 	}
-	
+
 	UE_LOG(LogEditorAutomationTests, Log, TEXT("Raw performance data has been captured and saved to the '//Game/Saved/Automation/Log/Performance/%s' folder location."), *ShortMapName);
-	
+
 	return true;
 }
 
 bool FGenerateEditorPerformanceCharts::Update()
 {
 	UE_LOG(LogEditorAutomationTests, Log, TEXT("Begin generating the editor performance charts."));
+
 	//Get the current changelist number
 	FString Changelist = GEngineVersion.ToString(EVersionComponent::Changelist);
 
@@ -582,17 +612,30 @@ bool FGenerateEditorPerformanceCharts::Update()
 	TArray<FString> MemoryAvailVirtualRawArray = CreateArrayFromFile(MemoryAvailVirtualFileLocation);
 	TArray<FString> MemoryUsedVirtualRawArray = CreateArrayFromFile(MemoryUsedVirtualFileLocation);
 	TArray<FString> MapLoadTimeRawArray = CreateArrayFromFile(MapLoadTimeFileLocation);
-	TArray<FString> TestRunTimeRawArray = CreateArrayFromFile(DurationFileLocation);
+	TArray<FString> DurationRawArray = CreateArrayFromFile(DurationFileLocation);
 	TArray<FString> MemoryUsedPeakArray = CreateArrayFromFile(MemoryUsedPeakFileLocation);
 	TArray<FString> MemoryUsedPeakVirtualArray = CreateArrayFromFile(MemoryUsedPeakVirtualFileLocation);
 	TArray<FString> TimeStampArray = CreateArrayFromFile(TimeStampFileLocation);
 
+	FString MapLoadTime;
+
 	//If the map load time doesn't exist then we give add a '0' to the Map Load time array.
-	if (IFileManager::Get().FileSize(MapLoadTimeFileLocation.GetCharArray().GetData()) <= 0)
+	if (!FPaths::FileExists(*MapLoadTimeFileLocation))
+	{
+		MapLoadTimeRawArray.Add("0");
+	}
+
+	//Make sure that the map load time actually has data to show even if the time is '0'
+	MapLoadTime = MapLoadTimeRawArray.Last();
+	if (MapLoadTime.IsEmpty())
 	{
-		MapLoadTimeRawArray.Add(TEXT("0"));
+		if (MapLoadTimeRawArray[MapLoadTimeRawArray.Num() - 2].IsNumeric())
+		{
+			MapLoadTime = MapLoadTimeRawArray[MapLoadTimeRawArray.Num() - 2];
+		}
 	}
 
+	
 	//Variable that will hold the old performance csv file data.
 	FString OldPerformanceCSVFile;
 
@@ -610,11 +653,11 @@ bool FGenerateEditorPerformanceCharts::Update()
 	}
 
 	//Add the top title row to the raw csv file
-	FString RAWCSVLine = (TEXT("Map Name, Changelist, Time Stamp, Test Run Time (in seconds), Map Load Time, Average FPS, Frame Time, Used Physical Memory, Used Virtual Memory, Used Peak Physical, Used Peak Virtual, Available Physical Memory, Available Virtual Memory\n"));
+	FString RAWCSVLine = (TEXT("Map Name, Changelist, Time Stamp, Map Load Time, Average FPS, Frame Time, Used Physical Memory, Used Virtual Memory, Used Peak Physical, Used Peak Virtual, Available Physical Memory, Available Virtual Memory\n"));
 	RAWCSVFile->Serialize(TCHAR_TO_ANSI(*RAWCSVLine), RAWCSVLine.Len());
 	
 	//Loop through the text files and add them to the raw csv file.
-	for (int32 I = 0; I < FPSRawArray.Num() - 1; I++)
+	for (int32 I = 0; I < (DurationRawArray.Num() - 1); I++)
 	{
 		//If the raw file isn't available to write to then we'll fail back this test.
 		if (!RAWCSVFile)
@@ -624,7 +667,7 @@ bool FGenerateEditorPerformanceCharts::Update()
 		}
 
 		//Create the line that will be written and then write it to the csv file.
-		RAWCSVLine = LoadedMapName + TEXT(",") + Changelist + TEXT(",") + TimeStampArray[I] + TEXT(",") + TestRunTimeRawArray[I] + TEXT(",") + MapLoadTimeRawArray[0] + TEXT(",") + FPSRawArray[I] + TEXT(",") + FrameTimeRawArray[I] + TEXT(",") + MemoryRawArray[I] + TEXT(",") + MemoryUsedVirtualRawArray[I] + TEXT(",") + MemoryUsedPeakArray[I] + TEXT(",") + MemoryUsedPeakVirtualArray[I] + TEXT(",") + MemoryAvailPhysRawArray[I] + TEXT(",") + MemoryAvailVirtualRawArray[I] + LINE_TERMINATOR;
+		RAWCSVLine = LoadedMapName + TEXT(",") + Changelist + TEXT(",") + TimeStampArray[I] + TEXT(",") + MapLoadTime + TEXT(",") + FPSRawArray[I] + TEXT(",") + FrameTimeRawArray[I] + TEXT(",") + MemoryRawArray[I] + TEXT(",") + MemoryUsedVirtualRawArray[I] + TEXT(",") + MemoryUsedPeakArray[I] + TEXT(",") + MemoryUsedPeakVirtualArray[I] + TEXT(",") + MemoryAvailPhysRawArray[I] + TEXT(",") + MemoryAvailVirtualRawArray[I] + LINE_TERMINATOR;
 		RAWCSVFile->Serialize(TCHAR_TO_ANSI(*RAWCSVLine), RAWCSVLine.Len());
 	}
 	
@@ -639,29 +682,15 @@ bool FGenerateEditorPerformanceCharts::Update()
 	float MemoryAvailPhysAvg = ValueFromStringArrayofNumbers(MemoryAvailPhysRawArray, true);
 	float MemoryAvailVirtualAvg = ValueFromStringArrayofNumbers(MemoryAvailVirtualRawArray, true);
 	float MemoryUsedVirtualAvg = ValueFromStringArrayofNumbers(MemoryUsedVirtualRawArray, true);
+	float TestRunTime = ValueFromStringArrayofNumbers(DurationRawArray, false);
 	
 	//Find the largest value in the array.
 	float MemoryUsedPeak = LargestValueFromStringArrayofNumbers(MemoryUsedPeakArray);
 	float MemoryUsedPeakVirtual = LargestValueFromStringArrayofNumbers(MemoryUsedPeakVirtualArray);
 	
-	//This gets the actual test run time by subtracting the last test run entry against the first test run entry.
-	FString TestRunStartTime = TestRunTimeRawArray[0];
-	FString TestRunEndTime;
-	for (int32 i = (TestRunTimeRawArray.Num() - 1); i >= 0 ; --i)
-	{
-		if (!TestRunTimeRawArray[i].IsEmpty())
-		{
-			TestRunEndTime = TestRunTimeRawArray[i];
-			break;
-		}
-	}
-	float TestRunTime = FCString::Atof(TestRunEndTime.GetCharArray().GetData()) - FCString::Atof(TestRunStartTime.GetCharArray().GetData());
-	//Get the map load time.
-	FString MapLoadTime = MapLoadTimeRawArray[0];
-	
 	//If the Performance csv file does not have a size then it must be new.
 	//This adds the title row to the performance csv file.
-	if (IFileManager::Get().FileSize(FinalCSVFilename.GetCharArray().GetData()) <= 0)
+	if (IFileManager::Get().FileSize(*FinalCSVFilename) <= 0)
 	{
 		FArchive* FinalCSVFile = IFileManager::Get().CreateFileWriter(*FinalCSVFilename);
 		//If the csv file can't be created or written to then we'll fail back this test.
@@ -704,6 +733,7 @@ bool FGenerateEditorPerformanceCharts::Update()
 	UE_LOG(LogEditorAutomationTests, Display, TEXT("AVG Used Virtual Memory: '%.0f' kb"), MemoryUsedVirtualAvg);
 	UE_LOG(LogEditorAutomationTests, Display, TEXT("Performance csv file is located here: %s"), *FPaths::ConvertRelativePathToFull(FinalCSVFilename));
 	UE_LOG(LogEditorAutomationTests, Log, TEXT("Performance csv file is located here: %s"), *FPaths::ConvertRelativePathToFull(FinalCSVFilename));
+	UE_LOG(LogEditorAutomationTests, Log, TEXT("Raw performance csv file is located here: %s"), *FPaths::ConvertRelativePathToFull(RAWCSVFilename));
 
 	//Delete the text files.
 	UE_LOG(LogEditorAutomationTests, Log, TEXT("Deleting temporary data text files."));
@@ -711,7 +741,6 @@ bool FGenerateEditorPerformanceCharts::Update()
 	IFileManager::Get().Delete(*FrameTimeLocation);
 	IFileManager::Get().Delete(*MemoryFileLocation);
 	IFileManager::Get().Delete(*MemoryAvailPhysFileLocation);
-	IFileManager::Get().Delete(*MapLoadTimeFileLocation);
 	IFileManager::Get().Delete(*DurationFileLocation);
 	IFileManager::Get().Delete(*MemoryAvailVirtualFileLocation);
 	IFileManager::Get().Delete(*MemoryUsedVirtualFileLocation);
@@ -774,6 +803,7 @@ bool FChangeViewportToFirstAvailableBookmarkCommand::Update()
 {
 	FEditorModeTools EditorModeTools;
 	FLevelEditorViewportClient* ViewportClient;
+	uint32 ViewportIndex = 0;
 
 	UE_LOG(LogEditorAutomationTests, Log, TEXT("Attempting to change the editor viewports view to the first set bookmark."));
 
@@ -782,11 +812,11 @@ bool FChangeViewportToFirstAvailableBookmarkCommand::Update()
 	{
 		ViewportClient = GEditor->LevelViewportClients[i];
 
-		for (uint32 ViewportIndex = 0; ViewportIndex <= AWorldSettings::MAX_BOOKMARK_NUMBER; ViewportIndex++)
+		for (ViewportIndex = 0; ViewportIndex <= AWorldSettings::MAX_BOOKMARK_NUMBER; ViewportIndex++)
 		{
 			if (EditorModeTools.CheckBookmark(ViewportIndex, ViewportClient))
 			{
-				UE_LOG(LogEditorAutomationTests, Log, TEXT("Changing a viewport view to the set bookmark %i"), ViewportIndex);
+				UE_LOG(LogEditorAutomationTests, VeryVerbose, TEXT("Changing a viewport view to the set bookmark %i"), ViewportIndex);
 				EditorModeTools.JumpToBookmark(ViewportIndex, true, ViewportClient);
 				break;
 			}
