Fixed crash when removing streaming level in Levels browser

diff --git a/Engine/Source/Editor/UnrealEd/Private/EditorLevelUtils.cpp b/Engine/Source/Editor/UnrealEd/Private/EditorLevelUtils.cpp
--- a/Engine/Source/Editor/UnrealEd/Private/EditorLevelUtils.cpp
+++ b/Engine/Source/Editor/UnrealEd/Private/EditorLevelUtils.cpp
@@ -488,14 +488,7 @@ namespace EditorLevelUtils
 			check(Level->bIsVisible == false);
 		}
 
-		const bool bSuccess = EditorDestroyLevel(Level);
-		// Since we just removed all the actors from this package, we do not want it to be saved out now
-		// and the user was warned they would lose any changes from before removing, so we're good to clear
-		// the dirty flag
-		UPackage* LevelPackage = Level->GetOutermost();
-		LevelPackage->SetDirtyFlag(false);
-
-		return bSuccess;
+		return EditorDestroyLevel(Level);
 	}
 	
 	bool EditorDestroyLevel( ULevel* InLevel )
