Make sure materials created from textures don't overlap expressions

Set the X coord of the texture sample when creating new material.

TTP#345550 - NUX: Auto-creating material from texture should not put the TextureSample node right on top of output node

diff --git a/Engine/Source/Editor/UnrealEd/Private/Factories/EditorFactories.cpp b/Engine/Source/Editor/UnrealEd/Private/Factories/EditorFactories.cpp
--- a/Engine/Source/Editor/UnrealEd/Private/Factories/EditorFactories.cpp
+++ b/Engine/Source/Editor/UnrealEd/Private/Factories/EditorFactories.cpp
@@ -166,6 +166,7 @@ UObject* UMaterialFactoryNew::FactoryCreateNew(UClass* Class,UObject* InParent,F
 		NewMaterial->Expressions.Add( Expression );
 
 		NewMaterial->BaseColor.Expression = Expression;
+		Expression->MaterialExpressionEditorX = -250;
 		Expression->Texture = InitialTexture;
 
 		TArray<FExpressionOutput> Outputs;
