FBX Refactor - Added import transform to static and skeletal meshes.

diff --git a/Engine/Source/Editor/DetailCustomizations/Private/FbxImportUIDetails.cpp b/Engine/Source/Editor/DetailCustomizations/Private/FbxImportUIDetails.cpp
--- a/Engine/Source/Editor/DetailCustomizations/Private/FbxImportUIDetails.cpp
+++ b/Engine/Source/Editor/DetailCustomizations/Private/FbxImportUIDetails.cpp
@@ -92,14 +92,15 @@ void FFbxImportUIDetails::CustomizeDetails( IDetailLayoutBuilder& DetailBuilder
 		if(IsImportTypeMetaDataValid(ImportType, ImportTypeMetaData))
 		{
 			// Decide on category
-			if(CategoryMetaData == TEXT("Transform"))
+			if(!CategoryMetaData.IsEmpty())
 			{
-				// Transform data can be held inside the asset uobjects (used for reimport later).
-				// Move it out into its own UI section.
-				TransformCategory.AddProperty(Handle);
+				// Populate custom categories.
+				IDetailCategoryBuilder& CustomCategory = DetailBuilder.EditCategory(*CategoryMetaData);
+				CustomCategory.AddProperty(Handle);
 			}
 			else
 			{
+				// No override, add to default mesh category
 				MeshCategory.AddProperty(Handle);
 			}
 		}
@@ -139,15 +140,20 @@ void FFbxImportUIDetails::CustomizeDetails( IDetailLayoutBuilder& DetailBuilder
 
 		for(TSharedPtr<IPropertyHandle> Handle : ExtraProperties)
 		{
-			if(Handle->GetProperty()->GetOuter() == UFbxAnimSequenceImportData::StaticClass())
+			FString CategoryMetaData = Handle->GetMetaData(TEXT("ImportCategory"));
+			if(Handle->GetProperty()->GetOuter() == UFbxAnimSequenceImportData::StaticClass()
+			   && CategoryMetaData.IsEmpty())
 			{
+				// Add to default anim category if no override specified
 				IDetailPropertyRow& PropertyRow = AnimCategory.AddProperty(Handle);
 			}
+			else if(ImportType == FBXIT_Animation && !CategoryMetaData.IsEmpty())
+			{
+				// Override category is available
+				IDetailCategoryBuilder& CustomCategory = DetailBuilder.EditCategory(*CategoryMetaData);
+				CustomCategory.AddProperty(Handle);
+			}
 		}
-
-		// Start and End frame for animations need a custom edit condition
-		TSharedRef<IPropertyHandle> StartFrameProp = DetailBuilder.GetProperty(GET_MEMBER_NAME_CHECKED(UFbxAnimSequenceImportData, StartFrame));
-		TSharedRef<IPropertyHandle> EndFrameProp = DetailBuilder.GetProperty(GET_MEMBER_NAME_CHECKED(UFbxAnimSequenceImportData, EndFrame));	
 	}
 	else
 	{
diff --git a/Engine/Source/Editor/Persona/Private/AnimationEditorViewportClient.cpp b/Engine/Source/Editor/Persona/Private/AnimationEditorViewportClient.cpp
--- a/Engine/Source/Editor/Persona/Private/AnimationEditorViewportClient.cpp
+++ b/Engine/Source/Editor/Persona/Private/AnimationEditorViewportClient.cpp
@@ -1104,8 +1104,9 @@ bool FAnimationViewportClient::InputWidgetDelta( FViewport* Viewport, EAxisList:
 
 				//Calculate the new delta rotation
 				FQuat DeltaQuat( BoneSpaceAxis, RotAngle );
+				DeltaQuat.Normalize();
 
-				FRotator NewRotation = ( CurrentSkelControlTM * FTransform( DeltaQuat ) ).Rotator();
+				FRotator NewRotation = ( CurrentSkelControlTM * FTransform( DeltaQuat )).Rotator();
 
 				if ( SelectedSocket )
 				{
@@ -1302,9 +1303,9 @@ FMatrix FAnimationViewportClient::GetWidgetCoordSystem() const
 		{
 			int32 BoneIndex = PreviewSkelMeshComp->BonesOfInterest.Last();
 
-			FMatrix BoneMatrix = PreviewSkelMeshComp->GetBoneMatrix(BoneIndex);
+			FTransform BoneMatrix = PreviewSkelMeshComp->GetBoneTransform(BoneIndex);
 
-			return BoneMatrix.RemoveTranslation();
+			return BoneMatrix.ToMatrixNoScale().RemoveTranslation();
 		}
 		else if( PreviewSkelMeshComp->SocketsOfInterest.Num() > 0 )
 		{
diff --git a/Engine/Source/Editor/UnrealEd/Classes/Factories/FbxAssetImportData.h b/Engine/Source/Editor/UnrealEd/Classes/Factories/FbxAssetImportData.h
--- a/Engine/Source/Editor/UnrealEd/Classes/Factories/FbxAssetImportData.h
+++ b/Engine/Source/Editor/UnrealEd/Classes/Factories/FbxAssetImportData.h
@@ -1,6 +1,7 @@
 // Copyright 1998-2014 Epic Games, Inc. All Rights Reserved.
 
 #pragma once
+
 #include "EditorFramework/AssetImportData.h"
 #include "FbxAssetImportData.generated.h"
 
@@ -12,12 +13,12 @@ class UNREALED_API UFbxAssetImportData : public UAssetImportData
 {
 	GENERATED_UCLASS_BODY()
 
-	UPROPERTY()
+	UPROPERTY(EditAnywhere, Category=Transform, meta=(ImportType="StaticMesh|SkeletalMesh|Animation", ImportCategory="Transform"))
 	FVector ImportTranslation;
 
-	UPROPERTY()
+	UPROPERTY(EditAnywhere, Category=Transform, meta=(ImportType="StaticMesh|SkeletalMesh|Animation", ImportCategory="Transform"))
 	FRotator ImportRotation;
 
-	UPROPERTY()
+	UPROPERTY(EditAnywhere, Category=Transform, meta=(ImportType="StaticMesh|SkeletalMesh|Animation", ImportCategory="Transform"))
 	float ImportUniformScale;
 };
\ No newline at end of file
diff --git a/Engine/Source/Editor/UnrealEd/Private/Factories/EditorFactories.cpp b/Engine/Source/Editor/UnrealEd/Private/Factories/EditorFactories.cpp
--- a/Engine/Source/Editor/UnrealEd/Private/Factories/EditorFactories.cpp
+++ b/Engine/Source/Editor/UnrealEd/Private/Factories/EditorFactories.cpp
@@ -5694,6 +5694,7 @@ EReimportResult::Type UReimportFbxStaticMeshFactory::Reimport( UObject* Obj )
 
 		if ( FFbxImporter->ImportFromFile( *Filename, FPaths::GetExtension( Filename ) ) )
 		{
+			FFbxImporter->ApplyTransformSettingsToFbxNode(FFbxImporter->Scene->GetRootNode(), ImportData);
 			const TArray<UAssetUserData*>* UserData = Mesh->GetAssetUserDataArray();
 			TArray<UAssetUserData*> UserDataCopy;
 			if (UserData)
diff --git a/Engine/Source/Editor/UnrealEd/Private/Fbx/FbxFactory.cpp b/Engine/Source/Editor/UnrealEd/Private/Fbx/FbxFactory.cpp
--- a/Engine/Source/Editor/UnrealEd/Private/Fbx/FbxFactory.cpp
+++ b/Engine/Source/Editor/UnrealEd/Private/Fbx/FbxFactory.cpp
@@ -240,6 +240,7 @@ UObject* UFbxFactory::FactoryCreateBinary
 
 			bool bImportStaticMeshLODs = ImportUI->StaticMeshImportData->bImportMeshLODs;
 			bool bCombineMeshes = ImportUI->bCombineMeshes;
+
 			if ( ImportUI->MeshTypeToImport == FBXIT_SkeletalMesh )
 			{
 				FbxImporter->FillFbxSkelMeshArrayInScene(RootNodeToImport, SkelMeshArray, false);
@@ -247,6 +248,8 @@ UObject* UFbxFactory::FactoryCreateBinary
 			}
 			else if( ImportUI->MeshTypeToImport == FBXIT_StaticMesh )
 			{
+				FbxImporter->ApplyTransformSettingsToFbxNode(RootNodeToImport, ImportUI->StaticMeshImportData);
+
 				if( bCombineMeshes && !bImportStaticMeshLODs )
 				{
 					// If Combine meshes and dont import mesh LODs, the interesting node count should be 1 so all the meshes are grouped together into one static mesh
@@ -387,7 +390,13 @@ UObject* UFbxFactory::FactoryCreateBinary
 
 								if ( NewMesh && ImportUI->bImportAnimations )
 								{
+									// We need to remove all scaling from the root node before we set up animation data.
+									// Othewise some of the global transform calculations will be incorrect.
+									FbxImporter->RemoveTransformSettingsFromFbxNode(RootNodeToImport, ImportUI->SkeletalMeshImportData);
 									FbxImporter->SetupAnimationDataFromMesh(NewMesh, InParent, SkelMeshNodeArray, ImportUI->AnimSequenceImportData, OutputName.ToString());
+
+									// Reapply the transforms for the rest of the import
+									FbxImporter->ApplyTransformSettingsToFbxNode(RootNodeToImport, ImportUI->SkeletalMeshImportData);
 								}
 							}
 							else if (NewObject) // the base skeletal mesh is imported successfully
diff --git a/Engine/Source/Editor/UnrealEd/Private/Fbx/FbxMainImport.cpp b/Engine/Source/Editor/UnrealEd/Private/Fbx/FbxMainImport.cpp
--- a/Engine/Source/Editor/UnrealEd/Private/Fbx/FbxMainImport.cpp
+++ b/Engine/Source/Editor/UnrealEd/Private/Fbx/FbxMainImport.cpp
@@ -186,16 +186,35 @@ void ApplyImportUIToImportOptions(UFbxImportUI* ImportUI, FBXImportOptions& InOu
 	{
 		UFbxStaticMeshImportData* StaticMeshData	= ImportUI->StaticMeshImportData;
 		InOutImportOptions.NormalImportMethod		= StaticMeshData->NormalImportMethod;
+		InOutImportOptions.ImportTranslation		= StaticMeshData->ImportTranslation;
+		InOutImportOptions.ImportRotation			= StaticMeshData->ImportRotation;
+		InOutImportOptions.ImportUniformScale		= StaticMeshData->ImportUniformScale;
 	}
 	else if ( ImportUI->MeshTypeToImport == FBXIT_SkeletalMesh )
 	{
 		UFbxSkeletalMeshImportData* SkeletalMeshData	= ImportUI->SkeletalMeshImportData;
 		InOutImportOptions.NormalImportMethod			= SkeletalMeshData->NormalImportMethod;
+		InOutImportOptions.ImportTranslation			= SkeletalMeshData->ImportTranslation;
+		InOutImportOptions.ImportRotation				= SkeletalMeshData->ImportRotation;
+		InOutImportOptions.ImportUniformScale			= SkeletalMeshData->ImportUniformScale;
+
+		if(ImportUI->bImportAnimations)
+		{
+			// Copy the transform information into the animation data to match the mesh.
+			UFbxAnimSequenceImportData* AnimData	= ImportUI->AnimSequenceImportData;
+			AnimData->ImportTranslation				= SkeletalMeshData->ImportTranslation;
+			AnimData->ImportRotation				= SkeletalMeshData->ImportRotation;
+			AnimData->ImportUniformScale			= SkeletalMeshData->ImportUniformScale;
+		}
+
 	}
 	else
 	{
 		UFbxAnimSequenceImportData* AnimData	= ImportUI->AnimSequenceImportData;
 		InOutImportOptions.NormalImportMethod = FBXNIM_ComputeNormals;
+		InOutImportOptions.ImportTranslation	= AnimData->ImportTranslation;
+		InOutImportOptions.ImportRotation		= AnimData->ImportRotation;
+		InOutImportOptions.ImportUniformScale	= AnimData->ImportUniformScale;
 	}
 
 	// only re-sample if they don't want to use default sample rate
@@ -1136,6 +1155,77 @@ void FFbxImporter::DumpFBXNode(FbxNode* Node)
 
 }
 
+void FFbxImporter::ApplyTransformSettingsToFbxNode(FbxNode* Node, UFbxAssetImportData* AssetData)
+{
+	check(Node);
+	check(AssetData);
+
+	FbxAMatrix TransformMatrix;
+	BuildFbxMatrixForImportTransform(TransformMatrix, AssetData);
+
+	FbxDouble3 ExistingTranslation = Node->LclTranslation.Get();
+	FbxDouble3 ExistingRotation = Node->LclRotation.Get();
+	FbxDouble3 ExistingScaling = Node->LclScaling.Get();
+
+	// A little slower to build up this information from the matrix, but it means we get
+	// the same result across all import types, as other areas can use the matrix directly
+	FbxVector4 AddedTranslation = TransformMatrix.GetT();
+	FbxVector4 AddedRotation = TransformMatrix.GetR();
+	FbxVector4 AddedScaling = TransformMatrix.GetS();
+
+	FbxDouble3 NewTranslation = FbxDouble3(ExistingTranslation[0] + AddedTranslation[0], ExistingTranslation[1] + AddedTranslation[1], ExistingTranslation[2] + AddedTranslation[2]);
+	FbxDouble3 NewRotation = FbxDouble3(ExistingRotation[0] + AddedRotation[0], ExistingRotation[1] + AddedRotation[1], ExistingRotation[2] + AddedRotation[2]);
+	FbxDouble3 NewScaling = FbxDouble3(ExistingScaling[0] * AddedScaling[0], ExistingScaling[1] * AddedScaling[1], ExistingScaling[2] * AddedScaling[2]);
+
+	Node->LclTranslation.Set(NewTranslation);
+	Node->LclRotation.Set(NewRotation);
+	Node->LclScaling.Set(NewScaling);
+}
+
+
+void FFbxImporter::RemoveTransformSettingsFromFbxNode(FbxNode* Node, UFbxAssetImportData* AssetData)
+{
+	check(Node);
+	check(AssetData);
+
+	FbxAMatrix TransformMatrix;
+	BuildFbxMatrixForImportTransform(TransformMatrix, AssetData);
+
+	FbxDouble3 ExistingTranslation = Node->LclTranslation.Get();
+	FbxDouble3 ExistingRotation = Node->LclRotation.Get();
+	FbxDouble3 ExistingScaling = Node->LclScaling.Get();
+
+	// A little slower to build up this information from the matrix, but it means we get
+	// the same result across all import types, as other areas can use the matrix directly
+	FbxVector4 AddedTranslation = TransformMatrix.GetT();
+	FbxVector4 AddedRotation = TransformMatrix.GetR();
+	FbxVector4 AddedScaling = TransformMatrix.GetS();
+
+	FbxDouble3 NewTranslation = FbxDouble3(ExistingTranslation[0] - AddedTranslation[0], ExistingTranslation[1] - AddedTranslation[1], ExistingTranslation[2] - AddedTranslation[2]);
+	FbxDouble3 NewRotation = FbxDouble3(ExistingRotation[0] - AddedRotation[0], ExistingRotation[1] - AddedRotation[1], ExistingRotation[2] - AddedRotation[2]);
+	FbxDouble3 NewScaling = FbxDouble3(ExistingScaling[0] / AddedScaling[0], ExistingScaling[1] / AddedScaling[1], ExistingScaling[2] / AddedScaling[2]);
+
+	Node->LclTranslation.Set(NewTranslation);
+	Node->LclRotation.Set(NewRotation);
+	Node->LclScaling.Set(NewScaling);
+}
+
+
+void FFbxImporter::BuildFbxMatrixForImportTransform(FbxAMatrix& OutMatrix, UFbxAssetImportData* AssetData)
+{
+	if(!AssetData)
+	{
+		OutMatrix.SetIdentity();
+		return;
+	}
+
+	FbxVector4 FbxAddedTranslation = Converter.ConvertToFbxPos(AssetData->ImportTranslation);
+	FbxVector4 FbxAddedScale = Converter.ConvertToFbxScale(FVector(AssetData->ImportUniformScale));
+	FbxVector4 FbxAddedRotation = Converter.ConvertToFbxRot(AssetData->ImportRotation.Euler());
+	
+	OutMatrix = FbxAMatrix(FbxAddedTranslation, FbxAddedRotation, FbxAddedScale);
+}
+
 /**
 * Get all Fbx skeletal mesh objects which are grouped by skeleton they bind to
 *
diff --git a/Engine/Source/Editor/UnrealEd/Private/Fbx/FbxSkeletalMeshImport.cpp b/Engine/Source/Editor/UnrealEd/Private/Fbx/FbxSkeletalMeshImport.cpp
--- a/Engine/Source/Editor/UnrealEd/Private/Fbx/FbxSkeletalMeshImport.cpp
+++ b/Engine/Source/Editor/UnrealEd/Private/Fbx/FbxSkeletalMeshImport.cpp
@@ -742,7 +742,7 @@ bool UnFbx::FFbxImporter::RetrievePoseFromBindPose(const TArray<FbxNode*>& NodeA
 	return (PoseArray.Size() > 0);
 }
 
-bool UnFbx::FFbxImporter::ImportBone(TArray<FbxNode*>& NodeArray, FSkeletalMeshImportData &ImportData, TArray<FbxNode*> &SortedLinks, bool& bOutDiffPose, bool bDisableMissingBindPoseWarning, bool & bUseTime0AsRefPose)
+bool UnFbx::FFbxImporter::ImportBone(TArray<FbxNode*>& NodeArray, FSkeletalMeshImportData &ImportData, UFbxSkeletalMeshImportData* TemplateData ,TArray<FbxNode*> &SortedLinks, bool& bOutDiffPose, bool bDisableMissingBindPoseWarning, bool & bUseTime0AsRefPose)
 {
 	bOutDiffPose = false;
 	int32 SkelType = 0; // 0 for skeletal mesh, 1 for rigid mesh
@@ -781,16 +781,16 @@ bool UnFbx::FFbxImporter::ImportBone(TArray<FbxNode*>& NodeArray, FSkeletalMeshI
 		for ( int32 i = 0; i < NodeArray.Num(); i++)
 		{
 			FbxMesh* FbxMesh = NodeArray[i]->GetMesh();
-				const int32 SkinDeformerCount = FbxMesh->GetDeformerCount(FbxDeformer::eSkin);
-				for (int32 DeformerIndex=0; DeformerIndex < SkinDeformerCount; DeformerIndex++)
+			const int32 SkinDeformerCount = FbxMesh->GetDeformerCount(FbxDeformer::eSkin);
+			for (int32 DeformerIndex=0; DeformerIndex < SkinDeformerCount; DeformerIndex++)
+			{
+				FbxSkin* Skin = (FbxSkin*)FbxMesh->GetDeformer(DeformerIndex, FbxDeformer::eSkin);
+				for ( int32 ClusterIndex = 0; ClusterIndex < Skin->GetClusterCount(); ClusterIndex++)
 				{
-					FbxSkin* Skin = (FbxSkin*)FbxMesh->GetDeformer(DeformerIndex, FbxDeformer::eSkin);
-					for ( int32 ClusterIndex = 0; ClusterIndex < Skin->GetClusterCount(); ClusterIndex++)
-					{
-						ClusterArray.Add(Skin->GetCluster(ClusterIndex));
-					}
-				}			
+					ClusterArray.Add(Skin->GetCluster(ClusterIndex));
+				}
 			}
+		}
 
 		if (ClusterArray.Num() == 0)
 		{
@@ -871,7 +871,7 @@ bool UnFbx::FFbxImporter::ImportBone(TArray<FbxNode*>& NodeArray, FSkeletalMeshI
 	FbxArray<FbxAMatrix> GlobalsPerLink;
 	GlobalsPerLink.Grow(SortedLinks.Num());
 	GlobalsPerLink[0].SetIdentity();
-	
+
 	bool GlobalLinkFoundFlag;
 	FbxVector4 LocalLinkT;
 	FbxQuaternion LocalLinkQ;
@@ -881,6 +881,8 @@ bool UnFbx::FFbxImporter::ImportBone(TArray<FbxNode*>& NodeArray, FSkeletalMeshI
 	FString LinksWithoutBindPoses;
 	int32 NumberOfRoot = 0;
 
+	int32 RootIdx = -1;
+
 	for (LinkIndex=0; LinkIndex<SortedLinks.Num(); LinkIndex++)
 	{
 		// Add a bone for each FBX Link
@@ -909,6 +911,7 @@ bool UnFbx::FFbxImporter::ImportBone(TArray<FbxNode*>& NodeArray, FSkeletalMeshI
 		if (ParentIndex == INDEX_NONE)
 		{
 			++NumberOfRoot;
+			RootIdx = LinkIndex;
 			if (NumberOfRoot > 1)
 			{
 				AddTokenizedErrorMessage(FTokenizedMessage::Create(EMessageSeverity::Error, LOCTEXT("MultipleRootsFound", "Multiple roots are found in the bone hierarchy. We only support single root bone.")), FFbxErrors::SkeletalMesh_MultipleRoots);
@@ -962,9 +965,6 @@ bool UnFbx::FFbxImporter::ImportBone(TArray<FbxNode*>& NodeArray, FSkeletalMeshI
 
 		if (!GlobalLinkFoundFlag)
 		{
-			// The node is not included in a bind pose and is not associated with any cluster, so just use its 'default'
-			// matrix as its bind pose.  
-			// Passing a time of Infinite (the default value) to EvaluateGlobalTransform returns the 'default' matrix for the node
 			GlobalsPerLink[LinkIndex] = Link->EvaluateGlobalTransform();
 		}
 		
@@ -978,8 +978,7 @@ bool UnFbx::FFbxImporter::ImportBone(TArray<FbxNode*>& NodeArray, FSkeletalMeshI
 			
 			GlobalsPerLink[LinkIndex] = T0Matrix;
 		}
-		
-
+	
 		if (LinkIndex)
 		{
 			FbxAMatrix	Matrix;
@@ -1038,6 +1037,17 @@ bool UnFbx::FFbxImporter::ImportBone(TArray<FbxNode*>& NodeArray, FSkeletalMeshI
 		JointMatrix.Transform.SetScale3D(Converter.ConvertScale(LocalLinkS));
 	}
 	
+	if(TemplateData)
+	{
+		FbxAMatrix FbxAddedMatrix;
+		BuildFbxMatrixForImportTransform(FbxAddedMatrix, TemplateData);
+		FMatrix AddedMatrix = Converter.ConvertMatrix(FbxAddedMatrix);
+
+		VBone& RootBone = ImportData.RefBonesBinary[RootIdx];
+		FTransform& RootTransform = RootBone.BonePos.Transform;
+		RootTransform.SetFromMatrix(RootTransform.ToMatrixWithScale() * AddedMatrix);
+	}
+	
 	if(bAnyLinksNotInBindPose)
 	{
 		AddTokenizedErrorMessage(FTokenizedMessage::Create(EMessageSeverity::Warning, FText::Format(LOCTEXT("FbxSkeletaLMeshimport_BonesAreMissingFromBindPose", "Warning: The following bones are missing from the bind pose.  This can happen for bones that are not vert weighted.  If they are not in the correct orientation after importing, please set the \"Use T0 as ref pose\" option or add them to the bind pose and reimport the skeletal mesh \n\n{0}"), FText::FromString(LinksWithoutBindPoses))), FFbxErrors::SkeletalMesh_BonesAreMissingFromBindPose);
@@ -1120,7 +1130,7 @@ USkeletalMesh* UnFbx::FFbxImporter::ImportSkeletalMesh(UObject* InParent, TArray
 	bool bUseTime0AsRefPose = ImportOptions->bUseT0AsRefPose;
 	// Note: importing morph data causes additional passes through this function, so disable the warning dialogs
 	// from popping up again on each additional pass.  
-	if ( !ImportBone(NodeArray, *SkelMeshImportDataPtr, SortedLinkArray,  bDiffPose, (FbxShapeArray != NULL), bUseTime0AsRefPose))
+	if ( !ImportBone(NodeArray, *SkelMeshImportDataPtr, TemplateImportData, SortedLinkArray,  bDiffPose, (FbxShapeArray != NULL), bUseTime0AsRefPose))
 	{
 		AddTokenizedErrorMessage(FTokenizedMessage::Create(EMessageSeverity::Error, LOCTEXT("FbxSkeletaLMeshimport_MultipleRootFound", "Multiple roots found")), FFbxErrors::SkeletalMesh_MultipleRoots);
 		// I can't delete object here since this is middle of import
@@ -1130,6 +1140,12 @@ USkeletalMesh* UnFbx::FFbxImporter::ImportSkeletalMesh(UObject* InParent, TArray
 		return NULL;
 	}
 
+ 	FbxNode* RootNode = Scene->GetRootNode();
+ 	if(RootNode && TemplateImportData)
+ 	{
+ 		ApplyTransformSettingsToFbxNode(RootNode, TemplateImportData);
+ 	}
+
 	// Create a list of all unique fbx materials.  This needs to be done as a separate pass before reading geometry
 	// so that we know about all possible materials before assigning material indices to each triangle
 	TArray<FbxSurfaceMaterial*> FbxMaterials;
diff --git a/Engine/Source/Editor/UnrealEd/Private/FbxImporter.h b/Engine/Source/Editor/UnrealEd/Private/FbxImporter.h
--- a/Engine/Source/Editor/UnrealEd/Private/FbxImporter.h
+++ b/Engine/Source/Editor/UnrealEd/Private/FbxImporter.h
@@ -108,6 +108,9 @@ struct FBXImportOptions
 	bool bUsedAsFullName;
 	bool bConvertScene;
 	bool bRemoveNameSpace;
+	FVector ImportTranslation;
+	FRotator ImportRotation;
+	float ImportUniformScale;
 	EFBXNormalImportMethod NormalImportMethod;
 	// Static Mesh options
 	bool bCombineToSingle;
@@ -549,6 +552,32 @@ public:
 
 	/** helper function **/
 	UNREALED_API static void DumpFBXNode(FbxNode* Node);
+
+	/**
+	 * Apply asset import settings for transform to an FBX node
+	 *
+	 * @param Node Node to apply transform settings too
+	 * @param AssetData the asset data object to get transform data from
+	 */
+	void ApplyTransformSettingsToFbxNode(FbxNode* Node, UFbxAssetImportData* AssetData);
+
+	/**
+	 * Remove asset import settings for transform to an FBX node
+	 *
+	 * @param Node Node to apply transform settings too
+	 * @param AssetData the asset data object to get transform data from
+	 */
+	void RemoveTransformSettingsFromFbxNode(FbxNode* Node, UFbxAssetImportData* AssetData);
+
+	/**
+	 * Populate the given matrix with the correct information for the asset data, in
+	 * a format that matches FBX internals or without conversion
+	 *
+	 * @param OutMatrix The matrix to fill
+	 * @param AssetData The asset data to extract the transform info from
+	 */
+	void BuildFbxMatrixForImportTransform(FbxAMatrix& OutMatrix, UFbxAssetImportData* AssetData);
+
 private:
 	/**
 	 * ActorX plug-in can export mesh and dummy as skeleton.
@@ -734,7 +763,7 @@ protected:
 	 * @param bDisableMissingBindPoseWarning
 	 * @param bUseTime0AsRefPose	in/out - Use Time 0 as Ref Pose 
 	 */
-	bool ImportBone(TArray<FbxNode*>& NodeArray, FSkeletalMeshImportData &ImportData, TArray<FbxNode*> &OutSortedLinks, bool& bOutDiffPose, bool bDisableMissingBindPoseWarning, bool & bUseTime0AsRefPose);
+	bool ImportBone(TArray<FbxNode*>& NodeArray, FSkeletalMeshImportData &ImportData, UFbxSkeletalMeshImportData* TemplateData, TArray<FbxNode*> &OutSortedLinks, bool& bOutDiffPose, bool bDisableMissingBindPoseWarning, bool & bUseTime0AsRefPose);
 	
 	/**
 	 * Skins the control points of the given mesh or shape using either the default pose for skinning or the first frame of the
diff --git a/Engine/Source/Editor/UnrealEd/Private/SkeletalMeshEdit.cpp b/Engine/Source/Editor/UnrealEd/Private/SkeletalMeshEdit.cpp
--- a/Engine/Source/Editor/UnrealEd/Private/SkeletalMeshEdit.cpp
+++ b/Engine/Source/Editor/UnrealEd/Private/SkeletalMeshEdit.cpp
@@ -905,7 +905,7 @@ namespace AnimationTransformDebug
 					}
 					
 					// now print information - it doesn't match well, find out what it is
-					UE_LOG(LogFbx, Warning, TEXT("IMPORT TRASNFORM ERROR : Bone (%s:%d) \r\nSource Global Transform (%s), \r\nConverted Global Trasnform (%s)"),
+					UE_LOG(LogFbx, Warning, TEXT("IMPORT TRANSFORM ERROR : Bone (%s:%d) \r\nSource Global Transform (%s), \r\nConverted Global Transform (%s)"),
 						*Data.BoneName.ToString(), Data.BoneIndex, *Data.SourceGlobalTransform[Key].ToString(), *GlobalTransform.ToString());
 				}
 			}
@@ -1102,6 +1102,12 @@ bool UnFbx::FFbxImporter::ImportAnimation(USkeleton* Skeleton, UAnimSequence * D
 
 	const bool bPreserveLocalTransform = FbxImporter->GetImportOptions()->bPreserveLocalTransform;
 
+	// Build additional transform matrix
+	UFbxAnimSequenceImportData* TemplateData = Cast<UFbxAnimSequenceImportData>(DestSeq->AssetImportData);
+	FbxAMatrix FbxAddedMatrix;
+	BuildFbxMatrixForImportTransform(FbxAddedMatrix, TemplateData);
+	FMatrix AddedMatrix = Converter.ConvertMatrix(FbxAddedMatrix);
+
 	int32 TotalNumKeys = 0;
 	for(int32 SourceTrackIdx = 0; SourceTrackIdx < FbxRawBoneNames.Num(); ++SourceTrackIdx)
 	{
@@ -1149,8 +1155,10 @@ bool UnFbx::FFbxImporter::ImportAnimation(USkeleton* Skeleton, UAnimSequence * D
 										"Track {0} did not yeild valid transform. Please report this to animation team."), FText::FromName(BoneName))), FFbxErrors::Animation_TransformError);
 					break;
 				}
-				// debug data
-				NewDebugData.SourceGlobalTransform.Add(GlobalTransform);
+
+				// debug data, including import transformation
+				FTransform AddedTransform(AddedMatrix);
+				NewDebugData.SourceGlobalTransform.Add(GlobalTransform * AddedTransform);
 
 				FTransform LocalTransform;
 				if( !bPreserveLocalTransform && LinkParent)
@@ -1168,12 +1176,21 @@ bool UnFbx::FFbxImporter::ImportAnimation(USkeleton* Skeleton, UAnimSequence * D
 					FbxVector4 NewLocalT = LocalMatrix.GetT();
 					FbxVector4 NewLocalS = LocalMatrix.GetS();
 					FbxQuaternion NewLocalQ = LocalMatrix.GetQ();
+
 					LocalTransform.SetTranslation(Converter.ConvertPos(NewLocalT));
 					LocalTransform.SetScale3D(Converter.ConvertScale(NewLocalS));
 					LocalTransform.SetRotation(Converter.ConvertRotToQuat(NewLocalQ));
+
 					NewDebugData.SourceParentGlobalTransform.Add(FTransform::Identity);
 				}
 
+				if(TemplateData && BoneTreeIndex == 0)
+				{
+					// If we found template data earlier, apply the import transform matrix to
+					// the root track.
+					LocalTransform.SetFromMatrix(LocalTransform.ToMatrixWithScale() * AddedMatrix);
+				}
+
 				if (LocalTransform.ContainsNaN())
 				{
 					bSuccess = false;
