Fixed crash on PIE exit when performing transactions on SelectedActors during PIE.
#ttp 347639 - EDITOR: CRASH: Selecting different actors in the scene outliner during PIE will crash
#branch UE4
#proj Editor.UnrealEd

diff --git a/Engine/Source/Editor/UnrealEd/Private/PlayLevel.cpp b/Engine/Source/Editor/UnrealEd/Private/PlayLevel.cpp
--- a/Engine/Source/Editor/UnrealEd/Private/PlayLevel.cpp
+++ b/Engine/Source/Editor/UnrealEd/Private/PlayLevel.cpp
@@ -56,6 +56,9 @@ void UEditorEngine::EndPlayMap()
 	// Enable screensavers when ending PIE.
 	EnableScreenSaver( true );
 
+	// Move SelectedActors object back to the transient package.
+	GetSelectedActors()->Rename(nullptr, GetTransientPackage(), REN_DoNotDirty | REN_DontCreateRedirectors | REN_NonTransactional);
+
 	// Make a list of all the actors that should be selected
 	TArray<UObject *> SelectedActors;
 	if ( ActorsThatWereSelected.Num() > 0 )
@@ -2366,6 +2369,10 @@ UGameInstance* UEditorEngine::CreatePIEGameInstance(int32 PIEInstance, bool bInS
 		}
 	}
 
+	// Move SelectedActors global object to the PIE package for the duration of the PIE session.
+	// This will stop any transactions on it from being saved during PIE.
+	GetSelectedActors()->Rename(nullptr, GWorld->GetOutermost(), REN_DoNotDirty | REN_DontCreateRedirectors | REN_NonTransactional);
+
 	// For play in editor, this is the viewport widget where the game is being displayed
 	TSharedPtr<SViewport> PieViewportWidget;
 
