Volumes are now not automatically assigned a layer based on the brush builder name upon creation (fixes issue in which placed volumes would disappear when placed in a new level and then reloaded).
#ttp 347867 - JIRA TEST: EDITOR: Reverb Volume becomes hidden after saving and reopening level
#branch UE4
#proj Editor.UnrealEd

diff --git a/Engine/Source/Editor/UnrealEd/Private/EditorBrushBuilder.cpp b/Engine/Source/Editor/UnrealEd/Private/EditorBrushBuilder.cpp
--- a/Engine/Source/Editor/UnrealEd/Private/EditorBrushBuilder.cpp
+++ b/Engine/Source/Editor/UnrealEd/Private/EditorBrushBuilder.cpp
@@ -60,8 +60,6 @@ bool UEditorBrushBuilder::EndBrush( UWorld* InWorld, ABrush* InBrush )
 	Brush->Modify();
 	BuilderBrush->Modify();
 
-	BuilderBrush->Layers.AddUnique( Layer );
-
 	FRotator Temp(0.0f,0.0f,0.0f);
 	FSnappingUtils::SnapToBSPVertex( Location, FVector::ZeroVector, Temp );
 	BuilderBrush->SetActorLocation(Location, false);
