Fixed cook on the fly server to initialize world correctly.

diff --git a/Engine/Source/Editor/UnrealEd/Private/CookOnTheFlyServer.cpp b/Engine/Source/Editor/UnrealEd/Private/CookOnTheFlyServer.cpp
--- a/Engine/Source/Editor/UnrealEd/Private/CookOnTheFlyServer.cpp
+++ b/Engine/Source/Editor/UnrealEd/Private/CookOnTheFlyServer.cpp
@@ -993,8 +993,9 @@ uint32 UCookOnTheFlyServer::TickCookOnTheSide( const float TimeSlice, uint32 &Co
 		
 		if ( PackagesToSave.Num() == 0 )
 		{
-			// an error occured or the file doesn't exist
-			// this happens often as we request files with different extensions when we are searching for files
+			// if we are iterative cooking the package might already be cooked
+			// so just add the package to the cooked packages list
+			// this could also happen if the source file doesn't exist which is often as we request files with different extensions when we are searching for files
 			// just return that we processed the cook request
 			// the network file manager will then handle the missing file and search somewhere else
 			UE_LOG(LogCookOnTheFly, Display, TEXT("Not cooking package %s"), *ToBuild.GetFilename().ToString());
@@ -1294,22 +1295,22 @@ uint32 UCookOnTheFlyServer::TickCookOnTheSide( const float TimeSlice, uint32 &Co
 		{
 		UE_LOG(LogCookCommandlet, Display, TEXT("Full GC..."));
 
-		CollectGarbage( RF_Native );
+			CollectGarbage( RF_Native );
 		NumProcessedSinceLastGC = 0;
 
 		if (bLeakTest)
 		{
-		for (FObjectIterator It; It; ++It)
-		{
+			for (FObjectIterator It; It; ++It)
+			{
 		if (!LastGCItems.Contains(FWeakObjectPtr(*It)))
-		{
+				{
 		UE_LOG(LogCookCommandlet, Warning, TEXT("\tLeaked %s"), *(It->GetFullName()));
 		LastGCItems.Add(FWeakObjectPtr(*It));
-		}
-		}
+				}
+			}
 		}
 		}*/
-
+		
 		if ( Timer.IsTimeUp() )
 		{
 			break;
@@ -1640,6 +1641,11 @@ bool UCookOnTheFlyServer::SaveCookedPackage( UPackage* Package, uint32 SaveFlags
 				if (World)
 				{
 					World->PersistentLevel->OwningWorld = World;
+					if ( !World->bIsWorldInitialized)
+					{
+						// we need to initialize the world - at least need physics scene since BP construction script runs during cooking, otherwise trace won't work
+						World->InitWorld(UWorld::InitializationValues().RequiresHitProxies(false).ShouldSimulatePhysics(false).EnableTraceCollision(false).CreateNavigation(false).AllowAudioPlayback(false).CreatePhysicsScene(true));
+					}
 				}
 
 				const FString FullFilename = FPaths::ConvertRelativePathToFull( PlatFilename );
@@ -1764,7 +1770,7 @@ void UCookOnTheFlyServer::CleanSandbox(const TArray<ITargetPlatform*>& Platforms
 		else
 		{
 			FPackageDependencyInfoModule& PDInfoModule = FModuleManager::LoadModuleChecked<FPackageDependencyInfoModule>("PackageDependencyInfo");
-			
+
 			// list of directories to skip
 			TArray<FString> DirectoriesToSkip;
 			TArray<FString> DirectoriesToNotRecurse;
@@ -1778,7 +1784,7 @@ void UCookOnTheFlyServer::CleanSandbox(const TArray<ITargetPlatform*>& Platforms
 				// use the timestamp grabbing visitor
 				IPlatformFile& PlatformFile = FPlatformFileManager::Get().GetPlatformFile();
 				FLocalTimestampDirectoryVisitor Visitor(PlatformFile, DirectoriesToSkip, DirectoriesToNotRecurse, false);
-			
+
 				PlatformFile.IterateDirectory(*SandboxDirectory, Visitor);
 
 				for (TMap<FString, FDateTime>::TIterator TimestampIt(Visitor.FileTimes); TimestampIt; ++TimestampIt)
@@ -2228,7 +2234,7 @@ void UCookOnTheFlyServer::StartCookByTheBook(const TArray<ITargetPlatform*>& Tar
 	bool bCookAll = false;
 	bool bMapsOnly = false;
 	bool bNoDev = false;
-
+	
 	bool bLeakTest = IsCookFlagSet(ECookInitializationFlags::LeakTest); // this won't work from the editor this needs to be standalone
 	check( !bLeakTest || CurrentCookMode == ECookMode::CookByTheBook );
 
