Changed FLightmassProcessor::BeginRun() to always use Development builds of UnrealLightmass on Mac and Linux, same as it does on Windows.

diff --git a/Engine/Source/Editor/UnrealEd/Private/Lightmass/Lightmass.cpp b/Engine/Source/Editor/UnrealEd/Private/Lightmass/Lightmass.cpp
--- a/Engine/Source/Editor/UnrealEd/Private/Lightmass/Lightmass.cpp
+++ b/Engine/Source/Editor/UnrealEd/Private/Lightmass/Lightmass.cpp
@@ -2460,18 +2460,6 @@ bool FLightmassProcessor::BeginRun()
 		TEXT("../Win64/UnrealLightmass-Json.dll")
 	};
 #elif PLATFORM_MAC
-#if UE_BUILD_DEBUG
-	const TCHAR* LightmassExecutable64 = TEXT("../Mac/UnrealLightmass-Mac-Debug");
-	const TCHAR* RequiredDependencyPaths64[] =
-	{
-		TEXT("../DotNET/Mac/AgentInterface.dll"),
-		TEXT("../Mac/UnrealLightmass-Core-Mac-Debug.dylib"),
-		TEXT("../Mac/UnrealLightmass-CoreUObject-Mac-Debug.dylib"),
-		TEXT("../Mac/UnrealLightmass-Json-Mac-Debug.dylib"),
-		TEXT("../Mac/UnrealLightmass-Projects-Mac-Debug.dylib"),
-		TEXT("../Mac/UnrealLightmass-SwarmInterface-Mac-Debug.dylib")
-	};
-#else
 	const TCHAR* LightmassExecutable64 = TEXT("../Mac/UnrealLightmass");
 	const TCHAR* RequiredDependencyPaths64[] =
 	{
@@ -2482,23 +2470,7 @@ bool FLightmassProcessor::BeginRun()
 		TEXT("../Mac/UnrealLightmass-Projects.dylib"),
 		TEXT("../Mac/UnrealLightmass-SwarmInterface.dylib")
 	};
-#endif
 #elif PLATFORM_LINUX
-#if UE_BUILD_DEBUG
-	const TCHAR* LightmassExecutable64 = TEXT("../Linux/UnrealLightmass-Linux-Debug");
-	const TCHAR* RequiredDependencyPaths64[] =
-	{
-		TEXT("../DotNET/Linux/AgentInterface.dll"),
-		TEXT("../Linux/libUnrealLightmass-Core-Linux-Debug.so"),
-		TEXT("../Linux/libUnrealLightmass-CoreUObject-Linux-Debug.so"),
-		TEXT("../Linux/libUnrealLightmass-Json-Linux-Debug.so"),
-		TEXT("../Linux/libUnrealLightmass-Projects-Linux-Debug.so"),
-		TEXT("../Linux/libUnrealLightmass-SwarmInterface-Linux-Debug.so"),
-		TEXT("../Linux/libUnrealLightmass-Networking-Linux-Debug.so"),
-		TEXT("../Linux/libUnrealLightmass-Messaging-Linux-Debug.so")
-		TEXT("../../Plugins/Messaging/UdpMessaging/Binaries/Linux/libUnrealLightmass-UdpMessaging-Linux-Debug.so")
-	};
-#else
 	const TCHAR* LightmassExecutable64 = TEXT("../Linux/UnrealLightmass");
 	const TCHAR* RequiredDependencyPaths64[] =
 	{
@@ -2512,7 +2484,6 @@ bool FLightmassProcessor::BeginRun()
 		TEXT("../Linux/libUnrealLightmass-Messaging.so"),
 		TEXT("../../Plugins/Messaging/UdpMessaging/Binaries/Linux/libUnrealLightmass-UdpMessaging.so")
 	};
-#endif // UE_BUILD_DEBUG
 #else // PLATFORM_LINUX
 #error "Unknown Lightmass platform"
 #endif
