Fixes issues with duplicating Blueprints, removed UK2Node_Variable::PostDuplicate and moved logic to happen during the Blueprint's PostDuplicate.

#ttp 349451 - CRITICAL: Regression: BP: Some blueprints will give a "Target must have a connection error when duplicated.
#ttp 348325 - Duplicating a blueprint leaves all (self) references pointing at the original

diff --git a/Engine/Source/Editor/BlueprintGraph/Classes/K2Node_Variable.h b/Engine/Source/Editor/BlueprintGraph/Classes/K2Node_Variable.h
--- a/Engine/Source/Editor/BlueprintGraph/Classes/K2Node_Variable.h
+++ b/Engine/Source/Editor/BlueprintGraph/Classes/K2Node_Variable.h
@@ -45,7 +45,6 @@ protected:
 
 public:
 	// Begin UObject interface
-	virtual void PostDuplicate(bool bDuplicateForPIE) override;
 	virtual void Serialize(FArchive& Ar) override;
 	// End UObject interface
 
diff --git a/Engine/Source/Editor/BlueprintGraph/Private/K2Node_Variable.cpp b/Engine/Source/Editor/BlueprintGraph/Private/K2Node_Variable.cpp
--- a/Engine/Source/Editor/BlueprintGraph/Private/K2Node_Variable.cpp
+++ b/Engine/Source/Editor/BlueprintGraph/Private/K2Node_Variable.cpp
@@ -322,34 +322,6 @@ UEdGraphPin* UK2Node_Variable::GetValuePin() const
 	return Pin;
 }
 
-void UK2Node_Variable::PostDuplicate(bool bDuplicateForPIE)
-{
-	Super::PostDuplicate(bDuplicateForPIE);
-	if (!bDuplicateForPIE && (!this->HasAnyFlags(RF_Transient)))
-	{
-		// Self context variable nodes need to be updated with the new Blueprint class
-		if(VariableReference.IsSelfContext())
-		{
-			const UEdGraphSchema_K2* K2Schema = GetDefault<UEdGraphSchema_K2>();
-			if(UEdGraphPin* SelfPin = K2Schema->FindSelfPin(*this, EGPD_Input))
-			{
-				UClass* TargetClass = nullptr;
-
-				if(UProperty* Property = VariableReference.ResolveMember<UProperty>(this))
-				{
-					TargetClass = Property->GetOwnerClass()->GetAuthoritativeClass();
-				}
-				else
-				{
-					TargetClass = GetBlueprint()->SkeletonGeneratedClass->GetAuthoritativeClass();
-				}
-
-				SelfPin->PinType.PinSubCategoryObject = TargetClass;
-			}
-		}
-	}
-}
-
 void UK2Node_Variable::ValidateNodeDuringCompilation(class FCompilerResultsLog& MessageLog) const
 {
 	Super::ValidateNodeDuringCompilation(MessageLog);
diff --git a/Engine/Source/Editor/UnrealEd/Private/Kismet2/BlueprintEditorUtils.cpp b/Engine/Source/Editor/UnrealEd/Private/Kismet2/BlueprintEditorUtils.cpp
--- a/Engine/Source/Editor/UnrealEd/Private/Kismet2/BlueprintEditorUtils.cpp
+++ b/Engine/Source/Editor/UnrealEd/Private/Kismet2/BlueprintEditorUtils.cpp
@@ -1488,6 +1488,32 @@ void FBlueprintEditorUtils::PostDuplicateBlueprint(UBlueprint* Blueprint)
 		for(auto& Node : AllGraphNodes)
 		{
 			Node->CreateNewGuid();
+
+			// Some variable nodes must be fixed up on duplicate, this cannot wait for individual 
+			// node calls to PostDuplicate because it happens after compilation and will still result in errors
+			if(UK2Node_Variable* VariableNode = Cast<UK2Node_Variable>(Node))
+			{
+				// Self context variable nodes need to be updated with the new Blueprint class
+				if(VariableNode->VariableReference.IsSelfContext())
+				{
+					const UEdGraphSchema_K2* K2Schema = GetDefault<UEdGraphSchema_K2>();
+					if(UEdGraphPin* SelfPin = K2Schema->FindSelfPin(*VariableNode, EGPD_Input))
+					{
+						UClass* TargetClass = nullptr;
+
+						if(UProperty* Property = VariableNode->VariableReference.ResolveMember<UProperty>(VariableNode))
+						{
+							TargetClass = Property->GetOwnerClass()->GetAuthoritativeClass();
+						}
+						else
+						{
+							TargetClass = Blueprint->SkeletonGeneratedClass->GetAuthoritativeClass();
+						}
+
+						SelfPin->PinType.PinSubCategoryObject = TargetClass;
+					}
+				}
+			}
 		}
 
 		// And compile again to make sure they go into the generated class, get cleaned up, etc...
