Rename FPostConstructInitializeProperties to something simpler

diff --git a/Engine/Source/Editor/UnrealEd/Private/Commandlets/PackageUtilities.cpp b/Engine/Source/Editor/UnrealEd/Private/Commandlets/PackageUtilities.cpp
--- a/Engine/Source/Editor/UnrealEd/Private/Commandlets/PackageUtilities.cpp
+++ b/Engine/Source/Editor/UnrealEd/Private/Commandlets/PackageUtilities.cpp
@@ -563,8 +563,8 @@ bool FContentHelper::QueryAssetsInCollection(FName InCollectionName, ECollection
 /*-----------------------------------------------------------------------------
 ULoadPackageCommandlet
 -----------------------------------------------------------------------------*/
-ULoadPackageCommandlet::ULoadPackageCommandlet(const class FPostConstructInitializeProperties& PCIP)
-	: Super(PCIP)
+ULoadPackageCommandlet::ULoadPackageCommandlet(const FObjectInitializer& ObjectInitializer)
+	: Super(ObjectInitializer)
 {
 	LogToConsole = false;
 }
@@ -1253,8 +1253,8 @@ void FPkgInfoReporter_Log::GeneratePackageReport( ULinkerLoad* InLinker/*=NULL*/
 	}
 }
 
-UPkgInfoCommandlet::UPkgInfoCommandlet(const class FPostConstructInitializeProperties& PCIP)
-	: Super(PCIP)
+UPkgInfoCommandlet::UPkgInfoCommandlet(const FObjectInitializer& ObjectInitializer)
+	: Super(ObjectInitializer)
 {
 	LogToConsole = false;
 }
@@ -2102,8 +2102,8 @@ struct CompressAnimationsFunctor
 	}
 };
 
-UCompressAnimationsCommandlet::UCompressAnimationsCommandlet(const class FPostConstructInitializeProperties& PCIP)
-	: Super(PCIP)
+UCompressAnimationsCommandlet::UCompressAnimationsCommandlet(const FObjectInitializer& ObjectInitializer)
+	: Super(ObjectInitializer)
 {
 	LogToConsole = false;
 }
@@ -2171,8 +2171,8 @@ int32 UCompressAnimationsCommandlet::Main( const FString& Params )
 //======================================================================
 // UReplaceActorCommandlet
 //======================================================================
-UReplaceActorCommandlet::UReplaceActorCommandlet(const class FPostConstructInitializeProperties& PCIP)
-	: Super(PCIP)
+UReplaceActorCommandlet::UReplaceActorCommandlet(const FObjectInitializer& ObjectInitializer)
+	: Super(ObjectInitializer)
 {
 	LogToConsole = false;
 }
diff --git a/Engine/Source/Editor/UnrealEd/Private/Commandlets/ParticleSystemAuditCommandlet.cpp b/Engine/Source/Editor/UnrealEd/Private/Commandlets/ParticleSystemAuditCommandlet.cpp
--- a/Engine/Source/Editor/UnrealEd/Private/Commandlets/ParticleSystemAuditCommandlet.cpp
+++ b/Engine/Source/Editor/UnrealEd/Private/Commandlets/ParticleSystemAuditCommandlet.cpp
@@ -12,8 +12,8 @@
 
 DEFINE_LOG_CATEGORY_STATIC(LogParticleSystemAuditCommandlet, Log, All);
 
-UParticleSystemAuditCommandlet::UParticleSystemAuditCommandlet(const class FPostConstructInitializeProperties& PCIP)
-	: Super(PCIP)
+UParticleSystemAuditCommandlet::UParticleSystemAuditCommandlet(const FObjectInitializer& ObjectInitializer)
+	: Super(ObjectInitializer)
 {
 	HighSpawnRateOrBurstThreshold = 35.f;
 	FarLODDistanceTheshold = 3000.f;
diff --git a/Engine/Source/Editor/UnrealEd/Private/Commandlets/RepairLocalizationDataCommandlet.cpp b/Engine/Source/Editor/UnrealEd/Private/Commandlets/RepairLocalizationDataCommandlet.cpp
--- a/Engine/Source/Editor/UnrealEd/Private/Commandlets/RepairLocalizationDataCommandlet.cpp
+++ b/Engine/Source/Editor/UnrealEd/Private/Commandlets/RepairLocalizationDataCommandlet.cpp
@@ -132,8 +132,8 @@ void RepairManifestAndArchives(TSharedRef<FInternationalizationManifest> Manifes
 	}
 }
 
-URepairLocalizationDataCommandlet::URepairLocalizationDataCommandlet(const class FPostConstructInitializeProperties& PCIP)
-	: Super(PCIP)
+URepairLocalizationDataCommandlet::URepairLocalizationDataCommandlet(const FObjectInitializer& ObjectInitializer)
+	: Super(ObjectInitializer)
 {
 }
 
diff --git a/Engine/Source/Editor/UnrealEd/Private/Commandlets/UpdateGameProjectCommandlet.cpp b/Engine/Source/Editor/UnrealEd/Private/Commandlets/UpdateGameProjectCommandlet.cpp
--- a/Engine/Source/Editor/UnrealEd/Private/Commandlets/UpdateGameProjectCommandlet.cpp
+++ b/Engine/Source/Editor/UnrealEd/Private/Commandlets/UpdateGameProjectCommandlet.cpp
@@ -7,8 +7,8 @@
 
 DEFINE_LOG_CATEGORY_STATIC(LogUpdateGameProjectCommandlet, Log, All);
 
-UUpdateGameProjectCommandlet::UUpdateGameProjectCommandlet(const class FPostConstructInitializeProperties& PCIP)
-	: Super(PCIP)
+UUpdateGameProjectCommandlet::UUpdateGameProjectCommandlet(const FObjectInitializer& ObjectInitializer)
+	: Super(ObjectInitializer)
 {
 	
 }
diff --git a/Engine/Source/Editor/UnrealEd/Private/CookOnTheFlyServer.cpp b/Engine/Source/Editor/UnrealEd/Private/CookOnTheFlyServer.cpp
--- a/Engine/Source/Editor/UnrealEd/Private/CookOnTheFlyServer.cpp
+++ b/Engine/Source/Editor/UnrealEd/Private/CookOnTheFlyServer.cpp
@@ -395,8 +395,8 @@ void LogCookerMessage( const FString& MessageText, EMessageSeverity::Type Severi
 /* UCookOnTheFlyServer structors
  *****************************************************************************/
 
-UCookOnTheFlyServer::UCookOnTheFlyServer(const class FPostConstructInitializeProperties& PCIP)
-	: Super(PCIP),
+UCookOnTheFlyServer::UCookOnTheFlyServer(const FObjectInitializer& ObjectInitializer)
+	: Super(ObjectInitializer),
 	CurrentCookMode(ECookMode::CookOnTheFly),
 	CookByTheBookOptions(NULL),
 	CookFlags(ECookInitializationFlags::None),
diff --git a/Engine/Source/Editor/UnrealEd/Private/CubemapUnwrapUtils.cpp b/Engine/Source/Editor/UnrealEd/Private/CubemapUnwrapUtils.cpp
--- a/Engine/Source/Editor/UnrealEd/Private/CubemapUnwrapUtils.cpp
+++ b/Engine/Source/Editor/UnrealEd/Private/CubemapUnwrapUtils.cpp
@@ -37,7 +37,7 @@ namespace CubemapHelpers
 		// If the source format is 8 bit per channel or less then select a LDR target format.
 		const EPixelFormat TargetPixelFormat = CalculateImageBytes(1, 1, 0, SourcePixelFormat) <= 4 ? PF_B8G8R8A8 : PF_FloatRGBA;
 
-		UTextureRenderTarget2D* RenderTargetLongLat = new UTextureRenderTarget2D(FPostConstructInitializeProperties());
+		UTextureRenderTarget2D* RenderTargetLongLat = new UTextureRenderTarget2D(FObjectInitializer());
 		check(RenderTargetLongLat);
 		RenderTargetLongLat->AddToRoot();
 		RenderTargetLongLat->ClearColor = FLinearColor(0.0f, 0.0f, 0.0f, 0.0f);
diff --git a/Engine/Source/Editor/UnrealEd/Private/EdGraphNode_Comment.cpp b/Engine/Source/Editor/UnrealEd/Private/EdGraphNode_Comment.cpp
--- a/Engine/Source/Editor/UnrealEd/Private/EdGraphNode_Comment.cpp
+++ b/Engine/Source/Editor/UnrealEd/Private/EdGraphNode_Comment.cpp
@@ -15,8 +15,8 @@
 /////////////////////////////////////////////////////
 // UEdGraphNode_Comment
 
-UEdGraphNode_Comment::UEdGraphNode_Comment(const class FPostConstructInitializeProperties& PCIP)
-	: Super(PCIP)
+UEdGraphNode_Comment::UEdGraphNode_Comment(const FObjectInitializer& ObjectInitializer)
+	: Super(ObjectInitializer)
 {
 	NodeWidth = 400.0f;
 	NodeHeight = 100.0f;
diff --git a/Engine/Source/Editor/UnrealEd/Private/Editor.cpp b/Engine/Source/Editor/UnrealEd/Private/Editor.cpp
--- a/Engine/Source/Editor/UnrealEd/Private/Editor.cpp
+++ b/Engine/Source/Editor/UnrealEd/Private/Editor.cpp
@@ -163,10 +163,10 @@ static inline USelection*& PrivateGetSelectedObjects()
 
 static void PrivateInitSelectedSets()
 {
-	PrivateGetSelectedActors() = new( GetTransientPackage(), TEXT("SelectedActors"), RF_Transactional ) USelection(FPostConstructInitializeProperties());
+	PrivateGetSelectedActors() = new( GetTransientPackage(), TEXT("SelectedActors"), RF_Transactional ) USelection(FObjectInitializer());
 	PrivateGetSelectedActors()->AddToRoot();
 
-	PrivateGetSelectedObjects() = new( GetTransientPackage(), TEXT("SelectedObjects"), RF_Transactional ) USelection(FPostConstructInitializeProperties());
+	PrivateGetSelectedObjects() = new( GetTransientPackage(), TEXT("SelectedObjects"), RF_Transactional ) USelection(FObjectInitializer());
 	PrivateGetSelectedObjects()->AddToRoot();
 }
 
@@ -180,8 +180,8 @@ static void PrivateDestroySelectedSets()
 #endif
 }
 
-UEditorEngine::UEditorEngine(const class FPostConstructInitializeProperties& PCIP)
-: Super(PCIP)
+UEditorEngine::UEditorEngine(const class FObjectInitializer& ObjectInitializer)
+: Super(ObjectInitializer)
 {
 	if (!IsRunningCommandlet())
 	{
@@ -417,8 +417,8 @@ void UEditorEngine::InitEditor(IEngineLoop* InEngineLoop)
 	UNavigationSystem::SetNavigationAutoUpdateEnabled(GetDefault<ULevelEditorMiscSettings>()->bNavigationAutoUpdate, EditorContext.World()->GetNavigationSystem() );
 
 	// Allocate temporary model.
-	TempModel = new UModel( FPostConstructInitializeProperties(),NULL, 1 );
-	ConversionTempModel = new UModel( FPostConstructInitializeProperties(),NULL, 1 );
+	TempModel = new UModel( FObjectInitializer(),NULL, 1 );
+	ConversionTempModel = new UModel( FObjectInitializer(),NULL, 1 );
 
 	// create the timer manager
 	TimerManager = MakeShareable( new FTimerManager() );
@@ -6163,7 +6163,7 @@ TArray<AActor*> UEditorEngine::AddExportTextActors(const FString& ExportText, bo
 	}
 
 	// Use a level factory to spawn all the actors using the ExportText
-	ULevelFactory* Factory = new ULevelFactory(FPostConstructInitializeProperties());
+	ULevelFactory* Factory = new ULevelFactory(FObjectInitializer());
 	FVector Location;
 	{
 		FScopedTransaction Transaction( NSLOCTEXT("UnrealEd", "AddActor", "Add Actor") );
diff --git a/Engine/Source/Editor/UnrealEd/Private/Editor/EditorGameAgnosticSettings.cpp b/Engine/Source/Editor/UnrealEd/Private/Editor/EditorGameAgnosticSettings.cpp
--- a/Engine/Source/Editor/UnrealEd/Private/Editor/EditorGameAgnosticSettings.cpp
+++ b/Engine/Source/Editor/UnrealEd/Private/Editor/EditorGameAgnosticSettings.cpp
@@ -2,8 +2,8 @@
 
 #include "UnrealEd.h"
 
-UEditorGameAgnosticSettings::UEditorGameAgnosticSettings(const class FPostConstructInitializeProperties& PCIP)
-	: Super(PCIP)
+UEditorGameAgnosticSettings::UEditorGameAgnosticSettings(const FObjectInitializer& ObjectInitializer)
+	: Super(ObjectInitializer)
 {
 	bCopyStarterContentPreference = true;
 	AutoScalabilityWorkScaleAmount = 1;
diff --git a/Engine/Source/Editor/UnrealEd/Private/Editor/EditorUserSettings.cpp b/Engine/Source/Editor/UnrealEd/Private/Editor/EditorUserSettings.cpp
--- a/Engine/Source/Editor/UnrealEd/Private/Editor/EditorUserSettings.cpp
+++ b/Engine/Source/Editor/UnrealEd/Private/Editor/EditorUserSettings.cpp
@@ -6,8 +6,8 @@
 
 #define LOCTEXT_NAMESPACE "EditorUserSettings"
 
-UEditorUserSettings::UEditorUserSettings(const class FPostConstructInitializeProperties& PCIP)
-	: Super(PCIP)
+UEditorUserSettings::UEditorUserSettings(const FObjectInitializer& ObjectInitializer)
+	: Super(ObjectInitializer)
 {
 	//Default to high quality
 	MaterialQualityLevel = 1;
diff --git a/Engine/Source/Editor/UnrealEd/Private/EditorActor.cpp b/Engine/Source/Editor/UnrealEd/Private/EditorActor.cpp
--- a/Engine/Source/Editor/UnrealEd/Private/EditorActor.cpp
+++ b/Engine/Source/Editor/UnrealEd/Private/EditorActor.cpp
@@ -85,8 +85,8 @@ public:
 	}
 };
 
-UUnrealEdEngine::UUnrealEdEngine(const class FPostConstructInitializeProperties& PCIP)
-	: Super(PCIP)
+UUnrealEdEngine::UUnrealEdEngine(const FObjectInitializer& ObjectInitializer)
+	: Super(ObjectInitializer)
 {
 }
 
@@ -253,7 +253,7 @@ void UUnrealEdEngine::edactPasteSelected(UWorld* InWorld, bool bDuplicate, bool
 	const TCHAR* Paste = *PasteString;
 
 	// Import the actors.
-	ULevelFactory* Factory = new ULevelFactory(FPostConstructInitializeProperties());
+	ULevelFactory* Factory = new ULevelFactory(FObjectInitializer());
 	Factory->FactoryCreateText( ULevel::StaticClass(), InWorld->GetCurrentLevel(), InWorld->GetCurrentLevel()->GetFName(), RF_Transactional, NULL, bDuplicate ? TEXT("move") : TEXT("paste"), Paste, Paste+FCString::Strlen(Paste), GWarn );
 
 	// Fire ULevel::LevelDirtiedEvent when falling out of scope.
diff --git a/Engine/Source/Editor/UnrealEd/Private/EditorActorFolders.cpp b/Engine/Source/Editor/UnrealEd/Private/EditorActorFolders.cpp
--- a/Engine/Source/Editor/UnrealEd/Private/EditorActorFolders.cpp
+++ b/Engine/Source/Editor/UnrealEd/Private/EditorActorFolders.cpp
@@ -6,8 +6,8 @@
 
 #define LOCTEXT_NAMESPACE "FActorFolders"
 
-UEditorActorFolders::UEditorActorFolders(const FPostConstructInitializeProperties& PCIP)
-	: Super(PCIP)
+UEditorActorFolders::UEditorActorFolders(const FObjectInitializer& ObjectInitializer)
+	: Super(ObjectInitializer)
 { }
 
 void UEditorActorFolders::Serialize(FArchive& Ar)
diff --git a/Engine/Source/Editor/UnrealEd/Private/EditorBrushBuilder.cpp b/Engine/Source/Editor/UnrealEd/Private/EditorBrushBuilder.cpp
--- a/Engine/Source/Editor/UnrealEd/Private/EditorBrushBuilder.cpp
+++ b/Engine/Source/Editor/UnrealEd/Private/EditorBrushBuilder.cpp
@@ -14,8 +14,8 @@
 /*-----------------------------------------------------------------------------
 	UEditorBrushBuilder.
 -----------------------------------------------------------------------------*/
-UEditorBrushBuilder::UEditorBrushBuilder(const class FPostConstructInitializeProperties& PCIP)
-	: Super(PCIP)
+UEditorBrushBuilder::UEditorBrushBuilder(const FObjectInitializer& ObjectInitializer)
+	: Super(ObjectInitializer)
 {
 	BitmapFilename = TEXT("BBGeneric");
 	ToolTip = TEXT("BrushBuilderName_Generic");
@@ -246,8 +246,8 @@ void UEditorBrushBuilder::PostEditChangeProperty(struct FPropertyChangedEvent& P
 	}
 }
 
-UConeBuilder::UConeBuilder(const class FPostConstructInitializeProperties& PCIP)
-	: Super(PCIP)
+UConeBuilder::UConeBuilder(const FObjectInitializer& ObjectInitializer)
+	: Super(ObjectInitializer)
 {
 	// Structure to hold one-time initialization
 	struct FConstructorStatics
@@ -324,8 +324,8 @@ bool UConeBuilder::Build( UWorld* InWorld, ABrush* InBrush )
 	return EndBrush( InWorld, InBrush );
 }
 
-UCubeBuilder::UCubeBuilder(const class FPostConstructInitializeProperties& PCIP)
-	: Super(PCIP)
+UCubeBuilder::UCubeBuilder(const FObjectInitializer& ObjectInitializer)
+	: Super(ObjectInitializer)
 {
 	// Structure to hold one-time initialization
 	struct FConstructorStatics
@@ -401,8 +401,8 @@ bool UCubeBuilder::Build( UWorld* InWorld, ABrush* InBrush )
 	return EndBrush( InWorld, InBrush );
 }
 
-UCurvedStairBuilder::UCurvedStairBuilder(const class FPostConstructInitializeProperties& PCIP)
-	: Super(PCIP)
+UCurvedStairBuilder::UCurvedStairBuilder(const FObjectInitializer& ObjectInitializer)
+	: Super(ObjectInitializer)
 {
 	// Structure to hold one-time initialization
 	struct FConstructorStatics
@@ -520,8 +520,8 @@ bool UCurvedStairBuilder::Build( UWorld* InWorld, ABrush* InBrush )
 	return EndBrush( InWorld, InBrush );
 }
 
-UCylinderBuilder::UCylinderBuilder(const class FPostConstructInitializeProperties& PCIP)
-	: Super(PCIP)
+UCylinderBuilder::UCylinderBuilder(const FObjectInitializer& ObjectInitializer)
+	: Super(ObjectInitializer)
 {
 	// Structure to hold one-time initialization
 	struct FConstructorStatics
@@ -596,8 +596,8 @@ bool UCylinderBuilder::Build( UWorld* InWorld, ABrush* InBrush )
 	return EndBrush( InWorld, InBrush );
 }
 
-ULinearStairBuilder::ULinearStairBuilder(const class FPostConstructInitializeProperties& PCIP)
-	: Super(PCIP)
+ULinearStairBuilder::ULinearStairBuilder(const FObjectInitializer& ObjectInitializer)
+	: Super(ObjectInitializer)
 {
 	// Structure to hold one-time initialization
 	struct FConstructorStatics
@@ -723,8 +723,8 @@ bool ULinearStairBuilder::Build( UWorld* InWorld, ABrush* InBrush )
 	return EndBrush( InWorld, InBrush );
 }
 
-USheetBuilder::USheetBuilder(const class FPostConstructInitializeProperties& PCIP)
-	: Super(PCIP)
+USheetBuilder::USheetBuilder(const FObjectInitializer& ObjectInitializer)
+	: Super(ObjectInitializer)
 {
 	// Structure to hold one-time initialization
 	struct FConstructorStatics
@@ -791,8 +791,8 @@ bool USheetBuilder::Build( UWorld* InWorld, ABrush* InBrush )
 	return EndBrush( InWorld, InBrush );
 }
 
-USpiralStairBuilder::USpiralStairBuilder(const class FPostConstructInitializeProperties& PCIP)
-	: Super(PCIP)
+USpiralStairBuilder::USpiralStairBuilder(const FObjectInitializer& ObjectInitializer)
+	: Super(ObjectInitializer)
 {
 	// Structure to hold one-time initialization
 	struct FConstructorStatics
@@ -914,8 +914,8 @@ bool USpiralStairBuilder::Build( UWorld* InWorld, ABrush* InBrush )
 	return EndBrush( InWorld, InBrush );
 }
 
-UTetrahedronBuilder::UTetrahedronBuilder(const class FPostConstructInitializeProperties& PCIP)
-	: Super(PCIP)
+UTetrahedronBuilder::UTetrahedronBuilder(const FObjectInitializer& ObjectInitializer)
+	: Super(ObjectInitializer)
 {
 	// Structure to hold one-time initialization
 	struct FConstructorStatics
@@ -983,8 +983,8 @@ bool UTetrahedronBuilder::Build( UWorld* InWorld, ABrush* InBrush )
 }
 
 
-UVolumetricBuilder::UVolumetricBuilder(const class FPostConstructInitializeProperties& PCIP)
-	: Super(PCIP)
+UVolumetricBuilder::UVolumetricBuilder(const FObjectInitializer& ObjectInitializer)
+	: Super(ObjectInitializer)
 {
 	// Structure to hold one-time initialization
 	struct FConstructorStatics
diff --git a/Engine/Source/Editor/UnrealEd/Private/EditorExporters.cpp b/Engine/Source/Editor/UnrealEd/Private/EditorExporters.cpp
--- a/Engine/Source/Editor/UnrealEd/Private/EditorExporters.cpp
+++ b/Engine/Source/Editor/UnrealEd/Private/EditorExporters.cpp
@@ -32,8 +32,8 @@ DEFINE_LOG_CATEGORY_STATIC(LogEditorExporters, Log, All);
 /*------------------------------------------------------------------------------
 	UTextBufferExporterTXT implementation.
 ------------------------------------------------------------------------------*/
-UTextBufferExporterTXT::UTextBufferExporterTXT(const class FPostConstructInitializeProperties& PCIP)
-	: Super(PCIP)
+UTextBufferExporterTXT::UTextBufferExporterTXT(const FObjectInitializer& ObjectInitializer)
+	: Super(ObjectInitializer)
 {
 	SupportedClass = UTextBuffer::StaticClass();
 	FormatExtension.Add(TEXT("TXT"));
@@ -63,8 +63,8 @@ bool UTextBufferExporterTXT::ExportText(const FExportObjectInnerContext* Context
 /*------------------------------------------------------------------------------
 	USoundExporterWAV implementation.
 ------------------------------------------------------------------------------*/
-USoundExporterWAV::USoundExporterWAV(const class FPostConstructInitializeProperties& PCIP)
-	: Super(PCIP)
+USoundExporterWAV::USoundExporterWAV(const FObjectInitializer& ObjectInitializer)
+	: Super(ObjectInitializer)
 {
 	SupportedClass = USoundWave::StaticClass();
 	bText = false;
@@ -96,8 +96,8 @@ bool USoundExporterWAV::ExportBinary( UObject* Object, const TCHAR* Type, FArchi
 /*------------------------------------------------------------------------------
 	USoundExporterOGG implementation.
 ------------------------------------------------------------------------------*/
-USoundExporterOGG::USoundExporterOGG(const class FPostConstructInitializeProperties& PCIP)
-	: Super(PCIP)
+USoundExporterOGG::USoundExporterOGG(const FObjectInitializer& ObjectInitializer)
+	: Super(ObjectInitializer)
 {
 
 	SupportedClass = USoundWave::StaticClass();
@@ -138,8 +138,8 @@ bool USoundExporterOGG::ExportBinary( UObject* Object, const TCHAR* Type, FArchi
 /*------------------------------------------------------------------------------
 	USoundSurroundExporterWAV implementation.
 ------------------------------------------------------------------------------*/
-USoundSurroundExporterWAV::USoundSurroundExporterWAV(const class FPostConstructInitializeProperties& PCIP)
-	: Super(PCIP)
+USoundSurroundExporterWAV::USoundSurroundExporterWAV(const FObjectInitializer& ObjectInitializer)
+	: Super(ObjectInitializer)
 {
 	
 	SupportedClass = USoundWave::StaticClass();
@@ -209,8 +209,8 @@ bool USoundSurroundExporterWAV::ExportBinary( UObject* Object, const TCHAR* Type
 /*------------------------------------------------------------------------------
 	UObjectExporterT3D implementation.
 ------------------------------------------------------------------------------*/
-UObjectExporterT3D::UObjectExporterT3D(const class FPostConstructInitializeProperties& PCIP)
-	: Super(PCIP)
+UObjectExporterT3D::UObjectExporterT3D(const FObjectInitializer& ObjectInitializer)
+	: Super(ObjectInitializer)
 {
 	SupportedClass = UObject::StaticClass();
 	bText = true;
@@ -233,8 +233,8 @@ bool UObjectExporterT3D::ExportText( const FExportObjectInnerContext* Context, U
 /*------------------------------------------------------------------------------
 	UPolysExporterT3D implementation.
 ------------------------------------------------------------------------------*/
-UPolysExporterT3D::UPolysExporterT3D(const class FPostConstructInitializeProperties& PCIP)
-	: Super(PCIP)
+UPolysExporterT3D::UPolysExporterT3D(const FObjectInitializer& ObjectInitializer)
+	: Super(ObjectInitializer)
 {
 
 	SupportedClass = UPolys::StaticClass();
@@ -299,8 +299,8 @@ bool UPolysExporterT3D::ExportText( const FExportObjectInnerContext* Context, UO
 /*------------------------------------------------------------------------------
 	UModelExporterT3D implementation.
 ------------------------------------------------------------------------------*/
-UModelExporterT3D::UModelExporterT3D(const class FPostConstructInitializeProperties& PCIP)
-	: Super(PCIP)
+UModelExporterT3D::UModelExporterT3D(const FObjectInitializer& ObjectInitializer)
+	: Super(ObjectInitializer)
 {
 	SupportedClass = UModel::StaticClass();
 	bText = true;
@@ -357,8 +357,8 @@ void ExporterHelper_DumpPackageInners(const FExportObjectInnerContext* Context,
 	}
 }
 
-ULevelExporterT3D::ULevelExporterT3D(const class FPostConstructInitializeProperties& PCIP)
-	: Super(PCIP)
+ULevelExporterT3D::ULevelExporterT3D(const FObjectInitializer& ObjectInitializer)
+	: Super(ObjectInitializer)
 {
 	SupportedClass = UWorld::StaticClass();
 	bText = true;
@@ -537,8 +537,8 @@ void ULevelExporterT3D::ExportPackageInners(FExportPackageParams& ExpPackagePara
 /*------------------------------------------------------------------------------
 	ULevelExporterSTL implementation.
 ------------------------------------------------------------------------------*/
-ULevelExporterSTL::ULevelExporterSTL(const class FPostConstructInitializeProperties& PCIP)
-	: Super(PCIP)
+ULevelExporterSTL::ULevelExporterSTL(const FObjectInitializer& ObjectInitializer)
+	: Super(ObjectInitializer)
 {
 	SupportedClass = UWorld::StaticClass();
 	bText = true;
@@ -1310,7 +1310,7 @@ bool GenerateExportMaterialPropertyData(
 
 		if(!RenderTarget)
 		{
-			RenderTarget = new UTextureRenderTarget2D(FPostConstructInitializeProperties());
+			RenderTarget = new UTextureRenderTarget2D(FObjectInitializer());
 			check(RenderTarget);
 			RenderTarget->AddToRoot();
 			RenderTarget->ClearColor = FLinearColor(0.0f, 0.0f, 0.0f, 0.0f);
@@ -1838,8 +1838,8 @@ struct FCompareMaterial
 /*------------------------------------------------------------------------------
 	ULevelExporterLOD implementation.
 ------------------------------------------------------------------------------*/
-ULevelExporterLOD::ULevelExporterLOD(const class FPostConstructInitializeProperties& PCIP)
-	: Super(PCIP)
+ULevelExporterLOD::ULevelExporterLOD(const FObjectInitializer& ObjectInitializer)
+	: Super(ObjectInitializer)
 {
 
 	SupportedClass = UWorld::StaticClass();
@@ -1986,8 +1986,8 @@ static void ExportPolys( UPolys* Polys, int32 &PolyNum, int32 TotalPolys, FFeedb
 	}
 }
 
-ULevelExporterOBJ::ULevelExporterOBJ(const class FPostConstructInitializeProperties& PCIP)
-	: Super(PCIP)
+ULevelExporterOBJ::ULevelExporterOBJ(const FObjectInitializer& ObjectInitializer)
+	: Super(ObjectInitializer)
 {
 	SupportedClass = UWorld::StaticClass();
 	bText = true;
@@ -2092,8 +2092,8 @@ bool ULevelExporterOBJ::ExportText(const FExportObjectInnerContext* Context, UOb
 /*------------------------------------------------------------------------------
 	ULevelExporterFBX implementation.
 ------------------------------------------------------------------------------*/
-ULevelExporterFBX::ULevelExporterFBX(const class FPostConstructInitializeProperties& PCIP)
-	: Super(PCIP)
+ULevelExporterFBX::ULevelExporterFBX(const FObjectInitializer& ObjectInitializer)
+	: Super(ObjectInitializer)
 {
 
 	SupportedClass = UWorld::StaticClass();
@@ -2144,8 +2144,8 @@ bool ULevelExporterFBX::ExportBinary( UObject* Object, const TCHAR* Type, FArchi
 /*------------------------------------------------------------------------------
 	UPolysExporterOBJ implementation.
 ------------------------------------------------------------------------------*/
-UPolysExporterOBJ::UPolysExporterOBJ(const class FPostConstructInitializeProperties& PCIP)
-	: Super(PCIP)
+UPolysExporterOBJ::UPolysExporterOBJ(const FObjectInitializer& ObjectInitializer)
+	: Super(ObjectInitializer)
 {
 
 	SupportedClass = UPolys::StaticClass();
@@ -2188,8 +2188,8 @@ bool UPolysExporterOBJ::ExportText(const FExportObjectInnerContext* Context, UOb
 /*------------------------------------------------------------------------------
 	USequenceExporterT3D implementation.
 ------------------------------------------------------------------------------*/
-USequenceExporterT3D::USequenceExporterT3D(const class FPostConstructInitializeProperties& PCIP)
-	: Super(PCIP)
+USequenceExporterT3D::USequenceExporterT3D(const FObjectInitializer& ObjectInitializer)
+	: Super(ObjectInitializer)
 {
 }
 
@@ -2201,8 +2201,8 @@ bool USequenceExporterT3D::ExportText(const FExportObjectInnerContext* Context,
 /*------------------------------------------------------------------------------
 UStaticMeshExporterOBJ implementation.
 ------------------------------------------------------------------------------*/
-UStaticMeshExporterOBJ::UStaticMeshExporterOBJ(const class FPostConstructInitializeProperties& PCIP)
-	: Super(PCIP)
+UStaticMeshExporterOBJ::UStaticMeshExporterOBJ(const FObjectInitializer& ObjectInitializer)
+	: Super(ObjectInitializer)
 {
 
 	SupportedClass = UStaticMesh::StaticClass();
@@ -2430,8 +2430,8 @@ bool UStaticMeshExporterOBJ::ExportText(const FExportObjectInnerContext* Context
 /*------------------------------------------------------------------------------
 UStaticMeshExporterFBX implementation.
 ------------------------------------------------------------------------------*/
-UStaticMeshExporterFBX::UStaticMeshExporterFBX(const class FPostConstructInitializeProperties& PCIP)
-	: Super(PCIP)
+UStaticMeshExporterFBX::UStaticMeshExporterFBX(const FObjectInitializer& ObjectInitializer)
+	: Super(ObjectInitializer)
 {
 
 
@@ -2458,8 +2458,8 @@ bool UStaticMeshExporterFBX::ExportBinary( UObject* Object, const TCHAR* Type, F
 /*------------------------------------------------------------------------------
 USkeletalMeshExporterFBX implementation.
 ------------------------------------------------------------------------------*/
-USkeletalMeshExporterFBX::USkeletalMeshExporterFBX(const class FPostConstructInitializeProperties& PCIP)
-	: Super(PCIP)
+USkeletalMeshExporterFBX::USkeletalMeshExporterFBX(const FObjectInitializer& ObjectInitializer)
+	: Super(ObjectInitializer)
 {
 
 	SupportedClass = USkeletalMesh::StaticClass();
@@ -2483,8 +2483,8 @@ bool USkeletalMeshExporterFBX::ExportBinary( UObject* Object, const TCHAR* Type,
 /*------------------------------------------------------------------------------
 UAnimSequenceExporterFBX implementation.
 ------------------------------------------------------------------------------*/
-UAnimSequenceExporterFBX::UAnimSequenceExporterFBX(const class FPostConstructInitializeProperties& PCIP)
-	: Super(PCIP)
+UAnimSequenceExporterFBX::UAnimSequenceExporterFBX(const FObjectInitializer& ObjectInitializer)
+	: Super(ObjectInitializer)
 {
 	SupportedClass = UAnimSequence::StaticClass();
 	bText = false;
@@ -2576,8 +2576,8 @@ void UEditorEngine::RebuildStaticNavigableGeometry(ULevel* Level)
 /*-----------------------------------------------------------------------------
 	UExportTextContainer
 -----------------------------------------------------------------------------*/
-UExportTextContainer::UExportTextContainer(const class FPostConstructInitializeProperties& PCIP)
-	: Super(PCIP)
+UExportTextContainer::UExportTextContainer(const FObjectInitializer& ObjectInitializer)
+	: Super(ObjectInitializer)
 {
 }
 
@@ -2640,7 +2640,7 @@ namespace MaterialExportUtils
 			OutFlattenMaterial.DiffuseSize.Y > 0)
 		{
 			// Create temporary render target
-			UTextureRenderTarget2D* RenderTargetDiffuse = new UTextureRenderTarget2D(FPostConstructInitializeProperties());
+			UTextureRenderTarget2D* RenderTargetDiffuse = new UTextureRenderTarget2D(FObjectInitializer());
 			check(RenderTargetDiffuse);
 			RenderTargetDiffuse->AddToRoot();
 			RenderTargetDiffuse->ClearColor = FLinearColor(0.0f, 0.0f, 0.0f, 0.0f);
@@ -2673,7 +2673,7 @@ namespace MaterialExportUtils
 			InMaterial->GetMaterial()->HasNormalConnected())
 		{
 			// Create temporary render target
-			UTextureRenderTarget2D* RenderTargetNormal = new UTextureRenderTarget2D(FPostConstructInitializeProperties());
+			UTextureRenderTarget2D* RenderTargetNormal = new UTextureRenderTarget2D(FObjectInitializer());
 			check(RenderTargetNormal);
 			RenderTargetNormal->AddToRoot();
 			RenderTargetNormal->ClearColor = FLinearColor(0.0f, 0.0f, 0.0f, 0.0f);
@@ -2735,7 +2735,7 @@ namespace MaterialExportUtils
 			if (OutFlattenMaterial.DiffuseSize.X > 0 && 
 				OutFlattenMaterial.DiffuseSize.Y > 0)
 			{
-				UTextureRenderTarget2D* RenderTargetTexture = new UTextureRenderTarget2D(FPostConstructInitializeProperties());
+				UTextureRenderTarget2D* RenderTargetTexture = new UTextureRenderTarget2D(FObjectInitializer());
 				check(RenderTargetTexture);
 				RenderTargetTexture->AddToRoot();
 				RenderTargetTexture->ClearColor = FLinearColor(0.0f, 0.0f, 0.0f, 0.0f);
@@ -2816,7 +2816,7 @@ namespace MaterialExportUtils
 			if (OutFlattenMaterial.NormalSize.X > 0 && 
 				OutFlattenMaterial.NormalSize.Y > 0)
 			{
-				UTextureRenderTarget2D* RenderTargetTexture = new UTextureRenderTarget2D(FPostConstructInitializeProperties());
+				UTextureRenderTarget2D* RenderTargetTexture = new UTextureRenderTarget2D(FObjectInitializer());
 				check(RenderTargetTexture);
 				RenderTargetTexture->AddToRoot();
 				RenderTargetTexture->ClearColor = FLinearColor(0.0f, 0.0f, 0.0f, 0.0f);
diff --git a/Engine/Source/Editor/UnrealEd/Private/EditorObject.cpp b/Engine/Source/Editor/UnrealEd/Private/EditorObject.cpp
--- a/Engine/Source/Editor/UnrealEd/Private/EditorObject.cpp
+++ b/Engine/Source/Editor/UnrealEd/Private/EditorObject.cpp
@@ -228,7 +228,7 @@ static const TCHAR* ImportProperties(
 					ExistingBrush->Rename();
 
 				// Create model.
-				UModelFactory* ModelFactory = new UModelFactory(FPostConstructInitializeProperties());
+				UModelFactory* ModelFactory = new UModelFactory(FObjectInitializer());
 				ModelFactory->FactoryCreateText( UModel::StaticClass(), SubobjectRoot, FName(BrushName, FNAME_Add, true), RF_NoFlags, NULL, TEXT("t3d"), SourceText, SourceText+FCString::Strlen(SourceText), Warn );
 				ImportedBrush = 1;
 			}
diff --git a/Engine/Source/Editor/UnrealEd/Private/EditorParameters.cpp b/Engine/Source/Editor/UnrealEd/Private/EditorParameters.cpp
--- a/Engine/Source/Editor/UnrealEd/Private/EditorParameters.cpp
+++ b/Engine/Source/Editor/UnrealEd/Private/EditorParameters.cpp
@@ -2,37 +2,37 @@
 
 #include "UnrealEd.h"
 
-UDEditorParameterValue::UDEditorParameterValue(const class FPostConstructInitializeProperties& PCIP)
-	: Super(PCIP)
+UDEditorParameterValue::UDEditorParameterValue(const FObjectInitializer& ObjectInitializer)
+	: Super(ObjectInitializer)
 {
 }
 
-UDEditorFontParameterValue::UDEditorFontParameterValue(const class FPostConstructInitializeProperties& PCIP)
-	: Super(PCIP)
+UDEditorFontParameterValue::UDEditorFontParameterValue(const FObjectInitializer& ObjectInitializer)
+	: Super(ObjectInitializer)
 {
 }
 
-UDEditorScalarParameterValue::UDEditorScalarParameterValue(const class FPostConstructInitializeProperties& PCIP)
-	: Super(PCIP)
+UDEditorScalarParameterValue::UDEditorScalarParameterValue(const FObjectInitializer& ObjectInitializer)
+	: Super(ObjectInitializer)
 {
 }
 
-UDEditorStaticComponentMaskParameterValue::UDEditorStaticComponentMaskParameterValue(const class FPostConstructInitializeProperties& PCIP)
-	: Super(PCIP)
+UDEditorStaticComponentMaskParameterValue::UDEditorStaticComponentMaskParameterValue(const FObjectInitializer& ObjectInitializer)
+	: Super(ObjectInitializer)
 {
 }
 
-UDEditorStaticSwitchParameterValue::UDEditorStaticSwitchParameterValue(const class FPostConstructInitializeProperties& PCIP)
-	: Super(PCIP)
+UDEditorStaticSwitchParameterValue::UDEditorStaticSwitchParameterValue(const FObjectInitializer& ObjectInitializer)
+	: Super(ObjectInitializer)
 {
 }
 
-UDEditorTextureParameterValue::UDEditorTextureParameterValue(const class FPostConstructInitializeProperties& PCIP)
-	: Super(PCIP)
+UDEditorTextureParameterValue::UDEditorTextureParameterValue(const FObjectInitializer& ObjectInitializer)
+	: Super(ObjectInitializer)
 {
 }
 
-UDEditorVectorParameterValue::UDEditorVectorParameterValue(const class FPostConstructInitializeProperties& PCIP)
-	: Super(PCIP)
+UDEditorVectorParameterValue::UDEditorVectorParameterValue(const FObjectInitializer& ObjectInitializer)
+	: Super(ObjectInitializer)
 {
 }
diff --git a/Engine/Source/Editor/UnrealEd/Private/EditorServer.cpp b/Engine/Source/Editor/UnrealEd/Private/EditorServer.cpp
--- a/Engine/Source/Editor/UnrealEd/Private/EditorServer.cpp
+++ b/Engine/Source/Editor/UnrealEd/Private/EditorServer.cpp
@@ -560,7 +560,7 @@ bool UEditorEngine::Exec_Brush( UWorld* InWorld, const TCHAR* Str, FOutputDevice
 			WorldBrush->SetActorLocation(SnapLocation - PrePivot, false);
 			WorldBrush->SetPrePivot( FVector::ZeroVector );
 			WorldBrush->Brush->Polys->Element.Empty();
-			UPolysFactory* It = new UPolysFactory(FPostConstructInitializeProperties());
+			UPolysFactory* It = new UPolysFactory(FObjectInitializer());
 			It->FactoryCreateText( UPolys::StaticClass(), WorldBrush->Brush->Polys->GetOuter(), *WorldBrush->Brush->Polys->GetName(), RF_NoFlags, WorldBrush->Brush->Polys, TEXT("t3d"), GStream, GStream+FCString::Strlen(GStream), GWarn );
 			// Do NOT merge faces.
 			FBSPOps::bspValidateBrush( WorldBrush->Brush, 0, 1 );
@@ -1079,7 +1079,7 @@ UTransactor* UEditorEngine::CreateTrans()
 		UndoBufferSize = 16;
 	}
 
-	UTransBuffer* TransBuffer = new UTransBuffer(FPostConstructInitializeProperties(), UndoBufferSize * 1024 * 1024);
+	UTransBuffer* TransBuffer = new UTransBuffer(FObjectInitializer(), UndoBufferSize * 1024 * 1024);
 	TransBuffer->OnBeforeRedoUndo().AddUObject(this, &UEditorEngine::HandleTransactorBeforeRedoUndo);
 	TransBuffer->OnRedo().AddUObject(this, &UEditorEngine::HandleTransactorRedo);
 	TransBuffer->OnUndo().AddUObject(this, &UEditorEngine::HandleTransactorUndo);
@@ -2730,11 +2730,11 @@ void UEditorEngine::DoMoveSelectedActorsToLevel( ULevel* InDestLevel )
 
 ULevel*  UEditorEngine::CreateTransLevelMoveBuffer( UWorld* InWorld )
 {
-	ULevel* BufferLevel = new(GetTransientPackage(), TEXT("TransLevelMoveBuffer")) ULevel(FPostConstructInitializeProperties(),FURL(NULL));
+	ULevel* BufferLevel = new(GetTransientPackage(), TEXT("TransLevelMoveBuffer")) ULevel(FObjectInitializer(),FURL(NULL));
 	check( BufferLevel );
 	BufferLevel->AddToRoot();
 	BufferLevel->OwningWorld = InWorld;
-	BufferLevel->Model = new( BufferLevel ) UModel( FPostConstructInitializeProperties(), NULL, true );
+	BufferLevel->Model = new( BufferLevel ) UModel( FObjectInitializer(), NULL, true );
 		
 	BufferLevel->SetFlags( RF_Transactional );
 	BufferLevel->Model->SetFlags( RF_Transactional );
@@ -2752,7 +2752,7 @@ ULevel*  UEditorEngine::CreateTransLevelMoveBuffer( UWorld* InWorld )
 	BufferDefaultBrush->Brush = CastChecked<UModel>(StaticFindObject(UModel::StaticClass(), BufferLevel->OwningWorld->GetOuter(), TEXT("Brush"), true), ECastCheckedType::NullAllowed);
 	if (!BufferDefaultBrush->Brush)
 	{
-		BufferDefaultBrush->Brush = new( InWorld, TEXT("Brush") )UModel( FPostConstructInitializeProperties(), BufferDefaultBrush, 1 );
+		BufferDefaultBrush->Brush = new( InWorld, TEXT("Brush") )UModel( FObjectInitializer(), BufferDefaultBrush, 1 );
 	}
 	BufferDefaultBrush->BrushComponent->Brush = BufferDefaultBrush->Brush;
 	BufferDefaultBrush->SetNotForClientOrServer();
diff --git a/Engine/Source/Editor/UnrealEd/Private/EditorTransaction.cpp b/Engine/Source/Editor/UnrealEd/Private/EditorTransaction.cpp
--- a/Engine/Source/Editor/UnrealEd/Private/EditorTransaction.cpp
+++ b/Engine/Source/Editor/UnrealEd/Private/EditorTransaction.cpp
@@ -331,18 +331,18 @@ void FTransaction::GetTransactionObjects(TArray<UObject*>& Objects)
 /*-----------------------------------------------------------------------------
 	Transaction tracking system.
 -----------------------------------------------------------------------------*/
-UTransactor::UTransactor(const class FPostConstructInitializeProperties& PCIP)
-	: Super(PCIP)
+UTransactor::UTransactor(const FObjectInitializer& ObjectInitializer)
+	: Super(ObjectInitializer)
 {
 }
 
-UTransBuffer::UTransBuffer(const class FPostConstructInitializeProperties& PCIP)
-	: Super(PCIP)
+UTransBuffer::UTransBuffer(const FObjectInitializer& ObjectInitializer)
+	: Super(ObjectInitializer)
 {
 }
 
-UTransBuffer::UTransBuffer( const class FPostConstructInitializeProperties& PCIP, SIZE_T InMaxMemory )
-:	UTransactor(PCIP)
+UTransBuffer::UTransBuffer( const FObjectInitializer& ObjectInitializer, SIZE_T InMaxMemory )
+:	UTransactor(ObjectInitializer)
 ,	MaxMemory( InMaxMemory )
 {
 	// Reset.
diff --git a/Engine/Source/Editor/UnrealEd/Private/Factories/ActorFactory.cpp b/Engine/Source/Editor/UnrealEd/Private/Factories/ActorFactory.cpp
--- a/Engine/Source/Editor/UnrealEd/Private/Factories/ActorFactory.cpp
+++ b/Engine/Source/Editor/UnrealEd/Private/Factories/ActorFactory.cpp
@@ -99,8 +99,8 @@ FQuat FindActorAlignmentRotation(const FQuat& InActorRotation, const FVector& In
 /*-----------------------------------------------------------------------------
 UActorFactory
 -----------------------------------------------------------------------------*/
-UActorFactory::UActorFactory(const class FPostConstructInitializeProperties& PCIP)
-	: Super(PCIP)
+UActorFactory::UActorFactory(const FObjectInitializer& ObjectInitializer)
+	: Super(ObjectInitializer)
 {
 	DisplayName = LOCTEXT("DefaultName","Actor");
 	bShowInEditorQuickMenu = false;
@@ -226,8 +226,8 @@ void UActorFactory::PostCreateBlueprint( UObject* Asset, AActor* CDO )
 /*-----------------------------------------------------------------------------
 UActorFactoryStaticMesh
 -----------------------------------------------------------------------------*/
-UActorFactoryStaticMesh::UActorFactoryStaticMesh(const class FPostConstructInitializeProperties& PCIP)
-	: Super(PCIP)
+UActorFactoryStaticMesh::UActorFactoryStaticMesh(const FObjectInitializer& ObjectInitializer)
+	: Super(ObjectInitializer)
 {
 	DisplayName = LOCTEXT("StaticMeshDisplayName", "Static Mesh");
 	NewActorClass = AStaticMeshActor::StaticClass();
@@ -297,8 +297,8 @@ FQuat UActorFactoryStaticMesh::AlignObjectToSurfaceNormal(const FVector& InSurfa
 /*-----------------------------------------------------------------------------
 UActorFactoryDeferredDecal
 -----------------------------------------------------------------------------*/
-UActorFactoryDeferredDecal::UActorFactoryDeferredDecal(const class FPostConstructInitializeProperties& PCIP)
-	: Super(PCIP)
+UActorFactoryDeferredDecal::UActorFactoryDeferredDecal(const FObjectInitializer& ObjectInitializer)
+	: Super(ObjectInitializer)
 { 
 	DisplayName = LOCTEXT("DeferredDecalDisplayName", "Deferred Decal");
 	NewActorClass = ADecalActor::StaticClass();
@@ -435,8 +435,8 @@ TargetMaterial :
 /*-----------------------------------------------------------------------------
 UActorFactoryTextRender
 -----------------------------------------------------------------------------*/
-UActorFactoryTextRender::UActorFactoryTextRender(const class FPostConstructInitializeProperties& PCIP)
-	: Super(PCIP)
+UActorFactoryTextRender::UActorFactoryTextRender(const FObjectInitializer& ObjectInitializer)
+	: Super(ObjectInitializer)
 {
 	// Property initialization
 	DisplayName = LOCTEXT("TextRenderDisplayName", "Text Render");
@@ -448,8 +448,8 @@ UActorFactoryTextRender::UActorFactoryTextRender(const class FPostConstructIniti
 /*-----------------------------------------------------------------------------
 UActorFactoryEmitter
 -----------------------------------------------------------------------------*/
-UActorFactoryEmitter::UActorFactoryEmitter(const class FPostConstructInitializeProperties& PCIP)
-	: Super(PCIP)
+UActorFactoryEmitter::UActorFactoryEmitter(const FObjectInitializer& ObjectInitializer)
+	: Super(ObjectInitializer)
 {
 	DisplayName = LOCTEXT("EmitterDisplayName", "Emitter");
 	NewActorClass = AEmitter::StaticClass();
@@ -521,8 +521,8 @@ void UActorFactoryEmitter::PostCreateBlueprint( UObject* Asset, AActor* CDO )
 /*-----------------------------------------------------------------------------
 UActorFactoryPlayerStart
 -----------------------------------------------------------------------------*/
-UActorFactoryPlayerStart::UActorFactoryPlayerStart(const class FPostConstructInitializeProperties& PCIP)
-	: Super(PCIP)
+UActorFactoryPlayerStart::UActorFactoryPlayerStart(const FObjectInitializer& ObjectInitializer)
+	: Super(ObjectInitializer)
 {
 	DisplayName = LOCTEXT("PlayerStartDisplayName", "Player Start");
 	NewActorClass = APlayerStart::StaticClass();
@@ -531,8 +531,8 @@ UActorFactoryPlayerStart::UActorFactoryPlayerStart(const class FPostConstructIni
 /*-----------------------------------------------------------------------------
 UActorFactoryTargetPoint
 -----------------------------------------------------------------------------*/
-UActorFactoryTargetPoint::UActorFactoryTargetPoint( const class FPostConstructInitializeProperties& PCIP )
-: Super( PCIP )
+UActorFactoryTargetPoint::UActorFactoryTargetPoint( const FObjectInitializer& ObjectInitializer )
+: Super( ObjectInitializer )
 {
 	DisplayName = LOCTEXT( "TargetPointDisplayName", "Target Point" );
 	NewActorClass = ATargetPoint::StaticClass();
@@ -541,8 +541,8 @@ UActorFactoryTargetPoint::UActorFactoryTargetPoint( const class FPostConstructIn
 /*-----------------------------------------------------------------------------
 UActorFactoryNote
 -----------------------------------------------------------------------------*/
-UActorFactoryNote::UActorFactoryNote( const class FPostConstructInitializeProperties& PCIP )
-: Super( PCIP )
+UActorFactoryNote::UActorFactoryNote( const FObjectInitializer& ObjectInitializer )
+: Super( ObjectInitializer )
 {
 	DisplayName = LOCTEXT( "NoteDisplayName", "Note" );
 	NewActorClass = ANote::StaticClass();
@@ -551,8 +551,8 @@ UActorFactoryNote::UActorFactoryNote( const class FPostConstructInitializeProper
 /*-----------------------------------------------------------------------------
 UActorFactoryPhysicsAsset
 -----------------------------------------------------------------------------*/
-UActorFactoryPhysicsAsset::UActorFactoryPhysicsAsset(const class FPostConstructInitializeProperties& PCIP)
-	: Super(PCIP)
+UActorFactoryPhysicsAsset::UActorFactoryPhysicsAsset(const FObjectInitializer& ObjectInitializer)
+	: Super(ObjectInitializer)
 {
 	DisplayName = LOCTEXT("PhysicsAssetDisplayName", "Skeletal Physics");
 	NewActorClass = ASkeletalMeshActor::StaticClass();
@@ -646,8 +646,8 @@ void UActorFactoryPhysicsAsset::PostCreateBlueprint( UObject* Asset, AActor* CDO
 /*-----------------------------------------------------------------------------
 UActorFactoryAnimationAsset
 -----------------------------------------------------------------------------*/
-UActorFactoryAnimationAsset::UActorFactoryAnimationAsset(const class FPostConstructInitializeProperties& PCIP)
-	: Super(PCIP)
+UActorFactoryAnimationAsset::UActorFactoryAnimationAsset(const FObjectInitializer& ObjectInitializer)
+	: Super(ObjectInitializer)
 {
 	DisplayName = LOCTEXT("SingleAnimSkeletalDisplayName", "Single Animation Skeletal");
 	NewActorClass = ASkeletalMeshActor::StaticClass();
@@ -824,8 +824,8 @@ void UActorFactoryAnimationAsset::PostCreateBlueprint( UObject* Asset,  AActor*
 /*-----------------------------------------------------------------------------
 UActorFactorySkeletalMesh
 -----------------------------------------------------------------------------*/
-UActorFactorySkeletalMesh::UActorFactorySkeletalMesh(const class FPostConstructInitializeProperties& PCIP)
-	: Super(PCIP)
+UActorFactorySkeletalMesh::UActorFactorySkeletalMesh(const FObjectInitializer& ObjectInitializer)
+	: Super(ObjectInitializer)
 { 
 	DisplayName = LOCTEXT("SkeletalMeshDisplayName", "Skeletal Mesh");
 	NewActorClass = ASkeletalMeshActor::StaticClass();
@@ -998,8 +998,8 @@ FQuat UActorFactorySkeletalMesh::AlignObjectToSurfaceNormal(const FVector& InSur
 /*-----------------------------------------------------------------------------
 UActorFactoryCameraActor
 -----------------------------------------------------------------------------*/
-UActorFactoryCameraActor::UActorFactoryCameraActor(const class FPostConstructInitializeProperties& PCIP)
-	: Super(PCIP)
+UActorFactoryCameraActor::UActorFactoryCameraActor(const FObjectInitializer& ObjectInitializer)
+	: Super(ObjectInitializer)
 {
 	DisplayName = LOCTEXT("CameraDisplayName", "Camera");
 	NewActorClass = ACameraActor::StaticClass();
@@ -1008,8 +1008,8 @@ UActorFactoryCameraActor::UActorFactoryCameraActor(const class FPostConstructIni
 /*-----------------------------------------------------------------------------
 UActorFactoryAmbientSound
 -----------------------------------------------------------------------------*/
-UActorFactoryAmbientSound::UActorFactoryAmbientSound(const class FPostConstructInitializeProperties& PCIP)
-	: Super(PCIP)
+UActorFactoryAmbientSound::UActorFactoryAmbientSound(const FObjectInitializer& ObjectInitializer)
+	: Super(ObjectInitializer)
 {
 	DisplayName = LOCTEXT("AmbientSoundDisplayName", "Ambient Sound");
 	NewActorClass = AAmbientSound::StaticClass();
@@ -1070,8 +1070,8 @@ void UActorFactoryAmbientSound::PostCreateBlueprint( UObject* Asset, AActor* CDO
 /*-----------------------------------------------------------------------------
 UActorFactoryClass
 -----------------------------------------------------------------------------*/
-UActorFactoryClass::UActorFactoryClass(const class FPostConstructInitializeProperties& PCIP)
-	: Super(PCIP)
+UActorFactoryClass::UActorFactoryClass(const FObjectInitializer& ObjectInitializer)
+	: Super(ObjectInitializer)
 {
 	DisplayName = LOCTEXT("ClassDisplayName", "Class");
 }
@@ -1145,8 +1145,8 @@ AActor* UActorFactoryClass::SpawnActor( UObject* Asset, ULevel* InLevel, const F
 /*-----------------------------------------------------------------------------
 UActorFactoryBlueprint
 -----------------------------------------------------------------------------*/
-UActorFactoryBlueprint::UActorFactoryBlueprint(const class FPostConstructInitializeProperties& PCIP)
-	: Super(PCIP)
+UActorFactoryBlueprint::UActorFactoryBlueprint(const FObjectInitializer& ObjectInitializer)
+	: Super(ObjectInitializer)
 {
 	DisplayName = LOCTEXT("BlueprintDisplayName", "Blueprint");
 }
@@ -1245,8 +1245,8 @@ void UActorFactoryBlueprint::PostSpawnActor( UObject* Asset, AActor* NewActor)
 /*-----------------------------------------------------------------------------
 UActorFactoryMatineeActor
 -----------------------------------------------------------------------------*/
-UActorFactoryMatineeActor::UActorFactoryMatineeActor(const class FPostConstructInitializeProperties& PCIP)
-	: Super(PCIP)
+UActorFactoryMatineeActor::UActorFactoryMatineeActor(const FObjectInitializer& ObjectInitializer)
+	: Super(ObjectInitializer)
 {
 	DisplayName = LOCTEXT("MatineeDisplayName", "Matinee");
 	NewActorClass = AMatineeActor::StaticClass();
@@ -1305,8 +1305,8 @@ void UActorFactoryMatineeActor::PostCreateBlueprint( UObject* Asset, AActor* CDO
 /*-----------------------------------------------------------------------------
 UActorFactoryDirectionalLight
 -----------------------------------------------------------------------------*/
-UActorFactoryDirectionalLight::UActorFactoryDirectionalLight(const class FPostConstructInitializeProperties& PCIP)
-	: Super(PCIP)
+UActorFactoryDirectionalLight::UActorFactoryDirectionalLight(const FObjectInitializer& ObjectInitializer)
+	: Super(ObjectInitializer)
 {
 	DisplayName = LOCTEXT("DirectionalLightDisplayName", "Directional Light");
 	NewActorClass = ADirectionalLight::StaticClass();
@@ -1317,8 +1317,8 @@ UActorFactoryDirectionalLight::UActorFactoryDirectionalLight(const class FPostCo
 /*-----------------------------------------------------------------------------
 UActorFactorySpotLight
 -----------------------------------------------------------------------------*/
-UActorFactorySpotLight::UActorFactorySpotLight(const class FPostConstructInitializeProperties& PCIP)
-	: Super(PCIP)
+UActorFactorySpotLight::UActorFactorySpotLight(const FObjectInitializer& ObjectInitializer)
+	: Super(ObjectInitializer)
 {
 	DisplayName = LOCTEXT("SpotLightDisplayName", "Spot Light");
 	NewActorClass = ASpotLight::StaticClass();
@@ -1329,8 +1329,8 @@ UActorFactorySpotLight::UActorFactorySpotLight(const class FPostConstructInitial
 /*-----------------------------------------------------------------------------
 UActorFactoryPointLight
 -----------------------------------------------------------------------------*/
-UActorFactoryPointLight::UActorFactoryPointLight(const class FPostConstructInitializeProperties& PCIP)
-	: Super(PCIP)
+UActorFactoryPointLight::UActorFactoryPointLight(const FObjectInitializer& ObjectInitializer)
+	: Super(ObjectInitializer)
 {
 	DisplayName = LOCTEXT("PointLightDisplayName", "Point Light");
 	NewActorClass = APointLight::StaticClass();
@@ -1341,8 +1341,8 @@ UActorFactoryPointLight::UActorFactoryPointLight(const class FPostConstructIniti
 /*-----------------------------------------------------------------------------
 UActorFactorySkyLight
 -----------------------------------------------------------------------------*/
-UActorFactorySkyLight::UActorFactorySkyLight( const class FPostConstructInitializeProperties& PCIP )
-: Super( PCIP )
+UActorFactorySkyLight::UActorFactorySkyLight( const FObjectInitializer& ObjectInitializer )
+: Super( ObjectInitializer )
 {
 	DisplayName = LOCTEXT( "SkyLightDisplayName", "Sky Light" );
 	NewActorClass = ASkyLight::StaticClass();
@@ -1351,8 +1351,8 @@ UActorFactorySkyLight::UActorFactorySkyLight( const class FPostConstructInitiali
 /*-----------------------------------------------------------------------------
 UActorFactorySphereReflectionCapture
 -----------------------------------------------------------------------------*/
-UActorFactorySphereReflectionCapture::UActorFactorySphereReflectionCapture(const class FPostConstructInitializeProperties& PCIP)
-	: Super(PCIP)
+UActorFactorySphereReflectionCapture::UActorFactorySphereReflectionCapture(const FObjectInitializer& ObjectInitializer)
+	: Super(ObjectInitializer)
 {
 	DisplayName = LOCTEXT("ReflectionCaptureSphereDisplayName", "Sphere Reflection Capture");
 	NewActorClass = ASphereReflectionCapture::StaticClass();
@@ -1363,8 +1363,8 @@ UActorFactorySphereReflectionCapture::UActorFactorySphereReflectionCapture(const
 /*-----------------------------------------------------------------------------
 UActorFactoryBoxReflectionCapture
 -----------------------------------------------------------------------------*/
-UActorFactoryBoxReflectionCapture::UActorFactoryBoxReflectionCapture(const class FPostConstructInitializeProperties& PCIP)
-	: Super(PCIP)
+UActorFactoryBoxReflectionCapture::UActorFactoryBoxReflectionCapture(const FObjectInitializer& ObjectInitializer)
+	: Super(ObjectInitializer)
 {
 	DisplayName = LOCTEXT("ReflectionCaptureBoxDisplayName", "Box Reflection Capture");
 	NewActorClass = ABoxReflectionCapture::StaticClass();
@@ -1375,8 +1375,8 @@ UActorFactoryBoxReflectionCapture::UActorFactoryBoxReflectionCapture(const class
 /*-----------------------------------------------------------------------------
 UActorFactoryPlaneReflectionCapture
 -----------------------------------------------------------------------------*/
-UActorFactoryPlaneReflectionCapture::UActorFactoryPlaneReflectionCapture(const class FPostConstructInitializeProperties& PCIP)
-	: Super(PCIP)
+UActorFactoryPlaneReflectionCapture::UActorFactoryPlaneReflectionCapture(const FObjectInitializer& ObjectInitializer)
+	: Super(ObjectInitializer)
 {
 	DisplayName = LOCTEXT("ReflectionCapturePlaneDisplayName", "Plane Reflection Capture");
 	NewActorClass = APlaneReflectionCapture::StaticClass();
@@ -1387,8 +1387,8 @@ UActorFactoryPlaneReflectionCapture::UActorFactoryPlaneReflectionCapture(const c
 /*-----------------------------------------------------------------------------
 UActorFactoryAtmosphericFog
 -----------------------------------------------------------------------------*/
-UActorFactoryAtmosphericFog::UActorFactoryAtmosphericFog(const class FPostConstructInitializeProperties& PCIP)
-	: Super(PCIP)
+UActorFactoryAtmosphericFog::UActorFactoryAtmosphericFog(const FObjectInitializer& ObjectInitializer)
+	: Super(ObjectInitializer)
 {
 	DisplayName = LOCTEXT("AtmosphericFogDisplayName", "Atmospheric Fog");
 	NewActorClass = AAtmosphericFog::StaticClass();
@@ -1397,8 +1397,8 @@ UActorFactoryAtmosphericFog::UActorFactoryAtmosphericFog(const class FPostConstr
 /*-----------------------------------------------------------------------------
 UActorFactoryExponentialHeightFog
 -----------------------------------------------------------------------------*/
-UActorFactoryExponentialHeightFog::UActorFactoryExponentialHeightFog(const class FPostConstructInitializeProperties& PCIP)
-	: Super(PCIP)
+UActorFactoryExponentialHeightFog::UActorFactoryExponentialHeightFog(const FObjectInitializer& ObjectInitializer)
+	: Super(ObjectInitializer)
 {
 	DisplayName = LOCTEXT("ExponentialHeightFogDisplayName", "Exponential Height Fog");
 	NewActorClass = AExponentialHeightFog::StaticClass();
@@ -1407,8 +1407,8 @@ UActorFactoryExponentialHeightFog::UActorFactoryExponentialHeightFog(const class
 /*-----------------------------------------------------------------------------
 UActorFactoryInteractiveFoliage
 -----------------------------------------------------------------------------*/
-UActorFactoryInteractiveFoliage::UActorFactoryInteractiveFoliage(const class FPostConstructInitializeProperties& PCIP)
-	: Super(PCIP)
+UActorFactoryInteractiveFoliage::UActorFactoryInteractiveFoliage(const FObjectInitializer& ObjectInitializer)
+	: Super(ObjectInitializer)
 {
 	DisplayName = LOCTEXT("InteractiveFoliageDisplayName", "Interactive Foliage");
 	NewActorClass = AInteractiveFoliageActor::StaticClass();
@@ -1417,8 +1417,8 @@ UActorFactoryInteractiveFoliage::UActorFactoryInteractiveFoliage(const class FPo
 /*-----------------------------------------------------------------------------
 UActorFactoryTriggerBox
 -----------------------------------------------------------------------------*/
-UActorFactoryTriggerBox::UActorFactoryTriggerBox(const class FPostConstructInitializeProperties& PCIP)
-	: Super(PCIP)
+UActorFactoryTriggerBox::UActorFactoryTriggerBox(const FObjectInitializer& ObjectInitializer)
+	: Super(ObjectInitializer)
 {
 	DisplayName = LOCTEXT("TriggerBoxDisplayName", "Box Trigger");
 	NewActorClass = ATriggerBox::StaticClass();
@@ -1427,8 +1427,8 @@ UActorFactoryTriggerBox::UActorFactoryTriggerBox(const class FPostConstructIniti
 /*-----------------------------------------------------------------------------
 UActorFactoryTriggerCapsule
 -----------------------------------------------------------------------------*/
-UActorFactoryTriggerCapsule::UActorFactoryTriggerCapsule(const class FPostConstructInitializeProperties& PCIP)
-	: Super(PCIP)
+UActorFactoryTriggerCapsule::UActorFactoryTriggerCapsule(const FObjectInitializer& ObjectInitializer)
+	: Super(ObjectInitializer)
 {
 	DisplayName = LOCTEXT("TriggerCapsuleDisplayName", "Capsule Trigger");
 	NewActorClass = ATriggerCapsule::StaticClass();
@@ -1437,8 +1437,8 @@ UActorFactoryTriggerCapsule::UActorFactoryTriggerCapsule(const class FPostConstr
 /*-----------------------------------------------------------------------------
 UActorFactoryTriggerSphere
 -----------------------------------------------------------------------------*/
-UActorFactoryTriggerSphere::UActorFactoryTriggerSphere(const class FPostConstructInitializeProperties& PCIP)
-	: Super(PCIP)
+UActorFactoryTriggerSphere::UActorFactoryTriggerSphere(const FObjectInitializer& ObjectInitializer)
+	: Super(ObjectInitializer)
 {
 	DisplayName = LOCTEXT("TriggerSphereDisplayName", "Sphere Trigger");
 	NewActorClass = ATriggerSphere::StaticClass();
@@ -1447,8 +1447,8 @@ UActorFactoryTriggerSphere::UActorFactoryTriggerSphere(const class FPostConstruc
 /*-----------------------------------------------------------------------------
 UActorFactoryDestructible
 -----------------------------------------------------------------------------*/
-UActorFactoryDestructible::UActorFactoryDestructible(const class FPostConstructInitializeProperties& PCIP)
-	: Super(PCIP)
+UActorFactoryDestructible::UActorFactoryDestructible(const FObjectInitializer& ObjectInitializer)
+	: Super(ObjectInitializer)
 {
 	DisplayName = LOCTEXT("DestructibleDisplayName", "Destructible");
 	NewActorClass = ADestructibleActor::StaticClass();
@@ -1513,8 +1513,8 @@ FQuat UActorFactoryDestructible::AlignObjectToSurfaceNormal(const FVector& InSur
 /*-----------------------------------------------------------------------------
 UActorFactoryVectorField
 -----------------------------------------------------------------------------*/
-UActorFactoryVectorFieldVolume::UActorFactoryVectorFieldVolume(const class FPostConstructInitializeProperties& PCIP)
-	: Super(PCIP)
+UActorFactoryVectorFieldVolume::UActorFactoryVectorFieldVolume(const FObjectInitializer& ObjectInitializer)
+	: Super(ObjectInitializer)
 {
 	DisplayName = LOCTEXT("VectorFieldVolumeDisplayName", "Vector Field Volume");
 	NewActorClass = AVectorFieldVolume::StaticClass();
@@ -1557,8 +1557,8 @@ void CreateBrushForVolumeActor( AVolume* NewActor, UBrushBuilder* BrushBuilder )
 		NewActor->PreEditChange(NULL);
 
 		NewActor->PolyFlags = 0;
-		NewActor->Brush = new( NewActor, NAME_None, RF_Transactional )UModel(FPostConstructInitializeProperties(), NULL, true );
-		NewActor->Brush->Polys = new( NewActor->Brush, NAME_None, RF_Transactional )UPolys(FPostConstructInitializeProperties());
+		NewActor->Brush = new( NewActor, NAME_None, RF_Transactional )UModel(FObjectInitializer(), NULL, true );
+		NewActor->Brush->Polys = new( NewActor->Brush, NAME_None, RF_Transactional )UPolys(FObjectInitializer());
 		NewActor->BrushComponent->Brush = NewActor->Brush;
 		if(BrushBuilder != nullptr)
 		{
@@ -1587,8 +1587,8 @@ void CreateBrushForVolumeActor( AVolume* NewActor, UBrushBuilder* BrushBuilder )
 /*-----------------------------------------------------------------------------
 UActorFactoryBoxVolume
 -----------------------------------------------------------------------------*/
-UActorFactoryBoxVolume::UActorFactoryBoxVolume( const class FPostConstructInitializeProperties& PCIP )
-: Super(PCIP)
+UActorFactoryBoxVolume::UActorFactoryBoxVolume( const FObjectInitializer& ObjectInitializer )
+: Super(ObjectInitializer)
 {
 	DisplayName = LOCTEXT( "BoxVolumeDisplayName", "Box Volume" );
 	NewActorClass = AVolume::StaticClass();
@@ -1627,8 +1627,8 @@ void UActorFactoryBoxVolume::PostSpawnActor( UObject* Asset, AActor* NewActor )
 /*-----------------------------------------------------------------------------
 UActorFactorySphereVolume
 -----------------------------------------------------------------------------*/
-UActorFactorySphereVolume::UActorFactorySphereVolume( const class FPostConstructInitializeProperties& PCIP )
-: Super(PCIP)
+UActorFactorySphereVolume::UActorFactorySphereVolume( const FObjectInitializer& ObjectInitializer )
+: Super(ObjectInitializer)
 {
 	DisplayName = LOCTEXT( "SphereVolumeDisplayName", "Sphere Volume" );
 	NewActorClass = AVolume::StaticClass();
@@ -1669,8 +1669,8 @@ void UActorFactorySphereVolume::PostSpawnActor( UObject* Asset, AActor* NewActor
 /*-----------------------------------------------------------------------------
 UActorFactoryCylinderVolume
 -----------------------------------------------------------------------------*/
-UActorFactoryCylinderVolume::UActorFactoryCylinderVolume( const class FPostConstructInitializeProperties& PCIP )
-: Super(PCIP)
+UActorFactoryCylinderVolume::UActorFactoryCylinderVolume( const FObjectInitializer& ObjectInitializer )
+: Super(ObjectInitializer)
 {
 	DisplayName = LOCTEXT( "CylinderVolumeDisplayName", "Cylinder Volume" );
 	NewActorClass = AVolume::StaticClass();
diff --git a/Engine/Source/Editor/UnrealEd/Private/Factories/AnimBlueprintFactory.cpp b/Engine/Source/Editor/UnrealEd/Private/Factories/AnimBlueprintFactory.cpp
--- a/Engine/Source/Editor/UnrealEd/Private/Factories/AnimBlueprintFactory.cpp
+++ b/Engine/Source/Editor/UnrealEd/Private/Factories/AnimBlueprintFactory.cpp
@@ -336,8 +336,8 @@ private:
 	UAnimBlueprintFactory implementation.
 ------------------------------------------------------------------------------*/
 
-UAnimBlueprintFactory::UAnimBlueprintFactory(const class FPostConstructInitializeProperties& PCIP)
-	: Super(PCIP)
+UAnimBlueprintFactory::UAnimBlueprintFactory(const FObjectInitializer& ObjectInitializer)
+	: Super(ObjectInitializer)
 {
 	bCreateNew = true;
 	bEditAfterNew = true;
diff --git a/Engine/Source/Editor/UnrealEd/Private/Factories/AnimCompositeFactory.cpp b/Engine/Source/Editor/UnrealEd/Private/Factories/AnimCompositeFactory.cpp
--- a/Engine/Source/Editor/UnrealEd/Private/Factories/AnimCompositeFactory.cpp
+++ b/Engine/Source/Editor/UnrealEd/Private/Factories/AnimCompositeFactory.cpp
@@ -11,8 +11,8 @@
 
 #define LOCTEXT_NAMESPACE "AnimCompositeFactory"
 
-UAnimCompositeFactory::UAnimCompositeFactory(const class FPostConstructInitializeProperties& PCIP)
-	: Super(PCIP)
+UAnimCompositeFactory::UAnimCompositeFactory(const FObjectInitializer& ObjectInitializer)
+	: Super(ObjectInitializer)
 {
 	bCreateNew = true;
 	SupportedClass = UAnimComposite::StaticClass();
diff --git a/Engine/Source/Editor/UnrealEd/Private/Factories/AnimMontageFactory.cpp b/Engine/Source/Editor/UnrealEd/Private/Factories/AnimMontageFactory.cpp
--- a/Engine/Source/Editor/UnrealEd/Private/Factories/AnimMontageFactory.cpp
+++ b/Engine/Source/Editor/UnrealEd/Private/Factories/AnimMontageFactory.cpp
@@ -11,8 +11,8 @@
 
 #define LOCTEXT_NAMESPACE "AnimMontageFactory"
 
-UAnimMontageFactory::UAnimMontageFactory(const class FPostConstructInitializeProperties& PCIP)
-	: Super(PCIP)
+UAnimMontageFactory::UAnimMontageFactory(const FObjectInitializer& ObjectInitializer)
+	: Super(ObjectInitializer)
 {
 	bCreateNew = true;
 	SupportedClass = UAnimMontage::StaticClass();
diff --git a/Engine/Source/Editor/UnrealEd/Private/Factories/CSVImportFactory.cpp b/Engine/Source/Editor/UnrealEd/Private/Factories/CSVImportFactory.cpp
--- a/Engine/Source/Editor/UnrealEd/Private/Factories/CSVImportFactory.cpp
+++ b/Engine/Source/Editor/UnrealEd/Private/Factories/CSVImportFactory.cpp
@@ -393,8 +393,8 @@ static UClass* GetCurveClass( ECSVImportType ImportType )
 }
 
 
-UCSVImportFactory::UCSVImportFactory(const class FPostConstructInitializeProperties& PCIP)
-	: Super(PCIP)
+UCSVImportFactory::UCSVImportFactory(const FObjectInitializer& ObjectInitializer)
+	: Super(ObjectInitializer)
 {
 
 	bCreateNew = false;
@@ -609,8 +609,8 @@ bool UCSVImportFactory::Reimport( UObject* Obj, const FString& Path )
 
 //////////////////////////////////////////////////////////////////////////
 
-UReimportDataTableFactory::UReimportDataTableFactory(const class FPostConstructInitializeProperties& PCIP)
-	: Super(PCIP)
+UReimportDataTableFactory::UReimportDataTableFactory(const FObjectInitializer& ObjectInitializer)
+	: Super(ObjectInitializer)
 {
 }
 
@@ -645,8 +645,8 @@ EReimportResult::Type UReimportDataTableFactory::Reimport( UObject* Obj )
 
 ////////////////////////////////////////////////////////////////////////////
 //
-UReimportCurveTableFactory::UReimportCurveTableFactory(const class FPostConstructInitializeProperties& PCIP)
-	: Super(PCIP)
+UReimportCurveTableFactory::UReimportCurveTableFactory(const FObjectInitializer& ObjectInitializer)
+	: Super(ObjectInitializer)
 {
 }
 
@@ -681,8 +681,8 @@ EReimportResult::Type UReimportCurveTableFactory::Reimport( UObject* Obj )
 
 ////////////////////////////////////////////////////////////////////////////
 //
-UReimportCurveFactory::UReimportCurveFactory(const class FPostConstructInitializeProperties& PCIP)
-	: Super(PCIP)
+UReimportCurveFactory::UReimportCurveFactory(const FObjectInitializer& ObjectInitializer)
+	: Super(ObjectInitializer)
 {
 }
 
