Rename FPostConstructInitializeProperties to something simpler

diff --git a/Engine/Source/Editor/UnrealEd/Private/Factories/EditorFactories.cpp b/Engine/Source/Editor/UnrealEd/Private/Factories/EditorFactories.cpp
--- a/Engine/Source/Editor/UnrealEd/Private/Factories/EditorFactories.cpp
+++ b/Engine/Source/Editor/UnrealEd/Private/Factories/EditorFactories.cpp
@@ -69,8 +69,8 @@ DEFINE_LOG_CATEGORY_STATIC(LogEditorFactories, Log, All);
 /*------------------------------------------------------------------------------
 	UTexture2DFactoryNew implementation.
 ------------------------------------------------------------------------------*/
-UTexture2DFactoryNew::UTexture2DFactoryNew(const class FPostConstructInitializeProperties& PCIP)
-	: Super(PCIP)
+UTexture2DFactoryNew::UTexture2DFactoryNew(const FObjectInitializer& ObjectInitializer)
+	: Super(ObjectInitializer)
 {
 	SupportedClass		= UTexture2D::StaticClass();
 	bCreateNew			= true;
@@ -116,8 +116,8 @@ UObject* UTexture2DFactoryNew::FactoryCreateNew( UClass* InClass, UObject* InPar
 /*------------------------------------------------------------------------------
 	UMaterialInstanceConstantFactoryNew implementation.
 ------------------------------------------------------------------------------*/
-UMaterialInstanceConstantFactoryNew::UMaterialInstanceConstantFactoryNew(const class FPostConstructInitializeProperties& PCIP)
-	: Super(PCIP)
+UMaterialInstanceConstantFactoryNew::UMaterialInstanceConstantFactoryNew(const FObjectInitializer& ObjectInitializer)
+	: Super(ObjectInitializer)
 {
 
 	SupportedClass = UMaterialInstanceConstant::StaticClass();
@@ -145,8 +145,8 @@ UObject* UMaterialInstanceConstantFactoryNew::FactoryCreateNew(UClass* Class,UOb
 /*------------------------------------------------------------------------------
 	UMaterialFactoryNew implementation.
 ------------------------------------------------------------------------------*/
-UMaterialFactoryNew::UMaterialFactoryNew(const class FPostConstructInitializeProperties& PCIP)
-	: Super(PCIP)
+UMaterialFactoryNew::UMaterialFactoryNew(const FObjectInitializer& ObjectInitializer)
+	: Super(ObjectInitializer)
 {
 
 	SupportedClass = UMaterial::StaticClass();
@@ -188,8 +188,8 @@ UObject* UMaterialFactoryNew::FactoryCreateNew(UClass* Class,UObject* InParent,F
 /*------------------------------------------------------------------------------
 	UMaterialFunctionFactoryNew implementation.
 ------------------------------------------------------------------------------*/
-UMaterialFunctionFactoryNew::UMaterialFunctionFactoryNew(const class FPostConstructInitializeProperties& PCIP)
-	: Super(PCIP)
+UMaterialFunctionFactoryNew::UMaterialFunctionFactoryNew(const FObjectInitializer& ObjectInitializer)
+	: Super(ObjectInitializer)
 {
 
 	SupportedClass = UMaterialFunction::StaticClass();
@@ -206,8 +206,8 @@ UObject* UMaterialFunctionFactoryNew::FactoryCreateNew(UClass* Class,UObject* In
 /*------------------------------------------------------------------------------
 	UMaterialParameterCollectionFactoryNew implementation.
 ------------------------------------------------------------------------------*/
-UMaterialParameterCollectionFactoryNew::UMaterialParameterCollectionFactoryNew(const class FPostConstructInitializeProperties& PCIP)
-	: Super(PCIP)
+UMaterialParameterCollectionFactoryNew::UMaterialParameterCollectionFactoryNew(const FObjectInitializer& ObjectInitializer)
+	: Super(ObjectInitializer)
 {
 
 	SupportedClass = UMaterialParameterCollection::StaticClass();
@@ -276,8 +276,8 @@ static bool ForceValidActorRefs(UStruct* Struct, uint8* Data)
 	return bChangedObjectPointer;
 }
 
-ULevelFactory::ULevelFactory(const class FPostConstructInitializeProperties& PCIP)
-	: Super(PCIP)
+ULevelFactory::ULevelFactory(const FObjectInitializer& ObjectInitializer)
+	: Super(ObjectInitializer)
 {
 	SupportedClass = UWorld::StaticClass();
 	Formats.Add(TEXT("t3d;Unreal World"));
@@ -425,7 +425,7 @@ UObject* ULevelFactory::FactoryCreateText
 				if (FParse::Value(Str, TEXT("NAME="), LevelName))
 				{
 					// create a new named level
-					World->SetCurrentLevel( new(World->GetOuter(), *LevelName)ULevel(FPostConstructInitializeProperties(),FURL(nullptr)) );
+					World->SetCurrentLevel( new(World->GetOuter(), *LevelName)ULevel(FObjectInitializer(),FURL(nullptr)) );
 				}
 			}
 #endif
@@ -752,7 +752,7 @@ UObject* ULevelFactory::FactoryCreateText
 
 	if (GIsImportingT3D && (MapPackageText.Len() > 0))
 	{
-		UPackageFactory* PackageFactory = new UPackageFactory(FPostConstructInitializeProperties());
+		UPackageFactory* PackageFactory = new UPackageFactory(FObjectInitializer());
 		check(PackageFactory);
 
 		FName NewPackageName(*(RootMapPackage->GetName()));
@@ -930,8 +930,8 @@ UObject* ULevelFactory::FactoryCreateText
 /*-----------------------------------------------------------------------------
 	UPackageFactory.
 -----------------------------------------------------------------------------*/
-UPackageFactory::UPackageFactory(const class FPostConstructInitializeProperties& PCIP)
-	: Super(PCIP)
+UPackageFactory::UPackageFactory(const FObjectInitializer& ObjectInitializer)
+	: Super(ObjectInitializer)
 {
 	SupportedClass = UPackage::StaticClass();
 	Formats.Add(TEXT("T3DPKG;Unreal Package"));
@@ -1115,8 +1115,8 @@ UObject* UPackageFactory::FactoryCreateText( UClass* Class, UObject* InParent, F
 	UPolysFactory.
 -----------------------------------------------------------------------------*/
 
-UPolysFactory::UPolysFactory(const class FPostConstructInitializeProperties& PCIP)
-	: Super(PCIP)
+UPolysFactory::UPolysFactory(const FObjectInitializer& ObjectInitializer)
+	: Super(ObjectInitializer)
 {
 
 	SupportedClass = UPolys::StaticClass();
@@ -1141,7 +1141,7 @@ UObject* UPolysFactory::FactoryCreateText
 	FEditorDelegates::OnAssetPreImport.Broadcast(this, Class, InParent, Name, Type);
 
 	// Create polys.
-	UPolys* Polys = Context ? CastChecked<UPolys>(Context) : new(InParent,Name,Flags)UPolys(FPostConstructInitializeProperties());
+	UPolys* Polys = Context ? CastChecked<UPolys>(Context) : new(InParent,Name,Flags)UPolys(FObjectInitializer());
 
 	// Eat up if present.
 	GetBEGIN( &Buffer, TEXT("POLYLIST") );
@@ -1388,8 +1388,8 @@ UObject* UPolysFactory::FactoryCreateText
 /*-----------------------------------------------------------------------------
 	UModelFactory.
 -----------------------------------------------------------------------------*/
-UModelFactory::UModelFactory(const class FPostConstructInitializeProperties& PCIP)
-	: Super(PCIP)
+UModelFactory::UModelFactory(const FObjectInitializer& ObjectInitializer)
+	: Super(ObjectInitializer)
 {
 	SupportedClass = UModel::StaticClass();
 	Formats.Add(TEXT("t3d;Unreal model text"));
@@ -1413,7 +1413,7 @@ UObject* UModelFactory::FactoryCreateText
 	FEditorDelegates::OnAssetPreImport.Broadcast(this, Class, InParent, Name, Type);
 
 	ABrush* TempOwner = (ABrush*)Context;
-	UModel* Model = new( InParent, Name, Flags )UModel( FPostConstructInitializeProperties(),TempOwner, 1 );
+	UModel* Model = new( InParent, Name, Flags )UModel( FObjectInitializer(),TempOwner, 1 );
 
 	const TCHAR* StrPtr;
 	FString StrLine;
@@ -1431,7 +1431,7 @@ UObject* UModelFactory::FactoryCreateText
 		}
 		else if( GetBEGIN (&StrPtr,TEXT("POLYLIST")) )
 		{
-			UPolysFactory* PolysFactory = new UPolysFactory(FPostConstructInitializeProperties());
+			UPolysFactory* PolysFactory = new UPolysFactory(FObjectInitializer());
 			Model->Polys = (UPolys*)PolysFactory->FactoryCreateText(UPolys::StaticClass(),Model,NAME_None,RF_Transactional,nullptr,Type,Buffer,BufferEnd,Warn);
 			check(Model->Polys);
 		}
@@ -1540,8 +1540,8 @@ void CreateSoundCue( USoundWave* Sound, UObject* InParent, EObjectFlags Flags, b
 /*-----------------------------------------------------------------------------
 	UReverbEffectFactory.
 -----------------------------------------------------------------------------*/
-UReverbEffectFactory::UReverbEffectFactory(const class FPostConstructInitializeProperties& PCIP)
-	: Super(PCIP)
+UReverbEffectFactory::UReverbEffectFactory(const FObjectInitializer& ObjectInitializer)
+	: Super(ObjectInitializer)
 {
 
 	SupportedClass = UReverbEffect::StaticClass();
@@ -1564,8 +1564,8 @@ UObject* UReverbEffectFactory::FactoryCreateNew( UClass* InClass, UObject* InPar
 
 static bool bSoundFactorySuppressImportOverwriteDialog = false;
 
-USoundFactory::USoundFactory(const class FPostConstructInitializeProperties& PCIP)
-	: Super(PCIP)
+USoundFactory::USoundFactory(const FObjectInitializer& ObjectInitializer)
+	: Super(ObjectInitializer)
 {
 
 
@@ -1713,7 +1713,7 @@ UObject* USoundFactory::FactoryCreateBinary
 
 		// Use pre-existing sound if it exists and we want to keep settings,
 		// otherwise create new sound and import raw data.
-		USoundWave* Sound = (bUseExistingSettings && ExistingSound) ? ExistingSound : new( InParent, Name, Flags ) USoundWave(FPostConstructInitializeProperties());
+		USoundWave* Sound = (bUseExistingSettings && ExistingSound) ? ExistingSound : new( InParent, Name, Flags ) USoundWave(FObjectInitializer());
 		
 		if (bUseExistingSettings && ExistingSound)
 		{
@@ -1786,8 +1786,8 @@ void USoundFactory::SuppressImportOverwriteDialog()
 	UDialogueVoiceFactory.
 -----------------------------------------------------------------------------*/
 
-UDialogueVoiceFactory::UDialogueVoiceFactory(const class FPostConstructInitializeProperties& PCIP)
-	: Super(PCIP)
+UDialogueVoiceFactory::UDialogueVoiceFactory(const FObjectInitializer& ObjectInitializer)
+	: Super(ObjectInitializer)
 {
 	SupportedClass = UDialogueVoice::StaticClass();
 	bCreateNew = true;
@@ -1803,8 +1803,8 @@ UObject* UDialogueVoiceFactory::FactoryCreateNew(UClass* Class,UObject* InParent
 	UDialogueWaveFactory.
 -----------------------------------------------------------------------------*/
 
-UDialogueWaveFactory::UDialogueWaveFactory(const class FPostConstructInitializeProperties& PCIP)
-	: Super(PCIP)
+UDialogueWaveFactory::UDialogueWaveFactory(const FObjectInitializer& ObjectInitializer)
+	: Super(ObjectInitializer)
 {
 	SupportedClass = UDialogueWave::StaticClass();
 	bCreateNew = true;
@@ -1819,8 +1819,8 @@ UObject* UDialogueWaveFactory::FactoryCreateNew(UClass* Class,UObject* InParent,
 /*-----------------------------------------------------------------------------
 	UReimportSoundFactory.
 -----------------------------------------------------------------------------*/
-UReimportSoundFactory::UReimportSoundFactory(const class FPostConstructInitializeProperties& PCIP)
-	: Super(PCIP)
+UReimportSoundFactory::UReimportSoundFactory(const FObjectInitializer& ObjectInitializer)
+	: Super(ObjectInitializer)
 {
 
 	SupportedClass = USoundWave::StaticClass();
@@ -1925,8 +1925,8 @@ const FString SurroundSpeakerLocations[SPEAKER_Count] =
 	TEXT( "_br" )			// SPEAKER_BackRight
 };
 
-USoundSurroundFactory::USoundSurroundFactory(const class FPostConstructInitializeProperties& PCIP)
-	: Super(PCIP)
+USoundSurroundFactory::USoundSurroundFactory(const FObjectInitializer& ObjectInitializer)
+	: Super(ObjectInitializer)
 {
 
 	SupportedClass = USoundWave::StaticClass();
@@ -2015,7 +2015,7 @@ UObject* USoundSurroundFactory::FactoryCreateBinary
 
 			if (Sound == nullptr)
 			{
-				Sound = new( InParent, BaseName, Flags ) USoundWave(FPostConstructInitializeProperties());
+				Sound = new( InParent, BaseName, Flags ) USoundWave(FObjectInitializer());
 			}
 		}
 
@@ -2149,8 +2149,8 @@ UObject* USoundSurroundFactory::FactoryCreateBinary
 /*-----------------------------------------------------------------------------
 	UReimportSoundSurroundFactory.
 -----------------------------------------------------------------------------*/
-UReimportSoundSurroundFactory::UReimportSoundSurroundFactory(const class FPostConstructInitializeProperties& PCIP)
-	: Super(PCIP)
+UReimportSoundSurroundFactory::UReimportSoundSurroundFactory(const FObjectInitializer& ObjectInitializer)
+	: Super(ObjectInitializer)
 {
 
 	SupportedClass = USoundWave::StaticClass();
@@ -2308,8 +2308,8 @@ EReimportResult::Type UReimportSoundSurroundFactory::Reimport( UObject* Obj )
 /*------------------------------------------------------------------------------
 	USoundCueFactoryNew.
 ------------------------------------------------------------------------------*/
-USoundCueFactoryNew::USoundCueFactoryNew(const class FPostConstructInitializeProperties& PCIP)
-	: Super(PCIP)
+USoundCueFactoryNew::USoundCueFactoryNew(const FObjectInitializer& ObjectInitializer)
+	: Super(ObjectInitializer)
 {
 
 	SupportedClass = USoundCue::StaticClass();
@@ -2339,8 +2339,8 @@ UObject* USoundCueFactoryNew::FactoryCreateNew( UClass* Class, UObject* InParent
 /*-----------------------------------------------------------------------------
 	USoundMixFactory.
 -----------------------------------------------------------------------------*/
-USoundMixFactory::USoundMixFactory(const class FPostConstructInitializeProperties& PCIP)
-	: Super(PCIP)
+USoundMixFactory::USoundMixFactory(const FObjectInitializer& ObjectInitializer)
+	: Super(ObjectInitializer)
 {
 
 	SupportedClass = USoundMix::StaticClass();
@@ -2361,8 +2361,8 @@ UObject* USoundMixFactory::FactoryCreateNew( UClass* InClass, UObject* InParent,
 /*-----------------------------------------------------------------------------
 	USoundClassFactory.
 -----------------------------------------------------------------------------*/
-USoundClassFactory::USoundClassFactory(const class FPostConstructInitializeProperties& PCIP)
-	: Super(PCIP)
+USoundClassFactory::USoundClassFactory(const FObjectInitializer& ObjectInitializer)
+	: Super(ObjectInitializer)
 {
 
 	SupportedClass = USoundClass::StaticClass();
@@ -2387,8 +2387,8 @@ UObject* USoundClassFactory::FactoryCreateNew( UClass* InClass, UObject* InParen
 /*------------------------------------------------------------------------------
 	USoundAttenuationFactory.
 ------------------------------------------------------------------------------*/
-USoundAttenuationFactory::USoundAttenuationFactory(const class FPostConstructInitializeProperties& PCIP)
-	: Super(PCIP)
+USoundAttenuationFactory::USoundAttenuationFactory(const FObjectInitializer& ObjectInitializer)
+	: Super(ObjectInitializer)
 {
 
 	SupportedClass = USoundAttenuation::StaticClass();
@@ -2405,8 +2405,8 @@ UObject* USoundAttenuationFactory::FactoryCreateNew( UClass* Class, UObject* InP
 /*------------------------------------------------------------------------------
 	UParticleSystemFactoryNew.
 ------------------------------------------------------------------------------*/
-UParticleSystemFactoryNew::UParticleSystemFactoryNew(const class FPostConstructInitializeProperties& PCIP)
-	: Super(PCIP)
+UParticleSystemFactoryNew::UParticleSystemFactoryNew(const FObjectInitializer& ObjectInitializer)
+	: Super(ObjectInitializer)
 {
 
 	SupportedClass = UParticleSystem::StaticClass();
@@ -2422,8 +2422,8 @@ UObject* UParticleSystemFactoryNew::FactoryCreateNew(UClass* Class,UObject* InPa
 /*------------------------------------------------------------------------------
 	UPhysicalMaterialFactoryNew.
 ------------------------------------------------------------------------------*/
-UPhysicalMaterialFactoryNew::UPhysicalMaterialFactoryNew(const class FPostConstructInitializeProperties& PCIP)
-	: Super(PCIP)
+UPhysicalMaterialFactoryNew::UPhysicalMaterialFactoryNew(const FObjectInitializer& ObjectInitializer)
+	: Super(ObjectInitializer)
 {
 
 	SupportedClass = UPhysicalMaterial::StaticClass();
@@ -2439,8 +2439,8 @@ UObject* UPhysicalMaterialFactoryNew::FactoryCreateNew(UClass* Class,UObject* In
 /*------------------------------------------------------------------------------
 	UInterpDataFactoryNew.
 ------------------------------------------------------------------------------*/
-UInterpDataFactoryNew::UInterpDataFactoryNew(const class FPostConstructInitializeProperties& PCIP)
-	: Super(PCIP)
+UInterpDataFactoryNew::UInterpDataFactoryNew(const FObjectInitializer& ObjectInitializer)
+	: Super(ObjectInitializer)
 {
 
 	SupportedClass = UInterpData::StaticClass();
@@ -2457,8 +2457,8 @@ UObject* UInterpDataFactoryNew::FactoryCreateNew(UClass* Class,UObject* InParent
 	UTextureRenderTargetFactoryNew
 -----------------------------------------------------------------------------*/
 
-UTextureRenderTargetFactoryNew::UTextureRenderTargetFactoryNew(const class FPostConstructInitializeProperties& PCIP)
-	: Super(PCIP)
+UTextureRenderTargetFactoryNew::UTextureRenderTargetFactoryNew(const FObjectInitializer& ObjectInitializer)
+	: Super(ObjectInitializer)
 {
 	SupportedClass = UTextureRenderTarget2D::StaticClass();
 	bCreateNew = true;
@@ -2484,8 +2484,8 @@ UObject* UTextureRenderTargetFactoryNew::FactoryCreateNew(UClass* Class,UObject*
 /*-----------------------------------------------------------------------------
 	UTextureRenderTargetCubeFactoryNew
 -----------------------------------------------------------------------------*/
-UTextureRenderTargetCubeFactoryNew::UTextureRenderTargetCubeFactoryNew(const class FPostConstructInitializeProperties& PCIP)
-	: Super(PCIP)
+UTextureRenderTargetCubeFactoryNew::UTextureRenderTargetCubeFactoryNew(const FObjectInitializer& ObjectInitializer)
+	: Super(ObjectInitializer)
 {
 
 	SupportedClass = UTextureRenderTargetCube::StaticClass();
@@ -3219,8 +3219,8 @@ UTexture2D* DecompressTGA(
 bool UTextureFactory::bSuppressImportOverwriteDialog = false;
 bool UTextureFactory::bSuppressImportResolutionWarnings = false;
 
-UTextureFactory::UTextureFactory(const class FPostConstructInitializeProperties& PCIP)
-	: Super(PCIP)
+UTextureFactory::UTextureFactory(const FObjectInitializer& ObjectInitializer)
+	: Super(ObjectInitializer)
 {
 	SupportedClass = UTexture::StaticClass();
 
@@ -4382,7 +4382,7 @@ UObject* UTextureFactory::FactoryCreateBinary
 		UPackage* MaterialPackage = CreatePackage(nullptr, *MaterialPackageName);
 
 		// Create the material
-		UMaterialFactoryNew* Factory = new UMaterialFactoryNew(FPostConstructInitializeProperties());
+		UMaterialFactoryNew* Factory = new UMaterialFactoryNew(FObjectInitializer());
 		UMaterial* Material = (UMaterial*)Factory->FactoryCreateNew( UMaterial::StaticClass(), MaterialPackage, *MaterialName, Flags, Context, Warn );
 
 		// Notify the asset registry
@@ -4541,8 +4541,8 @@ bool UTextureFactory::IsImportResolutionValid(int32 Width, int32 Height, bool bA
 /*------------------------------------------------------------------------------
 	UTextureExporterPCX implementation.
 ------------------------------------------------------------------------------*/
-UTextureExporterPCX::UTextureExporterPCX(const class FPostConstructInitializeProperties& PCIP)
-	: Super(PCIP)
+UTextureExporterPCX::UTextureExporterPCX(const FObjectInitializer& ObjectInitializer)
+	: Super(ObjectInitializer)
 {
 	SupportedClass = UTexture2D::StaticClass();
 	PreferredFormatIndex = 0;
@@ -4609,8 +4609,8 @@ bool UTextureExporterPCX::ExportBinary( UObject* Object, const TCHAR* Type, FArc
 /*------------------------------------------------------------------------------
 	UTextureExporterBMP implementation.
 ------------------------------------------------------------------------------*/
-UTextureExporterBMP::UTextureExporterBMP(const class FPostConstructInitializeProperties& PCIP)
-	: Super(PCIP)
+UTextureExporterBMP::UTextureExporterBMP(const FObjectInitializer& ObjectInitializer)
+	: Super(ObjectInitializer)
 {
 	SupportedClass = UTexture2D::StaticClass();
 	PreferredFormatIndex = 0;
@@ -4916,8 +4916,8 @@ public:
 	UTextureExporterHDR implementation.
 	Exports render targets.
 ------------------------------------------------------------------------------*/
-UTextureExporterHDR::UTextureExporterHDR(const class FPostConstructInitializeProperties& PCIP)
-	: Super(PCIP)
+UTextureExporterHDR::UTextureExporterHDR(const FObjectInitializer& ObjectInitializer)
+	: Super(ObjectInitializer)
 {
 	SupportedClass = UTextureRenderTarget::StaticClass();
 	PreferredFormatIndex = 0;
@@ -4946,8 +4946,8 @@ bool UTextureExporterHDR::ExportBinary(UObject* Object, const TCHAR* Type, FArch
 	UTextureCubeExporterHDR implementation.
 	Export UTextureCubes as .HDR
 ------------------------------------------------------------------------------*/
-UTextureCubeExporterHDR::UTextureCubeExporterHDR(const class FPostConstructInitializeProperties& PCIP)
-	: Super(PCIP)
+UTextureCubeExporterHDR::UTextureCubeExporterHDR(const FObjectInitializer& ObjectInitializer)
+	: Super(ObjectInitializer)
 {
 	SupportedClass = UTextureCube::StaticClass();
 	PreferredFormatIndex = 0;
@@ -4970,8 +4970,8 @@ bool UTextureCubeExporterHDR::ExportBinary(UObject* Object, const TCHAR* Type, F
 /*------------------------------------------------------------------------------
 	UTextureExporterTGA implementation.
 ------------------------------------------------------------------------------*/
-UTextureExporterTGA::UTextureExporterTGA(const class FPostConstructInitializeProperties& PCIP)
-	: Super(PCIP)
+UTextureExporterTGA::UTextureExporterTGA(const FObjectInitializer& ObjectInitializer)
+	: Super(ObjectInitializer)
 {
 
 	SupportedClass = UTexture2D::StaticClass();
@@ -5149,8 +5149,8 @@ static bool ScanFontBox( const uint8* Data, int32 X, int32 Y, int32& XL, int32&
 
 #define NUM_FONT_CHARS 256
 
-UFontFactory::UFontFactory(const class FPostConstructInitializeProperties& PCIP)
-	: Super(PCIP)
+UFontFactory::UFontFactory(const FObjectInitializer& ObjectInitializer)
+	: Super(ObjectInitializer)
 {
 	SupportedClass = UFont::StaticClass();
 	bEditorImport = false;
@@ -5175,7 +5175,7 @@ UObject* UFontFactory::FactoryCreateBinary
 	check(Type);
 
 	check(Class==UFont::StaticClass());
-	UFont* Font = new( InParent, Name, Flags )UFont(FPostConstructInitializeProperties());
+	UFont* Font = new( InParent, Name, Flags )UFont(FObjectInitializer());
 	// note RF_Public because font textures can be referenced direclty by material expressions
 	UTexture2D* Tex = CastChecked<UTexture2D>( UTextureFactory::FactoryCreateBinary( 
 		UTexture2D::StaticClass(), Font, NAME_None, RF_Public, Context, Type, Buffer, BufferEnd, Warn ) );
@@ -5436,8 +5436,8 @@ void FCustomizableTextObjectFactory::ProcessConstructedObject(UObject* CreatedOb
 /*-----------------------------------------------------------------------------
 UReimportTextureFactory.
 -----------------------------------------------------------------------------*/
-UReimportTextureFactory::UReimportTextureFactory(const class FPostConstructInitializeProperties& PCIP)
-	: Super(PCIP)
+UReimportTextureFactory::UReimportTextureFactory(const FObjectInitializer& ObjectInitializer)
+	: Super(ObjectInitializer)
 {
 
 	SupportedClass = UTexture::StaticClass();
@@ -5577,8 +5577,8 @@ EReimportResult::Type UReimportTextureFactory::Reimport( UObject* Obj )
 /*-----------------------------------------------------------------------------
 UReimportFbxStaticMeshFactory.
 -----------------------------------------------------------------------------*/
-UReimportFbxStaticMeshFactory::UReimportFbxStaticMeshFactory(const class FPostConstructInitializeProperties& PCIP)
-	: Super(PCIP)
+UReimportFbxStaticMeshFactory::UReimportFbxStaticMeshFactory(const FObjectInitializer& ObjectInitializer)
+	: Super(ObjectInitializer)
 {
 
 	SupportedClass = UStaticMesh::StaticClass();
@@ -5762,8 +5762,8 @@ EReimportResult::Type UReimportFbxStaticMeshFactory::Reimport( UObject* Obj )
 /*-----------------------------------------------------------------------------
 UReimportFbxSkeletalMeshFactory
 -----------------------------------------------------------------------------*/ 
-UReimportFbxSkeletalMeshFactory::UReimportFbxSkeletalMeshFactory(const class FPostConstructInitializeProperties& PCIP)
-	: Super(PCIP)
+UReimportFbxSkeletalMeshFactory::UReimportFbxSkeletalMeshFactory(const FObjectInitializer& ObjectInitializer)
+	: Super(ObjectInitializer)
 {
 
 	SupportedClass = USkeletalMesh::StaticClass();
@@ -5931,8 +5931,8 @@ USkeleton* ChooseSkeleton()
 	return SkeletonSelectorWindow->GetSelectedSkeleton();
 }
 
-UReimportFbxAnimSequenceFactory::UReimportFbxAnimSequenceFactory(const class FPostConstructInitializeProperties& PCIP)
-	: Super(PCIP)
+UReimportFbxAnimSequenceFactory::UReimportFbxAnimSequenceFactory(const FObjectInitializer& ObjectInitializer)
+	: Super(ObjectInitializer)
 {
 
 	SupportedClass = UAnimSequence::StaticClass();
@@ -6086,8 +6086,8 @@ public:
 /*------------------------------------------------------------------------------
 	UBlueprintFactory implementation.
 ------------------------------------------------------------------------------*/
-UBlueprintFactory::UBlueprintFactory(const class FPostConstructInitializeProperties& PCIP)
-	: Super(PCIP)
+UBlueprintFactory::UBlueprintFactory(const FObjectInitializer& ObjectInitializer)
+	: Super(ObjectInitializer)
 {
 
 	// Look in the config file to determine what the default base class is, if any
@@ -6181,8 +6181,8 @@ UObject* UBlueprintFactory::FactoryCreateNew(UClass* Class, UObject* InParent, F
 /*------------------------------------------------------------------------------
 	UBlueprintMacroFactory implementation.
 ------------------------------------------------------------------------------*/
-UBlueprintMacroFactory::UBlueprintMacroFactory(const class FPostConstructInitializeProperties& PCIP)
-	: Super(PCIP)
+UBlueprintMacroFactory::UBlueprintMacroFactory(const FObjectInitializer& ObjectInitializer)
+	: Super(ObjectInitializer)
 {
 	
 	bCreateNew = true;
@@ -6243,8 +6243,8 @@ FString UBlueprintMacroFactory::GetDefaultNewAssetName() const
 BlueprintFunctionLibraryFactory implementation.
 ------------------------------------------------------------------------------*/
 
-UBlueprintFunctionLibraryFactory::UBlueprintFunctionLibraryFactory(const class FPostConstructInitializeProperties& PCIP)
-: Super(PCIP)
+UBlueprintFunctionLibraryFactory::UBlueprintFunctionLibraryFactory(const FObjectInitializer& ObjectInitializer)
+: Super(ObjectInitializer)
 {
 	static FBoolConfigValueHelper CanCreateNewHelper(TEXT("CustomBlueprintFunctionLibrary"), TEXT("bCanCreateNew"));
 	bCreateNew = CanCreateNewHelper;
@@ -6309,8 +6309,8 @@ FString UBlueprintFunctionLibraryFactory::GetDefaultNewAssetName() const
 /*------------------------------------------------------------------------------
 	UBlueprintInterfaceFactory implementation.
 ------------------------------------------------------------------------------*/
-UBlueprintInterfaceFactory::UBlueprintInterfaceFactory(const class FPostConstructInitializeProperties& PCIP)
-	: Super(PCIP)
+UBlueprintInterfaceFactory::UBlueprintInterfaceFactory(const FObjectInitializer& ObjectInitializer)
+	: Super(ObjectInitializer)
 {
 
 	bCreateNew = true;
@@ -6378,8 +6378,8 @@ FString UBlueprintInterfaceFactory::GetDefaultNewAssetName() const
 	UCurveFactory implementation.
 ------------------------------------------------------------------------------*/
 
-UCurveFactory::UCurveFactory(const class FPostConstructInitializeProperties& PCIP)
-	: Super(PCIP)
+UCurveFactory::UCurveFactory(const FObjectInitializer& ObjectInitializer)
+	: Super(ObjectInitializer)
 {
 
 	bCreateNew = true;
@@ -6456,8 +6456,8 @@ UObject* UCurveFactory::FactoryCreateNew(UClass* Class, UObject* InParent, FName
 	UCurveImportFactory implementation.
 ------------------------------------------------------------------------------*/
 
-UCurveImportFactory::UCurveImportFactory(const class FPostConstructInitializeProperties& PCIP)
-	: Super(PCIP)
+UCurveImportFactory::UCurveImportFactory(const FObjectInitializer& ObjectInitializer)
+	: Super(ObjectInitializer)
 {
 
 	bCreateNew = false;
@@ -6556,8 +6556,8 @@ UObject* UCurveImportFactory::FactoryCreateText( UClass* InClass, UObject* InPar
 /*------------------------------------------------------------------------------
 	UArchetypeLibraryFactory implementation.
 ------------------------------------------------------------------------------*/
-UObjectLibraryFactory::UObjectLibraryFactory(const class FPostConstructInitializeProperties& PCIP)
-	: Super(PCIP)
+UObjectLibraryFactory::UObjectLibraryFactory(const FObjectInitializer& ObjectInitializer)
+	: Super(ObjectInitializer)
 {
 
 	bCreateNew = true;
@@ -6574,8 +6574,8 @@ UObject* UObjectLibraryFactory::FactoryCreateNew(UClass* Class, UObject* InParen
 	UDataAssetFactory implementation.
 ------------------------------------------------------------------------------*/
 
-UDataAssetFactory::UDataAssetFactory(const class FPostConstructInitializeProperties& PCIP)
-	: Super(PCIP)
+UDataAssetFactory::UDataAssetFactory(const FObjectInitializer& ObjectInitializer)
+	: Super(ObjectInitializer)
 {
 	bCreateNew = true;
 	bEditAfterNew = true;
@@ -6633,8 +6633,8 @@ UObject* UDataAssetFactory::FactoryCreateNew(UClass* Class, UObject* InParent, F
 /*------------------------------------------------------------------------------
 	UDestructibleMeshFactory implementation.
 ------------------------------------------------------------------------------*/
-UDestructibleMeshFactory::UDestructibleMeshFactory(const class FPostConstructInitializeProperties& PCIP)
-	: Super(PCIP)
+UDestructibleMeshFactory::UDestructibleMeshFactory(const FObjectInitializer& ObjectInitializer)
+	: Super(ObjectInitializer)
 {
 
 	bEditorImport = true;
@@ -6707,8 +6707,8 @@ UObject* UDestructibleMeshFactory::FactoryCreateBinary
 /*-----------------------------------------------------------------------------
 	UReimportDestructibleMeshFactory implementation.
 -----------------------------------------------------------------------------*/
-UReimportDestructibleMeshFactory::UReimportDestructibleMeshFactory(const class FPostConstructInitializeProperties& PCIP)
-	: Super(PCIP)
+UReimportDestructibleMeshFactory::UReimportDestructibleMeshFactory(const FObjectInitializer& ObjectInitializer)
+	: Super(ObjectInitializer)
 {
 
 	SupportedClass = UDestructibleMesh::StaticClass();
@@ -6839,8 +6839,8 @@ EReimportResult::Type UReimportDestructibleMeshFactory::Reimport( UObject* Obj )
 /*------------------------------------------------------------------------------
 	UBlendSpaceFactoryNew.
 ------------------------------------------------------------------------------*/
-UBlendSpaceFactoryNew::UBlendSpaceFactoryNew(const class FPostConstructInitializeProperties& PCIP)
-	: Super(PCIP)
+UBlendSpaceFactoryNew::UBlendSpaceFactoryNew(const FObjectInitializer& ObjectInitializer)
+	: Super(ObjectInitializer)
 {
 
 	SupportedClass = UBlendSpace::StaticClass();
@@ -6912,8 +6912,8 @@ void UBlendSpaceFactoryNew::OnTargetSkeletonSelected(const FAssetData& SelectedA
 /*------------------------------------------------------------------------------
 	UBlendSpaceFactory1D.
 ------------------------------------------------------------------------------*/
-UBlendSpaceFactory1D::UBlendSpaceFactory1D(const class FPostConstructInitializeProperties& PCIP)
-	: Super(PCIP)
+UBlendSpaceFactory1D::UBlendSpaceFactory1D(const FObjectInitializer& ObjectInitializer)
+	: Super(ObjectInitializer)
 {
 
 	SupportedClass = UBlendSpace1D::StaticClass();
@@ -6985,8 +6985,8 @@ void UBlendSpaceFactory1D::OnTargetSkeletonSelected(const FAssetData& SelectedAs
 /*------------------------------------------------------------------------------
 	UAimOffsetBlendSpaceFactoryNew.
 ------------------------------------------------------------------------------*/
-UAimOffsetBlendSpaceFactoryNew::UAimOffsetBlendSpaceFactoryNew(const class FPostConstructInitializeProperties& PCIP)
-	: Super(PCIP)
+UAimOffsetBlendSpaceFactoryNew::UAimOffsetBlendSpaceFactoryNew(const FObjectInitializer& ObjectInitializer)
+	: Super(ObjectInitializer)
 {
 	SupportedClass = UAimOffsetBlendSpace::StaticClass();
 	bCreateNew = true;
@@ -7009,8 +7009,8 @@ UObject* UAimOffsetBlendSpaceFactoryNew::FactoryCreateNew(UClass* Class,UObject*
 /*------------------------------------------------------------------------------
 	UAimOffsetBlendSpaceFactory1D.
 ------------------------------------------------------------------------------*/
-UAimOffsetBlendSpaceFactory1D::UAimOffsetBlendSpaceFactory1D(const class FPostConstructInitializeProperties& PCIP)
-	: Super(PCIP)
+UAimOffsetBlendSpaceFactory1D::UAimOffsetBlendSpaceFactory1D(const FObjectInitializer& ObjectInitializer)
+	: Super(ObjectInitializer)
 {
 	SupportedClass = UAimOffsetBlendSpace1D::StaticClass();
 	bCreateNew = true;
@@ -7033,8 +7033,8 @@ UObject* UAimOffsetBlendSpaceFactory1D::FactoryCreateNew(UClass* Class,UObject*
 /*------------------------------------------------------------------------------
 	UEnumFactory implementation.
 ------------------------------------------------------------------------------*/
-UEnumFactory::UEnumFactory(const class FPostConstructInitializeProperties& PCIP)
-	: Super(PCIP)
+UEnumFactory::UEnumFactory(const FObjectInitializer& ObjectInitializer)
+	: Super(ObjectInitializer)
 {
 	SupportedClass = UUserDefinedEnum::StaticClass();
 	bCreateNew = true;
@@ -7062,8 +7062,8 @@ UObject* UEnumFactory::FactoryCreateNew(UClass* Class, UObject* InParent, FName
 /*------------------------------------------------------------------------------
 	UStructureFactory implementation.
 ------------------------------------------------------------------------------*/
-UStructureFactory::UStructureFactory(const class FPostConstructInitializeProperties& PCIP)
-	: Super(PCIP)
+UStructureFactory::UStructureFactory(const FObjectInitializer& ObjectInitializer)
+	: Super(ObjectInitializer)
 {
 	SupportedClass = UUserDefinedStruct::StaticClass();
 	bCreateNew = FStructureEditorUtils::UserDefinedStructEnabled();
@@ -7079,8 +7079,8 @@ UObject* UStructureFactory::FactoryCreateNew(UClass* Class, UObject* InParent, F
 /*-----------------------------------------------------------------------------
 	UForceFeedbackEffectFactory implementation.
 -----------------------------------------------------------------------------*/
-UForceFeedbackEffectFactory::UForceFeedbackEffectFactory(const class FPostConstructInitializeProperties& PCIP)
-	: Super(PCIP)
+UForceFeedbackEffectFactory::UForceFeedbackEffectFactory(const FObjectInitializer& ObjectInitializer)
+	: Super(ObjectInitializer)
 {
 
 	SupportedClass = UForceFeedbackEffect::StaticClass();
@@ -7097,8 +7097,8 @@ UObject* UForceFeedbackEffectFactory::FactoryCreateNew( UClass* InClass, UObject
 /*-----------------------------------------------------------------------------
 USubsurfaceProfileFactory implementation.
 -----------------------------------------------------------------------------*/
-USubsurfaceProfileFactory::USubsurfaceProfileFactory(const class FPostConstructInitializeProperties& PCIP)
-: Super(PCIP)
+USubsurfaceProfileFactory::USubsurfaceProfileFactory(const FObjectInitializer& ObjectInitializer)
+: Super(ObjectInitializer)
 {
 
 	SupportedClass = USubsurfaceProfile::StaticClass();
@@ -7115,8 +7115,8 @@ UObject* USubsurfaceProfileFactory::FactoryCreateNew(UClass* InClass, UObject* I
 /*-----------------------------------------------------------------------------
 	UTouchInterfaceFactory implementation.
 -----------------------------------------------------------------------------*/
-UTouchInterfaceFactory::UTouchInterfaceFactory(const class FPostConstructInitializeProperties& PCIP)
-	: Super(PCIP)
+UTouchInterfaceFactory::UTouchInterfaceFactory(const FObjectInitializer& ObjectInitializer)
+	: Super(ObjectInitializer)
 {
 
 	SupportedClass = UTouchInterface::StaticClass();
@@ -7135,8 +7135,8 @@ UObject* UTouchInterfaceFactory::FactoryCreateNew( UClass* InClass, UObject* InP
 	UCameraAnimFactory implementation.
 ------------------------------------------------------------------------------*/
 
-UCameraAnimFactory::UCameraAnimFactory(const class FPostConstructInitializeProperties& PCIP)
-	: Super(PCIP)
+UCameraAnimFactory::UCameraAnimFactory(const FObjectInitializer& ObjectInitializer)
+	: Super(ObjectInitializer)
 {
 	SupportedClass = UCameraAnim::StaticClass();
 	bCreateNew = true;
diff --git a/Engine/Source/Editor/UnrealEd/Private/Factories/Factory.cpp b/Engine/Source/Editor/UnrealEd/Private/Factories/Factory.cpp
--- a/Engine/Source/Editor/UnrealEd/Private/Factories/Factory.cpp
+++ b/Engine/Source/Editor/UnrealEd/Private/Factories/Factory.cpp
@@ -19,8 +19,8 @@ int32 UFactory::OverwriteYesOrNoToAllState = -1;
 
 bool UFactory::bAllowOneTimeWarningMessages = true;
 
-UFactory::UFactory(const class FPostConstructInitializeProperties& PCIP)
-	: Super(PCIP)
+UFactory::UFactory(const FObjectInitializer& ObjectInitializer)
+	: Super(ObjectInitializer)
 {
 
 }
diff --git a/Engine/Source/Editor/UnrealEd/Private/Factories/NiagaraEffectFactory.cpp b/Engine/Source/Editor/UnrealEd/Private/Factories/NiagaraEffectFactory.cpp
--- a/Engine/Source/Editor/UnrealEd/Private/Factories/NiagaraEffectFactory.cpp
+++ b/Engine/Source/Editor/UnrealEd/Private/Factories/NiagaraEffectFactory.cpp
@@ -7,8 +7,8 @@
 
 #define LOCTEXT_NAMESPACE "NiagaraEffectFactory"
 
-UNiagaraEffectFactoryNew::UNiagaraEffectFactoryNew(const class FPostConstructInitializeProperties& PCIP)
-: Super(PCIP)
+UNiagaraEffectFactoryNew::UNiagaraEffectFactoryNew(const FObjectInitializer& ObjectInitializer)
+: Super(ObjectInitializer)
 {
 
 	SupportedClass = UNiagaraEffect::StaticClass();
diff --git a/Engine/Source/Editor/UnrealEd/Private/Factories/NiagaraScriptFactory.cpp b/Engine/Source/Editor/UnrealEd/Private/Factories/NiagaraScriptFactory.cpp
--- a/Engine/Source/Editor/UnrealEd/Private/Factories/NiagaraScriptFactory.cpp
+++ b/Engine/Source/Editor/UnrealEd/Private/Factories/NiagaraScriptFactory.cpp
@@ -7,8 +7,8 @@
 
 #define LOCTEXT_NAMESPACE "NiagaraScriptFactory"
 
-UNiagaraScriptFactoryNew::UNiagaraScriptFactoryNew(const class FPostConstructInitializeProperties& PCIP)
-	: Super(PCIP)
+UNiagaraScriptFactoryNew::UNiagaraScriptFactoryNew(const FObjectInitializer& ObjectInitializer)
+	: Super(ObjectInitializer)
 {
 
 	SupportedClass = UNiagaraScript::StaticClass();
diff --git a/Engine/Source/Editor/UnrealEd/Private/Factories/SkeletonFactory.cpp b/Engine/Source/Editor/UnrealEd/Private/Factories/SkeletonFactory.cpp
--- a/Engine/Source/Editor/UnrealEd/Private/Factories/SkeletonFactory.cpp
+++ b/Engine/Source/Editor/UnrealEd/Private/Factories/SkeletonFactory.cpp
@@ -16,8 +16,8 @@
 	USkeletonFactory implementation.
 ------------------------------------------------------------------------------*/
 
-USkeletonFactory::USkeletonFactory(const class FPostConstructInitializeProperties& PCIP)
-	: Super(PCIP)
+USkeletonFactory::USkeletonFactory(const FObjectInitializer& ObjectInitializer)
+	: Super(ObjectInitializer)
 {
 	bCreateNew = false;
 	SupportedClass = USkeleton::StaticClass();
diff --git a/Engine/Source/Editor/UnrealEd/Private/Factories/SlateBrushAssetFactory.cpp b/Engine/Source/Editor/UnrealEd/Private/Factories/SlateBrushAssetFactory.cpp
--- a/Engine/Source/Editor/UnrealEd/Private/Factories/SlateBrushAssetFactory.cpp
+++ b/Engine/Source/Editor/UnrealEd/Private/Factories/SlateBrushAssetFactory.cpp
@@ -6,8 +6,8 @@
 
 #define LOCTEXT_NAMESPACE "SlateBrushAssetFactory"
 
-USlateBrushAssetFactory::USlateBrushAssetFactory( const class FPostConstructInitializeProperties& PCIP )
- : Super(PCIP)
+USlateBrushAssetFactory::USlateBrushAssetFactory( const FObjectInitializer& ObjectInitializer )
+ : Super(ObjectInitializer)
 {
 	bCreateNew = true;
 	bEditAfterNew = true;
diff --git a/Engine/Source/Editor/UnrealEd/Private/Factories/SlateWidgetStyleAssetFactory.cpp b/Engine/Source/Editor/UnrealEd/Private/Factories/SlateWidgetStyleAssetFactory.cpp
--- a/Engine/Source/Editor/UnrealEd/Private/Factories/SlateWidgetStyleAssetFactory.cpp
+++ b/Engine/Source/Editor/UnrealEd/Private/Factories/SlateWidgetStyleAssetFactory.cpp
@@ -33,8 +33,8 @@ public:
 };
 
 
-USlateWidgetStyleAssetFactory::USlateWidgetStyleAssetFactory( const class FPostConstructInitializeProperties& PCIP )
- : Super(PCIP)
+USlateWidgetStyleAssetFactory::USlateWidgetStyleAssetFactory( const FObjectInitializer& ObjectInitializer )
+ : Super(ObjectInitializer)
 {
 	bCreateNew = true;
 	bEditAfterNew = true;
diff --git a/Engine/Source/Editor/UnrealEd/Private/Factories/TTFontImport.cpp b/Engine/Source/Editor/UnrealEd/Private/Factories/TTFontImport.cpp
--- a/Engine/Source/Editor/UnrealEd/Private/Factories/TTFontImport.cpp
+++ b/Engine/Source/Editor/UnrealEd/Private/Factories/TTFontImport.cpp
@@ -35,8 +35,8 @@ DEFINE_LOG_CATEGORY_STATIC(LogTTFontImport, Log, All);
 #define LOCTEXT_NAMESPACE "TTFontImport"
 
 
-UTrueTypeFontFactory::UTrueTypeFontFactory(const class FPostConstructInitializeProperties& PCIP)
-	: Super(PCIP)
+UTrueTypeFontFactory::UTrueTypeFontFactory(const FObjectInitializer& ObjectInitializer)
+	: Super(ObjectInitializer)
 {
 	SupportedClass = UFont::StaticClass();
 	bCreateNew = true;
@@ -149,7 +149,7 @@ UObject* UTrueTypeFontFactory::FactoryCreateNew(
 	}
 
 	// Create font and its texture.
-	UFont* Font = new( InParent, Name, Flags )UFont(FPostConstructInitializeProperties());
+	UFont* Font = new( InParent, Name, Flags )UFont(FObjectInitializer());
 	
 	if (ImportOptions->Data.bUseDistanceFieldAlpha)
 	{
@@ -768,7 +768,7 @@ UTexture2D* UTrueTypeFontFactory::CreateTextureFromDC( UFont* Font, HDC dc, int3
 	}
 
 	// Create texture for page.
-	UTexture2D* Texture = new(Font, *TextureString)UTexture2D(FPostConstructInitializeProperties());
+	UTexture2D* Texture = new(Font, *TextureString)UTexture2D(FObjectInitializer());
 
 	// note RF_Public because font textures can be referenced directly by material expressions
 	Texture->SetFlags(RF_Public);
@@ -1665,7 +1665,7 @@ UTexture2D* UTrueTypeFontFactory::CreateTextureFromBitmap( UFont* Font, uint8* B
 	}
 
 	// Create texture for page.
-	UTexture2D* Texture = new(Font, *TextureString)UTexture2D(FPostConstructInitializeProperties());
+	UTexture2D* Texture = new(Font, *TextureString)UTexture2D(FObjectInitializer());
 
 	// note RF_Public because font textures can be referenced directly by material expressions
 	Texture->SetFlags(RF_Public);
diff --git a/Engine/Source/Editor/UnrealEd/Private/Factories/VectorFieldFactory.cpp b/Engine/Source/Editor/UnrealEd/Private/Factories/VectorFieldFactory.cpp
--- a/Engine/Source/Editor/UnrealEd/Private/Factories/VectorFieldFactory.cpp
+++ b/Engine/Source/Editor/UnrealEd/Private/Factories/VectorFieldFactory.cpp
@@ -170,8 +170,8 @@ static bool ParseFGA(
 	return true;
 }
 
-UVectorFieldStaticFactory::UVectorFieldStaticFactory(const class FPostConstructInitializeProperties& PCIP)
-	: Super(PCIP)
+UVectorFieldStaticFactory::UVectorFieldStaticFactory(const FObjectInitializer& ObjectInitializer)
+	: Super(ObjectInitializer)
 {
 
 	SupportedClass = UVectorFieldStatic::StaticClass();
@@ -272,8 +272,8 @@ bool UVectorFieldStaticFactory::FactoryCanImport( const FString& Filename )
 /*-----------------------------------------------------------------------------
 	Static vector field reimport factory.
 -----------------------------------------------------------------------------*/
-UReimportVectorFieldStaticFactory::UReimportVectorFieldStaticFactory(const class FPostConstructInitializeProperties& PCIP)
-	: Super(PCIP)
+UReimportVectorFieldStaticFactory::UReimportVectorFieldStaticFactory(const FObjectInitializer& ObjectInitializer)
+	: Super(ObjectInitializer)
 {
 	SupportedClass = UVectorFieldStatic::StaticClass();
 
diff --git a/Engine/Source/Editor/UnrealEd/Private/Factories/WorldFactory.cpp b/Engine/Source/Editor/UnrealEd/Private/Factories/WorldFactory.cpp
--- a/Engine/Source/Editor/UnrealEd/Private/Factories/WorldFactory.cpp
+++ b/Engine/Source/Editor/UnrealEd/Private/Factories/WorldFactory.cpp
@@ -4,8 +4,8 @@
 
 #define LOCTEXT_NAMESPACE "WorldFactory"
 
-UWorldFactory::UWorldFactory(const class FPostConstructInitializeProperties& PCIP)
-	: Super(PCIP)
+UWorldFactory::UWorldFactory(const FObjectInitializer& ObjectInitializer)
+	: Super(ObjectInitializer)
 {
 	bCreateNew = UEditorEngine::IsUsingWorldAssets();
 	SupportedClass = UWorld::StaticClass();
diff --git a/Engine/Source/Editor/UnrealEd/Private/Fbx/FbxAnimSequenceImportData.cpp b/Engine/Source/Editor/UnrealEd/Private/Fbx/FbxAnimSequenceImportData.cpp
--- a/Engine/Source/Editor/UnrealEd/Private/Fbx/FbxAnimSequenceImportData.cpp
+++ b/Engine/Source/Editor/UnrealEd/Private/Fbx/FbxAnimSequenceImportData.cpp
@@ -2,8 +2,8 @@
 
 #include "UnrealEd.h"
 
-UFbxAnimSequenceImportData::UFbxAnimSequenceImportData(const class FPostConstructInitializeProperties& PCIP)
-	: Super(PCIP)
+UFbxAnimSequenceImportData::UFbxAnimSequenceImportData(const FObjectInitializer& ObjectInitializer)
+	: Super(ObjectInitializer)
 	, bImportCustomAttribute(true)
 {
 	
diff --git a/Engine/Source/Editor/UnrealEd/Private/Fbx/FbxAssetImportData.cpp b/Engine/Source/Editor/UnrealEd/Private/Fbx/FbxAssetImportData.cpp
--- a/Engine/Source/Editor/UnrealEd/Private/Fbx/FbxAssetImportData.cpp
+++ b/Engine/Source/Editor/UnrealEd/Private/Fbx/FbxAssetImportData.cpp
@@ -2,8 +2,8 @@
 
 #include "UnrealEd.h"
 
-UFbxAssetImportData::UFbxAssetImportData(const class FPostConstructInitializeProperties& PCIP)
-	: Super(PCIP)
+UFbxAssetImportData::UFbxAssetImportData(const FObjectInitializer& ObjectInitializer)
+	: Super(ObjectInitializer)
 	, ImportTranslation(0)
 	, ImportRotation(0)
 	, ImportUniformScale(1.0f)
diff --git a/Engine/Source/Editor/UnrealEd/Private/Fbx/FbxFactory.cpp b/Engine/Source/Editor/UnrealEd/Private/Fbx/FbxFactory.cpp
--- a/Engine/Source/Editor/UnrealEd/Private/Fbx/FbxFactory.cpp
+++ b/Engine/Source/Editor/UnrealEd/Private/Fbx/FbxFactory.cpp
@@ -40,8 +40,8 @@
 
 #define LOCTEXT_NAMESPACE "FBXFactory"
 
-UFbxFactory::UFbxFactory(const class FPostConstructInitializeProperties& PCIP)
-	: Super(PCIP)
+UFbxFactory::UFbxFactory(const FObjectInitializer& ObjectInitializer)
+	: Super(ObjectInitializer)
 {
 	SupportedClass = NULL;
 	Formats.Add(TEXT("fbx;FBX meshes and animations"));
@@ -580,8 +580,8 @@ void UFbxFactory::CleanUp()
 	}
 }
 
-UFbxImportUI::UFbxImportUI(const class FPostConstructInitializeProperties& PCIP)
-	: Super(PCIP)
+UFbxImportUI::UFbxImportUI(const FObjectInitializer& ObjectInitializer)
+	: Super(ObjectInitializer)
 {
 	bCombineMeshes = true;
 
diff --git a/Engine/Source/Editor/UnrealEd/Private/Fbx/FbxMaterialImport.cpp b/Engine/Source/Editor/UnrealEd/Private/Fbx/FbxMaterialImport.cpp
--- a/Engine/Source/Editor/UnrealEd/Private/Fbx/FbxMaterialImport.cpp
+++ b/Engine/Source/Editor/UnrealEd/Private/Fbx/FbxMaterialImport.cpp
@@ -111,7 +111,7 @@ UTexture* UnFbx::FFbxImporter::ImportTexture( FbxFileTexture* FbxTexture, bool b
 	{
 		UE_LOG(LogFbxMaterialImport, Verbose, TEXT("Loading texture file %s"),*Filename);
 		const uint8* PtrTexture = DataBinary.GetData();
-		UTextureFactory* TextureFact = new UTextureFactory(FPostConstructInitializeProperties());
+		UTextureFactory* TextureFact = new UTextureFactory(FObjectInitializer());
 		TextureFact->AddToRoot();
 
 		// save texture settings if texture exist
@@ -423,7 +423,7 @@ void UnFbx::FFbxImporter::CreateUnrealMaterial(FbxSurfaceMaterial& FbxMaterial,
 	
 
 	// create an unreal material asset
-	UMaterialFactoryNew* MaterialFactory = new UMaterialFactoryNew(FPostConstructInitializeProperties());
+	UMaterialFactoryNew* MaterialFactory = new UMaterialFactoryNew(FObjectInitializer());
 	
 	UMaterial* UnrealMaterial = (UMaterial*)MaterialFactory->FactoryCreateNew(
 		UMaterial::StaticClass(), Package, *MaterialFullName, RF_Standalone|RF_Public, NULL, GWarn );
diff --git a/Engine/Source/Editor/UnrealEd/Private/Fbx/FbxMeshImportData.cpp b/Engine/Source/Editor/UnrealEd/Private/Fbx/FbxMeshImportData.cpp
--- a/Engine/Source/Editor/UnrealEd/Private/Fbx/FbxMeshImportData.cpp
+++ b/Engine/Source/Editor/UnrealEd/Private/Fbx/FbxMeshImportData.cpp
@@ -2,8 +2,8 @@
 
 #include "UnrealEd.h"
 
-UFbxMeshImportData::UFbxMeshImportData(const class FPostConstructInitializeProperties& PCIP)
-	: Super(PCIP)
+UFbxMeshImportData::UFbxMeshImportData(const FObjectInitializer& ObjectInitializer)
+	: Super(ObjectInitializer)
 {
 	NormalImportMethod = FBXNIM_ComputeNormals;
 }
diff --git a/Engine/Source/Editor/UnrealEd/Private/Fbx/FbxSkeletalMeshImportData.cpp b/Engine/Source/Editor/UnrealEd/Private/Fbx/FbxSkeletalMeshImportData.cpp
--- a/Engine/Source/Editor/UnrealEd/Private/Fbx/FbxSkeletalMeshImportData.cpp
+++ b/Engine/Source/Editor/UnrealEd/Private/Fbx/FbxSkeletalMeshImportData.cpp
@@ -2,8 +2,8 @@
 
 #include "UnrealEd.h"
 
-UFbxSkeletalMeshImportData::UFbxSkeletalMeshImportData(const class FPostConstructInitializeProperties& PCIP)
-	: Super(PCIP)
+UFbxSkeletalMeshImportData::UFbxSkeletalMeshImportData(const FObjectInitializer& ObjectInitializer)
+	: Super(ObjectInitializer)
 	, bImportMeshesInBoneHierarchy(true)
 {
 }
diff --git a/Engine/Source/Editor/UnrealEd/Private/Fbx/FbxStaticMeshImport.cpp b/Engine/Source/Editor/UnrealEd/Private/Fbx/FbxStaticMeshImport.cpp
--- a/Engine/Source/Editor/UnrealEd/Private/Fbx/FbxStaticMeshImport.cpp
+++ b/Engine/Source/Editor/UnrealEd/Private/Fbx/FbxStaticMeshImport.cpp
@@ -936,7 +936,7 @@ UStaticMesh* UnFbx::FFbxImporter::ImportStaticMeshAsSingle(UObject* InParent, TA
 	}
 	else
 	{
-		StaticMesh = new(Package,FName(*MeshName),Flags|RF_Public) UStaticMesh(FPostConstructInitializeProperties());
+		StaticMesh = new(Package,FName(*MeshName),Flags|RF_Public) UStaticMesh(FObjectInitializer());
 	}
 
 	if (StaticMesh->SourceModels.Num() < LODIndex+1)
diff --git a/Engine/Source/Editor/UnrealEd/Private/Fbx/FbxStaticMeshImportData.cpp b/Engine/Source/Editor/UnrealEd/Private/Fbx/FbxStaticMeshImportData.cpp
--- a/Engine/Source/Editor/UnrealEd/Private/Fbx/FbxStaticMeshImportData.cpp
+++ b/Engine/Source/Editor/UnrealEd/Private/Fbx/FbxStaticMeshImportData.cpp
@@ -2,8 +2,8 @@
 
 #include "UnrealEd.h"
 
-UFbxStaticMeshImportData::UFbxStaticMeshImportData(const class FPostConstructInitializeProperties& PCIP)
-	: Super(PCIP)
+UFbxStaticMeshImportData::UFbxStaticMeshImportData(const FObjectInitializer& ObjectInitializer)
+	: Super(ObjectInitializer)
 {
 	StaticMeshLODGroup = NAME_None;
 	bRemoveDegenerates = true;
diff --git a/Engine/Source/Editor/UnrealEd/Private/Fbx/FbxTextureImportData.cpp b/Engine/Source/Editor/UnrealEd/Private/Fbx/FbxTextureImportData.cpp
--- a/Engine/Source/Editor/UnrealEd/Private/Fbx/FbxTextureImportData.cpp
+++ b/Engine/Source/Editor/UnrealEd/Private/Fbx/FbxTextureImportData.cpp
@@ -2,8 +2,8 @@
 
 #include "UnrealEd.h"
 
-UFbxTextureImportData::UFbxTextureImportData(const class FPostConstructInitializeProperties& PCIP)
-	: Super(PCIP)
+UFbxTextureImportData::UFbxTextureImportData(const FObjectInitializer& ObjectInitializer)
+	: Super(ObjectInitializer)
 {
 	
 }
diff --git a/Engine/Source/Editor/UnrealEd/Private/GeomFitUtils.cpp b/Engine/Source/Editor/UnrealEd/Private/GeomFitUtils.cpp
--- a/Engine/Source/Editor/UnrealEd/Private/GeomFitUtils.cpp
+++ b/Engine/Source/Editor/UnrealEd/Private/GeomFitUtils.cpp
@@ -82,7 +82,7 @@ int32 GenerateKDopAsSimpleCollision(UStaticMesh* StaticMesh, const TArray<FVecto
 		maxDist.Add(-MY_FLTMAX);
 
 	// Construct temporary UModel for kdop creation. We keep no refs to it, so it can be GC'd.
-	UModel* TempModel = new UModel(FPostConstructInitializeProperties(),NULL,1);
+	UModel* TempModel = new UModel(FObjectInitializer(),NULL,1);
 
 	// For each vertex, project along each kdop direction, to find the max in that direction.
 	const FStaticMeshLODResources& RenderData = StaticMesh->RenderData->LODResources[0];
diff --git a/Engine/Source/Editor/UnrealEd/Private/GroupActor.cpp b/Engine/Source/Editor/UnrealEd/Private/GroupActor.cpp
--- a/Engine/Source/Editor/UnrealEd/Private/GroupActor.cpp
+++ b/Engine/Source/Editor/UnrealEd/Private/GroupActor.cpp
@@ -8,12 +8,12 @@ const FLinearColor BOXCOLOR_LOCKEDGROUPS( 0.0f, 1.0f, 0.0f );
 const FLinearColor BOXCOLOR_UNLOCKEDGROUPS( 1.0f, 0.0f, 0.0f );
 
 
-AGroupActor::AGroupActor(const class FPostConstructInitializeProperties& PCIP)
-	: Super(PCIP)
+AGroupActor::AGroupActor(const FObjectInitializer& ObjectInitializer)
+	: Super(ObjectInitializer)
 {
 	bLocked = true;
 
-	TSubobjectPtr<USceneComponent> GroupComponent = PCIP.CreateDefaultSubobject<USceneComponent>(this, TEXT("GroupComponent"));
+	TSubobjectPtr<USceneComponent> GroupComponent = ObjectInitializer.CreateDefaultSubobject<USceneComponent>(this, TEXT("GroupComponent"));
 	RootComponent = GroupComponent;
 }
 
diff --git a/Engine/Source/Editor/UnrealEd/Private/LevelEditorViewport.cpp b/Engine/Source/Editor/UnrealEd/Private/LevelEditorViewport.cpp
--- a/Engine/Source/Editor/UnrealEd/Private/LevelEditorViewport.cpp
+++ b/Engine/Source/Editor/UnrealEd/Private/LevelEditorViewport.cpp
@@ -406,7 +406,7 @@ static UObject* GetOrCreateMaterialFromTexture( UTexture* UnrealTexture )
 	}
 
 	// create an unreal material asset
-	UMaterialFactoryNew* MaterialFactory = new UMaterialFactoryNew( FPostConstructInitializeProperties() );
+	UMaterialFactoryNew* MaterialFactory = new UMaterialFactoryNew( FObjectInitializer() );
 
 	UMaterial* UnrealMaterial = (UMaterial*)MaterialFactory->FactoryCreateNew(
 		UMaterial::StaticClass(), Package, *MaterialFullName, RF_Standalone | RF_Public, NULL, GWarn );
diff --git a/Engine/Source/Editor/UnrealEd/Private/Lightmass/LightmassRender.cpp b/Engine/Source/Editor/UnrealEd/Private/Lightmass/LightmassRender.cpp
--- a/Engine/Source/Editor/UnrealEd/Private/Lightmass/LightmassRender.cpp
+++ b/Engine/Source/Editor/UnrealEd/Private/Lightmass/LightmassRender.cpp
@@ -1048,7 +1048,7 @@ bool FLightmassMaterialRenderer::CreateRenderTarget(EPixelFormat InFormat, int32
 
 	if (RenderTarget == NULL)
 	{
-		RenderTarget = new UTextureRenderTarget2D(FPostConstructInitializeProperties());
+		RenderTarget = new UTextureRenderTarget2D(FObjectInitializer());
 		check(RenderTarget);
 		RenderTarget->AddToRoot();
 		RenderTarget->ClearColor = FLinearColor(0.0f, 0.0f, 0.0f, 0.0f);
diff --git a/Engine/Source/Editor/UnrealEd/Private/LightmassOptionsObject.cpp b/Engine/Source/Editor/UnrealEd/Private/LightmassOptionsObject.cpp
--- a/Engine/Source/Editor/UnrealEd/Private/LightmassOptionsObject.cpp
+++ b/Engine/Source/Editor/UnrealEd/Private/LightmassOptionsObject.cpp
@@ -2,7 +2,7 @@
 
 #include "UnrealEd.h"
 
-ULightmassOptionsObject::ULightmassOptionsObject(const class FPostConstructInitializeProperties& PCIP)
-	: Super(PCIP)
+ULightmassOptionsObject::ULightmassOptionsObject(const FObjectInitializer& ObjectInitializer)
+	: Super(ObjectInitializer)
 {
 }
diff --git a/Engine/Source/Editor/UnrealEd/Private/MaterialGraph.cpp b/Engine/Source/Editor/UnrealEd/Private/MaterialGraph.cpp
--- a/Engine/Source/Editor/UnrealEd/Private/MaterialGraph.cpp
+++ b/Engine/Source/Editor/UnrealEd/Private/MaterialGraph.cpp
@@ -14,8 +14,8 @@
 
 #define LOCTEXT_NAMESPACE "MaterialGraph"
 
-UMaterialGraph::UMaterialGraph(const FPostConstructInitializeProperties& PCIP)
-	: Super(PCIP)
+UMaterialGraph::UMaterialGraph(const FObjectInitializer& ObjectInitializer)
+	: Super(ObjectInitializer)
 {
 }
 
diff --git a/Engine/Source/Editor/UnrealEd/Private/MaterialGraphNode.cpp b/Engine/Source/Editor/UnrealEd/Private/MaterialGraphNode.cpp
--- a/Engine/Source/Editor/UnrealEd/Private/MaterialGraphNode.cpp
+++ b/Engine/Source/Editor/UnrealEd/Private/MaterialGraphNode.cpp
@@ -39,8 +39,8 @@
 /////////////////////////////////////////////////////
 // UMaterialGraphNode
 
-UMaterialGraphNode::UMaterialGraphNode(const class FPostConstructInitializeProperties& PCIP)
-	: Super(PCIP)
+UMaterialGraphNode::UMaterialGraphNode(const FObjectInitializer& ObjectInitializer)
+	: Super(ObjectInitializer)
 	, bPreviewNeedsUpdate(false)
 	, bIsErrorExpression(false)
 	, bIsPreviewExpression(false)
diff --git a/Engine/Source/Editor/UnrealEd/Private/MaterialGraphNode_Base.cpp b/Engine/Source/Editor/UnrealEd/Private/MaterialGraphNode_Base.cpp
--- a/Engine/Source/Editor/UnrealEd/Private/MaterialGraphNode_Base.cpp
+++ b/Engine/Source/Editor/UnrealEd/Private/MaterialGraphNode_Base.cpp
@@ -10,8 +10,8 @@
 /////////////////////////////////////////////////////
 // UMaterialGraphNode_Base
 
-UMaterialGraphNode_Base::UMaterialGraphNode_Base(const class FPostConstructInitializeProperties& PCIP)
-	: Super(PCIP)
+UMaterialGraphNode_Base::UMaterialGraphNode_Base(const FObjectInitializer& ObjectInitializer)
+	: Super(ObjectInitializer)
 {
 }
 
diff --git a/Engine/Source/Editor/UnrealEd/Private/MaterialGraphNode_Comment.cpp b/Engine/Source/Editor/UnrealEd/Private/MaterialGraphNode_Comment.cpp
--- a/Engine/Source/Editor/UnrealEd/Private/MaterialGraphNode_Comment.cpp
+++ b/Engine/Source/Editor/UnrealEd/Private/MaterialGraphNode_Comment.cpp
@@ -12,8 +12,8 @@
 /////////////////////////////////////////////////////
 // UMaterialGraphNode_Comment
 
-UMaterialGraphNode_Comment::UMaterialGraphNode_Comment(const class FPostConstructInitializeProperties& PCIP)
-	: Super(PCIP)
+UMaterialGraphNode_Comment::UMaterialGraphNode_Comment(const FObjectInitializer& ObjectInitializer)
+	: Super(ObjectInitializer)
 {
 }
 
diff --git a/Engine/Source/Editor/UnrealEd/Private/MaterialGraphNode_Root.cpp b/Engine/Source/Editor/UnrealEd/Private/MaterialGraphNode_Root.cpp
--- a/Engine/Source/Editor/UnrealEd/Private/MaterialGraphNode_Root.cpp
+++ b/Engine/Source/Editor/UnrealEd/Private/MaterialGraphNode_Root.cpp
@@ -14,8 +14,8 @@
 /////////////////////////////////////////////////////
 // UMaterialGraphNode_Root
 
-UMaterialGraphNode_Root::UMaterialGraphNode_Root(const class FPostConstructInitializeProperties& PCIP)
-	: Super(PCIP)
+UMaterialGraphNode_Root::UMaterialGraphNode_Root(const FObjectInitializer& ObjectInitializer)
+	: Super(ObjectInitializer)
 {
 }
 
diff --git a/Engine/Source/Editor/UnrealEd/Private/MaterialGraphSchema.cpp b/Engine/Source/Editor/UnrealEd/Private/MaterialGraphSchema.cpp
--- a/Engine/Source/Editor/UnrealEd/Private/MaterialGraphSchema.cpp
+++ b/Engine/Source/Editor/UnrealEd/Private/MaterialGraphSchema.cpp
@@ -144,8 +144,8 @@ UEdGraphNode* FMaterialGraphSchemaAction_Paste::PerformAction(class UEdGraph* Pa
 //////////////////////////
 // UMaterialGraphSchema //
 
-UMaterialGraphSchema::UMaterialGraphSchema(const class FPostConstructInitializeProperties& PCIP)
-	: Super(PCIP)
+UMaterialGraphSchema::UMaterialGraphSchema(const FObjectInitializer& ObjectInitializer)
+	: Super(ObjectInitializer)
 {
 	PC_Mask = TEXT("mask");
 	PC_Required = TEXT("required");
diff --git a/Engine/Source/Editor/UnrealEd/Private/Niagara/EdGraphSchema_Niagara.cpp b/Engine/Source/Editor/UnrealEd/Private/Niagara/EdGraphSchema_Niagara.cpp
--- a/Engine/Source/Editor/UnrealEd/Private/Niagara/EdGraphSchema_Niagara.cpp
+++ b/Engine/Source/Editor/UnrealEd/Private/Niagara/EdGraphSchema_Niagara.cpp
@@ -91,8 +91,8 @@ void FNiagaraSchemaAction_NewNode::AddReferencedObjects( FReferenceCollector& Co
 
 //////////////////////////////////////////////////////////////////////////
 
-UEdGraphSchema_Niagara::UEdGraphSchema_Niagara(const class FPostConstructInitializeProperties& PCIP)
-	: Super(PCIP)
+UEdGraphSchema_Niagara::UEdGraphSchema_Niagara(const FObjectInitializer& ObjectInitializer)
+	: Super(ObjectInitializer)
 {
 
 	PC_Float = TEXT("float");
diff --git a/Engine/Source/Editor/UnrealEd/Private/Niagara/NiagaraNode.cpp b/Engine/Source/Editor/UnrealEd/Private/Niagara/NiagaraNode.cpp
--- a/Engine/Source/Editor/UnrealEd/Private/Niagara/NiagaraNode.cpp
+++ b/Engine/Source/Editor/UnrealEd/Private/Niagara/NiagaraNode.cpp
@@ -2,8 +2,8 @@
 
 #include "UnrealEd.h"
 
-UNiagaraNode::UNiagaraNode(const class FPostConstructInitializeProperties& PCIP)
-	: Super(PCIP)
+UNiagaraNode::UNiagaraNode(const FObjectInitializer& ObjectInitializer)
+	: Super(ObjectInitializer)
 {
 }
 
diff --git a/Engine/Source/Editor/UnrealEd/Private/Niagara/NiagaraNodeGetAttr.cpp b/Engine/Source/Editor/UnrealEd/Private/Niagara/NiagaraNodeGetAttr.cpp
--- a/Engine/Source/Editor/UnrealEd/Private/Niagara/NiagaraNodeGetAttr.cpp
+++ b/Engine/Source/Editor/UnrealEd/Private/Niagara/NiagaraNodeGetAttr.cpp
@@ -2,8 +2,8 @@
 
 #include "UnrealEd.h"
 
-UNiagaraNodeGetAttr::UNiagaraNodeGetAttr(const class FPostConstructInitializeProperties& PCIP)
-	: Super(PCIP)
+UNiagaraNodeGetAttr::UNiagaraNodeGetAttr(const FObjectInitializer& ObjectInitializer)
+	: Super(ObjectInitializer)
 {
 }
 
diff --git a/Engine/Source/Editor/UnrealEd/Private/Niagara/NiagaraNodeOp.cpp b/Engine/Source/Editor/UnrealEd/Private/Niagara/NiagaraNodeOp.cpp
--- a/Engine/Source/Editor/UnrealEd/Private/Niagara/NiagaraNodeOp.cpp
+++ b/Engine/Source/Editor/UnrealEd/Private/Niagara/NiagaraNodeOp.cpp
@@ -5,8 +5,8 @@
 #include "GraphEditorSettings.h"
 #include "VectorVM.h"
 
-UNiagaraNodeOp::UNiagaraNodeOp(const class FPostConstructInitializeProperties& PCIP)
-	: Super(PCIP)
+UNiagaraNodeOp::UNiagaraNodeOp(const FObjectInitializer& ObjectInitializer)
+	: Super(ObjectInitializer)
 {
 }
 
diff --git a/Engine/Source/Editor/UnrealEd/Private/Niagara/NiagaraNodeOutputUpdate.cpp b/Engine/Source/Editor/UnrealEd/Private/Niagara/NiagaraNodeOutputUpdate.cpp
--- a/Engine/Source/Editor/UnrealEd/Private/Niagara/NiagaraNodeOutputUpdate.cpp
+++ b/Engine/Source/Editor/UnrealEd/Private/Niagara/NiagaraNodeOutputUpdate.cpp
@@ -4,8 +4,8 @@
 #include "BlueprintGraphDefinitions.h"
 #include "GraphEditorSettings.h"
 
-UNiagaraNodeOutputUpdate::UNiagaraNodeOutputUpdate(const class FPostConstructInitializeProperties& PCIP)
-	: Super(PCIP)
+UNiagaraNodeOutputUpdate::UNiagaraNodeOutputUpdate(const FObjectInitializer& ObjectInitializer)
+	: Super(ObjectInitializer)
 {
 }
 
diff --git a/Engine/Source/Editor/UnrealEd/Private/Niagara/NiagaraScriptSource.cpp b/Engine/Source/Editor/UnrealEd/Private/Niagara/NiagaraScriptSource.cpp
--- a/Engine/Source/Editor/UnrealEd/Private/Niagara/NiagaraScriptSource.cpp
+++ b/Engine/Source/Editor/UnrealEd/Private/Niagara/NiagaraScriptSource.cpp
@@ -8,8 +8,8 @@
 //////////////////////////////////////////////////////////////////////////
 // NiagaraGraph
 
-UNiagaraGraph::UNiagaraGraph(const class FPostConstructInitializeProperties& PCIP)
-	: Super(PCIP)
+UNiagaraGraph::UNiagaraGraph(const FObjectInitializer& ObjectInitializer)
+	: Super(ObjectInitializer)
 {
 
 	Schema = UEdGraphSchema_Niagara::StaticClass();
@@ -24,8 +24,8 @@ class UNiagaraScriptSource* UNiagaraGraph::GetSource() const
 //////////////////////////////////////////////////////////////////////////
 // UNiagraScriptSource
 
-UNiagaraScriptSource::UNiagaraScriptSource(const class FPostConstructInitializeProperties& PCIP)
-	: Super(PCIP)
+UNiagaraScriptSource::UNiagaraScriptSource(const FObjectInitializer& ObjectInitializer)
+	: Super(ObjectInitializer)
 {
 }
 
diff --git a/Engine/Source/Editor/UnrealEd/Private/PreferenceStubs.cpp b/Engine/Source/Editor/UnrealEd/Private/PreferenceStubs.cpp
--- a/Engine/Source/Editor/UnrealEd/Private/PreferenceStubs.cpp
+++ b/Engine/Source/Editor/UnrealEd/Private/PreferenceStubs.cpp
@@ -9,33 +9,33 @@
 //
 // UCascadeOptions
 // 
-UCascadeOptions::UCascadeOptions(const class FPostConstructInitializeProperties& PCIP)
-	: Super(PCIP)
+UCascadeOptions::UCascadeOptions(const FObjectInitializer& ObjectInitializer)
+	: Super(ObjectInitializer)
 {
 }
 
 ///////////////////////////////////////////////////////////////////////////////////////////////////
 ////////////////// UPhATSimOptions ////////////////////////////////////////////////////////////////
 ///////////////////////////////////////////////////////////////////////////////////////////////////
-UPhATSimOptions::UPhATSimOptions(const class FPostConstructInitializeProperties& PCIP)
-	: Super(PCIP)
+UPhATSimOptions::UPhATSimOptions(const FObjectInitializer& ObjectInitializer)
+	: Super(ObjectInitializer)
 {
 	PhysicsBlend = 1.0f;
 }
 
-UMaterialEditorOptions::UMaterialEditorOptions(const class FPostConstructInitializeProperties& PCIP)
-	: Super(PCIP)
+UMaterialEditorOptions::UMaterialEditorOptions(const FObjectInitializer& ObjectInitializer)
+	: Super(ObjectInitializer)
 {
 }
 
 
-UCurveEdOptions::UCurveEdOptions(const class FPostConstructInitializeProperties& PCIP)
-	: Super(PCIP)
+UCurveEdOptions::UCurveEdOptions(const FObjectInitializer& ObjectInitializer)
+	: Super(ObjectInitializer)
 {
 }
 
-UPersonaOptions::UPersonaOptions(const class FPostConstructInitializeProperties& PCIP)
-	: Super(PCIP)
+UPersonaOptions::UPersonaOptions(const FObjectInitializer& ObjectInitializer)
+	: Super(ObjectInitializer)
 	, DefaultLocalAxesSelection(2)
 {
 	ViewModeIndex = VMI_Lit;
diff --git a/Engine/Source/Editor/UnrealEd/Private/PreviewMaterial.cpp b/Engine/Source/Editor/UnrealEd/Private/PreviewMaterial.cpp
--- a/Engine/Source/Editor/UnrealEd/Private/PreviewMaterial.cpp
+++ b/Engine/Source/Editor/UnrealEd/Private/PreviewMaterial.cpp
@@ -150,8 +150,8 @@ public:
 };
 
 /** Implementation of Preview Material functions*/
-UPreviewMaterial::UPreviewMaterial(const class FPostConstructInitializeProperties& PCIP)
-	: Super(PCIP)
+UPreviewMaterial::UPreviewMaterial(const FObjectInitializer& ObjectInitializer)
+	: Super(ObjectInitializer)
 {
 }
 
@@ -160,8 +160,8 @@ FMaterialResource* UPreviewMaterial::AllocateResource()
 	return new FPreviewMaterial();
 }
 
-UMaterialEditorInstanceConstant::UMaterialEditorInstanceConstant(const class FPostConstructInitializeProperties& PCIP)
-	: Super(PCIP)
+UMaterialEditorInstanceConstant::UMaterialEditorInstanceConstant(const FObjectInitializer& ObjectInitializer)
+	: Super(ObjectInitializer)
 {
 }
 
@@ -722,7 +722,7 @@ void UMaterialEditorInstanceConstant::PostEditUndo()
 }
 #endif
 
-UMaterialEditorMeshComponent::UMaterialEditorMeshComponent(const class FPostConstructInitializeProperties& PCIP)
-	: Super(PCIP)
+UMaterialEditorMeshComponent::UMaterialEditorMeshComponent(const FObjectInitializer& ObjectInitializer)
+	: Super(ObjectInitializer)
 {
 }
diff --git a/Engine/Source/Editor/UnrealEd/Private/PropertyEditorTestObject.cpp b/Engine/Source/Editor/UnrealEd/Private/PropertyEditorTestObject.cpp
--- a/Engine/Source/Editor/UnrealEd/Private/PropertyEditorTestObject.cpp
+++ b/Engine/Source/Editor/UnrealEd/Private/PropertyEditorTestObject.cpp
@@ -3,7 +3,7 @@
 #include "UnrealEd.h"
 
 
-UPropertyEditorTestObject::UPropertyEditorTestObject(const class FPostConstructInitializeProperties& PCIP)
-	: Super(PCIP)
+UPropertyEditorTestObject::UPropertyEditorTestObject(const FObjectInitializer& ObjectInitializer)
+	: Super(ObjectInitializer)
 {
 }
diff --git a/Engine/Source/Editor/UnrealEd/Private/Settings/SettingsClasses.cpp b/Engine/Source/Editor/UnrealEd/Private/Settings/SettingsClasses.cpp
--- a/Engine/Source/Editor/UnrealEd/Private/Settings/SettingsClasses.cpp
+++ b/Engine/Source/Editor/UnrealEd/Private/Settings/SettingsClasses.cpp
@@ -1,4 +1,4 @@
-﻿// Copyright 1998-2014 Epic Games, Inc. All Rights Reserved.
+// Copyright 1998-2014 Epic Games, Inc. All Rights Reserved.
 
 #include "UnrealEd.h"
 #include "ISourceControlModule.h"
@@ -8,8 +8,8 @@
 
 UContentBrowserSettings::FSettingChangedEvent UContentBrowserSettings::SettingChangedEvent;
 
-UContentBrowserSettings::UContentBrowserSettings( const class FPostConstructInitializeProperties& PCIP )
-	: Super(PCIP)
+UContentBrowserSettings::UContentBrowserSettings( const FObjectInitializer& ObjectInitializer )
+	: Super(ObjectInitializer)
 { }
 
 
@@ -33,8 +33,8 @@ void UContentBrowserSettings::PostEditChangeProperty( struct FPropertyChangedEve
 /* UDestructableMeshEditorSettings interface
  *****************************************************************************/
 
-UDestructableMeshEditorSettings::UDestructableMeshEditorSettings( const class FPostConstructInitializeProperties& PCIP )
-	: Super(PCIP)
+UDestructableMeshEditorSettings::UDestructableMeshEditorSettings( const FObjectInitializer& ObjectInitializer )
+	: Super(ObjectInitializer)
 {
 	AnimPreviewLightingDirection = FRotator(-45.0f, 45.0f, 0);
 	AnimPreviewSkyColor = FColor(0, 0, 255);
@@ -48,8 +48,8 @@ UDestructableMeshEditorSettings::UDestructableMeshEditorSettings( const class FP
 /* UEditorExperimentalSettings interface
  *****************************************************************************/
 
-UEditorExperimentalSettings::UEditorExperimentalSettings( const class FPostConstructInitializeProperties& PCIP )
-	: Super(PCIP)
+UEditorExperimentalSettings::UEditorExperimentalSettings( const FObjectInitializer& ObjectInitializer )
+	: Super(ObjectInitializer)
 { 
 }
 
@@ -77,8 +77,8 @@ void UEditorExperimentalSettings::PostEditChangeProperty( struct FPropertyChange
 /* UEditorLoadingSavingSettings interface
  *****************************************************************************/
 
-UEditorLoadingSavingSettings::UEditorLoadingSavingSettings( const class FPostConstructInitializeProperties& PCIP )
-	: Super(PCIP)
+UEditorLoadingSavingSettings::UEditorLoadingSavingSettings( const FObjectInitializer& ObjectInitializer )
+	: Super(ObjectInitializer)
 {
 	TextDiffToolPath.FilePath = TEXT("P4Merge.exe");
 }
@@ -115,16 +115,16 @@ void UEditorLoadingSavingSettings::PostEditChangeProperty( struct FPropertyChang
 /* UEditorMiscSettings interface
  *****************************************************************************/
 
-UEditorMiscSettings::UEditorMiscSettings( const class FPostConstructInitializeProperties& PCIP )
-	: Super(PCIP)
+UEditorMiscSettings::UEditorMiscSettings( const FObjectInitializer& ObjectInitializer )
+	: Super(ObjectInitializer)
 { }
 
 
 /* ULevelEditorMiscSettings interface
  *****************************************************************************/
 
-ULevelEditorMiscSettings::ULevelEditorMiscSettings( const class FPostConstructInitializeProperties& PCIP )
-	: Super(PCIP)
+ULevelEditorMiscSettings::ULevelEditorMiscSettings( const FObjectInitializer& ObjectInitializer )
+	: Super(ObjectInitializer)
 {
 	bAutoApplyLightingEnable = true;
 }
@@ -154,8 +154,8 @@ void ULevelEditorMiscSettings::PostEditChangeProperty( struct FPropertyChangedEv
 /* ULevelEditorPlaySettings interface
  *****************************************************************************/
 
-ULevelEditorPlaySettings::ULevelEditorPlaySettings( const class FPostConstructInitializeProperties& PCIP )
-	: Super(PCIP)
+ULevelEditorPlaySettings::ULevelEditorPlaySettings( const FObjectInitializer& ObjectInitializer )
+	: Super(ObjectInitializer)
 {
 	ClientWindowWidth = 640;
 	ClientWindowHeight = 480;
@@ -169,8 +169,8 @@ ULevelEditorPlaySettings::ULevelEditorPlaySettings( const class FPostConstructIn
 /* ULevelEditorViewportSettings interface
  *****************************************************************************/
 
-ULevelEditorViewportSettings::ULevelEditorViewportSettings( const class FPostConstructInitializeProperties& PCIP )
-	: Super(PCIP)
+ULevelEditorViewportSettings::ULevelEditorViewportSettings( const FObjectInitializer& ObjectInitializer )
+	: Super(ObjectInitializer)
 {
 	bLevelStreamingVolumePrevis = false;
 	BillboardScale = 1.0f;
@@ -270,8 +270,8 @@ void ULevelEditorViewportSettings::PostEditChangeProperty( struct FPropertyChang
 /* UProjectPackagingSettings interface
  *****************************************************************************/
 
-UProjectPackagingSettings::UProjectPackagingSettings( const class FPostConstructInitializeProperties& PCIP )
-	: Super(PCIP)
+UProjectPackagingSettings::UProjectPackagingSettings( const FObjectInitializer& ObjectInitializer )
+	: Super(ObjectInitializer)
 {
 }
 
diff --git a/Engine/Source/Editor/UnrealEd/Private/SoundClassGraph.cpp b/Engine/Source/Editor/UnrealEd/Private/SoundClassGraph.cpp
--- a/Engine/Source/Editor/UnrealEd/Private/SoundClassGraph.cpp
+++ b/Engine/Source/Editor/UnrealEd/Private/SoundClassGraph.cpp
@@ -6,8 +6,8 @@
 #include "UnrealEd.h"
 #include "GraphEditor.h"
 
-USoundClassGraph::USoundClassGraph(const FPostConstructInitializeProperties& PCIP)
-	: Super(PCIP)
+USoundClassGraph::USoundClassGraph(const FObjectInitializer& ObjectInitializer)
+	: Super(ObjectInitializer)
 	, RootSoundClass(NULL)
 {
 }
diff --git a/Engine/Source/Editor/UnrealEd/Private/SoundClassGraphNode.cpp b/Engine/Source/Editor/UnrealEd/Private/SoundClassGraphNode.cpp
--- a/Engine/Source/Editor/UnrealEd/Private/SoundClassGraphNode.cpp
+++ b/Engine/Source/Editor/UnrealEd/Private/SoundClassGraphNode.cpp
@@ -11,8 +11,8 @@
 /////////////////////////////////////////////////////
 // USoundClassGraphNode
 
-USoundClassGraphNode::USoundClassGraphNode(const class FPostConstructInitializeProperties& PCIP)
-	: Super(PCIP)
+USoundClassGraphNode::USoundClassGraphNode(const FObjectInitializer& ObjectInitializer)
+	: Super(ObjectInitializer)
 	, ChildPin(NULL)
 	, ParentPin(NULL)
 {
diff --git a/Engine/Source/Editor/UnrealEd/Private/SoundClassGraphSchema.cpp b/Engine/Source/Editor/UnrealEd/Private/SoundClassGraphSchema.cpp
--- a/Engine/Source/Editor/UnrealEd/Private/SoundClassGraphSchema.cpp
+++ b/Engine/Source/Editor/UnrealEd/Private/SoundClassGraphSchema.cpp
@@ -18,8 +18,8 @@ UEdGraphNode* FSoundClassGraphSchemaAction_NewNode::PerformAction(class UEdGraph
 	return NULL;
 }
 
-USoundClassGraphSchema::USoundClassGraphSchema(const class FPostConstructInitializeProperties& PCIP)
-	: Super(PCIP)
+USoundClassGraphSchema::USoundClassGraphSchema(const FObjectInitializer& ObjectInitializer)
+	: Super(ObjectInitializer)
 {
 }
 
diff --git a/Engine/Source/Editor/UnrealEd/Private/SoundCueGraph.cpp b/Engine/Source/Editor/UnrealEd/Private/SoundCueGraph.cpp
--- a/Engine/Source/Editor/UnrealEd/Private/SoundCueGraph.cpp
+++ b/Engine/Source/Editor/UnrealEd/Private/SoundCueGraph.cpp
@@ -6,8 +6,8 @@
 #include "UnrealEd.h"
 #include "SoundDefinitions.h"
 
-USoundCueGraph::USoundCueGraph(const FPostConstructInitializeProperties& PCIP)
-	: Super(PCIP)
+USoundCueGraph::USoundCueGraph(const FObjectInitializer& ObjectInitializer)
+	: Super(ObjectInitializer)
 {
 }
 
diff --git a/Engine/Source/Editor/UnrealEd/Private/SoundCueGraphNode.cpp b/Engine/Source/Editor/UnrealEd/Private/SoundCueGraphNode.cpp
--- a/Engine/Source/Editor/UnrealEd/Private/SoundCueGraphNode.cpp
+++ b/Engine/Source/Editor/UnrealEd/Private/SoundCueGraphNode.cpp
@@ -16,8 +16,8 @@
 /////////////////////////////////////////////////////
 // USoundCueGraphNode
 
-USoundCueGraphNode::USoundCueGraphNode(const class FPostConstructInitializeProperties& PCIP)
-	: Super(PCIP)
+USoundCueGraphNode::USoundCueGraphNode(const FObjectInitializer& ObjectInitializer)
+	: Super(ObjectInitializer)
 {
 }
 
diff --git a/Engine/Source/Editor/UnrealEd/Private/SoundCueGraphNode_Base.cpp b/Engine/Source/Editor/UnrealEd/Private/SoundCueGraphNode_Base.cpp
--- a/Engine/Source/Editor/UnrealEd/Private/SoundCueGraphNode_Base.cpp
+++ b/Engine/Source/Editor/UnrealEd/Private/SoundCueGraphNode_Base.cpp
@@ -9,8 +9,8 @@
 /////////////////////////////////////////////////////
 // USoundCueGraphNode_Base
 
-USoundCueGraphNode_Base::USoundCueGraphNode_Base(const class FPostConstructInitializeProperties& PCIP)
-	: Super(PCIP)
+USoundCueGraphNode_Base::USoundCueGraphNode_Base(const FObjectInitializer& ObjectInitializer)
+	: Super(ObjectInitializer)
 {
 }
 
diff --git a/Engine/Source/Editor/UnrealEd/Private/SoundCueGraphNode_Root.cpp b/Engine/Source/Editor/UnrealEd/Private/SoundCueGraphNode_Root.cpp
--- a/Engine/Source/Editor/UnrealEd/Private/SoundCueGraphNode_Root.cpp
+++ b/Engine/Source/Editor/UnrealEd/Private/SoundCueGraphNode_Root.cpp
@@ -13,8 +13,8 @@
 /////////////////////////////////////////////////////
 // USoundCueGraphNode_Root
 
-USoundCueGraphNode_Root::USoundCueGraphNode_Root(const class FPostConstructInitializeProperties& PCIP)
-	: Super(PCIP)
+USoundCueGraphNode_Root::USoundCueGraphNode_Root(const FObjectInitializer& ObjectInitializer)
+	: Super(ObjectInitializer)
 {
 }
 
diff --git a/Engine/Source/Editor/UnrealEd/Private/SoundCueGraphSchema.cpp b/Engine/Source/Editor/UnrealEd/Private/SoundCueGraphSchema.cpp
--- a/Engine/Source/Editor/UnrealEd/Private/SoundCueGraphSchema.cpp
+++ b/Engine/Source/Editor/UnrealEd/Private/SoundCueGraphSchema.cpp
@@ -191,8 +191,8 @@ UEdGraphNode* FSoundCueGraphSchemaAction_Paste::PerformAction(class UEdGraph* Pa
 /////////////////////////////////////////////////////
 // USoundCueGraphSchema
 
-USoundCueGraphSchema::USoundCueGraphSchema(const class FPostConstructInitializeProperties& PCIP)
-	: Super(PCIP)
+USoundCueGraphSchema::USoundCueGraphSchema(const FObjectInitializer& ObjectInitializer)
+	: Super(ObjectInitializer)
 {
 }
 
diff --git a/Engine/Source/Editor/UnrealEd/Private/StaticMeshEdit.cpp b/Engine/Source/Editor/UnrealEd/Private/StaticMeshEdit.cpp
--- a/Engine/Source/Editor/UnrealEd/Private/StaticMeshEdit.cpp
+++ b/Engine/Source/Editor/UnrealEd/Private/StaticMeshEdit.cpp
@@ -251,7 +251,7 @@ private:
 void DecomposeUCXMesh( const TArray<FVector>& CollisionVertices, const TArray<int32>& CollisionFaceIdx, UBodySetup* BodySetup )
 {
 	// We keep no ref to this Model, so it will be GC'd at some point after the import.
-	UModel* TempModel = new UModel(FPostConstructInitializeProperties(),NULL,1);
+	UModel* TempModel = new UModel(FObjectInitializer(),NULL,1);
 
 	FMeshConnectivityBuilder ConnectivityBuilder;
 
@@ -508,7 +508,7 @@ UStaticMesh* CreateStaticMesh(struct FRawMesh& RawMesh,TArray<UMaterialInterface
 {
 	// Create the UStaticMesh object.
 	FStaticMeshComponentRecreateRenderStateContext RecreateRenderStateContext(FindObject<UStaticMesh>(InOuter,*InName.ToString()));
-	UStaticMesh* StaticMesh = new(InOuter,InName,RF_Public|RF_Standalone) UStaticMesh(FPostConstructInitializeProperties());
+	UStaticMesh* StaticMesh = new(InOuter,InName,RF_Public|RF_Standalone) UStaticMesh(FObjectInitializer());
 
 	// Add one LOD for the base mesh
 	FStaticMeshSourceModel* SrcModel = new(StaticMesh->SourceModels) FStaticMeshSourceModel();
diff --git a/Engine/Source/Editor/UnrealEd/Private/TemplateMapMetadata.cpp b/Engine/Source/Editor/UnrealEd/Private/TemplateMapMetadata.cpp
--- a/Engine/Source/Editor/UnrealEd/Private/TemplateMapMetadata.cpp
+++ b/Engine/Source/Editor/UnrealEd/Private/TemplateMapMetadata.cpp
@@ -10,7 +10,7 @@
 	UTemplateMapMetadata.
 -----------------------------------------------------------------------------*/
 
-UDEPRECATED_TemplateMapMetadata::UDEPRECATED_TemplateMapMetadata(const class FPostConstructInitializeProperties& PCIP)
-	: Super(PCIP)
+UDEPRECATED_TemplateMapMetadata::UDEPRECATED_TemplateMapMetadata(const FObjectInitializer& ObjectInitializer)
+	: Super(ObjectInitializer)
 {
 }
\ No newline at end of file
diff --git a/Engine/Source/Editor/UnrealEd/Private/TexAlignTools.cpp b/Engine/Source/Editor/UnrealEd/Private/TexAlignTools.cpp
--- a/Engine/Source/Editor/UnrealEd/Private/TexAlignTools.cpp
+++ b/Engine/Source/Editor/UnrealEd/Private/TexAlignTools.cpp
@@ -41,8 +41,8 @@ static bool ShouldFlipVectors( FVector InNormal, int32 InAxis )
 }
 
 
-UTexAligner::UTexAligner(const class FPostConstructInitializeProperties& PCIP)
-	: Super(PCIP)
+UTexAligner::UTexAligner(const FObjectInitializer& ObjectInitializer)
+	: Super(ObjectInitializer)
 {
 }
 
@@ -139,8 +139,8 @@ void UTexAligner::AlignSurf( ETexAlign InTexAlignType, UModel* InModel, FBspSurf
 {
 }
 
-UTexAlignerPlanar::UTexAlignerPlanar(const class FPostConstructInitializeProperties& PCIP)
-	: Super(PCIP)
+UTexAlignerPlanar::UTexAlignerPlanar(const FObjectInitializer& ObjectInitializer)
+	: Super(ObjectInitializer)
 {
 }
 
@@ -196,8 +196,8 @@ void UTexAlignerPlanar::AlignSurf( ETexAlign InTexAlignType, UModel* InModel, FB
 
 }
 
-UTexAlignerDefault::UTexAlignerDefault(const class FPostConstructInitializeProperties& PCIP)
-	: Super(PCIP)
+UTexAlignerDefault::UTexAlignerDefault(const FObjectInitializer& ObjectInitializer)
+	: Super(ObjectInitializer)
 {
 }
 
@@ -223,8 +223,8 @@ void UTexAlignerDefault::AlignSurf( ETexAlign InTexAlignType, UModel* InModel, F
 
 }
 
-UTexAlignerBox::UTexAlignerBox(const class FPostConstructInitializeProperties& PCIP)
-	: Super(PCIP)
+UTexAlignerBox::UTexAlignerBox(const FObjectInitializer& ObjectInitializer)
+	: Super(ObjectInitializer)
 {
 }
 
@@ -254,8 +254,8 @@ void UTexAlignerBox::AlignSurf( ETexAlign InTexAlignType, UModel* InModel, FBspS
 
 }
 
-UTexAlignerFit::UTexAlignerFit(const class FPostConstructInitializeProperties& PCIP)
-	: Super(PCIP)
+UTexAlignerFit::UTexAlignerFit(const FObjectInitializer& ObjectInitializer)
+	: Super(ObjectInitializer)
 {
 }
 
diff --git a/Engine/Source/Editor/UnrealEd/Private/ThumbnailManager.cpp b/Engine/Source/Editor/UnrealEd/Private/ThumbnailManager.cpp
--- a/Engine/Source/Editor/UnrealEd/Private/ThumbnailManager.cpp
+++ b/Engine/Source/Editor/UnrealEd/Private/ThumbnailManager.cpp
@@ -14,8 +14,8 @@ DEFINE_LOG_CATEGORY_STATIC(LogThumbnailManager, Log, All);
 
 UThumbnailManager* UThumbnailManager::ThumbnailManagerSingleton = nullptr;
 
-UThumbnailManager::UThumbnailManager(const class FPostConstructInitializeProperties& PCIP)
-	: Super(PCIP)
+UThumbnailManager::UThumbnailManager(const FObjectInitializer& ObjectInitializer)
+	: Super(ObjectInitializer)
 {
 	if (!IsRunningCommandlet())
 	{
diff --git a/Engine/Source/Editor/UnrealEd/Private/ThumbnailRendering/AnimBlueprintThumbnailRenderer.cpp b/Engine/Source/Editor/UnrealEd/Private/ThumbnailRendering/AnimBlueprintThumbnailRenderer.cpp
--- a/Engine/Source/Editor/UnrealEd/Private/ThumbnailRendering/AnimBlueprintThumbnailRenderer.cpp
+++ b/Engine/Source/Editor/UnrealEd/Private/ThumbnailRendering/AnimBlueprintThumbnailRenderer.cpp
@@ -7,8 +7,8 @@
 #include "EngineModule.h"
 #include "RendererInterface.h"
 
-UAnimBlueprintThumbnailRenderer::UAnimBlueprintThumbnailRenderer(const class FPostConstructInitializeProperties& PCIP)
-	: Super(PCIP)
+UAnimBlueprintThumbnailRenderer::UAnimBlueprintThumbnailRenderer(const FObjectInitializer& ObjectInitializer)
+	: Super(ObjectInitializer)
 {
 	ThumbnailScene = nullptr;
 }
diff --git a/Engine/Source/Editor/UnrealEd/Private/ThumbnailRendering/AnimSequenceThumbnailRenderer.cpp b/Engine/Source/Editor/UnrealEd/Private/ThumbnailRendering/AnimSequenceThumbnailRenderer.cpp
--- a/Engine/Source/Editor/UnrealEd/Private/ThumbnailRendering/AnimSequenceThumbnailRenderer.cpp
+++ b/Engine/Source/Editor/UnrealEd/Private/ThumbnailRendering/AnimSequenceThumbnailRenderer.cpp
@@ -7,8 +7,8 @@
 #include "EngineModule.h"
 #include "RendererInterface.h"
 
-UAnimSequenceThumbnailRenderer::UAnimSequenceThumbnailRenderer(const class FPostConstructInitializeProperties& PCIP)
-	: Super(PCIP)
+UAnimSequenceThumbnailRenderer::UAnimSequenceThumbnailRenderer(const FObjectInitializer& ObjectInitializer)
+	: Super(ObjectInitializer)
 {
 	ThumbnailScene = nullptr;
 }
diff --git a/Engine/Source/Editor/UnrealEd/Private/ThumbnailRendering/BlendSpaceThumbnailRenderer.cpp b/Engine/Source/Editor/UnrealEd/Private/ThumbnailRendering/BlendSpaceThumbnailRenderer.cpp
--- a/Engine/Source/Editor/UnrealEd/Private/ThumbnailRendering/BlendSpaceThumbnailRenderer.cpp
+++ b/Engine/Source/Editor/UnrealEd/Private/ThumbnailRendering/BlendSpaceThumbnailRenderer.cpp
@@ -7,8 +7,8 @@
 #include "EngineModule.h"
 #include "RendererInterface.h"
 
-UBlendSpaceThumbnailRenderer::UBlendSpaceThumbnailRenderer(const class FPostConstructInitializeProperties& PCIP)
-	: Super(PCIP)
+UBlendSpaceThumbnailRenderer::UBlendSpaceThumbnailRenderer(const FObjectInitializer& ObjectInitializer)
+	: Super(ObjectInitializer)
 {
 	ThumbnailScene = nullptr;
 }
diff --git a/Engine/Source/Editor/UnrealEd/Private/ThumbnailRendering/BlueprintThumbnailRenderer.cpp b/Engine/Source/Editor/UnrealEd/Private/ThumbnailRendering/BlueprintThumbnailRenderer.cpp
--- a/Engine/Source/Editor/UnrealEd/Private/ThumbnailRendering/BlueprintThumbnailRenderer.cpp
+++ b/Engine/Source/Editor/UnrealEd/Private/ThumbnailRendering/BlueprintThumbnailRenderer.cpp
@@ -7,8 +7,8 @@
 #include "EngineModule.h"
 #include "RendererInterface.h"
 
-UBlueprintThumbnailRenderer::UBlueprintThumbnailRenderer(const class FPostConstructInitializeProperties& PCIP)
-	: Super(PCIP)
+UBlueprintThumbnailRenderer::UBlueprintThumbnailRenderer(const FObjectInitializer& ObjectInitializer)
+	: Super(ObjectInitializer)
 {
 	ThumbnailScene = nullptr;
 }
diff --git a/Engine/Source/Editor/UnrealEd/Private/ThumbnailRendering/DefaultSizedThumbnailRenderer.cpp b/Engine/Source/Editor/UnrealEd/Private/ThumbnailRendering/DefaultSizedThumbnailRenderer.cpp
--- a/Engine/Source/Editor/UnrealEd/Private/ThumbnailRendering/DefaultSizedThumbnailRenderer.cpp
+++ b/Engine/Source/Editor/UnrealEd/Private/ThumbnailRendering/DefaultSizedThumbnailRenderer.cpp
@@ -2,8 +2,8 @@
 
 #include "UnrealEd.h"
 
-UDefaultSizedThumbnailRenderer::UDefaultSizedThumbnailRenderer(const class FPostConstructInitializeProperties& PCIP)
-	: Super(PCIP)
+UDefaultSizedThumbnailRenderer::UDefaultSizedThumbnailRenderer(const FObjectInitializer& ObjectInitializer)
+	: Super(ObjectInitializer)
 {
 }
 
diff --git a/Engine/Source/Editor/UnrealEd/Private/ThumbnailRendering/FontThumbnailRenderer.cpp b/Engine/Source/Editor/UnrealEd/Private/ThumbnailRendering/FontThumbnailRenderer.cpp
--- a/Engine/Source/Editor/UnrealEd/Private/ThumbnailRendering/FontThumbnailRenderer.cpp
+++ b/Engine/Source/Editor/UnrealEd/Private/ThumbnailRendering/FontThumbnailRenderer.cpp
@@ -2,8 +2,8 @@
 
 #include "UnrealEd.h"
 
-UFontThumbnailRenderer::UFontThumbnailRenderer(const class FPostConstructInitializeProperties& PCIP)
-	: Super(PCIP)
+UFontThumbnailRenderer::UFontThumbnailRenderer(const FObjectInitializer& ObjectInitializer)
+	: Super(ObjectInitializer)
 {
 }
 
diff --git a/Engine/Source/Editor/UnrealEd/Private/ThumbnailRendering/LevelThumbnailRenderer.cpp b/Engine/Source/Editor/UnrealEd/Private/ThumbnailRendering/LevelThumbnailRenderer.cpp
--- a/Engine/Source/Editor/UnrealEd/Private/ThumbnailRendering/LevelThumbnailRenderer.cpp
+++ b/Engine/Source/Editor/UnrealEd/Private/ThumbnailRendering/LevelThumbnailRenderer.cpp
@@ -4,8 +4,8 @@
 #include "EngineModule.h"
 #include "RendererInterface.h"
 
-ULevelThumbnailRenderer::ULevelThumbnailRenderer(const class FPostConstructInitializeProperties& PCIP)
-	: Super(PCIP)
+ULevelThumbnailRenderer::ULevelThumbnailRenderer(const FObjectInitializer& ObjectInitializer)
+	: Super(ObjectInitializer)
 {
 }
 
diff --git a/Engine/Source/Editor/UnrealEd/Private/ThumbnailRendering/MaterialFunctionThumbnailRenderer.cpp b/Engine/Source/Editor/UnrealEd/Private/ThumbnailRendering/MaterialFunctionThumbnailRenderer.cpp
--- a/Engine/Source/Editor/UnrealEd/Private/ThumbnailRendering/MaterialFunctionThumbnailRenderer.cpp
+++ b/Engine/Source/Editor/UnrealEd/Private/ThumbnailRendering/MaterialFunctionThumbnailRenderer.cpp
@@ -10,8 +10,8 @@
 #include "RendererInterface.h"
 #include "ShaderCompiler.h"
 
-UMaterialFunctionThumbnailRenderer::UMaterialFunctionThumbnailRenderer(const class FPostConstructInitializeProperties& PCIP)
-	: Super(PCIP)
+UMaterialFunctionThumbnailRenderer::UMaterialFunctionThumbnailRenderer(const FObjectInitializer& ObjectInitializer)
+	: Super(ObjectInitializer)
 {
 	ThumbnailScene = nullptr;
 }
diff --git a/Engine/Source/Editor/UnrealEd/Private/ThumbnailRendering/MaterialInstanceThumbnailRenderer.cpp b/Engine/Source/Editor/UnrealEd/Private/ThumbnailRendering/MaterialInstanceThumbnailRenderer.cpp
--- a/Engine/Source/Editor/UnrealEd/Private/ThumbnailRendering/MaterialInstanceThumbnailRenderer.cpp
+++ b/Engine/Source/Editor/UnrealEd/Private/ThumbnailRendering/MaterialInstanceThumbnailRenderer.cpp
@@ -7,8 +7,8 @@
 #include "EngineModule.h"
 #include "RendererInterface.h"
 
-UMaterialInstanceThumbnailRenderer::UMaterialInstanceThumbnailRenderer(const class FPostConstructInitializeProperties& PCIP)
-	: Super(PCIP)
+UMaterialInstanceThumbnailRenderer::UMaterialInstanceThumbnailRenderer(const FObjectInitializer& ObjectInitializer)
+	: Super(ObjectInitializer)
 {
 	ThumbnailScene = nullptr;
 }
diff --git a/Engine/Source/Editor/UnrealEd/Private/ThumbnailRendering/ParticleSystemThumbnailRenderer.cpp b/Engine/Source/Editor/UnrealEd/Private/ThumbnailRendering/ParticleSystemThumbnailRenderer.cpp
--- a/Engine/Source/Editor/UnrealEd/Private/ThumbnailRendering/ParticleSystemThumbnailRenderer.cpp
+++ b/Engine/Source/Editor/UnrealEd/Private/ThumbnailRendering/ParticleSystemThumbnailRenderer.cpp
@@ -9,8 +9,8 @@
 #include "EngineModule.h"
 #include "RendererInterface.h"
 
-UParticleSystemThumbnailRenderer::UParticleSystemThumbnailRenderer(const class FPostConstructInitializeProperties& PCIP)
-	: Super(PCIP)
+UParticleSystemThumbnailRenderer::UParticleSystemThumbnailRenderer(const FObjectInitializer& ObjectInitializer)
+	: Super(ObjectInitializer)
 {
 	// Structure to hold one-time initialization
 	struct FConstructorStatics
diff --git a/Engine/Source/Editor/UnrealEd/Private/ThumbnailRendering/SceneThumbnailInfo.cpp b/Engine/Source/Editor/UnrealEd/Private/ThumbnailRendering/SceneThumbnailInfo.cpp
--- a/Engine/Source/Editor/UnrealEd/Private/ThumbnailRendering/SceneThumbnailInfo.cpp
+++ b/Engine/Source/Editor/UnrealEd/Private/ThumbnailRendering/SceneThumbnailInfo.cpp
@@ -2,8 +2,8 @@
 
 #include "UnrealEd.h"
 
-USceneThumbnailInfo::USceneThumbnailInfo(const class FPostConstructInitializeProperties& PCIP)
-	: Super(PCIP)
+USceneThumbnailInfo::USceneThumbnailInfo(const FObjectInitializer& ObjectInitializer)
+	: Super(ObjectInitializer)
 {
 	OrbitPitch = -11.25;
 	OrbitYaw = -157.5;
diff --git a/Engine/Source/Editor/UnrealEd/Private/ThumbnailRendering/SceneThumbnailInfoWithPrimitive.cpp b/Engine/Source/Editor/UnrealEd/Private/ThumbnailRendering/SceneThumbnailInfoWithPrimitive.cpp
--- a/Engine/Source/Editor/UnrealEd/Private/ThumbnailRendering/SceneThumbnailInfoWithPrimitive.cpp
+++ b/Engine/Source/Editor/UnrealEd/Private/ThumbnailRendering/SceneThumbnailInfoWithPrimitive.cpp
@@ -2,8 +2,8 @@
 
 #include "UnrealEd.h"
 
-USceneThumbnailInfoWithPrimitive::USceneThumbnailInfoWithPrimitive(const class FPostConstructInitializeProperties& PCIP)
-	: Super(PCIP)
+USceneThumbnailInfoWithPrimitive::USceneThumbnailInfoWithPrimitive(const FObjectInitializer& ObjectInitializer)
+	: Super(ObjectInitializer)
 {
 	PrimitiveType = TPT_Sphere;
 	OrbitPitch = -35;
diff --git a/Engine/Source/Editor/UnrealEd/Private/ThumbnailRendering/SkeletalMeshThumbnailRenderer.cpp b/Engine/Source/Editor/UnrealEd/Private/ThumbnailRendering/SkeletalMeshThumbnailRenderer.cpp
--- a/Engine/Source/Editor/UnrealEd/Private/ThumbnailRendering/SkeletalMeshThumbnailRenderer.cpp
+++ b/Engine/Source/Editor/UnrealEd/Private/ThumbnailRendering/SkeletalMeshThumbnailRenderer.cpp
@@ -7,8 +7,8 @@
 #include "EngineModule.h"
 #include "RendererInterface.h"
 
-USkeletalMeshThumbnailRenderer::USkeletalMeshThumbnailRenderer(const class FPostConstructInitializeProperties& PCIP)
-	: Super(PCIP)
+USkeletalMeshThumbnailRenderer::USkeletalMeshThumbnailRenderer(const FObjectInitializer& ObjectInitializer)
+	: Super(ObjectInitializer)
 {
 	ThumbnailScene = nullptr;
 }
diff --git a/Engine/Source/Editor/UnrealEd/Private/ThumbnailRendering/SlateBrushThumbnailRenderer.cpp b/Engine/Source/Editor/UnrealEd/Private/ThumbnailRendering/SlateBrushThumbnailRenderer.cpp
--- a/Engine/Source/Editor/UnrealEd/Private/ThumbnailRendering/SlateBrushThumbnailRenderer.cpp
+++ b/Engine/Source/Editor/UnrealEd/Private/ThumbnailRendering/SlateBrushThumbnailRenderer.cpp
@@ -5,8 +5,8 @@
 // FPreviewScene derived helpers for rendering
 #include "ThumbnailHelpers.h"
 
-USlateBrushThumbnailRenderer::USlateBrushThumbnailRenderer(const class FPostConstructInitializeProperties& PCIP)
-	: Super(PCIP)
+USlateBrushThumbnailRenderer::USlateBrushThumbnailRenderer(const FObjectInitializer& ObjectInitializer)
+	: Super(ObjectInitializer)
 {
 }
 
diff --git a/Engine/Source/Editor/UnrealEd/Private/ThumbnailRendering/SoundWaveThumbnailRenderer.cpp b/Engine/Source/Editor/UnrealEd/Private/ThumbnailRendering/SoundWaveThumbnailRenderer.cpp
--- a/Engine/Source/Editor/UnrealEd/Private/ThumbnailRendering/SoundWaveThumbnailRenderer.cpp
+++ b/Engine/Source/Editor/UnrealEd/Private/ThumbnailRendering/SoundWaveThumbnailRenderer.cpp
@@ -2,8 +2,8 @@
 
 #include "UnrealEd.h"
 
-USoundWaveThumbnailRenderer::USoundWaveThumbnailRenderer(const class FPostConstructInitializeProperties& PCIP)
-	: Super(PCIP)
+USoundWaveThumbnailRenderer::USoundWaveThumbnailRenderer(const FObjectInitializer& ObjectInitializer)
+	: Super(ObjectInitializer)
 {
 }
 
diff --git a/Engine/Source/Editor/UnrealEd/Private/ThumbnailRendering/StaticMeshThumbnailRenderer.cpp b/Engine/Source/Editor/UnrealEd/Private/ThumbnailRendering/StaticMeshThumbnailRenderer.cpp
--- a/Engine/Source/Editor/UnrealEd/Private/ThumbnailRendering/StaticMeshThumbnailRenderer.cpp
+++ b/Engine/Source/Editor/UnrealEd/Private/ThumbnailRendering/StaticMeshThumbnailRenderer.cpp
@@ -7,8 +7,8 @@
 #include "EngineModule.h"
 #include "RendererInterface.h"
 
-UStaticMeshThumbnailRenderer::UStaticMeshThumbnailRenderer(const class FPostConstructInitializeProperties& PCIP)
-	: Super(PCIP)
+UStaticMeshThumbnailRenderer::UStaticMeshThumbnailRenderer(const FObjectInitializer& ObjectInitializer)
+	: Super(ObjectInitializer)
 {
 	ThumbnailScene = nullptr;
 }
diff --git a/Engine/Source/Editor/UnrealEd/Private/ThumbnailRendering/SubsurfaceProfileRenderer.cpp b/Engine/Source/Editor/UnrealEd/Private/ThumbnailRendering/SubsurfaceProfileRenderer.cpp
--- a/Engine/Source/Editor/UnrealEd/Private/ThumbnailRendering/SubsurfaceProfileRenderer.cpp
+++ b/Engine/Source/Editor/UnrealEd/Private/ThumbnailRendering/SubsurfaceProfileRenderer.cpp
@@ -2,8 +2,8 @@
 
 #include "UnrealEd.h"
 
-USubsurfaceProfileRenderer::USubsurfaceProfileRenderer(const class FPostConstructInitializeProperties& PCIP)
-	: Super(PCIP)
+USubsurfaceProfileRenderer::USubsurfaceProfileRenderer(const FObjectInitializer& ObjectInitializer)
+	: Super(ObjectInitializer)
 {
 }
 
diff --git a/Engine/Source/Editor/UnrealEd/Private/ThumbnailRendering/TextureCubeThumbnailRenderer.cpp b/Engine/Source/Editor/UnrealEd/Private/ThumbnailRendering/TextureCubeThumbnailRenderer.cpp
--- a/Engine/Source/Editor/UnrealEd/Private/ThumbnailRendering/TextureCubeThumbnailRenderer.cpp
+++ b/Engine/Source/Editor/UnrealEd/Private/ThumbnailRendering/TextureCubeThumbnailRenderer.cpp
@@ -2,8 +2,8 @@
 
 #include "UnrealEd.h"
 
-UTextureCubeThumbnailRenderer::UTextureCubeThumbnailRenderer(const class FPostConstructInitializeProperties& PCIP)
-	: Super(PCIP)
+UTextureCubeThumbnailRenderer::UTextureCubeThumbnailRenderer(const FObjectInitializer& ObjectInitializer)
+	: Super(ObjectInitializer)
 {
 }
 
diff --git a/Engine/Source/Editor/UnrealEd/Private/ThumbnailRendering/TextureThumbnailRenderer.cpp b/Engine/Source/Editor/UnrealEd/Private/ThumbnailRendering/TextureThumbnailRenderer.cpp
--- a/Engine/Source/Editor/UnrealEd/Private/ThumbnailRendering/TextureThumbnailRenderer.cpp
+++ b/Engine/Source/Editor/UnrealEd/Private/ThumbnailRendering/TextureThumbnailRenderer.cpp
@@ -5,8 +5,8 @@
 #include "CubemapUnwrapUtils.h"			// FMipLevelBatchedElementParameters
 #include "NormalMapPreview.h"
 
-UTextureThumbnailRenderer::UTextureThumbnailRenderer(const class FPostConstructInitializeProperties& PCIP)
-	: Super(PCIP)
+UTextureThumbnailRenderer::UTextureThumbnailRenderer(const FObjectInitializer& ObjectInitializer)
+	: Super(ObjectInitializer)
 {
 }
 
diff --git a/Engine/Source/Editor/UnrealEd/Private/ThumbnailRendering/ThumbnailRenderer.cpp b/Engine/Source/Editor/UnrealEd/Private/ThumbnailRendering/ThumbnailRenderer.cpp
--- a/Engine/Source/Editor/UnrealEd/Private/ThumbnailRendering/ThumbnailRenderer.cpp
+++ b/Engine/Source/Editor/UnrealEd/Private/ThumbnailRendering/ThumbnailRenderer.cpp
@@ -2,7 +2,7 @@
 
 #include "UnrealEd.h"
 
-UThumbnailRenderer::UThumbnailRenderer(const class FPostConstructInitializeProperties& PCIP)
-	: Super(PCIP)
+UThumbnailRenderer::UThumbnailRenderer(const FObjectInitializer& ObjectInitializer)
+	: Super(ObjectInitializer)
 {
 }
diff --git a/Engine/Source/Editor/UnrealEd/Private/ThumbnailRendering/WorldThumbnailInfo.cpp b/Engine/Source/Editor/UnrealEd/Private/ThumbnailRendering/WorldThumbnailInfo.cpp
--- a/Engine/Source/Editor/UnrealEd/Private/ThumbnailRendering/WorldThumbnailInfo.cpp
+++ b/Engine/Source/Editor/UnrealEd/Private/ThumbnailRendering/WorldThumbnailInfo.cpp
@@ -2,8 +2,8 @@
 
 #include "UnrealEd.h"
 
-UWorldThumbnailInfo::UWorldThumbnailInfo(const class FPostConstructInitializeProperties& PCIP)
-	: Super(PCIP)
+UWorldThumbnailInfo::UWorldThumbnailInfo(const FObjectInitializer& ObjectInitializer)
+	: Super(ObjectInitializer)
 {
 	CameraMode = ECameraProjectionMode::Perspective;
 	OrthoDirection = EOrthoThumbnailDirection::Top;
diff --git a/Engine/Source/Editor/UnrealEd/Private/ThumbnailRendering/WorldThumbnailRenderer.cpp b/Engine/Source/Editor/UnrealEd/Private/ThumbnailRendering/WorldThumbnailRenderer.cpp
--- a/Engine/Source/Editor/UnrealEd/Private/ThumbnailRendering/WorldThumbnailRenderer.cpp
+++ b/Engine/Source/Editor/UnrealEd/Private/ThumbnailRendering/WorldThumbnailRenderer.cpp
@@ -4,8 +4,8 @@
 #include "EngineModule.h"
 #include "ContentStreaming.h"
 #include "RendererInterface.h"
-UWorldThumbnailRenderer::UWorldThumbnailRenderer(const class FPostConstructInitializeProperties& PCIP)
-	: Super(PCIP)
+UWorldThumbnailRenderer::UWorldThumbnailRenderer(const FObjectInitializer& ObjectInitializer)
+	: Super(ObjectInitializer)
 {
 	GlobalOrbitPitchOffset = 0.f;
 	GlobalOrbitYawOffset = 0.f;
diff --git a/Engine/Source/Editor/UnrealEd/Private/UnrealEdEngine.cpp b/Engine/Source/Editor/UnrealEd/Private/UnrealEdEngine.cpp
--- a/Engine/Source/Editor/UnrealEd/Private/UnrealEdEngine.cpp
+++ b/Engine/Source/Editor/UnrealEd/Private/UnrealEdEngine.cpp
@@ -658,13 +658,13 @@ FString FClassPickerDefaults::GetDescription() const
 #undef LOCTEXT_NAMESPACE
 
 
-UUnrealEdKeyBindings::UUnrealEdKeyBindings(const class FPostConstructInitializeProperties& PCIP)
-	: Super(PCIP)
+UUnrealEdKeyBindings::UUnrealEdKeyBindings(const FObjectInitializer& ObjectInitializer)
+	: Super(ObjectInitializer)
 {
 }
 
-UUnrealEdOptions::UUnrealEdOptions(const class FPostConstructInitializeProperties& PCIP)
-	: Super(PCIP)
+UUnrealEdOptions::UUnrealEdOptions(const FObjectInitializer& ObjectInitializer)
+	: Super(ObjectInitializer)
 {
 }
 
diff --git a/Engine/Source/Editor/UnrealEd/Private/UnrealEdTypes.cpp b/Engine/Source/Editor/UnrealEd/Private/UnrealEdTypes.cpp
--- a/Engine/Source/Editor/UnrealEd/Private/UnrealEdTypes.cpp
+++ b/Engine/Source/Editor/UnrealEd/Private/UnrealEdTypes.cpp
@@ -2,7 +2,7 @@
 
 #include "UnrealEd.h"
 
-UUnrealEdTypes::UUnrealEdTypes(const class FPostConstructInitializeProperties& PCIP)
-	: Super(PCIP)
+UUnrealEdTypes::UUnrealEdTypes(const FObjectInitializer& ObjectInitializer)
+	: Super(ObjectInitializer)
 {
 }
diff --git a/Engine/Source/Editor/UnrealEd/Private/UserDefinedStructEditorData.cpp b/Engine/Source/Editor/UnrealEd/Private/UserDefinedStructEditorData.cpp
--- a/Engine/Source/Editor/UnrealEd/Private/UserDefinedStructEditorData.cpp
+++ b/Engine/Source/Editor/UnrealEd/Private/UserDefinedStructEditorData.cpp
@@ -17,8 +17,8 @@ FEdGraphPinType FStructVariableDescription::ToPinType() const
 	return FEdGraphPinType(Category, SubCategory, SubCategoryObject.Get(), bIsArray, false);
 }
 
-UUserDefinedStructEditorData::UUserDefinedStructEditorData(const class FPostConstructInitializeProperties& PCIP)
-: Super(PCIP)
+UUserDefinedStructEditorData::UUserDefinedStructEditorData(const FObjectInitializer& ObjectInitializer)
+: Super(ObjectInitializer)
 {
 }
 
diff --git a/Engine/Source/Editor/WorldBrowser/Private/Tiles/WorldTileCollectionModel.cpp b/Engine/Source/Editor/WorldBrowser/Private/Tiles/WorldTileCollectionModel.cpp
--- a/Engine/Source/Editor/WorldBrowser/Private/Tiles/WorldTileCollectionModel.cpp
+++ b/Engine/Source/Editor/WorldBrowser/Private/Tiles/WorldTileCollectionModel.cpp
@@ -1931,7 +1931,7 @@ bool FWorldTileCollectionModel::GenerateLODLevels(FLevelModelList InLevelList, i
 	
 			// Construct landscape static mesh
 			FString LandscapeMeshAssetName = TEXT("SM_") + LandscapeBaseAssetName;
-			UStaticMesh* StaticMesh = new(LODPackage, *LandscapeMeshAssetName, RF_Public|RF_Standalone) UStaticMesh(FPostConstructInitializeProperties());
+			UStaticMesh* StaticMesh = new(LODPackage, *LandscapeMeshAssetName, RF_Public|RF_Standalone) UStaticMesh(FObjectInitializer());
 			StaticMesh->InitResources();
 			{
 				FString OutputPath = StaticMesh->GetPathName();
