Fixed missing bone hierarchy error to point to the correct option for doign a rigid mesh import.
#ttp 349479 - FBX: IMPORT: No longer have the ability to import Morph Targets without skinned bone

diff --git a/Engine/Source/Editor/UnrealEd/Private/Fbx/FbxFactory.cpp b/Engine/Source/Editor/UnrealEd/Private/Fbx/FbxFactory.cpp
--- a/Engine/Source/Editor/UnrealEd/Private/Fbx/FbxFactory.cpp
+++ b/Engine/Source/Editor/UnrealEd/Private/Fbx/FbxFactory.cpp
@@ -468,7 +468,7 @@ UObject* UFbxFactory::FactoryCreateBinary
 				}
 				else if (ImportUI->MeshTypeToImport == FBXIT_SkeletalMesh)
 				{
-					FbxImporter->AddTokenizedErrorMessage(FTokenizedMessage::Create(EMessageSeverity::Error, LOCTEXT("FailedToImport_InvalidBone", "Failed to find any bone hierarchy. Try to import as Rigid Mesh option enabled. ")), FFbxErrors::SkeletalMesh_InvalidBone);
+					FbxImporter->AddTokenizedErrorMessage(FTokenizedMessage::Create(EMessageSeverity::Error, LOCTEXT("FailedToImport_InvalidBone", "Failed to find any bone hierarchy. Try disabling the \"Import As Skeletal\" option to import as a rigid mesh. ")), FFbxErrors::SkeletalMesh_InvalidBone);
 				}
 				else
 				{
