Added tooltip for Console for Gamepad Labels in Experimental preferences.
#ttp 323073 - EDITOR: PREFERENCES: Console for Gamepad Labels is missing tooltip
#branch UE4
#proj Editor.UnrealEd

diff --git a/Engine/Source/Editor/UnrealEd/Classes/Settings/EditorExperimentalSettings.h b/Engine/Source/Editor/UnrealEd/Classes/Settings/EditorExperimentalSettings.h
--- a/Engine/Source/Editor/UnrealEd/Classes/Settings/EditorExperimentalSettings.h
+++ b/Engine/Source/Editor/UnrealEd/Classes/Settings/EditorExperimentalSettings.h
@@ -41,6 +41,7 @@ public:
 	UPROPERTY(EditAnywhere, config, Category=Tools, meta=(DisplayName="Actor Merging"))
 	bool bActorMerging;
 
+	/** Specify which console-specific nomenclature to use for gamepad label text */
 	UPROPERTY(EditAnywhere, config, Category=UserInterface, meta=(DisplayName="Console for Gamepad Labels"))
 	TEnumAsByte<EConsoleForGamepadLabels::Type> ConsoleForGamepadLabels;
 
