Make it so component visualizers don't run in game view mode (requested by MattK)

diff --git a/Engine/Source/Editor/UnrealEd/Private/LevelEditorViewport.cpp b/Engine/Source/Editor/UnrealEd/Private/LevelEditorViewport.cpp
--- a/Engine/Source/Editor/UnrealEd/Private/LevelEditorViewport.cpp
+++ b/Engine/Source/Editor/UnrealEd/Private/LevelEditorViewport.cpp
@@ -3930,7 +3930,7 @@ void FLevelEditorViewportClient::Draw(const FSceneView* View,FPrimitiveDrawInter
 		FSnappingUtils::DrawSnappingHelpers( View, PDI );
 	}
 
-	if(GUnrealEd != NULL)
+	if(GUnrealEd != NULL && !IsInGameView())
 	{
 		GUnrealEd->DrawComponentVisualizers(View, PDI);
 	}
