Fix web display of compact name of some math and conversion nodes that use special characters

diff --git a/Engine/Source/Editor/UnrealEd/Private/Commandlets/GenerateBlueprintAPICommandlet.cpp b/Engine/Source/Editor/UnrealEd/Private/Commandlets/GenerateBlueprintAPICommandlet.cpp
--- a/Engine/Source/Editor/UnrealEd/Private/Commandlets/GenerateBlueprintAPICommandlet.cpp
+++ b/Engine/Source/Editor/UnrealEd/Private/Commandlets/GenerateBlueprintAPICommandlet.cpp
@@ -227,6 +227,11 @@ FString MakeJsonString( const FString& String )
 		default: OutString += *Char;
 		}
 	}
+
+	OutString = OutString.Replace(TEXT("\xD7"), TEXT("&times;"));
+	OutString = OutString.Replace(TEXT("\xF7"), TEXT("&divide;"));
+	OutString = OutString.Replace(TEXT("\x2022"), TEXT("&middot;"));
+
 	return OutString;
 }
 
