Cooking changes to allow DLC builds for UT

diff --git a/Engine/Source/Editor/UnrealEd/Classes/CookOnTheSide/CookOnTheFlyServer.h b/Engine/Source/Editor/UnrealEd/Classes/CookOnTheSide/CookOnTheFlyServer.h
--- a/Engine/Source/Editor/UnrealEd/Classes/CookOnTheSide/CookOnTheFlyServer.h
+++ b/Engine/Source/Editor/UnrealEd/Classes/CookOnTheSide/CookOnTheFlyServer.h
@@ -6,6 +6,7 @@
 
 #pragma once
 
+#include "IPlatformFileSandboxWrapper.h"
 #include "CookOnTheFlyServer.generated.h"
 
 
@@ -38,7 +39,7 @@ namespace ECookMode
 }
 
 UCLASS()
-class UCookOnTheFlyServer : public UObject, public FTickableEditorObject
+class UNREALED_API UCookOnTheFlyServer : public UObject, public FTickableEditorObject
 {
 	GENERATED_UCLASS_BODY()
 private:
@@ -103,7 +104,7 @@ private:
 			return Items.Num();
 		}
 
-		void Empty() 
+		void Empty()
 		{
 			ScopeLockType ScopeLock( &SynchronizationObject );
 			Items.Empty();
@@ -411,10 +412,10 @@ private:
 			FScopeLock ScopeLock( &SynchronizationObject );
 			return FilesProcessed.Remove( Filename );
 		}
-		void Empty()
+		void Empty(int32 ExpectedNumElements = 0)
 		{
 			FScopeLock ScopeLock( &SynchronizationObject );
-			FilesProcessed.Empty();
+			FilesProcessed.Empty(ExpectedNumElements);
 		}
 	};
 
@@ -608,6 +609,7 @@ private:
 
 public:
 
+	void WarmCookedPackages(const FString& AssetRegistryPath, const TArray<FName>& TargetPlatformNames);
 
 	
 
diff --git a/Engine/Source/Editor/UnrealEd/Private/Commandlets/CookCommandlet.cpp b/Engine/Source/Editor/UnrealEd/Private/Commandlets/CookCommandlet.cpp
--- a/Engine/Source/Editor/UnrealEd/Private/Commandlets/CookCommandlet.cpp
+++ b/Engine/Source/Editor/UnrealEd/Private/Commandlets/CookCommandlet.cpp
@@ -1029,6 +1029,18 @@ bool UCookCommandlet::NewCook( const TArray<ITargetPlatform*>& Platforms, TArray
 	//////////////////////////////////////////////////////////////////////////
 	// parse commandline options 
 
+	FString AssetRegistry;
+	if (FParse::Value(*Params, TEXT("SHIPPEDASSETREGISTRY="), AssetRegistry))
+	{
+		TArray<FName> TargetPlatformNames;
+		for (const auto &Platform : Platforms)
+		{
+			FName PlatformName = FName(*Platform->PlatformName());
+			TargetPlatformNames.Add(PlatformName); // build list of all target platform names
+		}
+		CookOnTheFlyServer->WarmCookedPackages(FPaths::GameContentDir() / AssetRegistry, TargetPlatformNames);
+	}
+
 	TArray<FString> CmdLineIniSections;
 	FString SectionStr;
 	if (FParse::Value(*Params, TEXT("MAPINISECTION="), SectionStr))
diff --git a/Engine/Source/Editor/UnrealEd/Private/CookOnTheFlyServer.cpp b/Engine/Source/Editor/UnrealEd/Private/CookOnTheFlyServer.cpp
--- a/Engine/Source/Editor/UnrealEd/Private/CookOnTheFlyServer.cpp
+++ b/Engine/Source/Editor/UnrealEd/Private/CookOnTheFlyServer.cpp
@@ -2422,3 +2422,39 @@ void UCookOnTheFlyServer::HandleNetworkFileServerRecompileShaders(const FShaderR
 		RecompileData.ModifiedFiles,
 		RecompileData.bCompileChangedShaders);
 }
+
+void UCookOnTheFlyServer::WarmCookedPackages(const FString& AssetRegistryPath, const TArray<FName>& TargetPlatformNames)
+{
+	FArrayReader SerializedAssetData;
+	if (FFileHelper::LoadFileToArray(SerializedAssetData, *AssetRegistryPath))
+	{
+		int32 LocalNumAssets = 0;
+		SerializedAssetData << LocalNumAssets;
+
+		// allocate one single block for all asset data structs (to reduce tens of thousands of heap allocations)
+		FAssetData* PreallocatedAssetDataBuffer = new FAssetData[LocalNumAssets];
+		CookedPackages.Empty(LocalNumAssets);
+
+		FFilePlatformRequest FileRequest;
+		for (const auto &Platform : TargetPlatformNames)
+		{
+			FileRequest.AddPlatform(Platform);
+		}
+
+		for (int32 AssetIndex = 0; AssetIndex < LocalNumAssets; AssetIndex++)
+		{
+			// make a new asset data object
+			FAssetData* NewAssetData = &PreallocatedAssetDataBuffer[AssetIndex];
+
+			// load it
+			SerializedAssetData << *NewAssetData;
+
+			UE_LOG(LogCookOnTheFly, Verbose, TEXT("Read package %s from %s"), *GetCachedStandardPackageFilename(NewAssetData->ObjectPath), *AssetRegistryPath);
+
+			FileRequest.SetFilename(GetCachedStandardPackageFilename(NewAssetData->ObjectPath));
+			CookedPackages.Add(FileRequest);
+		}
+
+		delete[] PreallocatedAssetDataBuffer;
+	}
+}
\ No newline at end of file
