Switching the Check Static Mesh UV's test to be EAutomationTestFlags::ATF_RequiresUser so that it doesn't run on the automation build farm.

diff --git a/Engine/Source/Editor/UnrealEd/Private/Tests/EditorAutomationTests.cpp b/Engine/Source/Editor/UnrealEd/Private/Tests/EditorAutomationTests.cpp
--- a/Engine/Source/Editor/UnrealEd/Private/Tests/EditorAutomationTests.cpp
+++ b/Engine/Source/Editor/UnrealEd/Private/Tests/EditorAutomationTests.cpp
@@ -2054,7 +2054,7 @@ bool FPerformUVTestCommand::Update()
 /**
  * StaticMeshUVsTest
  */
-IMPLEMENT_SIMPLE_AUTOMATION_TEST( FStaticMeshUVsTest, "Tools.Static Mesh.Static Mesh UVs Check", EAutomationTestFlags::ATF_Editor )
+IMPLEMENT_SIMPLE_AUTOMATION_TEST(FStaticMeshUVsTest, "Tools.Static Mesh.Static Mesh UVs Check", EAutomationTestFlags::ATF_Editor | EAutomationTestFlags::ATF_RequiresUser)
 
 /**
  * Find all static meshes and check the lightmap UVs
