[Github PR 515] PhoenixBlack - Component visualizers can now override DrawVisualizationHUD to do 2D visualizations as well.

diff --git a/Engine/Source/Editor/Kismet/Private/SCSEditorViewportClient.cpp b/Engine/Source/Editor/Kismet/Private/SCSEditorViewportClient.cpp
--- a/Engine/Source/Editor/Kismet/Private/SCSEditorViewportClient.cpp
+++ b/Engine/Source/Editor/Kismet/Private/SCSEditorViewportClient.cpp
@@ -242,6 +242,26 @@ void FSCSEditorViewportClient::DrawCanvas( FViewport& InViewport, FSceneView& Vi
 	AActor* PreviewActor = GetPreviewActor();
 	if(PreviewActor)
 	{
+		if (GUnrealEd != NULL)
+		{
+			TArray<FSCSEditorTreeNodePtrType> SelectedNodes = BlueprintEditorPtr.Pin()->GetSelectedSCSEditorTreeNodes();
+			for (int32 SelectionIndex = 0; SelectionIndex < SelectedNodes.Num(); ++SelectionIndex)
+			{
+				FSCSEditorTreeNodePtrType SelectedNode = SelectedNodes[SelectionIndex];
+
+				UActorComponent* Comp = Cast<USceneComponent>(SelectedNode->FindComponentInstanceInActor(PreviewActor, true));
+				if (Comp != NULL && Comp->IsRegistered())
+				{
+					// Try and find a visualizer
+					TSharedPtr<FComponentVisualizer> Visualizer = GUnrealEd->FindComponentVisualizer(Comp->GetClass());
+					if (Visualizer.IsValid())
+					{
+						Visualizer->DrawVisualizationHUD(Comp, &InViewport, &View, &Canvas);
+					}
+				}
+			}
+		}
+
 		TGuardValue<bool> AutoRestore(GAllowActorScriptExecutionInEditor, true);
 
 		const int32 HalfX = 0.5f * Viewport->GetSizeXY().X;
diff --git a/Engine/Source/Editor/UnrealEd/Classes/Editor/UnrealEdEngine.h b/Engine/Source/Editor/UnrealEd/Classes/Editor/UnrealEdEngine.h
--- a/Engine/Source/Editor/UnrealEd/Classes/Editor/UnrealEdEngine.h
+++ b/Engine/Source/Editor/UnrealEd/Classes/Editor/UnrealEdEngine.h
@@ -234,6 +234,8 @@ class UNREALED_API UUnrealEdEngine : public UEditorEngine, public FNotifyHook
 
 	/** Draw component visualizers for components for selected actors */
 	void DrawComponentVisualizers(const FSceneView* View, FPrimitiveDrawInterface* PDI);
+	/** Draw component visualizers HUD elements for components for selected actors */
+	void DrawComponentVisualizersHUD(const FViewport* Viewport, const FSceneView* View, FCanvas* Canvas);
 
 	/** Updates the property windows of selected actors */
 	virtual void UpdateFloatingPropertyWindows();
diff --git a/Engine/Source/Editor/UnrealEd/Private/LevelEditorViewport.cpp b/Engine/Source/Editor/UnrealEd/Private/LevelEditorViewport.cpp
--- a/Engine/Source/Editor/UnrealEd/Private/LevelEditorViewport.cpp
+++ b/Engine/Source/Editor/UnrealEd/Private/LevelEditorViewport.cpp
@@ -4031,6 +4031,12 @@ void FLevelEditorViewportClient::SetupViewForRendering( FSceneViewFamily& ViewFa
 
 void FLevelEditorViewportClient::DrawCanvas( FViewport& InViewport, FSceneView& View, FCanvas& Canvas )
 {
+	// HUD for components visualizers
+	if (GUnrealEd != NULL)
+	{
+		GUnrealEd->DrawComponentVisualizersHUD(&InViewport, &View, &Canvas);
+	}
+
 	// Information string
 	Canvas.DrawShadowedString( 4,4, *GLevelEditorModeTools().InfoString, GEngine->GetSmallFont(), FColor(255,255,255) );
 
diff --git a/Engine/Source/Editor/UnrealEd/Private/UnrealEdEngine.cpp b/Engine/Source/Editor/UnrealEd/Private/UnrealEdEngine.cpp
--- a/Engine/Source/Editor/UnrealEd/Private/UnrealEdEngine.cpp
+++ b/Engine/Source/Editor/UnrealEd/Private/UnrealEdEngine.cpp
@@ -1136,6 +1136,37 @@ void UUnrealEdEngine::DrawComponentVisualizers(const FSceneView* View, FPrimitiv
 }
 
 
+void UUnrealEdEngine::DrawComponentVisualizersHUD(const FViewport* Viewport, const FSceneView* View, FCanvas* Canvas)
+{
+	// Iterate over all selected actors
+	for (FSelectionIterator It(GetSelectedActorIterator()); It; ++It)
+	{
+		AActor* Actor = Cast<AActor>(*It);
+		if (Actor != NULL)
+		{
+			// Then iterate over components of that actor
+			TArray<UActorComponent*> Components;
+			Actor->GetComponents(Components);
+
+			for (int32 CompIdx = 0; CompIdx<Components.Num(); CompIdx++)
+			{
+				UActorComponent* Comp = Components[CompIdx];
+				if (Comp->IsRegistered())
+				{
+					// Try and find a visualizer
+
+					TSharedPtr<FComponentVisualizer> Visualizer = FindComponentVisualizer(Comp->GetClass());
+					if (Visualizer.IsValid())
+					{
+						Visualizer->DrawVisualizationHUD(Comp, Viewport, View, Canvas);
+					}
+				}
+			}
+		}
+	}
+}
+
+
 EWriteDisallowedWarningState UUnrealEdEngine::GetWarningStateForWritePermission(const FString& PackageName) const
 {
 	EWriteDisallowedWarningState WarningState = WDWS_WarningUnnecessary;
diff --git a/Engine/Source/Editor/UnrealEd/Public/ComponentVisualizer.h b/Engine/Source/Editor/UnrealEd/Public/ComponentVisualizer.h
--- a/Engine/Source/Editor/UnrealEd/Public/ComponentVisualizer.h
+++ b/Engine/Source/Editor/UnrealEd/Public/ComponentVisualizer.h
@@ -32,6 +32,8 @@ public:
 	virtual void OnRegister() {}
 	/** Draw visualization for the supplied component */
 	virtual void DrawVisualization(const UActorComponent* Component, const FSceneView* View, FPrimitiveDrawInterface* PDI) {}
+	/** Draw HUD on viewport for the supplied component */
+	virtual void DrawVisualizationHUD(const UActorComponent* Component, const FViewport* Viewport, const FSceneView* View, FCanvas* Canvas) {}
 	/** */
 	virtual bool VisProxyHandleClick(FLevelEditorViewportClient* InViewportClient, HComponentVisProxy* VisProxy, const FViewportClick& Click) { return false; }
 	/** */
