Merging from 4.5 branch:

Properly display event and switch nodes

diff --git a/Engine/Source/Editor/UnrealEd/Private/Commandlets/GenerateBlueprintAPICommandlet.cpp b/Engine/Source/Editor/UnrealEd/Private/Commandlets/GenerateBlueprintAPICommandlet.cpp
--- a/Engine/Source/Editor/UnrealEd/Private/Commandlets/GenerateBlueprintAPICommandlet.cpp
+++ b/Engine/Source/Editor/UnrealEd/Private/Commandlets/GenerateBlueprintAPICommandlet.cpp
@@ -7,6 +7,7 @@
 #include "EdGraphSchema_K2_Actions.h"
 #include "K2Node.h"
 #include "K2Node_CommutativeAssociativeBinaryOperator.h"
+#include "K2Node_SwitchEnum.h"
 #include "AnimationGraph.h"
 #include "AnimGraphNode_StateMachine.h"
 #include "BlueprintEditorUtils.h"
@@ -834,12 +835,21 @@ static void GenerateBlueprintAPIUtils::DumpActionMenuItem(uint32 Indent, FGraphA
 		{
 			ActionEntry += IndentedNewline + "\"NodeType\"     : \"pure\"";
 		}
+		else if (Node->IsA<UK2Node_Event>())
+		{
+			ActionEntry += IndentedNewline + "\"NodeType\"     : \"event\"";
+		}
+		else if (Node->IsA<UK2Node_Switch>())
+		{
+			ActionEntry += IndentedNewline + "\"NodeType\"     : \"switch\"";
+		}
 		else
 		{
 			ActionEntry += IndentedNewline + "\"NodeType\"     : \"function\"";
 		}
 
-		if (Node->IsA<UK2Node_CommutativeAssociativeBinaryOperator>())
+		if (Node->IsA<UK2Node_CommutativeAssociativeBinaryOperator>()
+			|| (Node->IsA<UK2Node_Switch>() && !Node->IsA<UK2Node_SwitchEnum>()))
 		{
 			ActionEntry += IndentedNewline + "\"ShowAddPin\"   : \"true\"";
 		}
diff --git a/Engine/Source/Programs/UnrealDocTool/APIDocTool/APIDocTool/APIAction.cs b/Engine/Source/Programs/UnrealDocTool/APIDocTool/APIDocTool/APIAction.cs
--- a/Engine/Source/Programs/UnrealDocTool/APIDocTool/APIDocTool/APIAction.cs
+++ b/Engine/Source/Programs/UnrealDocTool/APIDocTool/APIDocTool/APIAction.cs
@@ -104,7 +104,7 @@ namespace APIDocTool
 					WritePinObjects(Writer, Pins.Where(x => x.bInputPin));
 					Writer.LeaveParam();
 					Writer.EnterParam("outputs");
-					WritePinObjects(Writer, Pins.Where(x => !x.bInputPin));
+					WritePinObjects(Writer, Pins.Where(x => !x.bInputPin && x.GetTypeText() != "delegate"));
 
 					if (bShowAddPin)
 					{
diff --git a/Engine/Source/Programs/UnrealDocTool/APIDocTool/APIDocTool/APIActionPin.cs b/Engine/Source/Programs/UnrealDocTool/APIDocTool/APIDocTool/APIActionPin.cs
--- a/Engine/Source/Programs/UnrealDocTool/APIDocTool/APIDocTool/APIActionPin.cs
+++ b/Engine/Source/Programs/UnrealDocTool/APIDocTool/APIDocTool/APIActionPin.cs
@@ -117,7 +117,7 @@ namespace APIDocTool
 			}
 		}
 
-		private string GetTypeText()
+		public string GetTypeText()
 		{
 			if (PinCategory == "struct")
 			{
