Refactored the FEditorMapPerf automation test so that it doesn't actually have to write out a text file to save data.
UEAUTO-81

diff --git a/Engine/Source/Editor/UnrealEd/Private/Tests/AutomationEditorCommon.cpp b/Engine/Source/Editor/UnrealEd/Private/Tests/AutomationEditorCommon.cpp
--- a/Engine/Source/Editor/UnrealEd/Private/Tests/AutomationEditorCommon.cpp
+++ b/Engine/Source/Editor/UnrealEd/Private/Tests/AutomationEditorCommon.cpp
@@ -287,77 +287,75 @@ namespace AutomationEditorCommonUtils
 * Writes a number to a text file.
 * @param FolderNameForTest is the folder that has the same name as the test. (For Example: "Performance").
 * @param FolderNameForBeingTested is the name for the thing that is being tested. (For Example: "MapName").
-* @param FileName is the name of the file with an extension (no need to add the .txt)
+* @param FileName is the name of the file with an extension
 * @param NumberToWriteToFile is the float number that is expected to be written to the file.
 * @param Delimiter is the delimiter to be used. TEXT(",")
 */
 void WriteToTextFile(const FString& FolderNameForTest ,const FString& FolderNameForBeingTested, const FString& FileName, const float& NumberToWriteToFile, const FString& Delimiter)
 {
 	//Performance file locations and setups.
-	FString FileSaveLocation = FPaths::AutomationLogDir() + FolderNameForTest + TEXT("/") + FolderNameForBeingTested + TEXT("/") + FileName + TEXT(".txt");
+	FString FileSaveLocation = FPaths::Combine(*FPaths::AutomationLogDir(), *FolderNameForTest, *FolderNameForBeingTested, *FileName);
 
 	//Variables that hold the content from the text files
 	FString ReadFromTextFile;
 
 	//Log out to a text file the Duration.
 	FString CurrentNumberToWrite = FString::Printf(TEXT("%f"), NumberToWriteToFile);
-	FFileHelper::LoadFileToString(ReadFromTextFile, *FileSaveLocation);
-	FString FileSetup = ReadFromTextFile + CurrentNumberToWrite + Delimiter;
-	FFileHelper::SaveStringToFile(FileSetup, *FileSaveLocation);
+	if (FPaths::FileExists(FileSaveLocation))
+	{
+		FFileHelper::LoadFileToString(ReadFromTextFile, *FileSaveLocation);
+		FString FileSetup = ReadFromTextFile + Delimiter + CurrentNumberToWrite;
+		FFileHelper::SaveStringToFile(FileSetup, *FileSaveLocation);
+		return;
+	}
 
+	FFileHelper::SaveStringToFile(CurrentNumberToWrite, *FileSaveLocation);
 }
 
 /**
-* Returns the sum of the numbers available in a FString array.
-* @param NumberArray is the name of the array intended to be used.
-* @param bAverageOnly will return the average of the available numbers instead of the sum.
+* Returns the sum of the numbers available in an array of float.
+* @param ArrayofNumbers is the name of the array intended to be used.
+* @param bisAveragedInstead will return the average of the available numbers instead of the sum.
 */
-float ValueFromStringArrayofNumbers(const TArray<FString>& NumberArray, bool bAverageInstead)
+float TotalFromFloatArray(const TArray<float>& ArrayofNumbers, bool bisAveragedInstead)
 {
 	//Total Value holds the sum of all the numbers available in the array.
 	float TotalValue = 0;
-	int32 NumberCounter = 0;
 
 	//Loop through the array.
-	for (int32 I = 0; I < NumberArray.Num(); ++I)
+	for (int32 I = 0; I < ArrayofNumbers.Num(); ++I)
 	{
 		//If the character in the array is a number then we'll want to add it to our total.
-		if (NumberArray[I].IsNumeric())
-		{
-			NumberCounter++;
-			TotalValue += FCString::Atof(NumberArray[I].GetCharArray().GetData());
-		}
-	}	
+		TotalValue += ArrayofNumbers[I];
+	}
 
 	//If bAverageInstead equals true then only the average is returned.
 	//Otherwise return the total value.
-	if (bAverageInstead)
+	if (bisAveragedInstead)
 	{
-		UE_LOG(LogEditorAutomationTests, VeryVerbose, TEXT("Average value of the Array is %f"), (TotalValue / NumberCounter));
-		return (TotalValue / NumberCounter);
+		UE_LOG(LogEditorAutomationTests, VeryVerbose, TEXT("Average value of the Array is %f"), (TotalValue / ArrayofNumbers.Num()));
+		return (TotalValue / ArrayofNumbers.Num());
 	}
-	
+
 	UE_LOG(LogEditorAutomationTests, VeryVerbose, TEXT("Total Value of the Array is %f"), TotalValue);
 	return TotalValue;
-
 }
 
-float LargestValueFromStringArrayofNumbers(const TArray<FString>& NumberArray)
+/**
+* Returns the largest value from an array of float numbers.
+* @param ArrayofNumbers is the name of the array intended to be used.
+*/
+float LargetValueInFloatArray(const TArray<float>& ArrayOfNumbers)
 {
 	//Total Value holds the sum of all the numbers available in the array.
-	float LargestValue = FCString::Atof(NumberArray[0].GetCharArray().GetData());
+	float LargestValue = 0;
 
 	//Loop through the array.
-	for (int32 I = 0; I < NumberArray.Num(); ++I)
+	for (int32 I = 0; I < ArrayOfNumbers.Num(); ++I)
 	{
-		//If the character in the array is a number then we'll want to add it to our total.
-		if (NumberArray[I].IsNumeric())
+		if (LargestValue < ArrayOfNumbers[I])
 		{
-			float CurrentValue = FCString::Atof(NumberArray[I].GetCharArray().GetData());
-			if ( LargestValue < CurrentValue)
-			{
-				LargestValue = CurrentValue;
-			}
+			LargestValue = ArrayOfNumbers[I];
 		}
 	}
 	UE_LOG(LogEditorAutomationTests, VeryVerbose, TEXT("The Largest value of the array is %f"), LargestValue);
@@ -366,13 +364,11 @@ float LargestValueFromStringArrayofNumbers(const TArray<FString>& NumberArray)
 
 TArray<FString> CreateArrayFromFile(const FString& FullFileLocation)
 {
-
 	FString RawData;
 	TArray<FString> DataArray;
 
 	if (FPaths::FileExists(*FullFileLocation))
 	{
-
 		UE_LOG(LogEditorAutomationTests, VeryVerbose, TEXT("Loading and parsing the data from '%s' into an array."), *FullFileLocation);
 		FFileHelper::LoadFileToString(RawData, *FullFileLocation);
 		RawData.ParseIntoArray(&DataArray, TEXT(","), false);
@@ -472,177 +468,99 @@ bool FEndPlayMapCommand::Update()
 }
 
 /**
-* This command grabs the FPS and Memory for the current editor session.
+* This command grabs the FPS and Memory stats for the current editor session.
 */
-bool FEditorPerformanceCommand::Update()
+bool FEditorPerfCaptureCommand::Update()
 {
-	//Gets the main frame module to get the name of our current level.
-	const IMainFrameModule& MainFrameModule = FModuleManager::GetModuleChecked< IMainFrameModule >("MainFrame");
-	FString ShortMapName = MainFrameModule.GetLoadedLevelName();
+	//Capture the current time stamp and format it to YYYY-MM-DD HH:MM:SS.mmm.
+	FDateTime CurrentDateAndTime = FDateTime::Now();
 
-	//Variables that hold the content from the text files
-	FString SavedAverageFPS;
-	FString SavedMemory;
-	FString SavedDuration;
-	static SIZE_T StaticLastTotalAllocated = 0;
-	static SIZE_T AvailablePhysicalMemory = 0;
-	static SIZE_T AvailableVirtualMemory = 0;
-	static SIZE_T UsedVirtualMemory = 0;
-	static SIZE_T UsedPeakPhysical = 0;
-	static SIZE_T UsedPeakVirtual = 0;
-
-	//Grab the FPS and working memory number.
-	//float CurrentTime = FPlatformTime::Seconds();
-	//float ElapsedTime = CurrentTime - StartTime;
-	float CurrentTime = 0;
-	float ElapsedTime = 0;
-	FDateTime CurrentDateAndTime = FDateTime::UtcNow();
-
-	//If the last entry and the current time are the same then return false until it is. We only want to grab data for every actual second.
-	FString TestDurationFileLocation = FPaths::Combine(*FPaths::AutomationLogDir(), TEXT("Performance"), *ShortMapName, TEXT("RAWDuration.txt"));
-	if (FPaths::FileExists(TestDurationFileLocation))
+	//This is how long it has been since the last run through.
+	FTimespan ElapsedTime = 0;
+	FTimespan TimeBetweenCaptures;
+
+	//We want to only capture data every whole second.
+	if (EditorPerfStats.TimeStamp.Num() > 0)
 	{
-		FDateTime PerformanceRunTime = IFileManager::Get().GetTimeStamp(*TestDurationFileLocation);
-		FTimespan TimeDifferenceBetweenRuns = CurrentDateAndTime.GetTicks() - PerformanceRunTime.GetTicks();
-		if (TimeDifferenceBetweenRuns.GetTicks() < ETimespan::TicksPerSecond)
+		TimeBetweenCaptures = CurrentDateAndTime.GetTicks() - EditorPerfStats.TimeStamp.Last().GetTicks();
+		if (TimeBetweenCaptures.GetTicks() < ETimespan::TicksPerSecond)
 		{
 			return false;
 		}
-
-		CurrentTime = FMath::RoundToInt(TimeDifferenceBetweenRuns.GetTotalSeconds());
-		TArray<FString> TestRunTimeRawArray = CreateArrayFromFile(TestDurationFileLocation);
-		ElapsedTime = (ValueFromStringArrayofNumbers(TestRunTimeRawArray, false)) + CurrentTime;
+		
+		ElapsedTime = CurrentDateAndTime.GetTicks() - EditorPerfStats.TimeStamp[0].GetTicks();
 	}
 
-	if (ElapsedTime < Duration || (FPlatformTime::Seconds() - StartTime) < Duration)
+	if (ElapsedTime.GetTotalSeconds() <= EditorPerfStats.TestDuration + 1)
 	{
 		//Find the Average FPS
 		//Clamp to avoid huge averages at startup or after hitches
 		const float CurrentFPS = 1.0f / FSlateApplication::Get().GetAverageDeltaTime();
 		const float ClampedFPS = (CurrentFPS < 0.0f || CurrentFPS > 4000.0f) ? 0.0f : CurrentFPS;
+		EditorPerfStats.AverageFPS.Add(ClampedFPS);
 
 		//Find the Frame Time in ms.
 		//Clamp to avoid huge averages at startup or after hitches
 		const float AverageMS = FSlateApplication::Get().GetAverageDeltaTime() * 1000.0f;
 		const float ClampedMS = (AverageMS < 0.0f || AverageMS > 4000.0f) ? 0.0f : AverageMS;
+		EditorPerfStats.AverageFrameTime.Add(ClampedMS);
 
 		//Query OS for process memory used.
 		FPlatformMemoryStats MemoryStats = FPlatformMemory::GetStats();
-		StaticLastTotalAllocated = MemoryStats.UsedPhysical;
+		EditorPerfStats.UsedPhysical.Add((float)MemoryStats.UsedPhysical / 1024);
 
 		//Query OS for available physical memory
-		AvailablePhysicalMemory = MemoryStats.AvailablePhysical;
+		EditorPerfStats.AvailablePhysical.Add((float)MemoryStats.AvailablePhysical / 1024);
 
 		//Query OS for available virtual memory
-		AvailableVirtualMemory = MemoryStats.AvailableVirtual;
+		EditorPerfStats.AvailableVirtual.Add((float)MemoryStats.AvailableVirtual / 1024);
 
 		//Query OS for used virtual memory
-		UsedVirtualMemory = MemoryStats.UsedVirtual;
+		EditorPerfStats.UsedVirtual.Add((float)MemoryStats.UsedVirtual / 1024);
 
 		//Query OS for used Peak Used physical memory
-		UsedPeakPhysical = MemoryStats.PeakUsedPhysical;
+		EditorPerfStats.PeakUsedPhysical.Add((float)MemoryStats.PeakUsedPhysical / 1024);
 
 		//Query OS for used Peak Used virtual memory
-		UsedPeakVirtual = MemoryStats.PeakUsedVirtual;
-
-		//Log out to a text file the Duration.
-		WriteToTextFile(TEXT("Performance"), ShortMapName, TEXT("RAWDuration"), CurrentTime, TEXT(","));
-		//Log out to a text file the AverageFPS
-		WriteToTextFile(TEXT("Performance"), ShortMapName, TEXT("RAWAverageFPS"), ClampedFPS, TEXT(","));
-		//Log out to a text file the AverageFrameTime
-		WriteToTextFile(TEXT("Performance"), ShortMapName, TEXT("RAWFrameTime"), ClampedMS, TEXT(","));
-		//Log out to a text file the Used Physical Memory in kilobytes.
-		WriteToTextFile(TEXT("Performance"), ShortMapName, TEXT("RAWMemoryUsedPhysical"), (float)StaticLastTotalAllocated / 1024, TEXT(","));
-		//Log out to a text file the Used Peaked Memory in kilobytes.
-		WriteToTextFile(TEXT("Performance"), ShortMapName, TEXT("RAWMemoryUsedPeak"), (float)UsedPeakPhysical / 1024, TEXT(","));
-		//Log out to a text file the Used Peaked virtual Memory in kilobytes.
-		WriteToTextFile(TEXT("Performance"), ShortMapName, TEXT("RAWMemoryUsedPeakVirtual"), (float)UsedPeakVirtual / 1024, TEXT(","));
-		//Log out to a text file the Available Physical Memory in kilobytes.
-		WriteToTextFile(TEXT("Performance"), ShortMapName, TEXT("RAWMemoryAvailablePhysical"), (float)AvailablePhysicalMemory / 1024, TEXT(","));
-		//Log out to a text file the Available Virtual Memory in kilobytes.
-		WriteToTextFile(TEXT("Performance"), ShortMapName, TEXT("RAWMemoryAvailableVirtual"), (float)AvailableVirtualMemory / 1024, TEXT(","));
-		//Log out to a text file the Used Virtual Memory in kilobytes.
-		WriteToTextFile(TEXT("Performance"), ShortMapName, TEXT("RAWMemoryUsedVirtual"), (float)UsedVirtualMemory / 1024, TEXT(","));
-
-		//Log out to a text file the current time stamp.
-		FString OldCurrentTimeData;
+		EditorPerfStats.PeakUsedVirtual.Add((float)MemoryStats.PeakUsedVirtual / 1024);
+
+		//Capturing the time stamp.
 		FString FormatedTimeStamp = FString::Printf(TEXT("%04i-%02i-%02i %02i:%02i:%02i.%03i"), CurrentDateAndTime.GetYear(), CurrentDateAndTime.GetMonth(), CurrentDateAndTime.GetDay(), CurrentDateAndTime.GetHour(), CurrentDateAndTime.GetMinute(), CurrentDateAndTime.GetSecond(), CurrentDateAndTime.GetMillisecond());
-		FString OldFileLocation = FPaths::AutomationLogDir() + TEXT("Performance/") + ShortMapName + TEXT("/RAWTimeStamp.txt");
-		FFileHelper::LoadFileToString(OldCurrentTimeData, *OldFileLocation);
-		FString FileSetup = OldCurrentTimeData + FormatedTimeStamp + TEXT(",");
-		FFileHelper::SaveStringToFile(FileSetup, *OldFileLocation);
+		EditorPerfStats.FormattedTimeStamp.Add(FormatedTimeStamp);
+		EditorPerfStats.TimeStamp.Add(CurrentDateAndTime);
 
 		return false;
 	}
 
-	UE_LOG(LogEditorAutomationTests, Log, TEXT("Raw performance data has been captured and saved to the '//Game/Saved/Automation/Log/Performance/%s' folder location."), *ShortMapName);
+	//Combine performance data into one chart.
+	ADD_LATENT_AUTOMATION_COMMAND(FGenerateEditorPerfCaptureCommand(EditorPerfStats));
 
 	return true;
 }
 
-bool FGenerateEditorPerformanceCharts::Update()
+bool FGenerateEditorPerfCaptureCommand::Update()
 {
 	UE_LOG(LogEditorAutomationTests, Log, TEXT("Begin generating the editor performance charts."));
 
 	//Get the current changelist number
 	FString Changelist = GEngineVersion.ToString(EVersionComponent::Changelist);
 
-	//Gets the main frame module to get the name of our current level.
-	const IMainFrameModule& MainFrameModule = FModuleManager::GetModuleChecked< IMainFrameModule >("MainFrame");
-	FString LoadedMapName = MainFrameModule.GetLoadedLevelName();
+	//The locations to save the csv files.
+	FString FileLocation = FPaths::Combine(*FPaths::AutomationLogDir(), TEXT("Performance"), *EditorPerfStats.MapName);
 
-	//The locations of the raw data.
-	FString FileLocation = FPaths::AutomationLogDir() + TEXT("Performance/") + LoadedMapName + TEXT("/");
-	FString FPSFileLocation = FileLocation + TEXT("RAWAverageFPS.txt");
-	FString FrameTimeLocation = FileLocation + TEXT("RAWFrameTime.txt");
-	FString MemoryFileLocation = FileLocation + TEXT("RAWMemoryUsedPhysical.txt");
-	FString MemoryAvailPhysFileLocation = FileLocation + TEXT("RAWMemoryAvailablePhysical.txt");
-	FString MemoryAvailVirtualFileLocation = FileLocation + TEXT("RAWMemoryAvailableVirtual.txt");
-	FString MemoryUsedVirtualFileLocation = FileLocation + TEXT("RAWMemoryUsedVirtual.txt");
-	FString MapLoadTimeFileLocation = FileLocation + TEXT("RAWMapLoadTime.txt");
-	FString DurationFileLocation = FileLocation + TEXT("RAWDuration.txt");
-	FString MemoryUsedPeakFileLocation = FileLocation + TEXT("RAWMemoryUsedPeak.txt");
-	FString MemoryUsedPeakVirtualFileLocation = FileLocation + TEXT("RAWMemoryUsedPeakVirtual.txt");
-	FString TimeStampFileLocation = FileLocation + TEXT("RAWTimeStamp.txt");
-	//Get the data from the raw files and put them into an array.
-	TArray<FString> FPSRawArray = CreateArrayFromFile(FPSFileLocation);
-	TArray<FString> FrameTimeRawArray = CreateArrayFromFile(FrameTimeLocation);
-	TArray<FString> MemoryRawArray = CreateArrayFromFile(MemoryFileLocation);
-	TArray<FString> MemoryAvailPhysRawArray = CreateArrayFromFile(MemoryAvailPhysFileLocation);
-	TArray<FString> MemoryAvailVirtualRawArray = CreateArrayFromFile(MemoryAvailVirtualFileLocation);
-	TArray<FString> MemoryUsedVirtualRawArray = CreateArrayFromFile(MemoryUsedVirtualFileLocation);
-	TArray<FString> MapLoadTimeRawArray = CreateArrayFromFile(MapLoadTimeFileLocation);
-	TArray<FString> DurationRawArray = CreateArrayFromFile(DurationFileLocation);
-	TArray<FString> MemoryUsedPeakArray = CreateArrayFromFile(MemoryUsedPeakFileLocation);
-	TArray<FString> MemoryUsedPeakVirtualArray = CreateArrayFromFile(MemoryUsedPeakVirtualFileLocation);
-	TArray<FString> TimeStampArray = CreateArrayFromFile(TimeStampFileLocation);
-
-	FString MapLoadTime;
-
-	//If the map load time doesn't exist then we give add a '0' to the Map Load time array.
-	if (!FPaths::FileExists(*MapLoadTimeFileLocation))
+	//If a map load time file exists then grab the last load time from it.
+	FString MapLoadTimeFileLocation = FPaths::Combine(*FileLocation, TEXT("RAWMapLoadTime.txt"));
+	EditorPerfStats.MapLoadTime = 0;
+	if (FPaths::FileExists(*MapLoadTimeFileLocation))
 	{
-		MapLoadTimeRawArray.Add("0");
+		TArray<FString> SavedMapLoadTimes = CreateArrayFromFile(MapLoadTimeFileLocation);
+		EditorPerfStats.MapLoadTime = FCString::Atof(*SavedMapLoadTimes.Last());
 	}
 
-	//Make sure that the map load time actually has data to show even if the time is '0'
-	MapLoadTime = MapLoadTimeRawArray.Last();
-	if (MapLoadTime.IsEmpty())
-	{
-		if (MapLoadTimeRawArray[MapLoadTimeRawArray.Num() - 2].IsNumeric())
-		{
-			MapLoadTime = MapLoadTimeRawArray[MapLoadTimeRawArray.Num() - 2];
-		}
-	}
-
-	
-	//Variable that will hold the old performance csv file data.
-	FString OldPerformanceCSVFile;
-
-	//RAW csv holds all of the information from the text files
-	//Final csv only holds the averaged information.
-	FString FinalCSVFilename = FString::Printf(TEXT("%s%s_Performance.csv"), *FileLocation, *LoadedMapName);
-	FString RAWCSVFilename = FString::Printf(TEXT("%sRAW_%s_%s.csv"), *FileLocation, *LoadedMapName, *FDateTime::Now().ToString());
+	//RAW csv holds the data gathered from a single test ran.
+	//Final csv only holds the final information for each full test run.
+	FString FinalCSVFilename = FString::Printf(TEXT("%s/%s_Performance.csv"), *FileLocation, *EditorPerfStats.MapName);
+	FString RAWCSVFilename = FString::Printf(TEXT("%s/RAW_%s_%s.csv"), *FileLocation, *EditorPerfStats.MapName, *FDateTime::Now().ToString());
 
 	//Create the .csv file that will hold the raw data from the text files.
 	FArchive* RAWCSVFile = IFileManager::Get().CreateFileWriter(*RAWCSVFilename);
@@ -657,7 +575,7 @@ bool FGenerateEditorPerformanceCharts::Update()
 	RAWCSVFile->Serialize(TCHAR_TO_ANSI(*RAWCSVLine), RAWCSVLine.Len());
 	
 	//Loop through the text files and add them to the raw csv file.
-	for (int32 I = 0; I < (DurationRawArray.Num() - 1); I++)
+	for (int32 I = 0; I < EditorPerfStats.TimeStamp.Num(); I++)
 	{
 		//If the raw file isn't available to write to then we'll fail back this test.
 		if (!RAWCSVFile)
@@ -667,26 +585,25 @@ bool FGenerateEditorPerformanceCharts::Update()
 		}
 
 		//Create the line that will be written and then write it to the csv file.
-		RAWCSVLine = LoadedMapName + TEXT(",") + Changelist + TEXT(",") + TimeStampArray[I] + TEXT(",") + MapLoadTime + TEXT(",") + FPSRawArray[I] + TEXT(",") + FrameTimeRawArray[I] + TEXT(",") + MemoryRawArray[I] + TEXT(",") + MemoryUsedVirtualRawArray[I] + TEXT(",") + MemoryUsedPeakArray[I] + TEXT(",") + MemoryUsedPeakVirtualArray[I] + TEXT(",") + MemoryAvailPhysRawArray[I] + TEXT(",") + MemoryAvailVirtualRawArray[I] + LINE_TERMINATOR;
+		RAWCSVLine = FString::Printf(TEXT("%s,%s,%s,%.3f,%.1f,%.1f,%.0f,%.0f,%.0f,%.0f,%.0f,%.0f%s"), *EditorPerfStats.MapName, *Changelist, *EditorPerfStats.TimeStamp[I].ToString(), EditorPerfStats.MapLoadTime, EditorPerfStats.AverageFPS[I], EditorPerfStats.AverageFrameTime[I], EditorPerfStats.UsedPhysical[I], EditorPerfStats.UsedVirtual[I], EditorPerfStats.PeakUsedPhysical[I], EditorPerfStats.PeakUsedVirtual[I], EditorPerfStats.AvailablePhysical[I], EditorPerfStats.AvailableVirtual[I], LINE_TERMINATOR);
 		RAWCSVFile->Serialize(TCHAR_TO_ANSI(*RAWCSVLine), RAWCSVLine.Len());
 	}
-	
 	//Close the raw csv file.
 	RAWCSVFile->Close();
 	
 	//Get the final data for the Performance csv file.
-	//This gets the averaged numbers using the raw data arrays.
-	float FPSAvg = ValueFromStringArrayofNumbers(FPSRawArray, true);
-	float FrameTimeAvg = ValueFromStringArrayofNumbers(FrameTimeRawArray, true);
-	float MemoryAvg = ValueFromStringArrayofNumbers(MemoryRawArray, true);
-	float MemoryAvailPhysAvg = ValueFromStringArrayofNumbers(MemoryAvailPhysRawArray, true);
-	float MemoryAvailVirtualAvg = ValueFromStringArrayofNumbers(MemoryAvailVirtualRawArray, true);
-	float MemoryUsedVirtualAvg = ValueFromStringArrayofNumbers(MemoryUsedVirtualRawArray, true);
-	float TestRunTime = ValueFromStringArrayofNumbers(DurationRawArray, false);
-	
-	//Find the largest value in the array.
-	float MemoryUsedPeak = LargestValueFromStringArrayofNumbers(MemoryUsedPeakArray);
-	float MemoryUsedPeakVirtual = LargestValueFromStringArrayofNumbers(MemoryUsedPeakVirtualArray);
+	//This gets the averaged numbers.
+	float FPSAvg = TotalFromFloatArray(EditorPerfStats.AverageFPS, true);
+	float FrameTimeAvg = TotalFromFloatArray(EditorPerfStats.AverageFrameTime, true);
+	float MemoryUsedPhysical = TotalFromFloatArray(EditorPerfStats.UsedPhysical, true);
+	float MemoryAvailPhysAvg = TotalFromFloatArray(EditorPerfStats.AvailablePhysical, true);
+	float MemoryAvailVirtualAvg = TotalFromFloatArray(EditorPerfStats.AvailableVirtual, true);
+	float MemoryUsedVirtualAvg = TotalFromFloatArray(EditorPerfStats.UsedVirtual, true);
+	float MemoryUsedPeak = LargetValueInFloatArray(EditorPerfStats.PeakUsedPhysical);
+	float MemoryUsedPeakVirtual = LargetValueInFloatArray(EditorPerfStats.PeakUsedVirtual);
+
+	//This is how long the test lasted.
+	FTimespan TestRunDuration = (EditorPerfStats.TimeStamp.Last().GetTicks() - EditorPerfStats.TimeStamp[0].GetTicks()) + ETimespan::TicksPerSecond;
 	
 	//If the Performance csv file does not have a size then it must be new.
 	//This adds the title row to the performance csv file.
@@ -707,47 +624,33 @@ bool FGenerateEditorPerformanceCharts::Update()
 	}
 
 	//Load the currently existing performance csv data.
+	FString OldPerformanceCSVFile;
 	FFileHelper::LoadFileToString(OldPerformanceCSVFile, *FinalCSVFilename);
 	FArchive* FinalCSVFile = IFileManager::Get().CreateFileWriter(*FinalCSVFilename);
 	//If the performance file isn't available to write to then we'll fail back this test.
 	if (!FinalCSVFile)
 	{
-		UE_LOG(LogEditorAutomationTests, Error, TEXT("Failed to write to the Performance csv file: '%s'.  \nIt may have been opened or deleted while in use. \nData Text files will not be deleted at this time."), *FPaths::ConvertRelativePathToFull(FinalCSVFilename));
+		UE_LOG(LogEditorAutomationTests, Error, TEXT("Failed to write to the Performance csv file: '%s'.  \nIt may have been opened or deleted while in use."), *FPaths::ConvertRelativePathToFull(FinalCSVFilename));
 		return true;
 	}
 	
 	//Write the old performance csv file data to the new csv file.
 	FinalCSVFile->Serialize(TCHAR_TO_ANSI(*OldPerformanceCSVFile), OldPerformanceCSVFile.Len());
 
-	//Write the averaged numbers to the Performance CSV file.
-	FString FinalCSVLine = FString::Printf(TEXT("%s,%s,%s,%.0f,%s,%.3f,%.3f,%.0f,%.0f,%.0f,%.0f,%.0f,%.0f%s"), *FDateTime::Now().ToString(), *LoadedMapName, *Changelist, TestRunTime, *MapLoadTime, FPSAvg, FrameTimeAvg, MemoryAvg, MemoryUsedVirtualAvg, MemoryUsedPeak, MemoryUsedPeakVirtual, MemoryAvailPhysAvg, MemoryAvailVirtualAvg, LINE_TERMINATOR);
+	//Write the stats to the Performance CSV file and then close it so we can edit it while the engine is still running.
+	FString FinalCSVLine = FString::Printf(TEXT("%s,%s,%s,%.0f,%.3f,%.1f,%.1f,%.0f,%.0f,%.0f,%.0f,%.0f,%.0f%s"), *FDateTime::Now().ToString(), *EditorPerfStats.MapName, *Changelist, TestRunDuration.GetTotalSeconds(), EditorPerfStats.MapLoadTime, FPSAvg, FrameTimeAvg, MemoryUsedPhysical, MemoryUsedVirtualAvg, MemoryUsedPeak, MemoryUsedPeakVirtual, MemoryAvailPhysAvg, MemoryAvailVirtualAvg, LINE_TERMINATOR);
 	FinalCSVFile->Serialize(TCHAR_TO_ANSI(*FinalCSVLine), FinalCSVLine.Len());
-
-	//Close the csv files so we can use them while the editor is still open.
 	FinalCSVFile->Close();
 
-	//Display the averaged numbers.
+	//Display the test results to the user.
 	UE_LOG(LogEditorAutomationTests, Display, TEXT("AVG FPS: '%.1f'"), FPSAvg);
 	UE_LOG(LogEditorAutomationTests, Display, TEXT("AVG Frame Time: '%.1f' ms"), FrameTimeAvg);
-	UE_LOG(LogEditorAutomationTests, Display, TEXT("AVG Used Physical Memory: '%.0f' kb"), MemoryAvg);
+	UE_LOG(LogEditorAutomationTests, Display, TEXT("AVG Used Physical Memory: '%.0f' kb"), MemoryUsedPhysical);
 	UE_LOG(LogEditorAutomationTests, Display, TEXT("AVG Used Virtual Memory: '%.0f' kb"), MemoryUsedVirtualAvg);
 	UE_LOG(LogEditorAutomationTests, Display, TEXT("Performance csv file is located here: %s"), *FPaths::ConvertRelativePathToFull(FinalCSVFilename));
 	UE_LOG(LogEditorAutomationTests, Log, TEXT("Performance csv file is located here: %s"), *FPaths::ConvertRelativePathToFull(FinalCSVFilename));
 	UE_LOG(LogEditorAutomationTests, Log, TEXT("Raw performance csv file is located here: %s"), *FPaths::ConvertRelativePathToFull(RAWCSVFilename));
 
-	//Delete the text files.
-	UE_LOG(LogEditorAutomationTests, Log, TEXT("Deleting temporary data text files."));
-	IFileManager::Get().Delete(*FPSFileLocation);
-	IFileManager::Get().Delete(*FrameTimeLocation);
-	IFileManager::Get().Delete(*MemoryFileLocation);
-	IFileManager::Get().Delete(*MemoryAvailPhysFileLocation);
-	IFileManager::Get().Delete(*DurationFileLocation);
-	IFileManager::Get().Delete(*MemoryAvailVirtualFileLocation);
-	IFileManager::Get().Delete(*MemoryUsedVirtualFileLocation);
-	IFileManager::Get().Delete(*MemoryUsedPeakFileLocation);
-	IFileManager::Get().Delete(*MemoryUsedPeakVirtualFileLocation);
-	IFileManager::Get().Delete(*TimeStampFileLocation);
-
 	return true;
 }
 
@@ -778,7 +681,7 @@ bool FEditorLoadMap::Update()
 	UE_LOG(LogEditorAutomationTests, Log, TEXT("%s has been loaded."), *ShortMapName);
 
 	//Log out to a text file the time it takes to load the map.
-	WriteToTextFile(TEXT("Performance"), LoadedMapName, TEXT("RAWMapLoadTime"), MapLoadTime, TEXT(","));
+	WriteToTextFile(TEXT("Performance"), LoadedMapName, TEXT("RAWMapLoadTime.txt"), MapLoadTime, TEXT(","));
 	
 	UE_LOG(LogEditorAutomationTests, Display, TEXT("%s took %.3f to load."), *LoadedMapName, MapLoadTime);
 	
diff --git a/Engine/Source/Editor/UnrealEd/Private/Tests/EditorPerformanceTests.cpp b/Engine/Source/Editor/UnrealEd/Private/Tests/EditorPerformanceTests.cpp
--- a/Engine/Source/Editor/UnrealEd/Private/Tests/EditorPerformanceTests.cpp
+++ b/Engine/Source/Editor/UnrealEd/Private/Tests/EditorPerformanceTests.cpp
@@ -80,25 +80,27 @@ void FMapPerformanceInEditor::GetTests(TArray<FString>& OutBeautifiedNames, TArr
 			}
 		}
 	}
-
 }
 
 bool FMapPerformanceInEditor::RunTest(const FString& Parameters)
 {
+	//This is used to hold the data that is being captured.
+	EditorPerfCaptureParameters EditorPerformanceData;
+
 	//Get the map name from the parameters.
 	FString MapName = Parameters;
-
+	
 	//Gets the main frame module to get the name of our current level.
 	const IMainFrameModule& MainFrameModule = FModuleManager::GetModuleChecked< IMainFrameModule >("MainFrame");
-	FString ShortMapName = MainFrameModule.GetLoadedLevelName();
+	EditorPerformanceData.MapName = MainFrameModule.GetLoadedLevelName();
 
 	//Duration variable will be used to indicate how long the test will run for.  Defaults to a minute.
-	float Duration = 0;
+	EditorPerformanceData.TestDuration = 60;
 
-	if (!MapName.Equals(TEXT("CurrentLevel")))
+	if (!MapName.Equals(TEXT("CurrentLevel"), ESearchCase::IgnoreCase))
 	{
 		//Get the base filename for the map that will be used.
-		ShortMapName = FPaths::GetBaseFilename(MapName);
+		EditorPerformanceData.MapName = FPaths::GetBaseFilename(MapName);
 
 		//Load Map and get the time it took to take to load the map.
 		ADD_LATENT_AUTOMATION_COMMAND(FEditorLoadMap(MapName));
@@ -112,26 +114,25 @@ bool FMapPerformanceInEditor::RunTest(const FString& Parameters)
 	for (auto Iter = AutomationTestSettings->EditorPerformanceTestMaps.CreateConstIterator(); Iter; ++Iter)
 	{
 		FString IterMapName = FPaths::GetBaseFilename(Iter->PerformanceTestmap.FilePath);
-		if ((IterMapName == ShortMapName) || (IterMapName == TEXT("CurrentLevel")))
+		if ((IterMapName == EditorPerformanceData.MapName || (IterMapName == TEXT("CurrentLevel"))))
 		{
-			Duration = ((Iter->TestTimer));
+			EditorPerformanceData.TestDuration = ((Iter->TestTimer));
 			//If the duration is equal to 0 then we simply warn the user that they need to set the test timer option for the performance test.
 			//If the duration is less than 0 then we fail this test.
-			if (Duration == 0)
+			if (EditorPerformanceData.TestDuration == 0)
 			{
-				UE_LOG(LogEditorAutomationTests, Warning, TEXT("Please set the test timer for '%s' in the automation preferences or engine.ini."), *ShortMapName);
+				UE_LOG(LogEditorAutomationTests, Warning, TEXT("Please set the test timer for '%s' in the automation preferences or engine.ini."), *EditorPerformanceData.MapName);
 			}
 			else
-			if (Duration < 0)
+			if (EditorPerformanceData.TestDuration < 0)
 			{
-				UE_LOG(LogEditorAutomationTests, Error, TEXT("Test timer preference option '%s' is less than 0."), *ShortMapName);
+				UE_LOG(LogEditorAutomationTests, Error, TEXT("Test timer preference option for '%s' is less than 0."), *EditorPerformanceData.MapName);
 				return false;
 			}
 		}
 	}
 
-
-	UE_LOG(LogEditorAutomationTests, Log, TEXT("Running the performance capture test for %.0f seconds on %s"), Duration, *ShortMapName);
+	UE_LOG(LogEditorAutomationTests, Log, TEXT("Running the performance capture test for %.0i seconds on %s"), EditorPerformanceData.TestDuration, *EditorPerformanceData.MapName);
 	
 	//Move the viewport views to the first bookmark
 	ADD_LATENT_AUTOMATION_COMMAND(FChangeViewportToFirstAvailableBookmarkCommand);
@@ -142,12 +143,7 @@ bool FMapPerformanceInEditor::RunTest(const FString& Parameters)
 	
 	//Grab the performance numbers based on the duration.
 	ADD_LATENT_AUTOMATION_COMMAND(FWaitLatentCommand(0.5f));
-	ADD_LATENT_AUTOMATION_COMMAND(FEditorPerformanceCommand(Duration));
-	
-
-	//Combine performance data into one chart.
-	ADD_LATENT_AUTOMATION_COMMAND(FWaitLatentCommand(0.5f));
-	ADD_LATENT_AUTOMATION_COMMAND(FGenerateEditorPerformanceCharts(ShortMapName));
+	ADD_LATENT_AUTOMATION_COMMAND(FEditorPerfCaptureCommand(EditorPerformanceData));
 
 	return true;
 }
diff --git a/Engine/Source/Editor/UnrealEd/Public/Tests/AutomationEditorCommon.h b/Engine/Source/Editor/UnrealEd/Public/Tests/AutomationEditorCommon.h
--- a/Engine/Source/Editor/UnrealEd/Public/Tests/AutomationEditorCommon.h
+++ b/Engine/Source/Editor/UnrealEd/Public/Tests/AutomationEditorCommon.h
@@ -4,7 +4,6 @@
 #include "AssetRegistryModule.h"
 #include "ModuleManager.h"
 #include "AutomationCommon.h"
-
 #include "Tests/AutomationTestSettings.h"
 
 //Includes needed for opening certain assets
@@ -96,7 +95,33 @@ namespace AutomationEditorCommonUtils
 	void ApplyCustomFactorySettings(UFactory* InFactory, const TArray<FImportFactorySettingValues>& FactorySettings);
 }
 
-
+struct EditorPerfCaptureParameters
+{
+	//Basic Test Info
+	FString MapName;
+	int32 TestDuration;
+
+	//Saved Performance Values
+	float MapLoadTime;
+	TArray<FTimespan> TimeBetweenCaptures;
+	TArray<float> AverageFPS;
+	TArray<float> AverageFrameTime;
+	TArray<float> UsedPhysical;
+	TArray<float> AvailablePhysical;
+	TArray<float> AvailableVirtual;
+	TArray<float> UsedVirtual;
+	TArray<float> PeakUsedVirtual;
+	TArray<float> PeakUsedPhysical;
+	TArray<FDateTime> TimeStamp;
+	TArray<FString> FormattedTimeStamp;
+
+	EditorPerfCaptureParameters()
+		: MapName(TEXT("None"))
+		, TestDuration(60)
+		, MapLoadTime(0)
+	{
+	}
+};
 
 //////////////////////////////////////////////////////////////////////////
 //Common latent commands used for automated editor testing.
@@ -130,7 +155,19 @@ DEFINE_LATENT_AUTOMATION_COMMAND(FEndPlayMapCommand);
 /**
 * This will display the average FPS and Memory numbers over a duration of time.
 */
-DEFINE_LATENT_AUTOMATION_COMMAND_ONE_PARAMETER(FEditorPerformanceCommand, float, Duration);
+//DEFINE_LATENT_AUTOMATION_COMMAND_ONE_PARAMETER(FEditorPerformanceCommand, float, Duration);
+
+/**
+* This will capture the average FPS and Memory numbers over a duration of time.
+* @param The name of the EditorPerfCaptureParameters struct that will be used to hold the data.
+*/
+DEFINE_LATENT_AUTOMATION_COMMAND_ONE_PARAMETER(FEditorPerfCaptureCommand, EditorPerfCaptureParameters, EditorPerfStats);
+
+/**
+* Generates the editor performance capture chart.
+* @param The name of the EditorPerfCaptureParameters struct that will be used to hold the data.
+*/
+DEFINE_LATENT_AUTOMATION_COMMAND_ONE_PARAMETER(FGenerateEditorPerfCaptureCommand, EditorPerfCaptureParameters, EditorPerfStats);
 
 /**
 * Loads a map
@@ -141,7 +178,8 @@ DEFINE_LATENT_AUTOMATION_COMMAND_ONE_PARAMETER(FEditorLoadMap, FString, MapName)
 * Generates the editor performance capture chart.
 * @param Input the name of the map.
 */
-DEFINE_LATENT_AUTOMATION_COMMAND_ONE_PARAMETER(FGenerateEditorPerformanceCharts, FString, MapName);
+//DEFINE_LATENT_AUTOMATION_COMMAND_ONE_PARAMETER(FGenerateEditorPerformanceCharts, FString, MapName);
+
 
 /**
 * Waits for shaders to finish compiling before moving on to the next thing.
