Minor fix to the FEditorPerfCommand.

diff --git a/Engine/Source/Editor/UnrealEd/Private/Tests/AutomationEditorCommon.cpp b/Engine/Source/Editor/UnrealEd/Private/Tests/AutomationEditorCommon.cpp
--- a/Engine/Source/Editor/UnrealEd/Private/Tests/AutomationEditorCommon.cpp
+++ b/Engine/Source/Editor/UnrealEd/Private/Tests/AutomationEditorCommon.cpp
@@ -491,7 +491,7 @@ bool FEditorPerfCaptureCommand::Update()
 		ElapsedTime = CurrentDateAndTime.GetTicks() - EditorPerfStats.TimeStamp[0].GetTicks();
 	}
 
-	if (ElapsedTime.GetTotalSeconds() <= EditorPerfStats.TestDuration + 1)
+	if (ElapsedTime.GetTotalSeconds() <= EditorPerfStats.TestDuration)
 	{
 		//Find the Average FPS
 		//Clamp to avoid huge averages at startup or after hitches
