Fixed issue where cultures to stage in project packaging settings would not save to file due to an uninitialized class member inhibiting change notifications.

diff --git a/Engine/Source/Editor/UnrealEd/Private/Settings/ProjectPackagingSettingsCustomization.h b/Engine/Source/Editor/UnrealEd/Private/Settings/ProjectPackagingSettingsCustomization.h
--- a/Engine/Source/Editor/UnrealEd/Private/Settings/ProjectPackagingSettingsCustomization.h
+++ b/Engine/Source/Editor/UnrealEd/Private/Settings/ProjectPackagingSettingsCustomization.h
@@ -110,6 +110,7 @@ protected:
 
 	FProjectPackagingSettingsCustomization()
 		: FilterCulturesChoice(EFilterCulturesChoices::AllAvailableCultures)
+		, IsInBatchSelectOperation(false)
 	{
 
 	}
@@ -357,6 +358,7 @@ protected:
 		{
 			RemoveCulture(Culture->GetName());
 		}
+		
 	}
 
 	bool IsCultureSelected(FCulturePtr Culture)
