Improvements to NiagaraCompiler and VM.

- Support for Matrix and scalar ops.
- Constant handling improved to support saclars and matrices. Also reduced constant table to only constants being used. (Matrix and scalar attributes still to do.)
- Refactored compiler to separate the VM from UnrealEd and allow future compilers for compute etc.
- Reduced temp register usage.
- Improved output attriubte writing.
- Improved support for tooltips / descriptions and other UI niceties for Niagara Ops.

- Support for VectorVM ops with 4 operands.
- Refactored VM to be neater and fixed a few bugs that were caused by the mess and code duplication.
- New VM ops for splatting, shuffling and streaming output.

- Fixed a few bugs in the UI Node code.

diff --git a/Engine/Source/Editor/GraphEditor/Private/KismetPins/SGraphPinNum.cpp b/Engine/Source/Editor/GraphEditor/Private/KismetPins/SGraphPinNum.cpp
--- a/Engine/Source/Editor/GraphEditor/Private/KismetPins/SGraphPinNum.cpp
+++ b/Engine/Source/Editor/GraphEditor/Private/KismetPins/SGraphPinNum.cpp
@@ -16,16 +16,17 @@ void SGraphPinNum::SetTypeInValue(const FText& NewTypeInValue, ETextCommit::Type
 	FString TypeValueString = NewTypeInValue.ToString();
 	if (TypeValueString.IsNumeric())
 	{
+		//This is rubbish.
+		//TODO: Make not rubbish.
 		const UEdGraphSchema_K2* K2Schema = Cast<const UEdGraphSchema_K2>(GraphPinObj->GetSchema());
-
-		if (GraphPinObj->PinType.PinCategory == K2Schema->PC_Int || GraphPinObj->PinType.PinCategory == K2Schema->PC_Byte)
+		if ((K2Schema && (GraphPinObj->PinType.PinCategory == K2Schema->PC_Int || GraphPinObj->PinType.PinCategory == K2Schema->PC_Byte)))
 		{
 			int32 IntValue = FCString::Atoi(*TypeValueString);
 			GraphPinObj->GetSchema()->TrySetDefaultValue(*GraphPinObj, FString::FromInt(IntValue));
 		}
 		else
 		{
-			GraphPinObj->GetSchema()->TrySetDefaultValue(*GraphPinObj, *TypeValueString );
+			GraphPinObj->GetSchema()->TrySetDefaultValue(*GraphPinObj, *TypeValueString);
 		}
 	}
 }
diff --git a/Engine/Source/Editor/GraphEditor/Private/NiagaraPins/SGraphPinVector4.cpp b/Engine/Source/Editor/GraphEditor/Private/NiagaraPins/SGraphPinVector4.cpp
--- a/Engine/Source/Editor/GraphEditor/Private/NiagaraPins/SGraphPinVector4.cpp
+++ b/Engine/Source/Editor/GraphEditor/Private/NiagaraPins/SGraphPinVector4.cpp
@@ -118,14 +118,14 @@ public:
 								[
 									SNew(STextBlock)
 									.Font(FEditorStyle::GetFontStyle("Graph.VectorEditableTextBox"))
-									.Text(LOCTEXT("VectorNodeZAxisValueLabel", "W"))
+									.Text(LOCTEXT("VectorNodeWAxisValueLabel", "W"))
 									.ColorAndOpacity(LabelClr)
 								]
 								.Value(this, &SVector4TextBox::GetTypeInValue_3)
 									.OnValueCommitted(InArgs._OnFloatCommitted_Box_3)
 									.Font(FEditorStyle::GetFontStyle("Graph.VectorEditableTextBox"))
 									.UndeterminedString(LOCTEXT("MultipleValues", "Multiple Values"))
-									.ToolTipText(LOCTEXT("VectorNodeZAxisValueLabel_ToolTip", "W value"))
+									.ToolTipText(LOCTEXT("VectorNodeWAxisValueLabel_ToolTip", "W value"))
 									.EditableTextBoxStyle(&FEditorStyle::GetWidgetStyle<FEditableTextBoxStyle>("Graph.VectorEditableTextBox"))
 									.BorderForegroundColor(FLinearColor::White)
 									.BorderBackgroundColor(FLinearColor::White)
@@ -271,7 +271,7 @@ void SGraphPinVector4::OnChangedValueTextBox_2(float NewValue, ETextCommit::Type
 
 	FString DefaultValue;
 	//Update Z value
-	DefaultValue = GetValue(TextBox_0) + FString(TEXT(",")) + GetValue(TextBox_1) + FString(TEXT(",")) + ValueStr + FString(TEXT(",")) + GetValue(TextBox_2);
+	DefaultValue = GetValue(TextBox_0) + FString(TEXT(",")) + GetValue(TextBox_1) + FString(TEXT(",")) + ValueStr + FString(TEXT(",")) + GetValue(TextBox_3);
 	//Set new default value
 	GraphPinObj->GetSchema()->TrySetDefaultValue(*GraphPinObj, DefaultValue);
 }
diff --git a/Engine/Source/Editor/GraphEditor/Private/NiagaraPins/SGraphPinVector4.h b/Engine/Source/Editor/GraphEditor/Private/NiagaraPins/SGraphPinVector4.h
--- a/Engine/Source/Editor/GraphEditor/Private/NiagaraPins/SGraphPinVector4.h
+++ b/Engine/Source/Editor/GraphEditor/Private/NiagaraPins/SGraphPinVector4.h
@@ -54,6 +54,7 @@ private:
 	void OnChangedValueTextBox_2(float NewValue, ETextCommit::Type CommitInfo);
 	void OnChangedValueTextBox_3(float NewValue, ETextCommit::Type CommitInfo);
 
+	static FVector4 ParseValue(FString Value);
 private:
 	/** Flag is true if the widget is used to represent a rotator; false otherwise */
 	bool bIsRotator;
diff --git a/Engine/Source/Editor/GraphEditor/Private/NodeFactory.cpp b/Engine/Source/Editor/GraphEditor/Private/NodeFactory.cpp
--- a/Engine/Source/Editor/GraphEditor/Private/NodeFactory.cpp
+++ b/Engine/Source/Editor/GraphEditor/Private/NodeFactory.cpp
@@ -374,11 +374,18 @@ TSharedPtr<SGraphPin> FNodeFactory::CreatePinWidget(UEdGraphPin* InPin)
 
 	if (const UEdGraphSchema_Niagara* NSchema = Cast<const UEdGraphSchema_Niagara>(InPin->GetSchema()))
 	{
-		UScriptStruct* VectorStruct = FindObjectChecked<UScriptStruct>(UObject::StaticClass(), TEXT("Vector"));
-		if (InPin->PinType.PinSubCategoryObject == VectorStruct)
+		if (InPin->PinType.PinCategory == NSchema->PC_Float)
+		{
+			return SNew(SGraphPinNum, InPin);
+		}
+		if (InPin->PinType.PinCategory == NSchema->PC_Vector)
 		{
 			return SNew(SGraphPinVector4, InPin);
 		}
+		if (InPin->PinType.PinCategory == NSchema->PC_Matrix)
+		{
+			return SNew(SGraphPin, InPin);
+		}
 	}
 
 
diff --git a/Engine/Source/Editor/NiagaraEditor/Private/NiagaraCompiler.cpp b/Engine/Source/Editor/NiagaraEditor/Private/NiagaraCompiler.cpp
--- a/Engine/Source/Editor/NiagaraEditor/Private/NiagaraCompiler.cpp
+++ b/Engine/Source/Editor/NiagaraEditor/Private/NiagaraCompiler.cpp
@@ -8,414 +8,281 @@
 #include "EdGraphUtilities.h"
 #include "VectorVM.h"
 #include "ComponentReregisterContext.h"
+#include "NiagaraCompiler_VectorVM.h"
+#include "NiagaraEditorCommon.h"
+
+#define LOCTEXT_NAMESPACE "NiagaraCompiler"
 
 DEFINE_LOG_CATEGORY_STATIC(LogNiagaraCompiler,All,All);
 
-/** Information related to an expression. */
-struct FNiagaraExpr
+void FNiagaraEditorModule::CompileScript(UNiagaraScript* ScriptToCompile)
 {
-	/** VM opcode. */
-	uint8 OpIndex;
-	/** Source operands. */
-	int32 Src[3];
-	/* Source operand type vector (bit 0-2, 1=constant)*/
-	uint8 SrcOperandTypeVector;
-
-	/** Default constructor. */
-	FNiagaraExpr()
+	check(ScriptToCompile != NULL);
+	if (ScriptToCompile->Source == NULL)
 	{
-		OpIndex = 0;
-		SrcOperandTypeVector = 0;
-		for (int32 i = 0; i < 3; i++)
-		{
-			Src[i] = INDEX_NONE;
-		}
+		UE_LOG(LogNiagaraCompiler, Error, TEXT("No source for Niagara script: %s"), *ScriptToCompile->GetPathName());
+		return;
 	}
-};
-
-/** Compiler context, passed around during compilation. */
-struct FNiagaraCompilerContext
-{
-	/** Expression. */
-	TArray<FNiagaraExpr> Expressions;
-	/** Constant table. */
-	TArray<FVector4> Constants;
-	/** Constant names. */
-	TArray<FName> ConstantNames;
-	/** Map from pin to expression to prevent executing subtrees multiple times. */
-	TMap<UEdGraphPin*,int32> PinToExpression;
-	/** Name of all attributes. */
-	TArray<FName> Attributes;
-	/** The compilation log. */
-	FCompilerResultsLog& Log;
-
-	/** Initialization constructor. */
-	explicit FNiagaraCompilerContext(FCompilerResultsLog& InLog)
-		: Log(InLog)
-	{
-		ConstantNames.Add(TEXT("__zero__"));
-		Constants.Add(FVector4(0.0f, 0.0f, 0.0f, 0.0f));
 
-		// Setup built-in constants.
-		for (uint32 i = 0; i < NiagaraConstants::NumBuiltinConstants; i++)
-		{
-			ConstantNames.Add(NiagaraConstants::ConstantNames[i]);
-			Constants.Add(FVector4(0.0f, 0.0f, 0.0f, 0.0f));
-		}
-	}
-};
+	TComponentReregisterContext<UNiagaraComponent> ComponentReregisterContext;
 
-/** Convert an expression index to an attribute index. */
-static int32 AttrIndexFromExpressionIndex(int32 ExpressionIndex)
-{
-	return ExpressionIndex & 0x3fffffff;
-}
+	FNiagaraCompiler_VectorVM Compiler;
+	Compiler.CompileScript(ScriptToCompile);
 
-/** Returns true if the expression index refers to an attribute. */
-static bool IsAttrIndex(int32 ExpressionIndex)
-{
-	return ExpressionIndex != INDEX_NONE && (ExpressionIndex & 0x80000000);
+	//Compile for compute here too?
 }
 
-/** Convert an expression index to a constant index. */
-static int32 ConstIndexFromExpressionIndex(int32 ExpressionIndex)
-{
-	return ExpressionIndex & 0x3fffffff;
-}
+//////////////////////////////////////////////////////////////////////////
 
-/** Returns true if the expression index refers to a constant. */
-static bool IsConstIndex(int32 ExpressionIndex)
+/** Expression that gets an input attribute. */
+class FNiagaraExpression_GetAttribute : public FNiagaraExpression
 {
-	return ExpressionIndex != INDEX_NONE && (ExpressionIndex & 0x40000000);
-}
+public:
 
-/** Returns true if the expression index is valid. */
-static bool IsValidExpressionIndex(FNiagaraCompilerContext& Context, int32 ExpressionIndex)
-{
-	if (ExpressionIndex == INDEX_NONE)
+	FNiagaraExpression_GetAttribute(class FNiagaraCompiler* InCompiler, ENiagaraDataType InType, FName InName)
+		: FNiagaraExpression(InCompiler, InType)
+		, AttributeName(InName)
 	{
-		return false;
+		ResultLocation = ENiagaraExpressionResultLocation::InputData;
 	}
-	if (IsAttrIndex(ExpressionIndex))
+
+	virtual void Process()override
 	{
-		return Context.Attributes.IsValidIndex(AttrIndexFromExpressionIndex(ExpressionIndex));
+		check(ResultLocation == ENiagaraExpressionResultLocation::InputData);
+		ResultIndex = Compiler->GetAttributeIndex(AttributeName);
 	}
-	return Context.Expressions.IsValidIndex(ExpressionIndex);
-}
 
-/** Returns true if the op has valid source operands. */
-static bool IsValidOp(FNiagaraCompilerContext& Context, FNiagaraExpr& Op)
+	FName AttributeName;
+};
+
+/** Expression that gets a constant. */
+class FNiagaraExpression_GetConstant : public FNiagaraExpression
 {
-	bool bValid = true;
-	for (int32 SrcIndex = 0; bValid && SrcIndex < 3; ++SrcIndex)
+public:
+
+	FNiagaraExpression_GetConstant(class FNiagaraCompiler* InCompiler, ENiagaraDataType InType, FName InName, bool bIsInternal)
+		: FNiagaraExpression(InCompiler, InType)
+		, ConstantName(InName)
+		, bInternal(bIsInternal)
 	{
-		VectorVM::EOpSrc::Type Type = VectorVM::GetOpCodeInfo(Op.OpIndex).SrcTypes[SrcIndex];
-		int32 ExpressionIndex = Op.Src[SrcIndex];
-		switch (Type)
+		ResultLocation = ENiagaraExpressionResultLocation::Constants;
+		//For matrix constants we must add 4 input expressions that will be filled in as the constant is processed.
+		//They must at least exist now so that other expressions can reference them.
+		if (ResultType == ENiagaraDataType::Matrix)
 		{
-		case VectorVM::EOpSrc::Register:
-			bValid = bValid && IsValidExpressionIndex(Context, ExpressionIndex);
-			break;
-		case VectorVM::EOpSrc::Const:
-			bValid = bValid && IsConstIndex(ExpressionIndex)
-				&& Context.Constants.IsValidIndex(ConstIndexFromExpressionIndex(ExpressionIndex));
-			break;
-		default:
-			break;
+			SourceExpressions.Add(MakeShareable(new FNiagaraExpression(InCompiler, ENiagaraDataType::Vector)));
+			SourceExpressions.Add(MakeShareable(new FNiagaraExpression(InCompiler, ENiagaraDataType::Vector)));
+			SourceExpressions.Add(MakeShareable(new FNiagaraExpression(InCompiler, ENiagaraDataType::Vector)));
+			SourceExpressions.Add(MakeShareable(new FNiagaraExpression(InCompiler, ENiagaraDataType::Vector)));
 		}
 	}
-	return bValid;
-}
 
-/** Evaluates the graph connected to a pin. */
-static int32 EvaluateGraph(FNiagaraCompilerContext& Context, UEdGraphPin* Pin)
-{
-	int32 ExpressionIndex = INDEX_NONE;
-	if (Pin)
+	virtual void Process()override
 	{
-		if (Pin->LinkedTo.Num() == 1)
+		check(ResultLocation == ENiagaraExpressionResultLocation::Constants);
+		//Get the index of the constant. Also gets a component index if this is for a packed scalar constant.
+		Compiler->GetConstantResultIndex(ConstantName, bInternal, ResultIndex, ComponentIndex);
+		
+		if (ResultType == ENiagaraDataType::Matrix)
 		{
-			check(Pin->Direction == EGPD_Input);
-			Pin = Pin->LinkedTo[0];
-			check(Pin->Direction == EGPD_Output);
-
-			int32* CachedExpressionIndex = Context.PinToExpression.Find(Pin);
-			if (CachedExpressionIndex)
-			{
-				ExpressionIndex = *CachedExpressionIndex;
-			}
-			else
-			{
-				UEdGraphNode* Node = Pin->GetOwningNode();
-				if (Node->IsA(UNiagaraNodeOp::StaticClass()))
-				{
-					UNiagaraNodeOp* OpNode = (UNiagaraNodeOp*)Node;
-					VectorVM::FVectorVMOpInfo const& OpInfo = VectorVM::GetOpCodeInfo(OpNode->OpIndex);
-
-					FNiagaraExpr NewOp;
-					NewOp.OpIndex = OpNode->OpIndex;
-					NewOp.SrcOperandTypeVector = 0;
-
-					for (int32 SrcIndex = 0; SrcIndex < 3; ++SrcIndex)
-					{
-						if (OpInfo.SrcTypes[SrcIndex] != VectorVM::EOpSrc::Invalid)
-						{
-							UEdGraphPin* SrcPin = OpNode->FindPinChecked(UNiagaraNodeOp::InPinNames[SrcIndex]);
-							NewOp.Src[SrcIndex] = EvaluateGraph(Context,SrcPin);
-							if (NewOp.Src[SrcIndex] & 0x40000000)
-							{
-								NewOp.SrcOperandTypeVector |= (1 << SrcIndex);
-							}
-						}
-					}
-
-					ExpressionIndex = Context.Expressions.Add(NewOp);
-				}
-				else if (Node->IsA(UNiagaraNodeGetAttr::StaticClass()))
-				{
-					UNiagaraNodeGetAttr* AttrNode = (UNiagaraNodeGetAttr*)Node;
-					int32 AttrIndex = INDEX_NONE;
-					if (Context.Attributes.Find(AttrNode->AttrName, AttrIndex))
-					{
-						ExpressionIndex = AttrIndex | 0x80000000;
-					}
-					else if (Context.ConstantNames.Find(AttrNode->AttrName, AttrIndex))
-					{
-						ExpressionIndex = AttrIndex | 0x40000000;
-					}
-				}
-				Context.PinToExpression.Add(Pin, ExpressionIndex);
-			}
-		}
-		else if (!Pin->bDefaultValueIsIgnored)
-		{
-			FString ConstString = Pin->GetDefaultAsString();
-			FName ConstName(*ConstString);
-			int32 ConstIndex = Context.ConstantNames.Find(ConstName);
-			if (ConstIndex == INDEX_NONE)
-			{
-				TArray<FString> ResultString;
-				ConstString.Trim();
-				ConstString.TrimTrailing();
-				ConstString.ParseIntoArray(&ResultString, TEXT(","), true);
-				ConstIndex = 0;
-				if (ResultString.Num() == 4)
-				{
-					ConstIndex = Context.Constants.Add(FVector4(FCString::Atof(*ResultString[0]), FCString::Atof(*ResultString[1]), FCString::Atof(*ResultString[2]), FCString::Atof(*ResultString[3])));
-					check(Context.ConstantNames.Num() == ConstIndex);
-					Context.ConstantNames.Add(ConstName);
-				}
-			}
-			check(Context.Constants.IsValidIndex(ConstIndex));
-			ExpressionIndex = ConstIndex | 0x40000000;			
+			check(SourceExpressions.Num() == 4);
+			TNiagaraExprPtr Src0 = GetSourceExpression(0);
+			TNiagaraExprPtr Src1 = GetSourceExpression(1);
+			TNiagaraExprPtr Src2 = GetSourceExpression(2);
+			TNiagaraExprPtr Src3 = GetSourceExpression(3);
+			Src0->ResultLocation = ENiagaraExpressionResultLocation::Constants;
+			Src0->ResultIndex = ResultIndex;
+			Src1->ResultLocation = ENiagaraExpressionResultLocation::Constants;
+			Src1->ResultIndex = ResultIndex + 1;
+			Src2->ResultLocation = ENiagaraExpressionResultLocation::Constants;
+			Src2->ResultIndex = ResultIndex + 2;
+			Src3->ResultLocation = ENiagaraExpressionResultLocation::Constants;
+			Src3->ResultIndex = ResultIndex + 3;
 		}
 	}
-	return ExpressionIndex;
-}
 
-void FNiagaraEditorModule::CompileScript(UNiagaraScript* ScriptToCompile)
+	FName ConstantName;
+	bool bInternal;
+};
+
+/** Expression that just collects some other expressions together. E.g for a matrix. */
+class FNiagaraExpression_Collection : public FNiagaraExpression
 {
-	check(ScriptToCompile != NULL);
-	if (ScriptToCompile->Source == NULL)
+public:
+
+	FNiagaraExpression_Collection(class FNiagaraCompiler* InCompiler, ENiagaraDataType Type, TArray<TNiagaraExprPtr>& InSourceExpressions)
+		: FNiagaraExpression(InCompiler, Type)
 	{
-		UE_LOG(LogNiagaraCompiler, Error, TEXT("No source for Niagara script: %s"), *ScriptToCompile->GetPathName());
-		return;
+		SourceExpressions = InSourceExpressions;
 	}
+};
 
-	TComponentReregisterContext<UNiagaraComponent> ComponentReregisterContext;
-
-	UNiagaraScriptSource* Source = CastChecked<UNiagaraScriptSource>(ScriptToCompile->Source);
-	check(Source->UpdateGraph);
-
-	// Results log.
-	FCompilerResultsLog MessageLog;
-
-	// Clone the source graph so we can modify it as needed; merging in the child graphs
-	UEdGraph* UpdateGraph = FEdGraphUtilities::CloneGraph(Source->UpdateGraph, Source, &MessageLog, true); 
-	FEdGraphUtilities::MergeChildrenGraphsIn(UpdateGraph, UpdateGraph, /*bRequireSchemaMatch=*/ true);
+//////////////////////////////////////////////////////////////////////////
 
-	// Find the output node.
-	UNiagaraNodeOutputUpdate* OutputNode = NULL;
+void FNiagaraCompiler::CheckInputs(FName OpName, TArray<TNiagaraExprPtr>& Inputs)
+{
+	//check the types of the input expressions.
+	const FNiagaraOpInfo* OpInfo = FNiagaraOpInfo::GetOpInfo(OpName);
+	check(OpInfo);
+	int32 NumInputs = Inputs.Num();
+	check(OpInfo->Inputs.Num() == NumInputs);
+	for (int32 i = 0; i < NumInputs ; ++i)
 	{
-		TArray<UNiagaraNodeOutputUpdate*> OutputNodes;
-		UpdateGraph->GetNodesOfClass(OutputNodes);
-		if (OutputNodes.Num() != 1)
+		if (!(OpInfo->Inputs[i].DataType == Inputs[i]->ResultType || Inputs[i]->ResultType == ENiagaraDataType::Unknown))
 		{
-			UE_LOG(LogNiagaraCompiler, Error, TEXT("Script contains %s output nodes: %s"),
-				OutputNodes.Num() == 0 ? TEXT("no") : TEXT("too many"),
-				*ScriptToCompile->GetPathName()
-				);
-			return;
+			UE_LOG(LogNiagaraCompiler, Error, TEXT("Exression %s has incorrect inputs!\nExpected: %d - Actual: %d"), *OpName.ToString(), (int32)OpInfo->Inputs[i].DataType, (int32)((TNiagaraExprPtr)Inputs[i])->ResultType);
 		}
-		OutputNode = OutputNodes[0];
 	}
-	check(OutputNode);
-
-	// Traverse the node graph for each output and generate expressions as we go.
-	FNiagaraCompilerContext Context(MessageLog);
-	Source->GetParticleAttributes(Context.Attributes);
-	TArray<int32> OutputExpressions;
-	OutputExpressions.AddUninitialized(Context.Attributes.Num());
-	FMemory::Memset(OutputExpressions.GetData(), 0xff, Context.Attributes.Num() * sizeof(int32));
+}
 
-	for (int32 PinIndex = 0; PinIndex < OutputNode->Pins.Num(); ++PinIndex)
+void FNiagaraCompiler::CheckOutputs(FName OpName, TArray<TNiagaraExprPtr>& Outputs)
+{
+	//check the types of the input expressions.
+	const FNiagaraOpInfo* OpInfo = FNiagaraOpInfo::GetOpInfo(OpName);
+	check(OpInfo);
+	int32 NumOutputs = Outputs.Num();
+	check(OpInfo->Outputs.Num() == NumOutputs);
+	for (int32 i = 0; i < NumOutputs; ++i)
 	{
-		UEdGraphPin* Pin = OutputNode->Pins[PinIndex];
-		int32 AttrIndex = Context.Attributes.Find(FName(*Pin->PinName));
-		if (AttrIndex != INDEX_NONE)
+		if (!(OpInfo->Outputs[i].DataType == Outputs[i]->ResultType || Outputs[i]->ResultType == ENiagaraDataType::Unknown))
 		{
-			int32 ExpressionIndex = EvaluateGraph(Context, Pin);
-			OutputExpressions[AttrIndex] = ExpressionIndex;
+			UE_LOG(LogNiagaraCompiler, Error, TEXT("Exression %s has incorrect outputs!\nExpected: %d - Actual: %d"), *OpName.ToString(), (int32)OpInfo->Outputs[i].DataType, (int32)((TNiagaraExprPtr)Outputs[i])->ResultType);
 		}
 	}
+}
+
+int32 FNiagaraCompiler::GetAttributeIndex(FName Name)
+{
+	return Attributes.Find(Name);
+}
 
-	// Generate pass-thru ops for any outputs that are not connected.
-	for (int32 AttrIndex = 0; AttrIndex < OutputExpressions.Num(); ++AttrIndex)
+TNiagaraExprPtr FNiagaraCompiler::CompileOutputPin(UEdGraphPin* Pin)
+{
+	check(Pin->Direction == EGPD_Output);
+
+	TNiagaraExprPtr Ret;
+
+	TNiagaraExprPtr* Expr = PinToExpression.Find(Pin);
+	if (Expr)
 	{
-		if (OutputExpressions[AttrIndex] == INDEX_NONE)
-		{
-			// Generate a pass-thru op.
-			FNiagaraExpr PassThru;
-			PassThru.OpIndex = VectorVM::EOp::add;
-			PassThru.Src[0] = AttrIndex | 0x80000000;
-			PassThru.Src[1] = 0 | 0x40000000;
-			PassThru.Src[2] = INDEX_NONE;
-			PassThru.SrcOperandTypeVector = 0x2;
-			OutputExpressions[AttrIndex] = Context.Expressions.Add(PassThru);
-		}
+		Ret = *Expr; //We've compiled this pin before. Return it's expression.
 	}
-	
-	// Figure out the lifetime of each expression.
-	TArray<int32> ExpressionLifetimes;
-	ExpressionLifetimes.AddUninitialized(Context.Expressions.Num());
-	for (int32 i = 0; i < Context.Expressions.Num(); ++i)
+	else
 	{
-		ExpressionLifetimes[i] = i;
-		FNiagaraExpr& Expr = Context.Expressions[i];
-		VectorVM::FVectorVMOpInfo const& OpInfo = VectorVM::GetOpCodeInfo(Expr.OpIndex);
-		for (int32 k = 0; k < 3; ++k)
+		//Otherwise we need to compile the node to get its output pins.
+		UNiagaraNode* Node = Cast<UNiagaraNode>(Pin->GetOwningNode());
+		TArray<FNiagaraNodeResult> Outputs;
+		Node->Compile(this, Outputs);
+		for (int32 i=0; i<Outputs.Num(); ++i)
 		{
-			//if (OpInfo.SrcTypes[k] == VectorVM::EOpSrc::Register
-			if ((Expr.SrcOperandTypeVector & (1 << k)) == 0 && OpInfo.SrcTypes[k] != VectorVM::EOpSrc::Invalid
-				&& !IsAttrIndex(Expr.Src[k]))
-			{
-				check(Context.Expressions.IsValidIndex(Expr.Src[k]));
-				check(Expr.Src[k] < i);
-				ExpressionLifetimes[Expr.Src[k]] = i;
-			}
+			PinToExpression.Add(Outputs[i].OutputPin, Outputs[i].Expression);
+			//We also grab the expression for the pin we're currently interested in. Otherwise we'd have to search the map for it.
+			Ret = Outputs[i].OutputPin == Pin ? Outputs[i].Expression : Ret;
 		}
 	}
 
-	// Allocate temporary registers for the output of each expression.
-	int32 Registers[VectorVM::NumTempRegisters];
-	FMemory::Memset(Registers, 0xff, sizeof(Registers));
-	TArray<int32> RegisterAssignments;
-	RegisterAssignments.AddUninitialized(Context.Expressions.Num());
-	for (int32 i = 0; i < Context.Expressions.Num(); ++i)
-	{
-		FNiagaraExpr& Expr = Context.Expressions[i];
-		VectorVM::FVectorVMOpInfo const& OpInfo = VectorVM::GetOpCodeInfo(Expr.OpIndex);
+	return Ret;
+}
 
-		for (int32 j = 0; j < 3; ++j)
+TNiagaraExprPtr FNiagaraCompiler::CompilePin(UEdGraphPin* Pin)
+{
+	check(Pin);
+	TNiagaraExprPtr Ret;
+	if (Pin->Direction == EGPD_Input)
+	{
+		if (Pin->LinkedTo.Num() > 0)
 		{
-			//if (OpInfo.SrcTypes[j] == VectorVM::EOpSrc::Register)
-			if ((Expr.SrcOperandTypeVector & (1 << j)) == 0 && OpInfo.SrcTypes[j] != VectorVM::EOpSrc::Invalid)
+			// > 1 ???
+			Ret = CompileOutputPin(Pin->LinkedTo[0]);
+		}
+		else if (!Pin->bDefaultValueIsIgnored)
+		{
+			//No connections to this input so add the default as a const expression.
+			const UEdGraphSchema_Niagara* Schema = Cast<const UEdGraphSchema_Niagara>(Pin->GetSchema());
+			ENiagaraDataType DataType = Schema->GetPinDataType(Pin);
+
+			FString ConstString = Pin->GetDefaultAsString();
+			FName ConstName(*ConstString);
+			
+			if (DataType == ENiagaraDataType::Scalar)
 			{
-				if (IsAttrIndex(Expr.Src[j]))
-				{
-					Expr.Src[j] = VectorVM::FirstInputRegister + AttrIndexFromExpressionIndex(Expr.Src[j]);
-				}
-				else
-				{
-					Expr.Src[j] = RegisterAssignments[Expr.Src[j]];
-				}
+				float Default;
+				Schema->GetPinDefaultValue(Pin, Default);
+				ConstantData.SetOrAddInternal(ConstName, Default);
 			}
-			//else if (OpInfo.SrcTypes[j] == VectorVM::EOpSrc::Const)
-			else if ((Expr.SrcOperandTypeVector & (1 << j)) == 1)
+			else if (DataType == ENiagaraDataType::Vector)
 			{
-				// VectorVM::EOpSrc::Const.
-				Expr.Src[j] = ConstIndexFromExpressionIndex(Expr.Src[j]);
-				//Expr.SrcOperandTypeVector |= 1<<j;
-			}
-		}
-
-		// now that we've figured out register/constant assignments, we revert back to the base opcode; the kernel Exec function will handle passing 
-		// the data in correctly based on the SrcOperandTypeVector. It would be good to change the compiler so we don't need the permutation
-		// opcodes at all, i.e. figure out source operand types without needing them in the OpInfo.
-		Expr.OpIndex = OpInfo.BaseOpcode;
+				FVector4 Default;
+				Schema->GetPinDefaultValue(Pin, Default);
+				ConstantData.SetOrAddInternal(ConstName, Default);
 
-
-		int32 AttrIndex = INDEX_NONE;
-		if (OutputExpressions.Find(i, AttrIndex))
-		{
-			// Assign to an output register.
-			RegisterAssignments[i] = VectorVM::FirstOutputRegister + AttrIndex;
-		}
-		else
-		{
-			int32 AssignedRegister = INDEX_NONE;
-			for (int32 j = 0; j < VectorVM::NumTempRegisters; ++j)
+			}
+			else if (DataType == ENiagaraDataType::Matrix)
 			{
-				if (Registers[j] == INDEX_NONE
-					|| ExpressionLifetimes[Registers[j]] < i)
-				{
-					AssignedRegister = j;
-					break;
-				}
+				FMatrix Default;
+				Schema->GetPinDefaultValue(Pin, Default);
+				ConstantData.SetOrAddInternal(ConstName, Default);
 			}
-			check(AssignedRegister != INDEX_NONE && AssignedRegister < VectorVM::NumTempRegisters);
-
-			Registers[AssignedRegister] = i;
-			RegisterAssignments[i] = AssignedRegister;
+			Ret = Expression_GetInternalConstant(ConstName, DataType);
 		}
 	}
-
-	// Generate bytecode!
-	TArray<uint8>& Code = ScriptToCompile->ByteCode;
-	Code.Empty();
-	for (int32 i = 0; i < Context.Expressions.Num(); ++i)
+	else
 	{
-		FNiagaraExpr& Expr = Context.Expressions[i];
-		VectorVM::FVectorVMOpInfo const& OpInfo = VectorVM::GetOpCodeInfo(Expr.OpIndex);
+		Ret = CompileOutputPin(Pin);
+	}
 
-		int32 DestRegister = RegisterAssignments[i];
-		check(DestRegister < VectorVM::MaxRegisters);
+	return Ret;
+}
 
-		check(Expr.OpIndex < VectorVM::EOp::NumOpcodes);
-		check(Expr.OpIndex == OpInfo.BaseOpcode);
+void FNiagaraCompiler::GetParticleAttributes(TArray<FName>& OutAttributes)
+{
+	check(Source);
+	Source->GetParticleAttributes(OutAttributes);
+}
 
-		Code.Add(Expr.OpIndex);
-		Code.Add(DestRegister);
-		Code.Add(Expr.SrcOperandTypeVector);
-		for (int32 j = 0; j < 3; ++j)
-		{
-			if (OpInfo.SrcTypes[j] == VectorVM::EOpSrc::Invalid)
-				break;
-			Code.Add(Expr.Src[j]);
-		}
+TNiagaraExprPtr FNiagaraCompiler::Expression_GetAttribute(FName AttributeName)
+{
+	int32 Index = Expressions.Add(MakeShareable(new FNiagaraExpression_GetAttribute(this, ENiagaraDataType::Vector, AttributeName)));
+	return Expressions[Index];
+}
 
-		// If the expression is output to multiple attributes, copy them here.
-		if (DestRegister >= VectorVM::FirstOutputRegister)
-		{
-			int32 AttrIndex = DestRegister - VectorVM::FirstOutputRegister;
-			for (int32 j = AttrIndex+1; j < OutputExpressions.Num(); ++j)
-			{
-				if (OutputExpressions[j] == i)
-				{
-					Code.Add(VectorVM::EOp::add);
-					Code.Add(VectorVM::FirstOutputRegister + j);
-					Code.Add(0x0);	// all inputs are registers
-					Code.Add(DestRegister);
-					Code.Add(0);
-				}
-			}
-		}
-	}
+TNiagaraExprPtr FNiagaraCompiler::Expression_GetExternalConstant(FName ConstantName, ENiagaraDataType Type)
+{
+	int32 Index = Expressions.Add(MakeShareable(new FNiagaraExpression_GetConstant(this, Type, ConstantName, false)));
+	return Expressions[Index];
+}
 
-	// Terminate with the 'done' opcode.
-	Code.Add(VectorVM::EOp::done);
+TNiagaraExprPtr FNiagaraCompiler::Expression_GetInternalConstant(FName ConstantName, ENiagaraDataType Type)
+{
+	int32 Index = Expressions.Add(MakeShareable(new FNiagaraExpression_GetConstant(this, Type, ConstantName, true)));
+	return Expressions[Index];
+}
+
+TNiagaraExprPtr FNiagaraCompiler::Expression_Collection(TArray<TNiagaraExprPtr>& SourceExpressions)
+{
+	int32 Index = Expressions.Add(MakeShareable(new FNiagaraExpression_Collection(this, ENiagaraDataType::Unknown, SourceExpressions)));
+	return Expressions[Index];
+}
+
+TNiagaraExprPtr FNiagaraCompiler::GetAttribute(FName AttributeName)
+{
+	return Expression_GetAttribute(AttributeName);
+}
 
-	// And copy the constant table and attributes.
-	ScriptToCompile->ConstantTable = Context.Constants;
-	ScriptToCompile->Attributes = Context.Attributes;
+TNiagaraExprPtr FNiagaraCompiler::GetExternalConstant(FName ConstantName, float Default)
+{
+	SetOrAddConstant(false, ConstantName, Default);
+	return Expression_GetExternalConstant(ConstantName, ENiagaraDataType::Scalar);
+}
+
+TNiagaraExprPtr FNiagaraCompiler::GetExternalConstant(FName ConstantName, FVector4 Default)
+{
+	SetOrAddConstant(false, ConstantName, Default);
+	return Expression_GetExternalConstant(ConstantName, ENiagaraDataType::Vector);
+}
+
+TNiagaraExprPtr FNiagaraCompiler::GetExternalConstant(FName ConstantName, const FMatrix& Default)
+{
+	SetOrAddConstant(false, ConstantName, Default);
+	return Expression_GetExternalConstant(ConstantName, ENiagaraDataType::Matrix);
 }
diff --git a/Engine/Source/Editor/NiagaraEditor/Private/NiagaraCompiler.h b/Engine/Source/Editor/NiagaraEditor/Private/NiagaraCompiler.h
new file mode 100644
--- /dev/null
+++ b/Engine/Source/Editor/NiagaraEditor/Private/NiagaraCompiler.h
@@ -0,0 +1,85 @@
+// Copyright 1998-2014 Epic Games, Inc. All Rights Reserved.
+
+#pragma once
+
+#include "INiagaraCompiler.h"
+#include "CompilerResultsLog.h"
+#include "Niagara/NiagaraScriptConstantData.h"
+
+/** Base class for Niagara compilers. Children of this will include a compiler for the VectorVM and for Compute shaders. Possibly others. */
+class NIAGARAEDITOR_API FNiagaraCompiler : public INiagaraCompiler
+{
+protected:
+	/** The script we are compiling. */
+	UNiagaraScript* Script;
+
+	/** The source of the script we are compiling. */
+	UNiagaraScriptSource* Source;
+
+	/** The set of expressions generated from the script source. */
+	TArray<TNiagaraExprPtr> Expressions;
+
+	/** The particle attribute names. */
+	TArray<FName> Attributes;
+	
+	//All internal and external constants used in the graph.
+	FNiagaraScriptConstantData ConstantData;
+
+	/** Map of Pins to expressions. Allows us to reuse expressions for pins that have already been compiled. */
+	TMap<UEdGraphPin*, TNiagaraExprPtr> PinToExpression;
+
+	//Message log. Automatically handles marking the NodeGraph with errors.
+	FCompilerResultsLog MessageLog;
+
+	/** Compiles an output Pin on a graph node. Caches the result for any future inputs connected to it. */
+	TNiagaraExprPtr CompileOutputPin(UEdGraphPin* Pin);
+
+	/** Creates an expression retrieving a particle attribute. */
+	TNiagaraExprPtr Expression_GetAttribute(FName AttributeName);
+
+	/** Creates an expression retrieving an external constant. */
+	TNiagaraExprPtr Expression_GetExternalConstant(FName ConstantName, ENiagaraDataType Type);
+
+	/** Creates an expression retrieving an internal constant. */
+	TNiagaraExprPtr Expression_GetInternalConstant(FName ConstantName, ENiagaraDataType Type);
+
+	/** Creates an expression collecting some source expressions together for use by future expressions. */
+	TNiagaraExprPtr Expression_Collection(TArray<TNiagaraExprPtr>& SourceExpressions);
+
+	/** Sets the value of a constant. Overwrites the value if it already exists. */
+	template< typename T >
+	void SetOrAddConstant(bool bInternal, FName Name, const T& Default);
+
+public:
+
+	//Begin INiagaraCompiler Interface
+	virtual TNiagaraExprPtr CompilePin(UEdGraphPin* Pin)override;
+	virtual void GetParticleAttributes(TArray<FName>& OutAttributes)override;
+	virtual TNiagaraExprPtr GetAttribute(FName AttributeName)override;
+	virtual TNiagaraExprPtr GetExternalConstant(FName ConstantName, float Default)override;
+	virtual TNiagaraExprPtr GetExternalConstant(FName ConstantName, FVector4 Default)override;
+	virtual TNiagaraExprPtr GetExternalConstant(FName ConstantName, const FMatrix& Default)override;
+
+	virtual void CheckInputs(FName OpName, TArray<TNiagaraExprPtr>& Inputs)override;
+	virtual void CheckOutputs(FName OpName, TArray<TNiagaraExprPtr>& Outputs)override;
+	//End INiagaraCompiler Interface
+	
+	/** Gets the index into a constants table of the constant specified by Name and bInternal. */
+	virtual void GetConstantResultIndex(FName Name, bool bInternal, int32& OutResultIndex, int32& OutComponentIndex) = 0;
+	/**	Gets the index of a named attribute into either the input or output data. */
+	int32 GetAttributeIndex(FName Name);
+
+	/**	Gets the index of a free temporary location. */
+	virtual int32 AquireTemporary() = 0;
+	/** Frees the passed temporary index for use by other expressions. */
+	virtual void FreeTemporary(int32 TempIndex) = 0;
+};
+
+template< typename T >
+void FNiagaraCompiler::SetOrAddConstant(bool bInternal, FName Name, const T& Default)
+{
+	if (bInternal)
+		ConstantData.SetOrAddInternal(Name, Default);
+	else
+		ConstantData.SetOrAddExternal(Name, Default);
+}
\ No newline at end of file
diff --git a/Engine/Source/Editor/NiagaraEditor/Private/NiagaraCompiler_VectorVM.cpp b/Engine/Source/Editor/NiagaraEditor/Private/NiagaraCompiler_VectorVM.cpp
new file mode 100644
--- /dev/null
+++ b/Engine/Source/Editor/NiagaraEditor/Private/NiagaraCompiler_VectorVM.cpp
@@ -0,0 +1,595 @@
+// Copyright 1998-2014 Epic Games, Inc. All Rights Reserved.
+
+#include "NiagaraEditorPrivatePCH.h"
+#include "Engine/NiagaraScript.h"
+#include "Engine/NiagaraConstants.h"
+#include "Components/NiagaraComponent.h"
+#include "CompilerResultsLog.h"
+#include "EdGraphUtilities.h"
+#include "VectorVM.h"
+#include "ComponentReregisterContext.h"
+#include "NiagaraCompiler_VectorVM.h"
+
+#define LOCTEXT_NAMESPACE "NiagaraCompiler_VectorVM"
+
+DEFINE_LOG_CATEGORY_STATIC(LogNiagaraCompiler_VectorVM, All, All);
+
+/** Expression handling native operations of the VectorVM. */
+class FNiagaraExpression_VMOperation : public FNiagaraExpression
+{
+public:
+	VectorVM::EOp OpCode;
+
+	FNiagaraExpression_VMOperation(FNiagaraCompiler* Compiler, ENiagaraDataType InType, TArray<TNiagaraExprPtr>& InputExpressions, VectorVM::EOp Op)
+		: FNiagaraExpression(Compiler, InType)
+		, OpCode(Op)
+	{
+		SourceExpressions = InputExpressions;
+	}
+
+	FNiagaraExpression_VMOperation::~FNiagaraExpression_VMOperation()
+	{
+		if (ResultLocation == ENiagaraExpressionResultLocation::Temporaries && ResultIndex != INDEX_NONE)
+		{
+			Compiler->FreeTemporary(ResultIndex);
+		}
+	}
+
+	virtual void Process()override
+	{
+		FNiagaraCompiler_VectorVM* VMCompiler = (FNiagaraCompiler_VectorVM*)Compiler;
+
+		//Should be invalid at this point.
+		check(ResultIndex == INDEX_NONE);
+		check(ComponentIndex == INDEX_NONE);
+
+		//Special case where we're splatting but we don't know the component until now. E.g. Getting constants.
+		//Choose the component now based on the source expressions component index
+		if (OpCode == VectorVM::EOp::splatx)
+		{
+			check(SourceExpressions.Num() == 1);
+			TNiagaraExprPtr Src0 = GetSourceExpression(0);
+			if (Src0->ComponentIndex != INDEX_NONE)
+			{
+				switch (Src0->ComponentIndex)
+				{
+				case 0: OpCode = VectorVM::EOp::splatx; break;
+				case 1: OpCode = VectorVM::EOp::splaty; break;
+				case 2: OpCode = VectorVM::EOp::splatz; break;
+				case 3: OpCode = VectorVM::EOp::splatw; break;
+				}
+			}
+		}
+
+		ResultIndex = VMCompiler->AquireTemporary();
+		ResultLocation = ENiagaraExpressionResultLocation::Temporaries;
+		check(ResultIndex < VectorVM::NumTempRegisters);
+
+		VMCompiler->WriteCode((uint8)OpCode);
+		VMCompiler->WriteCode(VMCompiler->GetResultVMIndex( this ));
+
+		//Add the bitfield defining whether each source operand comes from a constant or a register.
+		int32 NumOperands = SourceExpressions.Num();
+		check(NumOperands <= 4);
+		VMCompiler->WriteCode(
+			VectorVM::CreateSrcOperandMask(
+			NumOperands > 0 ? VMCompiler->ExpressionIsConstant(GetSourceExpression(0).Get()) : false,
+			NumOperands > 1 ? VMCompiler->ExpressionIsConstant(GetSourceExpression(1).Get()) : false,
+			NumOperands > 2 ? VMCompiler->ExpressionIsConstant(GetSourceExpression(2).Get()) : false,
+			NumOperands > 3 ? VMCompiler->ExpressionIsConstant(GetSourceExpression(3).Get()) : false)
+			);
+
+		//Add the locations for each of the source operands.
+		for (int32 SrcIdx = 0; SrcIdx < SourceExpressions.Num(); ++SrcIdx)
+		{
+			VMCompiler->WriteCode(VMCompiler->GetResultVMIndex(GetSourceExpression(SrcIdx).Get()));
+		}
+	}
+};
+
+/** Expression to write attribute data output in the VM. */
+class FNiagaraExpression_VMOutput : public FNiagaraExpression
+{
+public:
+
+	FNiagaraExpression_VMOutput(class FNiagaraCompiler* InCompiler, ENiagaraDataType InType, FName InName, TNiagaraExprPtr& InSourceExpression)
+		: FNiagaraExpression(InCompiler, InType)
+		, AttributeName(InName)
+	{
+		ResultLocation = ENiagaraExpressionResultLocation::OutputData;
+		SourceExpressions.Add(InSourceExpression);
+		check(SourceExpressions.Num() == 1);
+	}
+
+	FName AttributeName;
+
+	virtual void Process()override
+	{
+		FNiagaraCompiler_VectorVM* VMCompiler = (FNiagaraCompiler_VectorVM*)Compiler;
+		check(ResultLocation == ENiagaraExpressionResultLocation::OutputData);
+		ResultIndex = VMCompiler->GetAttributeIndex(AttributeName);
+		
+		FNiagaraExpression* SrcExpr = GetSourceExpression(0).Get();
+		check(SourceExpressions.Num() == 1);
+
+		uint8 DestIndex = VMCompiler->GetResultVMIndex(this);
+		uint8 SrcIndex = VMCompiler->GetResultVMIndex(SrcExpr);
+
+		VMCompiler->WriteCode((uint8)VectorVM::EOp::output);
+		VMCompiler->WriteCode(DestIndex);
+		VMCompiler->WriteCode(VectorVM::CreateSrcOperandMask(VMCompiler->ExpressionIsConstant(SrcExpr)));
+		VMCompiler->WriteCode(SrcIndex);
+	}
+};
+
+bool FNiagaraCompiler_VectorVM::ExpressionIsConstant(FNiagaraExpression*  Expression)
+{
+	return Expression->ResultLocation == ENiagaraExpressionResultLocation::Constants;
+}
+
+uint8 FNiagaraCompiler_VectorVM::GetResultVMIndex(FNiagaraExpression* Expression)
+{
+	if (Expression->ResultLocation == ENiagaraExpressionResultLocation::InputData)
+	{
+		return VectorVM::FirstInputRegister + Expression->ResultIndex;
+	}
+	else if (Expression->ResultLocation == ENiagaraExpressionResultLocation::OutputData)
+	{
+		return VectorVM::FirstOutputRegister + Expression->ResultIndex;
+	}
+	else
+	{
+		return Expression->ResultIndex;//TempRegister or Constant.
+	}
+}
+
+void FNiagaraCompiler_VectorVM::GetConstantResultIndex(FName Name, bool bInternal, int32& OutResultIndex, int32& OutComponentIndex)
+{
+	ENiagaraDataType Type;
+	ConstantData.GetTableIndex(Name, bInternal, OutResultIndex, OutComponentIndex, Type);
+}
+
+int32 FNiagaraCompiler_VectorVM::AquireTemporary()
+{
+	int32 Free = TempRegisters.Find(false);
+	TempRegisters[Free] = true;
+	return Free;
+}
+
+void FNiagaraCompiler_VectorVM::FreeTemporary(int32 TempIndex)
+{
+	check(TempIndex < TempRegisters.Num() && TempIndex >= 0);
+	TempRegisters[TempIndex] = false;
+}
+
+void FNiagaraCompiler_VectorVM::WriteCode(uint8 InCode)
+{
+	check(Script);
+	Script->ByteCode.Add(InCode);
+}
+
+TNiagaraExprPtr FNiagaraCompiler_VectorVM::Output(FName OutputName, TNiagaraExprPtr& SourceExpression)
+{
+	int32 Index = Expressions.Add(MakeShareable(new FNiagaraExpression_VMOutput(this, ENiagaraDataType::Vector, OutputName, SourceExpression)));
+	return Expressions[Index];
+}
+
+void FNiagaraCompiler_VectorVM::CompileScript(UNiagaraScript* InScript)
+{
+	//TODO - Most of this function can be refactored up into FNiagaraCompiler once we know how the compute script and VM script will coexist in the UNiagaraScript.
+
+	check(InScript);
+
+	Script = InScript;
+	UNiagaraScriptSource* Source = CastChecked<UNiagaraScriptSource>(InScript->Source);
+
+	// Clone the source graph so we can modify it as needed; merging in the child graphs
+	UEdGraph* UpdateGraph = FEdGraphUtilities::CloneGraph(Source->UpdateGraph, Source, &MessageLog, true);
+	FEdGraphUtilities::MergeChildrenGraphsIn(UpdateGraph, UpdateGraph, /*bRequireSchemaMatch=*/ true);
+
+	TempRegisters.AddZeroed(VectorVM::NumTempRegisters);
+
+	// Find the output node.
+	UNiagaraNodeOutputUpdate* OutputNode = NULL;
+	{
+		TArray<UNiagaraNodeOutputUpdate*> OutputNodes;
+		Source->UpdateGraph->GetNodesOfClass(OutputNodes);
+		if (OutputNodes.Num() != 1)
+		{
+			UE_LOG(LogNiagaraCompiler, Error, TEXT("Script contains %s output nodes: %s"),
+				OutputNodes.Num() == 0 ? TEXT("no") : TEXT("too many"),
+				*InScript->GetPathName()
+				);
+			return;
+		}
+		OutputNode = OutputNodes[0];
+	}
+	check(OutputNode);
+
+	//Todo - Replace with attributes defined by an Emitter object.
+	Source->GetParticleAttributes(Attributes);
+
+	TArray<FNiagaraNodeResult> OutputExpressions;
+	OutputNode->Compile(this, OutputExpressions);
+
+	//Todo - track usage of temp registers so we can allocate and free and reuse from a smaller set.
+	int32 NextTempRegister = 0;
+	Script->ByteCode.Empty();
+
+	//Now we have a set of expressions from which we can generate the bytecode.
+	for (int32 i = 0; i < Expressions.Num() ; ++i)
+	{
+		TNiagaraExprPtr Expr = Expressions[i];
+		if (Expr.IsValid())
+		{
+			Expr->Process();
+			Expr->PostProcess();
+		}
+		else
+		{
+			check(0);//Replace with graceful bailout?
+		}
+
+		Expressions[i].Reset();
+	}
+
+	PinToExpression.Empty();//TODO - Replace this with expression caching at a lower level. Will be more useful and nicer.
+
+	// Terminate with the 'done' opcode.
+	Script->ByteCode.Add((int8)VectorVM::EOp::done);
+
+	Script->ConstantData = ConstantData;
+	Script->Attributes = Attributes;
+	//To do - scalar and matrix attributes.
+}
+
+TNiagaraExprPtr FNiagaraCompiler_VectorVM::Expression_VMNative(VectorVM::EOp Op, TArray<TNiagaraExprPtr>& InputExpressions)
+{
+	int32 Index = Expressions.Add(MakeShareable(new FNiagaraExpression_VMOperation(this, ENiagaraDataType::Vector, InputExpressions, Op)));
+	return Expressions[Index];
+}
+
+TNiagaraExprPtr FNiagaraCompiler_VectorVM::Expression_VMNative(VectorVM::EOp Op, TNiagaraExprPtr A)
+{
+	TArray<TNiagaraExprPtr> In; In.Add(A);
+	return Expression_VMNative(Op, In);
+}
+
+TNiagaraExprPtr FNiagaraCompiler_VectorVM::Expression_VMNative(VectorVM::EOp Op, TNiagaraExprPtr A, TNiagaraExprPtr B)
+{
+	TArray<TNiagaraExprPtr> In; In.Add(A); In.Add(B);
+	return Expression_VMNative(Op, In);
+}
+
+TNiagaraExprPtr FNiagaraCompiler_VectorVM::Expression_VMNative(VectorVM::EOp Op, TNiagaraExprPtr A, TNiagaraExprPtr B, TNiagaraExprPtr C)
+{
+	TArray<TNiagaraExprPtr> In; In.Add(A); In.Add(B); In.Add(C);
+	return Expression_VMNative(Op, In);
+}
+
+TNiagaraExprPtr FNiagaraCompiler_VectorVM::Expression_VMNative(VectorVM::EOp Op, TNiagaraExprPtr A, TNiagaraExprPtr B, TNiagaraExprPtr C, TNiagaraExprPtr D)
+{
+	TArray<TNiagaraExprPtr> In; In.Add(A); In.Add(B); In.Add(C); In.Add(D);
+	return Expression_VMNative(Op, In);
+}
+
+void FNiagaraCompiler_VectorVM::Add_Internal(TArray<TNiagaraExprPtr>& InputExpressions, TArray<TNiagaraExprPtr>& OutputExpressions)
+{
+	OutputExpressions.Add(Expression_VMNative(VectorVM::EOp::add, InputExpressions));
+}
+
+void FNiagaraCompiler_VectorVM::Subtract_Internal(TArray<TNiagaraExprPtr>& InputExpressions, TArray<TNiagaraExprPtr>& OutputExpressions)
+{
+	OutputExpressions.Add((Expression_VMNative(VectorVM::EOp::sub, InputExpressions)));
+}
+
+void FNiagaraCompiler_VectorVM::Multiply_Internal(TArray<TNiagaraExprPtr>& InputExpressions, TArray<TNiagaraExprPtr>& OutputExpressions)
+{
+	OutputExpressions.Add(Expression_VMNative(VectorVM::EOp::mul, InputExpressions));
+}
+
+void FNiagaraCompiler_VectorVM::MultiplyAdd_Internal(TArray<TNiagaraExprPtr>& InputExpressions, TArray<TNiagaraExprPtr>& OutputExpressions)
+{
+	OutputExpressions.Add(Expression_VMNative(VectorVM::EOp::mad, InputExpressions));
+}
+
+void FNiagaraCompiler_VectorVM::Lerp_Internal(TArray<TNiagaraExprPtr>& InputExpressions, TArray<TNiagaraExprPtr>& OutputExpressions)
+{
+	OutputExpressions.Add(Expression_VMNative(VectorVM::EOp::lerp, InputExpressions));
+}
+
+void FNiagaraCompiler_VectorVM::Reciprocal_Internal(TArray<TNiagaraExprPtr>& InputExpressions, TArray<TNiagaraExprPtr>& OutputExpressions)
+{
+	OutputExpressions.Add(Expression_VMNative(VectorVM::EOp::rcp, InputExpressions));
+}
+
+void FNiagaraCompiler_VectorVM::ReciprocalSqrt_Internal(TArray<TNiagaraExprPtr>& InputExpressions, TArray<TNiagaraExprPtr>& OutputExpressions)
+{
+	OutputExpressions.Add(Expression_VMNative(VectorVM::EOp::rsq, InputExpressions));
+}
+
+void FNiagaraCompiler_VectorVM::Sqrt_Internal(TArray<TNiagaraExprPtr>& InputExpressions, TArray<TNiagaraExprPtr>& OutputExpressions)
+{
+	OutputExpressions.Add(Expression_VMNative(VectorVM::EOp::sqrt, InputExpressions));
+}
+
+void FNiagaraCompiler_VectorVM::Negate_Internal(TArray<TNiagaraExprPtr>& InputExpressions, TArray<TNiagaraExprPtr>& OutputExpressions)
+{
+	OutputExpressions.Add(Expression_VMNative(VectorVM::EOp::neg, InputExpressions));
+}
+
+void FNiagaraCompiler_VectorVM::Abs_Internal(TArray<TNiagaraExprPtr>& InputExpressions, TArray<TNiagaraExprPtr>& OutputExpressions)
+{
+	OutputExpressions.Add(Expression_VMNative(VectorVM::EOp::abs, InputExpressions));
+}
+
+void FNiagaraCompiler_VectorVM::Exp_Internal(TArray<TNiagaraExprPtr>& InputExpressions, TArray<TNiagaraExprPtr>& OutputExpressions)
+{
+	OutputExpressions.Add(Expression_VMNative(VectorVM::EOp::exp, InputExpressions));
+}
+
+void FNiagaraCompiler_VectorVM::Exp2_Internal(TArray<TNiagaraExprPtr>& InputExpressions, TArray<TNiagaraExprPtr>& OutputExpressions)
+{
+	OutputExpressions.Add(Expression_VMNative(VectorVM::EOp::exp2, InputExpressions));
+}
+
+void FNiagaraCompiler_VectorVM::Log_Internal(TArray<TNiagaraExprPtr>& InputExpressions, TArray<TNiagaraExprPtr>& OutputExpressions)
+{
+	OutputExpressions.Add(Expression_VMNative(VectorVM::EOp::log, InputExpressions));
+}
+
+void FNiagaraCompiler_VectorVM::LogBase2_Internal(TArray<TNiagaraExprPtr>& InputExpressions, TArray<TNiagaraExprPtr>& OutputExpressions)
+{
+	OutputExpressions.Add(Expression_VMNative(VectorVM::EOp::log2, InputExpressions));
+}
+
+void FNiagaraCompiler_VectorVM::Sin_Internal(TArray<TNiagaraExprPtr>& InputExpressions, TArray<TNiagaraExprPtr>& OutputExpressions)
+{
+	OutputExpressions.Add(Expression_VMNative(VectorVM::EOp::sin, InputExpressions));
+}
+
+void FNiagaraCompiler_VectorVM::Cos_Internal(TArray<TNiagaraExprPtr>& InputExpressions, TArray<TNiagaraExprPtr>& OutputExpressions)
+{
+	OutputExpressions.Add(Expression_VMNative(VectorVM::EOp::cos, InputExpressions));
+}
+
+void FNiagaraCompiler_VectorVM::Tan_Internal(TArray<TNiagaraExprPtr>& InputExpressions, TArray<TNiagaraExprPtr>& OutputExpressions)
+{
+	OutputExpressions.Add(Expression_VMNative(VectorVM::EOp::tan, InputExpressions));
+}
+
+void FNiagaraCompiler_VectorVM::ASin_Internal(TArray<TNiagaraExprPtr>& InputExpressions, TArray<TNiagaraExprPtr>& OutputExpressions)
+{
+	OutputExpressions.Add(Expression_VMNative(VectorVM::EOp::asin, InputExpressions));
+}
+
+void FNiagaraCompiler_VectorVM::ACos_Internal(TArray<TNiagaraExprPtr>& InputExpressions, TArray<TNiagaraExprPtr>& OutputExpressions)
+{
+	OutputExpressions.Add(Expression_VMNative(VectorVM::EOp::acos, InputExpressions));
+}
+
+void FNiagaraCompiler_VectorVM::ATan_Internal(TArray<TNiagaraExprPtr>& InputExpressions, TArray<TNiagaraExprPtr>& OutputExpressions)
+{
+	OutputExpressions.Add(Expression_VMNative(VectorVM::EOp::atan, InputExpressions));
+}
+
+void FNiagaraCompiler_VectorVM::ATan2_Internal(TArray<TNiagaraExprPtr>& InputExpressions, TArray<TNiagaraExprPtr>& OutputExpressions)
+{
+	OutputExpressions.Add(Expression_VMNative(VectorVM::EOp::atan2, InputExpressions));
+}
+
+void FNiagaraCompiler_VectorVM::Ceil_Internal(TArray<TNiagaraExprPtr>& InputExpressions, TArray<TNiagaraExprPtr>& OutputExpressions)
+{
+	OutputExpressions.Add(Expression_VMNative(VectorVM::EOp::ceil, InputExpressions));
+}
+
+void FNiagaraCompiler_VectorVM::Floor_Internal(TArray<TNiagaraExprPtr>& InputExpressions, TArray<TNiagaraExprPtr>& OutputExpressions)
+{
+	OutputExpressions.Add(Expression_VMNative(VectorVM::EOp::floor, InputExpressions));
+}
+
+void FNiagaraCompiler_VectorVM::Modulo_Internal(TArray<TNiagaraExprPtr>& InputExpressions, TArray<TNiagaraExprPtr>& OutputExpressions)
+{
+	OutputExpressions.Add(Expression_VMNative(VectorVM::EOp::fmod, InputExpressions));
+}
+
+void FNiagaraCompiler_VectorVM::Fractional_Internal(TArray<TNiagaraExprPtr>& InputExpressions, TArray<TNiagaraExprPtr>& OutputExpressions)
+{
+	OutputExpressions.Add(Expression_VMNative(VectorVM::EOp::frac, InputExpressions));
+}
+
+void FNiagaraCompiler_VectorVM::Truncate_Internal(TArray<TNiagaraExprPtr>& InputExpressions, TArray<TNiagaraExprPtr>& OutputExpressions)
+{
+	OutputExpressions.Add(Expression_VMNative(VectorVM::EOp::trunc, InputExpressions));
+}
+
+void FNiagaraCompiler_VectorVM::Clamp_Internal(TArray<TNiagaraExprPtr>& InputExpressions, TArray<TNiagaraExprPtr>& OutputExpressions)
+{
+	OutputExpressions.Add(Expression_VMNative(VectorVM::EOp::clamp, InputExpressions));
+}
+
+void FNiagaraCompiler_VectorVM::Min_Internal(TArray<TNiagaraExprPtr>& InputExpressions, TArray<TNiagaraExprPtr>& OutputExpressions)
+{
+	OutputExpressions.Add(Expression_VMNative(VectorVM::EOp::min, InputExpressions));
+}
+
+void FNiagaraCompiler_VectorVM::Max_Internal(TArray<TNiagaraExprPtr>& InputExpressions, TArray<TNiagaraExprPtr>& OutputExpressions)
+{
+	OutputExpressions.Add(Expression_VMNative(VectorVM::EOp::max, InputExpressions));
+}
+
+void FNiagaraCompiler_VectorVM::Pow_Internal(TArray<TNiagaraExprPtr>& InputExpressions, TArray<TNiagaraExprPtr>& OutputExpressions)
+{
+	OutputExpressions.Add(Expression_VMNative(VectorVM::EOp::pow, InputExpressions));
+}
+
+void FNiagaraCompiler_VectorVM::Sign_Internal(TArray<TNiagaraExprPtr>& InputExpressions, TArray<TNiagaraExprPtr>& OutputExpressions)
+{
+	OutputExpressions.Add(Expression_VMNative(VectorVM::EOp::sign, InputExpressions));
+}
+
+void FNiagaraCompiler_VectorVM::Step_Internal(TArray<TNiagaraExprPtr>& InputExpressions, TArray<TNiagaraExprPtr>& OutputExpressions)
+{
+	OutputExpressions.Add(Expression_VMNative(VectorVM::EOp::step, InputExpressions));
+}
+
+void FNiagaraCompiler_VectorVM::Dot_Internal(TArray<TNiagaraExprPtr>& InputExpressions, TArray<TNiagaraExprPtr>& OutputExpressions)
+{
+	OutputExpressions.Add(Expression_VMNative(VectorVM::EOp::dot, InputExpressions));
+}
+
+void FNiagaraCompiler_VectorVM::Cross_Internal(TArray<TNiagaraExprPtr>& InputExpressions, TArray<TNiagaraExprPtr>& OutputExpressions)
+{
+	OutputExpressions.Add(Expression_VMNative(VectorVM::EOp::cross, InputExpressions));
+}
+
+void FNiagaraCompiler_VectorVM::Normalize_Internal(TArray<TNiagaraExprPtr>& InputExpressions, TArray<TNiagaraExprPtr>& OutputExpressions)
+{
+	OutputExpressions.Add(Expression_VMNative(VectorVM::EOp::normalize, InputExpressions));
+}
+
+void FNiagaraCompiler_VectorVM::Random_Internal(TArray<TNiagaraExprPtr>& InputExpressions, TArray<TNiagaraExprPtr>& OutputExpressions)
+{
+	OutputExpressions.Add(Expression_VMNative(VectorVM::EOp::random, InputExpressions));
+}
+
+void FNiagaraCompiler_VectorVM::Length_Internal(TArray<TNiagaraExprPtr>& InputExpressions, TArray<TNiagaraExprPtr>& OutputExpressions)
+{
+	OutputExpressions.Add(Expression_VMNative(VectorVM::EOp::length, InputExpressions));
+}
+
+void FNiagaraCompiler_VectorVM::Sin4_Internal(TArray<TNiagaraExprPtr>& InputExpressions, TArray<TNiagaraExprPtr>& OutputExpressions)
+{
+	OutputExpressions.Add(Expression_VMNative(VectorVM::EOp::sin4, InputExpressions));
+}
+
+void FNiagaraCompiler_VectorVM::Noise_Internal(TArray<TNiagaraExprPtr>& InputExpressions, TArray<TNiagaraExprPtr>& OutputExpressions)
+{
+	OutputExpressions.Add(Expression_VMNative(VectorVM::EOp::noise, InputExpressions));
+}
+
+void FNiagaraCompiler_VectorVM::Scalar_Internal(TArray<TNiagaraExprPtr>& InputExpressions, TArray<TNiagaraExprPtr>& OutputExpressions)
+{
+	GetX_Internal(InputExpressions,OutputExpressions);
+}
+
+void FNiagaraCompiler_VectorVM::Vector_Internal(TArray<TNiagaraExprPtr>& InputExpressions, TArray<TNiagaraExprPtr>& OutputExpressions)
+{
+	OutputExpressions.Add(InputExpressions[0]);
+}
+
+void FNiagaraCompiler_VectorVM::Matrix_Internal(TArray<TNiagaraExprPtr>& InputExpressions, TArray<TNiagaraExprPtr>& OutputExpressions)
+{
+	//Just collect the matrix rows into an expression for what ever use subsequent expressions have.
+	TNiagaraExprPtr MatrixExpr = Expression_Collection(InputExpressions);
+	MatrixExpr->ResultType = ENiagaraDataType::Matrix;
+	OutputExpressions.Add(MatrixExpr);
+}
+
+void FNiagaraCompiler_VectorVM::Compose_Internal(TArray<TNiagaraExprPtr>& InputExpressions, TArray<TNiagaraExprPtr>& OutputExpressions)
+{
+	OutputExpressions.Add(Expression_VMNative(VectorVM::EOp::compose, InputExpressions));
+}
+
+void FNiagaraCompiler_VectorVM::ComposeX_Internal(TArray<TNiagaraExprPtr>& InputExpressions, TArray<TNiagaraExprPtr>& OutputExpressions)
+{
+	OutputExpressions.Add(Expression_VMNative(VectorVM::EOp::composex, InputExpressions));
+}
+
+void FNiagaraCompiler_VectorVM::ComposeY_Internal(TArray<TNiagaraExprPtr>& InputExpressions, TArray<TNiagaraExprPtr>& OutputExpressions)
+{
+	OutputExpressions.Add(Expression_VMNative(VectorVM::EOp::composey, InputExpressions));
+}
+
+void FNiagaraCompiler_VectorVM::ComposeZ_Internal(TArray<TNiagaraExprPtr>& InputExpressions, TArray<TNiagaraExprPtr>& OutputExpressions)
+{
+	OutputExpressions.Add(Expression_VMNative(VectorVM::EOp::composez, InputExpressions));
+}
+
+void FNiagaraCompiler_VectorVM::ComposeW_Internal(TArray<TNiagaraExprPtr>& InputExpressions, TArray<TNiagaraExprPtr>& OutputExpressions)
+{
+	OutputExpressions.Add(Expression_VMNative(VectorVM::EOp::composew, InputExpressions));
+}
+
+void FNiagaraCompiler_VectorVM::Split_Internal(TArray<TNiagaraExprPtr>& InputExpressions, TArray<TNiagaraExprPtr>& OutputExpressions)
+{
+	GetX_Internal(InputExpressions, OutputExpressions);
+	GetY_Internal(InputExpressions, OutputExpressions);
+	GetZ_Internal(InputExpressions, OutputExpressions);
+	GetW_Internal(InputExpressions, OutputExpressions);
+}
+
+void FNiagaraCompiler_VectorVM::GetX_Internal(TArray<TNiagaraExprPtr>& InputExpressions, TArray<TNiagaraExprPtr>& OutputExpressions)
+{
+	OutputExpressions.Add(Expression_VMNative(VectorVM::EOp::splatx, InputExpressions));
+}
+
+void FNiagaraCompiler_VectorVM::GetY_Internal(TArray<TNiagaraExprPtr>& InputExpressions, TArray<TNiagaraExprPtr>& OutputExpressions)
+{
+	OutputExpressions.Add(Expression_VMNative(VectorVM::EOp::splaty, InputExpressions));
+}
+
+void FNiagaraCompiler_VectorVM::GetZ_Internal(TArray<TNiagaraExprPtr>& InputExpressions, TArray<TNiagaraExprPtr>& OutputExpressions)
+{
+	OutputExpressions.Add(Expression_VMNative(VectorVM::EOp::splatz, InputExpressions));
+}
+
+void FNiagaraCompiler_VectorVM::GetW_Internal(TArray<TNiagaraExprPtr>& InputExpressions, TArray<TNiagaraExprPtr>& OutputExpressions)
+{
+	OutputExpressions.Add(Expression_VMNative(VectorVM::EOp::splatw, InputExpressions));
+}
+
+void FNiagaraCompiler_VectorVM::TransformVector_Internal(TArray<TNiagaraExprPtr>& InputExpressions, TArray<TNiagaraExprPtr>& OutputExpressions)
+{
+	TNiagaraExprPtr MatrixExpr = InputExpressions[0];
+	check(MatrixExpr->SourceExpressions.Num() == 4);
+	TNiagaraExprPtr Row0 = MatrixExpr->GetSourceExpression(0);
+	TNiagaraExprPtr Row1 = MatrixExpr->GetSourceExpression(1);
+	TNiagaraExprPtr Row2 = MatrixExpr->GetSourceExpression(2);
+	TNiagaraExprPtr Row3 = MatrixExpr->GetSourceExpression(3);
+
+	TNiagaraExprPtr VectorExpr = InputExpressions[1];
+
+	TArray<TNiagaraExprPtr> Temp;
+	Temp.Add(Expression_VMNative(VectorVM::EOp::splatx, VectorExpr));
+	Temp.Add(Expression_VMNative(VectorVM::EOp::splaty, VectorExpr));
+	Temp.Add(Expression_VMNative(VectorVM::EOp::splatz, VectorExpr));
+	Temp.Add(Expression_VMNative(VectorVM::EOp::splatw, VectorExpr));
+
+	Temp[0] = Expression_VMNative(VectorVM::EOp::mul, MatrixExpr->GetSourceExpression(0), Temp[0]);
+	Temp[1] = Expression_VMNative(VectorVM::EOp::mul, MatrixExpr->GetSourceExpression(1), Temp[1]);
+	Temp[2] = Expression_VMNative(VectorVM::EOp::mul, MatrixExpr->GetSourceExpression(2), Temp[2]);
+	Temp[3] = Expression_VMNative(VectorVM::EOp::mul, MatrixExpr->GetSourceExpression(3), Temp[3]);
+
+	Temp[0] = Expression_VMNative(VectorVM::EOp::add, Temp[0], Temp[1]);
+	Temp[1] = Expression_VMNative(VectorVM::EOp::add, Temp[2], Temp[3]);
+	OutputExpressions.Add(Expression_VMNative(VectorVM::EOp::add, Temp[0], Temp[1]));
+}
+
+void FNiagaraCompiler_VectorVM::Transpose_Internal(TArray<TNiagaraExprPtr>& InputExpressions, TArray<TNiagaraExprPtr>& OutputExpressions)
+{
+	TNiagaraExprPtr MatrixExpr = InputExpressions[0];
+	check(MatrixExpr->SourceExpressions.Num() == 4);
+	TNiagaraExprPtr Row0 = MatrixExpr->GetSourceExpression(0);
+	TNiagaraExprPtr Row1 = MatrixExpr->GetSourceExpression(1);
+	TNiagaraExprPtr Row2 = MatrixExpr->GetSourceExpression(2);
+	TNiagaraExprPtr Row3 = MatrixExpr->GetSourceExpression(3);
+
+	//TODO: can probably make this faster if I provide direct shuffle ops in the VM rather than using the flexible, catch-all op compose.
+	//Maybe don't want to make that externally usable though? Maybe for power users...
+	TArray<TNiagaraExprPtr> Columns;
+	Columns.Add(Expression_VMNative(VectorVM::EOp::composex, Row0, Row1, Row2, Row3));
+	Columns.Add(Expression_VMNative(VectorVM::EOp::composey, Row0, Row1, Row2, Row3));
+	Columns.Add(Expression_VMNative(VectorVM::EOp::composez, Row0, Row1, Row2, Row3));
+	Columns.Add(Expression_VMNative(VectorVM::EOp::composew, Row0, Row1, Row2, Row3));
+
+	Matrix_Internal(Columns, OutputExpressions);
+}
+
+
+void FNiagaraCompiler_VectorVM::Inverse_Internal(TArray<TNiagaraExprPtr>& InputExpressions, TArray<TNiagaraExprPtr>& OutputExpressions)
+{
+	//OutputExpressions.Add(Expression_UnknownError());
+
+	
+}
diff --git a/Engine/Source/Editor/NiagaraEditor/Private/NiagaraCompiler_VectorVM.h b/Engine/Source/Editor/NiagaraEditor/Private/NiagaraCompiler_VectorVM.h
new file mode 100644
--- /dev/null
+++ b/Engine/Source/Editor/NiagaraEditor/Private/NiagaraCompiler_VectorVM.h
@@ -0,0 +1,48 @@
+// Copyright 1998-2014 Epic Games, Inc. All Rights Reserved.
+
+#pragma once
+
+#include "NiagaraCompiler.h"
+#include "VectorVM.h"
+
+/** A Niagara compiler that produces byte code to run on the VectorVM. */
+class NIAGARAEDITOR_API FNiagaraCompiler_VectorVM : public FNiagaraCompiler
+{
+private:
+
+	/** Creates an expression directly mapping to a VM operation. */
+	TNiagaraExprPtr Expression_VMNative(VectorVM::EOp Op, TArray<TNiagaraExprPtr>& InputExpressions);
+	TNiagaraExprPtr Expression_VMNative(VectorVM::EOp Op, TNiagaraExprPtr A);
+	TNiagaraExprPtr Expression_VMNative(VectorVM::EOp Op, TNiagaraExprPtr A, TNiagaraExprPtr B);
+	TNiagaraExprPtr Expression_VMNative(VectorVM::EOp Op, TNiagaraExprPtr A, TNiagaraExprPtr B, TNiagaraExprPtr C);
+	TNiagaraExprPtr Expression_VMNative(VectorVM::EOp Op, TNiagaraExprPtr A, TNiagaraExprPtr B, TNiagaraExprPtr C, TNiagaraExprPtr D);
+
+	/** Array tracking which temp registers are free allowing their re-use and thus a smaller total number. */
+	TArray<bool> TempRegisters;
+public:
+
+	/** Returns true if the expression is stored in the constant data. */
+	bool ExpressionIsConstant(FNiagaraExpression*  Expression);
+	
+	/** Gets the final index used in the VM bytecode for the passed expression. */
+	uint8 GetResultVMIndex(FNiagaraExpression*  Expression);
+
+	/** Writes to the scripts bytecode. */
+	void WriteCode(uint8 InCode);
+
+	//Begin INiagaraCompiler Interface	
+	virtual void CompileScript(UNiagaraScript* InScript) override;
+
+#define NiagaraOp(OpName) virtual void OpName##_Internal(TArray<TNiagaraExprPtr>& InputExpressions, TArray<TNiagaraExprPtr>& OutputExpressions)override;
+		NiagaraOpList;
+#undef NiagaraOp
+
+	virtual TNiagaraExprPtr Output(FName OutputName, TNiagaraExprPtr& SourceExpression)override;
+	//End INiagaraCompiler;
+
+	//Begin FNiagaraCompiler Interface
+	virtual void GetConstantResultIndex(FName Name, bool bInternal, int32& OutResultIndex, int32& OutComponentIndex)override;
+	virtual int32 AquireTemporary()override;
+	virtual void FreeTemporary(int32 TempIndex) override;
+	//End INiagaraCompiler Interface
+};
diff --git a/Engine/Source/Editor/NiagaraEditor/Private/NiagaraEditorPrivatePCH.h b/Engine/Source/Editor/NiagaraEditor/Private/NiagaraEditorPrivatePCH.h
--- a/Engine/Source/Editor/NiagaraEditor/Private/NiagaraEditorPrivatePCH.h
+++ b/Engine/Source/Editor/NiagaraEditor/Private/NiagaraEditorPrivatePCH.h
@@ -6,6 +6,7 @@
 
 #include "NiagaraEditorModule.h"
 #include "NiagaraEditor.h"
+#include "NiagaraCompiler.h"
 #include "NiagaraEffectEditor.h"
 
 #include "GraphEditor.h"
diff --git a/Engine/Source/Editor/UnrealEd/Classes/Niagara/EdGraphSchema_Niagara.h b/Engine/Source/Editor/UnrealEd/Classes/Niagara/EdGraphSchema_Niagara.h
--- a/Engine/Source/Editor/UnrealEd/Classes/Niagara/EdGraphSchema_Niagara.h
+++ b/Engine/Source/Editor/UnrealEd/Classes/Niagara/EdGraphSchema_Niagara.h
@@ -2,6 +2,7 @@
 
 #pragma once
 #include "EdGraph/EdGraphSchema.h"
+#include "INiagaraCompiler.h"
 #include "EdGraphSchema_Niagara.generated.h"
 
 /** Action to add a node to the graph */
@@ -50,7 +51,9 @@ class UEdGraphSchema_Niagara : public UEdGraphSchema
 	UPROPERTY()
 	FString PC_Float;
 	UPROPERTY()
-	FString PC_Struct;
+	FString PC_Vector;
+	UPROPERTY()
+	FString PC_Matrix;
 
 	// Begin EdGraphSchema interface
 	virtual void GetGraphContextActions(FGraphContextMenuBuilder& ContextMenuBuilder) const override;
@@ -58,5 +61,10 @@ class UEdGraphSchema_Niagara : public UEdGraphSchema
 	virtual FLinearColor GetPinTypeColor(const FEdGraphPinType& PinType) const override;
 	virtual bool ShouldHidePinDefaultValue(UEdGraphPin* Pin) const override;
 	// End EdGraphSchema interface
+
+	UNREALED_API ENiagaraDataType GetPinDataType(UEdGraphPin* Pin)const;
+	UNREALED_API void GetPinDefaultValue(UEdGraphPin* Pin, float& OutDefault)const;
+	UNREALED_API void GetPinDefaultValue(UEdGraphPin* Pin, FVector4& OutDefault)const;
+	UNREALED_API void GetPinDefaultValue(UEdGraphPin* Pin, FMatrix& OutDefault)const;
 };
 
diff --git a/Engine/Source/Editor/UnrealEd/Classes/Niagara/NiagaraNode.h b/Engine/Source/Editor/UnrealEd/Classes/Niagara/NiagaraNode.h
--- a/Engine/Source/Editor/UnrealEd/Classes/Niagara/NiagaraNode.h
+++ b/Engine/Source/Editor/UnrealEd/Classes/Niagara/NiagaraNode.h
@@ -3,15 +3,18 @@
 #pragma once
 #include "NiagaraNode.generated.h"
 
-UCLASS(MinimalAPI)
-class UNiagaraNode : public UEdGraphNode
+UCLASS()
+class UNREALED_API UNiagaraNode : public UEdGraphNode
 {
 	GENERATED_UCLASS_BODY()
 
+public:
+
 	/** Get the Niagara graph that owns this node */
 	class UNiagaraGraph* GetNiagaraGraph();
 
 	/** Get the source object */
 	class UNiagaraScriptSource* GetSource();
-};
 
+	virtual void Compile(class INiagaraCompiler* Compiler, TArray<FNiagaraNodeResult>& Outputs){ check(0); }
+};
diff --git a/Engine/Source/Editor/UnrealEd/Classes/Niagara/NiagaraNodeConstant.h b/Engine/Source/Editor/UnrealEd/Classes/Niagara/NiagaraNodeConstant.h
new file mode 100644
--- /dev/null
+++ b/Engine/Source/Editor/UnrealEd/Classes/Niagara/NiagaraNodeConstant.h
@@ -0,0 +1,39 @@
+// Copyright 1998-2014 Epic Games, Inc. All Rights Reserved.
+
+#pragma once
+
+#include "Niagara/NiagaraCommon.h"
+#include "NiagaraNodeConstant.generated.h"
+
+UCLASS(MinimalAPI)
+class UNiagaraNodeConstant : public UNiagaraNode
+{
+	GENERATED_UCLASS_BODY()
+
+public:
+	/** The type of the constant we're creating. */
+	UPROPERTY(EditAnywhere, meta = (DisplayName = "Name"), Category="Constant")
+	FName ConstName;
+
+	/** The type of the constant we're creating. */
+	UPROPERTY(EditAnywhere, Category = "Constant")
+	ENiagaraDataType DataType;
+	
+	UPROPERTY(EditAnywhere, Category = "Constant")
+	uint32 bNeedsDefault : 1;
+
+	// Begin UObject interface
+	virtual void PostEditChangeProperty(struct FPropertyChangedEvent& PropertyChangedEvent);
+	// End UObject interface
+	
+	// Begin EdGraphNode interface
+	virtual void AllocateDefaultPins() override;
+	virtual FText GetNodeTitle(ENodeTitleType::Type TitleType) const override;
+	virtual FLinearColor GetNodeTitleColor() const override;
+	// End EdGraphNode interface
+
+	// Begin UNiagaraNode interface
+	UNREALED_API virtual void Compile(class INiagaraCompiler* Compiler, TArray<FNiagaraNodeResult>& Outputs)override;
+	// End UNiagaraNode interface
+};
+
diff --git a/Engine/Source/Editor/UnrealEd/Classes/Niagara/NiagaraNodeGetAttr.h b/Engine/Source/Editor/UnrealEd/Classes/Niagara/NiagaraNodeGetAttr.h
--- a/Engine/Source/Editor/UnrealEd/Classes/Niagara/NiagaraNodeGetAttr.h
+++ b/Engine/Source/Editor/UnrealEd/Classes/Niagara/NiagaraNodeGetAttr.h
@@ -8,6 +8,8 @@ class UNiagaraNodeGetAttr : public UNiagaraNode
 {
 	GENERATED_UCLASS_BODY()
 
+public:
+
 	/** Name of attribute we are getting */
 	UPROPERTY()
 	FName	AttrName;
@@ -17,5 +19,9 @@ class UNiagaraNodeGetAttr : public UNiagaraNode
 	virtual FText GetNodeTitle(ENodeTitleType::Type TitleType) const override;
 	virtual FLinearColor GetNodeTitleColor() const override;
 	// End EdGraphNode interface
+
+	// Begin UNiagaraNode interface
+	UNREALED_API virtual void Compile(class INiagaraCompiler* Compiler, TArray<FNiagaraNodeResult>& Output) override;
+	// End UNiagaraNode interface
 };
 
diff --git a/Engine/Source/Editor/UnrealEd/Classes/Niagara/NiagaraNodeOp.h b/Engine/Source/Editor/UnrealEd/Classes/Niagara/NiagaraNodeOp.h
--- a/Engine/Source/Editor/UnrealEd/Classes/Niagara/NiagaraNodeOp.h
+++ b/Engine/Source/Editor/UnrealEd/Classes/Niagara/NiagaraNodeOp.h
@@ -8,18 +8,22 @@ class UNiagaraNodeOp : public UNiagaraNode
 {
 	GENERATED_UCLASS_BODY()
 
-	/** Index of operation */
+public:
+
+	/** Name of operation */
 	UPROPERTY()
-	uint8 OpIndex;
+	FName OpName;
 
 	// Begin EdGraphNode interface
 	virtual void AllocateDefaultPins() override;
 	virtual FText GetNodeTitle(ENodeTitleType::Type TitleType) const override;
+	virtual FText GetTooltipText() const override;
 	virtual FLinearColor GetNodeTitleColor() const override;
 	// End EdGraphNode interface
 
-	/** Names to use for each input pin. */
-	UNREALED_API static const TCHAR* InPinNames[3];
+	// Begin UNiagaraNode interface
+	UNREALED_API virtual void Compile(class INiagaraCompiler* Compiler, TArray<FNiagaraNodeResult>& Outputs) override;
+	// End UNiagaraNode interface
 };
 
 
diff --git a/Engine/Source/Editor/UnrealEd/Classes/Niagara/NiagaraNodeOutputUpdate.h b/Engine/Source/Editor/UnrealEd/Classes/Niagara/NiagaraNodeOutputUpdate.h
--- a/Engine/Source/Editor/UnrealEd/Classes/Niagara/NiagaraNodeOutputUpdate.h
+++ b/Engine/Source/Editor/UnrealEd/Classes/Niagara/NiagaraNodeOutputUpdate.h
@@ -8,11 +8,15 @@ class UNiagaraNodeOutputUpdate : public UNiagaraNode
 {
 	GENERATED_UCLASS_BODY()
 
+public:
+
 	// Begin EdGraphNode interface
 	virtual void AllocateDefaultPins() override;
 	virtual bool CanUserDeleteNode() const override { return false; }
 	virtual FText GetNodeTitle(ENodeTitleType::Type TitleType) const override;
 	virtual FLinearColor GetNodeTitleColor() const override;
 	// End EdGraphNode interface
+
+	void Compile(class INiagaraCompiler* Compiler, TArray<FNiagaraNodeResult>& Outputs)override;
 };
 
diff --git a/Engine/Source/Editor/UnrealEd/Classes/Niagara/NiagaraScriptSource.h b/Engine/Source/Editor/UnrealEd/Classes/Niagara/NiagaraScriptSource.h
--- a/Engine/Source/Editor/UnrealEd/Classes/Niagara/NiagaraScriptSource.h
+++ b/Engine/Source/Editor/UnrealEd/Classes/Niagara/NiagaraScriptSource.h
@@ -18,5 +18,5 @@ class UNiagaraScriptSource : public UNiagaraScriptSourceBase
 	virtual void PostLoad() override;
 
 	UNREALED_API void GetParticleAttributes(TArray<FName>& VectorOutputs);
-	UNREALED_API void GetEmitterAttributes(TArray<FName>& ScalarInputs);
+	UNREALED_API void GetEmitterAttributes(TArray<FName>& VectorInputs, TArray<FName>& MatrixInputs);
 };
diff --git a/Engine/Source/Editor/UnrealEd/Private/Factories/NiagaraEffectFactory.cpp b/Engine/Source/Editor/UnrealEd/Private/Factories/NiagaraEffectFactory.cpp
--- a/Engine/Source/Editor/UnrealEd/Private/Factories/NiagaraEffectFactory.cpp
+++ b/Engine/Source/Editor/UnrealEd/Private/Factories/NiagaraEffectFactory.cpp
@@ -3,7 +3,6 @@
 #include "UnrealEd.h"
 #include "Engine/NiagaraEffect.h"
 #include "NiagaraEditorModule.h"
-#include "VectorVM.h"
 
 #define LOCTEXT_NAMESPACE "NiagaraEffectFactory"
 
diff --git a/Engine/Source/Editor/UnrealEd/Private/Factories/NiagaraScriptFactory.cpp b/Engine/Source/Editor/UnrealEd/Private/Factories/NiagaraScriptFactory.cpp
--- a/Engine/Source/Editor/UnrealEd/Private/Factories/NiagaraScriptFactory.cpp
+++ b/Engine/Source/Editor/UnrealEd/Private/Factories/NiagaraScriptFactory.cpp
@@ -3,7 +3,6 @@
 #include "UnrealEd.h"
 #include "Engine/NiagaraScript.h"
 #include "NiagaraEditorModule.h"
-#include "VectorVM.h"
 
 #define LOCTEXT_NAMESPACE "NiagaraScriptFactory"
 
@@ -41,12 +40,12 @@ UObject* UNiagaraScriptFactoryNew::FactoryCreateNew(UClass* Class,UObject* InPar
 			UNiagaraNodeOutputUpdate* OutputNode = OutputNodeCreator.CreateNode();
 			OutputNodeCreator.Finalize();
 
-			FGraphNodeCreator<UNiagaraNodeGetAttr> TimeNodeCreator(*CreatedGraph);
-			UNiagaraNodeGetAttr* DeltaTimeNode = TimeNodeCreator.CreateNode();
-			DeltaTimeNode->AttrName = FName(TEXT("DeltaTime"));
-			TimeNodeCreator.Finalize();
-
-			UEdGraphPin* DeltaTimePin = DeltaTimeNode->FindPinChecked(DeltaTimeNode->AttrName.ToString());
+// 			FGraphNodeCreator<UNiagaraNodeGetAttr> TimeNodeCreator(*CreatedGraph);
+// 			UNiagaraNodeGetAttr* DeltaTimeNode = TimeNodeCreator.CreateNode();
+// 			DeltaTimeNode->AttrName = FName(TEXT("DeltaTime"));
+// 			DeltaTimeNode->bIsEmitterAttribute = true;
+// 			TimeNodeCreator.Finalize();
+			//UEdGraphPin* DeltaTimePin = DeltaTimeNode->FindPinChecked(DeltaTimeNode->AttrName.ToString());
 
 			TArray<FName> OutputNames;
 			Source->GetParticleAttributes(OutputNames);
diff --git a/Engine/Source/Editor/UnrealEd/Private/Niagara/EdGraphSchema_Niagara.cpp b/Engine/Source/Editor/UnrealEd/Private/Niagara/EdGraphSchema_Niagara.cpp
--- a/Engine/Source/Editor/UnrealEd/Private/Niagara/EdGraphSchema_Niagara.cpp
+++ b/Engine/Source/Editor/UnrealEd/Private/Niagara/EdGraphSchema_Niagara.cpp
@@ -5,7 +5,11 @@
 #include "BlueprintGraphDefinitions.h"
 #include "GraphEditorSettings.h"
 
-#include "VectorVM.h"
+
+#include "INiagaraEditor.h"
+#include "INiagaraCompiler.h"
+
+#define LOCTEXT_NAMESPACE "NiagaraSchema"
 
 #define SNAP_GRID (16) // @todo ensure this is the same as SNodePanel::GetSnapGridSize()
 
@@ -96,7 +100,8 @@ UEdGraphSchema_Niagara::UEdGraphSchema_Niagara(const FObjectInitializer& ObjectI
 {
 
 	PC_Float = TEXT("float");
-	PC_Struct = TEXT("struct");
+	PC_Vector = TEXT("vector");
+	PC_Matrix = TEXT("matrix");
 }
 
 TSharedPtr<FNiagaraSchemaAction_NewNode> AddNewNodeAction(FGraphContextMenuBuilder& ContextMenuBuilder, const FString& Category, const FText& MenuDesc, const FString& Tooltip)
@@ -113,7 +118,6 @@ void UEdGraphSchema_Niagara::GetGraphContextActions(FGraphContextMenuBuilder& Co
 
 	TArray<FName> InputAttributeNames;
 	Source->GetParticleAttributes(InputAttributeNames);
-	Source->GetEmitterAttributes(InputAttributeNames);
 	for(int32 i=0; i<InputAttributeNames.Num(); i++)
 	{
 		const FName AttrName = InputAttributeNames[i];
@@ -129,20 +133,62 @@ void UEdGraphSchema_Niagara::GetGraphContextActions(FGraphContextMenuBuilder& Co
 		GetAttrAction->NodeTemplate = GetAttrNode;
 	}
 
-	// Then get ops.
-	uint8 NumOps = VectorVM::GetNumOpCodes();
-	for(uint8 OpIdx=0; OpIdx<NumOps; OpIdx++)
+	TArray<FName> EmitterConstantNames_Vectors;
+	TArray<FName> EmitterConstantNames_Matrices;
+	Source->GetEmitterAttributes(EmitterConstantNames_Vectors, EmitterConstantNames_Matrices);
+	for (int32 i = 0; i < EmitterConstantNames_Vectors.Num(); i++)
 	{
-		VectorVM::FVectorVMOpInfo const& Info = VectorVM::GetOpCodeInfo(OpIdx);
-		if(Info.IsImplemented())
-		{
-			TSharedPtr<FNiagaraSchemaAction_NewNode> AddOpAction = AddNewNodeAction(ContextMenuBuilder, TEXT("Add Operation"), FText::FromString(Info.FriendlyName), TEXT(""));
+		const FName ConstName = EmitterConstantNames_Vectors[i];
 
-			UNiagaraNodeOp* OpNode = NewObject<UNiagaraNodeOp>(ContextMenuBuilder.OwnerOfTemporaries);
-			OpNode->OpIndex = OpIdx;
-			AddOpAction->NodeTemplate = OpNode;
-		}
+		FFormatNamedArguments Args;
+		Args.Add(TEXT("Constant"), FText::FromName(ConstName));
+		const FText MenuDesc = FText::Format(NSLOCTEXT("Niagara", "GetConstant", "Get {Constant}"), Args);
+
+		TSharedPtr<FNiagaraSchemaAction_NewNode> GetAttrAction = AddNewNodeAction(ContextMenuBuilder, TEXT("Get Constant"), MenuDesc, TEXT(""));
+
+		UNiagaraNodeConstant* GetConstNode = NewObject<UNiagaraNodeConstant>(ContextMenuBuilder.OwnerOfTemporaries);
+		GetConstNode->ConstName = ConstName;
+		GetConstNode->DataType = ENiagaraDataType::Vector;
+		GetConstNode->bNeedsDefault = false;
+		GetAttrAction->NodeTemplate = GetConstNode;
+	}
+
+	for (int32 i = 0; i < EmitterConstantNames_Matrices.Num(); i++)
+	{
+		const FName ConstName = EmitterConstantNames_Matrices[i];
+
+		FFormatNamedArguments Args;
+		Args.Add(TEXT("Constant"), FText::FromName(ConstName));
+		const FText MenuDesc = FText::Format(NSLOCTEXT("Niagara", "GetConstant", "Get {Constant}"), Args);
+
+		TSharedPtr<FNiagaraSchemaAction_NewNode> GetAttrAction = AddNewNodeAction(ContextMenuBuilder, TEXT("Get Constant"), MenuDesc, TEXT(""));
+
+		UNiagaraNodeConstant* GetConstNode = NewObject<UNiagaraNodeConstant>(ContextMenuBuilder.OwnerOfTemporaries);
+		GetConstNode->ConstName = ConstName;
+		GetConstNode->DataType = ENiagaraDataType::Matrix;
+		GetConstNode->bNeedsDefault = false;
+		GetAttrAction->NodeTemplate = GetConstNode;
 	}
+
+#define NiagaraOp(OPNAME) \
+		if(const FNiagaraOpInfo* OpInfo = FNiagaraOpInfo::GetOpInfo(FNiagaraOpInfo::##OPNAME##))\
+		{\
+			TSharedPtr<FNiagaraSchemaAction_NewNode> AddOpAction = AddNewNodeAction(ContextMenuBuilder, TEXT("Add Operation"), OpInfo->FriendlyName, TEXT(""));\
+			UNiagaraNodeOp* OpNode = NewObject<UNiagaraNodeOp>(ContextMenuBuilder.OwnerOfTemporaries); \
+			OpNode->OpName = OpInfo->Name; \
+			AddOpAction->NodeTemplate = OpNode; \
+		}
+
+		
+	NiagaraOpList
+
+#undef NiagaraOp
+
+	UNiagaraNodeConstant* ConstantNode = NewObject<UNiagaraNodeConstant>(ContextMenuBuilder.OwnerOfTemporaries);
+	ConstantNode->DataType = ENiagaraDataType::Vector;
+	ConstantNode->bNeedsDefault = true;
+	TSharedPtr<FNiagaraSchemaAction_NewNode> AddOpAction = AddNewNodeAction(ContextMenuBuilder, TEXT("Constants"), ConstantNode->GetNodeTitle(ENodeTitleType::ListView), TEXT(""));
+	AddOpAction->NodeTemplate = ConstantNode;
 }
 
 const FPinConnectionResponse UEdGraphSchema_Niagara::CanCreateConnection(const UEdGraphPin* PinA, const UEdGraphPin* PinB) const
@@ -196,9 +242,13 @@ FLinearColor UEdGraphSchema_Niagara::GetPinTypeColor(const FEdGraphPinType& PinT
 	{
 		return Settings->FloatPinTypeColor;
 	}
-	else if (TypeString == PC_Struct)
+	else if (TypeString == PC_Vector)
 	{
-		return Settings->StructPinTypeColor;
+		return Settings->VectorPinTypeColor;
+	}
+	else if (TypeString == PC_Matrix)
+	{
+		return Settings->TransformPinTypeColor;
 	}
 
 	return Settings->DefaultPinTypeColor;
@@ -215,3 +265,62 @@ bool UEdGraphSchema_Niagara::ShouldHidePinDefaultValue(UEdGraphPin* Pin) const
 
 	return false;
 }
+
+ENiagaraDataType UEdGraphSchema_Niagara::GetPinDataType(UEdGraphPin* Pin) const
+{
+	if (Pin->PinType.PinCategory == PC_Float)
+	{
+		return ENiagaraDataType::Scalar;
+	}
+	else if (Pin->PinType.PinCategory == PC_Vector)
+	{
+		return ENiagaraDataType::Vector;
+	}
+	else if (Pin->PinType.PinCategory == PC_Matrix)
+	{
+		return ENiagaraDataType::Matrix;
+	}
+	return ENiagaraDataType::Invalid;
+}
+
+void UEdGraphSchema_Niagara::GetPinDefaultValue(UEdGraphPin* Pin, float& OutDefault)const
+{
+	//Ugh, this string parsing is rubbish. Surely there's a consistent api for dealing with default values in the edgraph without all this?
+	FString ResultString = Pin->DefaultValue.IsEmpty() ? Pin->AutogeneratedDefaultValue : Pin->DefaultValue;
+	ResultString.Trim();
+	ResultString.TrimTrailing();
+	OutDefault = FCString::Atof(*ResultString);
+}
+
+void UEdGraphSchema_Niagara::GetPinDefaultValue(UEdGraphPin* Pin, FVector4& OutDefault)const
+{
+	//Ugh, this string parsing is rubbish. Surely there's a consistent api for dealing with default values in the edgraph without all this?
+	TArray<FString> ResultString;
+	FString ValueString = Pin->DefaultValue.IsEmpty() ? Pin->AutogeneratedDefaultValue : Pin->DefaultValue;
+	ValueString.Trim();
+	ValueString.TrimTrailing();
+	ValueString.ParseIntoArray(&ResultString, TEXT(","), true);
+	if (ResultString.Num() == 4)
+	{
+		OutDefault = FVector4(FCString::Atof(*ResultString[0]), FCString::Atof(*ResultString[1]), FCString::Atof(*ResultString[2]), FCString::Atof(*ResultString[3]));
+	}
+}
+
+void UEdGraphSchema_Niagara::GetPinDefaultValue(UEdGraphPin* Pin, FMatrix& OutDefault)const
+{
+	TArray<FString> ResultString;
+	FString ValueString = Pin->DefaultValue.IsEmpty() ? Pin->AutogeneratedDefaultValue : Pin->DefaultValue;
+	ValueString.Trim();
+	ValueString.TrimTrailing();
+	ValueString.ParseIntoArray(&ResultString, TEXT(","), true);
+	check(ResultString.Num() == 16);
+	if (ResultString.Num() == 16)
+	{
+		OutDefault.M[0][0] = FCString::Atof(*ResultString[0]); OutDefault.M[0][1] = FCString::Atof(*ResultString[1]); OutDefault.M[0][2] = FCString::Atof(*ResultString[2]); OutDefault.M[0][3] = FCString::Atof(*ResultString[3]);
+		OutDefault.M[1][0] = FCString::Atof(*ResultString[4]); OutDefault.M[1][1] = FCString::Atof(*ResultString[5]); OutDefault.M[1][2] = FCString::Atof(*ResultString[6]); OutDefault.M[1][3] = FCString::Atof(*ResultString[7]);
+		OutDefault.M[2][0] = FCString::Atof(*ResultString[8]); OutDefault.M[2][1] = FCString::Atof(*ResultString[9]); OutDefault.M[2][2] = FCString::Atof(*ResultString[10]); OutDefault.M[2][3] = FCString::Atof(*ResultString[11]);
+		OutDefault.M[3][0] = FCString::Atof(*ResultString[12]); OutDefault.M[3][1] = FCString::Atof(*ResultString[13]); OutDefault.M[3][2] = FCString::Atof(*ResultString[14]); OutDefault.M[3][3] = FCString::Atof(*ResultString[15]);
+	}
+}
+
+#undef LOCTEXT_NAMESPACE
\ No newline at end of file
diff --git a/Engine/Source/Editor/UnrealEd/Private/Niagara/NiagaraEditorCommon.cpp b/Engine/Source/Editor/UnrealEd/Private/Niagara/NiagaraEditorCommon.cpp
new file mode 100644
--- /dev/null
+++ b/Engine/Source/Editor/UnrealEd/Private/Niagara/NiagaraEditorCommon.cpp
@@ -0,0 +1,578 @@
+// Copyright 1998-2014 Epic Games, Inc. All Rights Reserved.
+
+#include "UnrealEd.h"
+
+DECLARE_LOG_CATEGORY_EXTERN(LogNiagaraCompiler, All, All);
+DEFINE_LOG_CATEGORY(LogNiagaraCompiler);
+
+#define NiagaraOp(OpName) const FName FNiagaraOpInfo::OpName = FName(TEXT(#OpName));
+NiagaraOpList
+#undef NiagaraOp
+
+void FNiagaraOpInfo::Init()
+{
+	//Common input and output names.
+	FName Result(TEXT("Result"));	FText ResultText = NSLOCTEXT("NiagaraOpInfo", "Operation Result", "Result");
+	FName A(TEXT("A"));	FText AText = NSLOCTEXT("NiagaraOpInfo", "First Function Param", "A");
+	FName B(TEXT("B"));	FText BText = NSLOCTEXT("NiagaraOpInfo", "Second Function Param", "B");
+	FName C(TEXT("C"));	FText CText = NSLOCTEXT("NiagaraOpInfo", "Third Function Param", "C");
+	FName D(TEXT("D"));	FText DText = NSLOCTEXT("NiagaraOpInfo", "Fourth Function Param", "D");
+	FName X(TEXT("X"));	FText XText = NSLOCTEXT("NiagaraOpInfo", "First Vector Component", "X");
+	FName Y(TEXT("Y"));	FText YText = NSLOCTEXT("NiagaraOpInfo", "Second Vector Component", "Y");
+	FName Z(TEXT("Z"));	FText ZText = NSLOCTEXT("NiagaraOpInfo", "Third Vector Component", "Z");
+	FName W(TEXT("W"));	FText WText = NSLOCTEXT("NiagaraOpInfo", "Fourth Vector Component", "W");
+
+	TMap<FName, FNiagaraOpInfo>& OpInfos = GetOpInfoMap();
+
+	FString DefaultStr_Zero(TEXT("0.0"));
+	FString DefaultStr_One(TEXT("1.0"));
+	FString DefaultStr_VecZero(TEXT("0.0,0.0,0.0,0.0"));
+	FString DefaultStr_VecOne(TEXT("1.0,1.0,1.0,1.0"));
+
+	FNiagaraOpInfo* Op = &OpInfos.Add(Add);
+	Op->Name = Add;
+	Op->FriendlyName = NSLOCTEXT("NiagaraOpInfo", "Add Name", "Add");
+	Op->Description = NSLOCTEXT("NiagaraOpInfo", "Add Desc", "Result = A + B");
+	Op->Inputs.Add(FNiagaraOpInOutInfo(A, ENiagaraDataType::Vector, AText, AText, DefaultStr_VecZero));
+	Op->Inputs.Add(FNiagaraOpInOutInfo(B, ENiagaraDataType::Vector, BText, BText, DefaultStr_VecZero));
+	Op->Outputs.Add(FNiagaraOpInOutInfo(Result, ENiagaraDataType::Vector, ResultText, ResultText, DefaultStr_VecZero));
+	Op->OpDelegate.BindStatic(&INiagaraCompiler::Add);
+
+	Op = &OpInfos.Add(Subtract);
+	Op->Name = Subtract;
+	Op->FriendlyName = NSLOCTEXT("NiagaraOpInfo", "Subtract Name", "Subtract");
+	Op->Description = NSLOCTEXT("NiagaraOpInfo", "Subtract Desc", "Result = A - B");
+	Op->Inputs.Add(FNiagaraOpInOutInfo(A, ENiagaraDataType::Vector, AText, AText, DefaultStr_VecOne));
+	Op->Inputs.Add(FNiagaraOpInOutInfo(B, ENiagaraDataType::Vector, BText, BText, DefaultStr_VecOne));
+	Op->Outputs.Add(FNiagaraOpInOutInfo(Result, ENiagaraDataType::Vector, ResultText, ResultText, DefaultStr_VecOne));
+	Op->OpDelegate.BindStatic(&INiagaraCompiler::Subtract);
+
+	Op = &OpInfos.Add(Multiply);
+	Op->Name = Multiply;
+	Op->FriendlyName = NSLOCTEXT("NiagaraOpInfo", "Multiply Name", "Multiply");
+	Op->Description = NSLOCTEXT("NiagaraOpInfo", "Multiply Desc", "Result = A * B = {{A.x * B.x, A.y * B.y, A.z * B.z, A.w * B.w}}");
+	Op->Inputs.Add(FNiagaraOpInOutInfo(A, ENiagaraDataType::Vector, AText, AText, DefaultStr_VecOne));
+	Op->Inputs.Add(FNiagaraOpInOutInfo(B, ENiagaraDataType::Vector, BText, BText, DefaultStr_VecOne));
+	Op->Outputs.Add(FNiagaraOpInOutInfo(Result, ENiagaraDataType::Vector, ResultText, ResultText, DefaultStr_VecOne));
+	Op->OpDelegate.BindStatic(&INiagaraCompiler::Multiply);
+
+	Op = &OpInfos.Add(MultiplyAdd);
+	Op->Name = MultiplyAdd;
+	Op->FriendlyName = NSLOCTEXT("NiagaraOpInfo", "MultiplyAdd Name", "MultiplyAdd");
+	Op->Description = NSLOCTEXT("NiagaraOpInfo", "MultiplyAdd Desc", "Result = (A * B) + C");
+	Op->Inputs.Add(FNiagaraOpInOutInfo(A, ENiagaraDataType::Vector, AText, AText, DefaultStr_VecOne));
+	Op->Inputs.Add(FNiagaraOpInOutInfo(B, ENiagaraDataType::Vector, BText, BText, DefaultStr_VecOne));
+	Op->Inputs.Add(FNiagaraOpInOutInfo(C, ENiagaraDataType::Vector, CText, CText, DefaultStr_VecZero));
+	Op->Outputs.Add(FNiagaraOpInOutInfo(Result, ENiagaraDataType::Vector, ResultText, ResultText, DefaultStr_VecOne));
+	Op->OpDelegate.BindStatic(&INiagaraCompiler::MultiplyAdd);
+
+	Op = &OpInfos.Add(Lerp);
+	Op->Name = Lerp;
+	Op->FriendlyName = NSLOCTEXT("NiagaraOpInfo", "Lerp Name", "Lerp");
+	Op->Description = NSLOCTEXT("NiagaraOpInfo", "Lerp Desc", "Result = (A * (1 - C)) + (B * C)");
+	Op->Inputs.Add(FNiagaraOpInOutInfo(A, ENiagaraDataType::Vector, AText, AText, DefaultStr_VecZero));
+	Op->Inputs.Add(FNiagaraOpInOutInfo(B, ENiagaraDataType::Vector, BText, BText, DefaultStr_VecOne));
+	Op->Inputs.Add(FNiagaraOpInOutInfo(C, ENiagaraDataType::Vector, CText, CText, DefaultStr_VecZero));
+	Op->Outputs.Add(FNiagaraOpInOutInfo(Result, ENiagaraDataType::Vector, ResultText, ResultText, DefaultStr_VecZero));
+	Op->OpDelegate.BindStatic(&INiagaraCompiler::Lerp);
+
+	Op = &OpInfos.Add(Reciprocal);
+	Op->Name = Reciprocal;
+	Op->FriendlyName = NSLOCTEXT("NiagaraOpInfo", "Reciprocal Name", "Reciprocal");
+	Op->Description = NSLOCTEXT("NiagaraOpInfo", "Reciprocal Desc", "Result = 1 / A");
+	Op->Inputs.Add(FNiagaraOpInOutInfo(A, ENiagaraDataType::Vector, AText, AText, DefaultStr_VecOne));
+	Op->Outputs.Add(FNiagaraOpInOutInfo(Result, ENiagaraDataType::Vector, ResultText, ResultText, DefaultStr_VecOne));
+	Op->OpDelegate.BindStatic(&INiagaraCompiler::Reciprocal);
+
+	Op = &OpInfos.Add(ReciprocalSqrt);
+	Op->Name = ReciprocalSqrt;
+	Op->FriendlyName = NSLOCTEXT("NiagaraOpInfo", "ReciprocalSqrt Name", "ReciprocalSqrt");
+	Op->Description = NSLOCTEXT("NiagaraOpInfo", "ReciprocalSqrt Desc", "Result = 1 / Sqrt(A)");
+	Op->Inputs.Add(FNiagaraOpInOutInfo(A, ENiagaraDataType::Vector, AText, AText, DefaultStr_VecOne));
+	Op->Outputs.Add(FNiagaraOpInOutInfo(Result, ENiagaraDataType::Vector, ResultText, ResultText, DefaultStr_VecOne));
+	Op->OpDelegate.BindStatic(&INiagaraCompiler::ReciprocalSqrt);
+
+	Op = &OpInfos.Add(Sqrt);
+	Op->Name = Sqrt;
+	Op->FriendlyName = NSLOCTEXT("NiagaraOpInfo", "Sqrt Name", "Sqrt");
+	Op->Description = NSLOCTEXT("NiagaraOpInfo", "Sqrt Desc", "Result = A ^ 0.5");
+	Op->Inputs.Add(FNiagaraOpInOutInfo(A, ENiagaraDataType::Vector, AText, AText, DefaultStr_VecOne));
+	Op->Outputs.Add(FNiagaraOpInOutInfo(Result, ENiagaraDataType::Vector, ResultText, ResultText, DefaultStr_VecOne));
+	Op->OpDelegate.BindStatic(&INiagaraCompiler::Sqrt);
+
+	Op = &OpInfos.Add(Negate);
+	Op->Name = Negate;
+	Op->FriendlyName = NSLOCTEXT("NiagaraOpInfo", "Negate Name", "Negate");
+	Op->Description = NSLOCTEXT("NiagaraOpInfo", "Negate Desc", "Result = 1 - A");
+	Op->Inputs.Add(FNiagaraOpInOutInfo(A, ENiagaraDataType::Vector, AText, AText, DefaultStr_VecOne));
+	Op->Outputs.Add(FNiagaraOpInOutInfo(Result, ENiagaraDataType::Vector, ResultText, ResultText, DefaultStr_VecOne));
+	Op->OpDelegate.BindStatic(&INiagaraCompiler::Negate);
+
+	Op = &OpInfos.Add(Abs);
+	Op->Name = Abs;
+	Op->FriendlyName = NSLOCTEXT("NiagaraOpInfo", "Abs Name", "Abs");
+	Op->Description = NSLOCTEXT("NiagaraOpInfo", "Abs Desc", "Result = |A|");
+	Op->Inputs.Add(FNiagaraOpInOutInfo(A, ENiagaraDataType::Vector, AText, AText, DefaultStr_VecOne));
+	Op->Inputs.Add(FNiagaraOpInOutInfo(B, ENiagaraDataType::Vector, BText, BText, DefaultStr_VecOne));
+	Op->Inputs.Add(FNiagaraOpInOutInfo(C, ENiagaraDataType::Vector, CText, CText, DefaultStr_VecOne));
+	Op->Outputs.Add(FNiagaraOpInOutInfo(Result, ENiagaraDataType::Vector, ResultText, ResultText, DefaultStr_VecOne));
+	Op->OpDelegate.BindStatic(&INiagaraCompiler::Abs);
+
+	Op = &OpInfos.Add(Exp);
+	Op->Name = Exp;
+	Op->FriendlyName = NSLOCTEXT("NiagaraOpInfo", "Exp Name", "Exp");
+	Op->Description = NSLOCTEXT("NiagaraOpInfo", "Exp Desc", "Exp");
+	Op->Inputs.Add(FNiagaraOpInOutInfo(A, ENiagaraDataType::Vector, AText, AText, DefaultStr_VecOne));
+	Op->Outputs.Add(FNiagaraOpInOutInfo(Result, ENiagaraDataType::Vector, ResultText, ResultText, DefaultStr_VecOne));
+	Op->OpDelegate.BindStatic(&INiagaraCompiler::Exp);
+
+	Op = &OpInfos.Add(Exp2);
+	Op->Name = Exp2;
+	Op->FriendlyName = NSLOCTEXT("NiagaraOpInfo", "Exp2 Name", "Exp2");
+	Op->Description = NSLOCTEXT("NiagaraOpInfo", "Exp2 Desc", "Exp2");
+	Op->Inputs.Add(FNiagaraOpInOutInfo(A, ENiagaraDataType::Vector, AText, AText, DefaultStr_VecOne));
+	Op->Outputs.Add(FNiagaraOpInOutInfo(Result, ENiagaraDataType::Vector, ResultText, ResultText, DefaultStr_VecOne));
+	Op->OpDelegate.BindStatic(&INiagaraCompiler::Exp2);
+
+	Op = &OpInfos.Add(Log);
+	Op->Name = Log;
+	Op->FriendlyName = NSLOCTEXT("NiagaraOpInfo", "Log Name", "Log");
+	Op->Description = NSLOCTEXT("NiagaraOpInfo", "Log Desc", "Log");
+	Op->Inputs.Add(FNiagaraOpInOutInfo(A, ENiagaraDataType::Vector, AText, AText, DefaultStr_VecOne));
+	Op->Outputs.Add(FNiagaraOpInOutInfo(Result, ENiagaraDataType::Vector, ResultText, ResultText, DefaultStr_VecOne));
+	Op->OpDelegate.BindStatic(&INiagaraCompiler::Log);
+
+	Op = &OpInfos.Add(LogBase2);
+	Op->Name = LogBase2;
+	Op->FriendlyName = NSLOCTEXT("NiagaraOpInfo", "LogBase2 Name", "LogBase2");
+	Op->Description = NSLOCTEXT("NiagaraOpInfo", "LogBase2 Desc", "LogBase2");
+	Op->Inputs.Add(FNiagaraOpInOutInfo(A, ENiagaraDataType::Vector, AText, AText, DefaultStr_VecOne));
+	Op->Outputs.Add(FNiagaraOpInOutInfo(Result, ENiagaraDataType::Vector, ResultText, ResultText, DefaultStr_VecOne));
+	Op->OpDelegate.BindStatic(&INiagaraCompiler::LogBase2);
+
+	Op = &OpInfos.Add(Sin);
+	Op->Name = Sin;
+	Op->FriendlyName = NSLOCTEXT("NiagaraOpInfo", "Sin Name", "Sin");
+	Op->Description = NSLOCTEXT("NiagaraOpInfo", "Sin Desc", "Sin");
+	Op->Inputs.Add(FNiagaraOpInOutInfo(A, ENiagaraDataType::Vector, AText, AText, DefaultStr_VecZero));
+	Op->Outputs.Add(FNiagaraOpInOutInfo(Result, ENiagaraDataType::Vector, ResultText, ResultText, DefaultStr_VecZero));
+	Op->OpDelegate.BindStatic(&INiagaraCompiler::Sin);
+
+	Op = &OpInfos.Add(Cos);
+	Op->Name = Cos;
+	Op->FriendlyName = NSLOCTEXT("NiagaraOpInfo", "Cos Name", "Cos");
+	Op->Description = NSLOCTEXT("NiagaraOpInfo", "Cos Desc", "Cos");
+	Op->Inputs.Add(FNiagaraOpInOutInfo(A, ENiagaraDataType::Vector, AText, AText, DefaultStr_VecZero));
+	Op->Outputs.Add(FNiagaraOpInOutInfo(Result, ENiagaraDataType::Vector, ResultText, ResultText, DefaultStr_VecZero));
+	Op->OpDelegate.BindStatic(&INiagaraCompiler::Cos);
+
+	Op = &OpInfos.Add(Tan);
+	Op->Name = Tan;
+	Op->FriendlyName = NSLOCTEXT("NiagaraOpInfo", "Tan Name", "Tan");
+	Op->Description = NSLOCTEXT("NiagaraOpInfo", "Tan Desc", "Tan");
+	Op->Inputs.Add(FNiagaraOpInOutInfo(A, ENiagaraDataType::Vector, AText, AText, DefaultStr_VecZero));
+	Op->Outputs.Add(FNiagaraOpInOutInfo(Result, ENiagaraDataType::Vector, ResultText, ResultText, DefaultStr_VecZero));
+	Op->OpDelegate.BindStatic(&INiagaraCompiler::Tan);
+
+	Op = &OpInfos.Add(ASin);
+	Op->Name = ASin;
+	Op->FriendlyName = NSLOCTEXT("NiagaraOpInfo", "ASin Name", "ASin");
+	Op->Description = NSLOCTEXT("NiagaraOpInfo", "ASin Desc", "ASin");
+	Op->Inputs.Add(FNiagaraOpInOutInfo(A, ENiagaraDataType::Vector, AText, AText, DefaultStr_VecZero));
+	Op->Outputs.Add(FNiagaraOpInOutInfo(Result, ENiagaraDataType::Vector, ResultText, ResultText, DefaultStr_VecZero));
+	Op->OpDelegate.BindStatic(&INiagaraCompiler::ASin);
+
+	Op = &OpInfos.Add(ACos);
+	Op->Name = ACos;
+	Op->FriendlyName = NSLOCTEXT("NiagaraOpInfo", "ACos Name", "ACos");
+	Op->Description = NSLOCTEXT("NiagaraOpInfo", "ACos Desc", "ACos");
+	Op->Inputs.Add(FNiagaraOpInOutInfo(A, ENiagaraDataType::Vector, AText, AText, DefaultStr_VecZero));
+	Op->Outputs.Add(FNiagaraOpInOutInfo(Result, ENiagaraDataType::Vector, ResultText, ResultText, DefaultStr_VecZero));
+	Op->OpDelegate.BindStatic(&INiagaraCompiler::ACos);
+
+	Op = &OpInfos.Add(ATan);
+	Op->Name = ATan;
+	Op->FriendlyName = NSLOCTEXT("NiagaraOpInfo", "ATan Name", "ATan");
+	Op->Description = NSLOCTEXT("NiagaraOpInfo", "ATan Desc", "ATan");
+	Op->Inputs.Add(FNiagaraOpInOutInfo(A, ENiagaraDataType::Vector, AText, AText, DefaultStr_VecZero));
+	Op->Outputs.Add(FNiagaraOpInOutInfo(Result, ENiagaraDataType::Vector, ResultText, ResultText, DefaultStr_VecZero));
+	Op->OpDelegate.BindStatic(&INiagaraCompiler::ATan);
+
+	Op = &OpInfos.Add(ATan2);
+	Op->Name = ATan2;
+	Op->FriendlyName = NSLOCTEXT("NiagaraOpInfo", "ATan2 Name", "ATan2");
+	Op->Description = NSLOCTEXT("NiagaraOpInfo", "ATan2 Desc", "ATan2");
+	Op->Inputs.Add(FNiagaraOpInOutInfo(A, ENiagaraDataType::Vector, AText, AText, DefaultStr_VecZero));
+	Op->Inputs.Add(FNiagaraOpInOutInfo(B, ENiagaraDataType::Vector, BText, BText, DefaultStr_VecZero));
+	Op->Outputs.Add(FNiagaraOpInOutInfo(Result, ENiagaraDataType::Vector, ResultText, ResultText, DefaultStr_VecZero));
+	Op->OpDelegate.BindStatic(&INiagaraCompiler::ATan2);
+
+	Op = &OpInfos.Add(Ceil);
+	Op->Name = Ceil;
+	Op->FriendlyName = NSLOCTEXT("NiagaraOpInfo", "Ceil Name", "Ceil");
+	Op->Description = NSLOCTEXT("NiagaraOpInfo", "Ceil Desc", "Ceil");
+	Op->Inputs.Add(FNiagaraOpInOutInfo(A, ENiagaraDataType::Vector, AText, AText, DefaultStr_VecOne));
+	Op->Outputs.Add(FNiagaraOpInOutInfo(Result, ENiagaraDataType::Vector, ResultText, ResultText, DefaultStr_VecOne));
+	Op->OpDelegate.BindStatic(&INiagaraCompiler::Ceil);
+
+	Op = &OpInfos.Add(Floor);
+	Op->Name = Floor;
+	Op->FriendlyName = NSLOCTEXT("NiagaraOpInfo", "Floor Name", "Floor");
+	Op->Description = NSLOCTEXT("NiagaraOpInfo", "Floor Desc", "Floor");
+	Op->Inputs.Add(FNiagaraOpInOutInfo(A, ENiagaraDataType::Vector, AText, AText, DefaultStr_VecOne));
+	Op->Outputs.Add(FNiagaraOpInOutInfo(Result, ENiagaraDataType::Vector, ResultText, ResultText, DefaultStr_VecOne));
+	Op->OpDelegate.BindStatic(&INiagaraCompiler::Floor);
+
+	Op = &OpInfos.Add(Modulo);
+	Op->Name = Modulo;
+	Op->FriendlyName = NSLOCTEXT("NiagaraOpInfo", "Modulo Name", "Modulo");
+	Op->Description = NSLOCTEXT("NiagaraOpInfo", "Modulo Desc", "Modulo");
+	Op->Inputs.Add(FNiagaraOpInOutInfo(A, ENiagaraDataType::Vector, AText, AText, DefaultStr_VecOne));
+	Op->Inputs.Add(FNiagaraOpInOutInfo(B, ENiagaraDataType::Vector, BText, BText, DefaultStr_VecOne));
+	Op->Outputs.Add(FNiagaraOpInOutInfo(Result, ENiagaraDataType::Vector, ResultText, ResultText, DefaultStr_VecOne));
+	Op->OpDelegate.BindStatic(&INiagaraCompiler::Modulo);
+
+	Op = &OpInfos.Add(Fractional);
+	Op->Name = Fractional;
+	Op->FriendlyName = NSLOCTEXT("NiagaraOpInfo", "Fractional Name", "Fractional");
+	Op->Description = NSLOCTEXT("NiagaraOpInfo", "Fractional Desc", "Fractional");
+	Op->Inputs.Add(FNiagaraOpInOutInfo(A, ENiagaraDataType::Vector, AText, AText, DefaultStr_VecOne));
+	Op->Outputs.Add(FNiagaraOpInOutInfo(Result, ENiagaraDataType::Vector, ResultText, ResultText, DefaultStr_VecOne));
+	Op->OpDelegate.BindStatic(&INiagaraCompiler::Fractional);
+
+	Op = &OpInfos.Add(Truncate);
+	Op->Name = Truncate;
+	Op->FriendlyName = NSLOCTEXT("NiagaraOpInfo", "Truncate Name", "Truncate");
+	Op->Description = NSLOCTEXT("NiagaraOpInfo", "Truncate Desc", "Truncate");
+	Op->Inputs.Add(FNiagaraOpInOutInfo(A, ENiagaraDataType::Vector, AText, AText, DefaultStr_VecOne));
+	Op->Outputs.Add(FNiagaraOpInOutInfo(Result, ENiagaraDataType::Vector, ResultText, ResultText, DefaultStr_VecOne));
+	Op->OpDelegate.BindStatic(&INiagaraCompiler::Truncate);
+
+	Op = &OpInfos.Add(Clamp);
+	Op->Name = Clamp;
+	Op->FriendlyName = NSLOCTEXT("NiagaraOpInfo", "Clamp Name", "Clamp");
+	Op->Description = NSLOCTEXT("NiagaraOpInfo", "Clamp Desc", "Clamp");
+	Op->Inputs.Add(FNiagaraOpInOutInfo(A, ENiagaraDataType::Vector, AText, AText, DefaultStr_VecZero));
+	Op->Inputs.Add(FNiagaraOpInOutInfo(B, ENiagaraDataType::Vector, BText, BText, DefaultStr_VecOne));
+	Op->Inputs.Add(FNiagaraOpInOutInfo(C, ENiagaraDataType::Vector, CText, CText, DefaultStr_VecOne));
+	Op->Outputs.Add(FNiagaraOpInOutInfo(Result, ENiagaraDataType::Vector, ResultText, ResultText, DefaultStr_VecOne));
+	Op->OpDelegate.BindStatic(&INiagaraCompiler::Clamp);
+
+	Op = &OpInfos.Add(Min);
+	Op->Name = Min;
+	Op->FriendlyName = NSLOCTEXT("NiagaraOpInfo", "Min Name", "Min");
+	Op->Description = NSLOCTEXT("NiagaraOpInfo", "Min Desc", "Min");
+	Op->Inputs.Add(FNiagaraOpInOutInfo(A, ENiagaraDataType::Vector, AText, AText, DefaultStr_VecZero));
+	Op->Inputs.Add(FNiagaraOpInOutInfo(B, ENiagaraDataType::Vector, BText, BText, DefaultStr_VecOne));
+	Op->Outputs.Add(FNiagaraOpInOutInfo(Result, ENiagaraDataType::Vector, ResultText, ResultText, DefaultStr_VecOne));
+	Op->OpDelegate.BindStatic(&INiagaraCompiler::Min);
+
+	Op = &OpInfos.Add(Max);
+	Op->Name = Max;
+	Op->FriendlyName = NSLOCTEXT("NiagaraOpInfo", "Max Name", "Max");
+	Op->Description = NSLOCTEXT("NiagaraOpInfo", "Max Desc", "Max");
+	Op->Inputs.Add(FNiagaraOpInOutInfo(A, ENiagaraDataType::Vector, AText, AText, DefaultStr_VecZero));
+	Op->Inputs.Add(FNiagaraOpInOutInfo(B, ENiagaraDataType::Vector, BText, BText, DefaultStr_VecOne));
+	Op->Outputs.Add(FNiagaraOpInOutInfo(Result, ENiagaraDataType::Vector, ResultText, ResultText, DefaultStr_VecOne));
+	Op->OpDelegate.BindStatic(&INiagaraCompiler::Max);
+
+	Op = &OpInfos.Add(Pow);
+	Op->Name = Pow;
+	Op->FriendlyName = NSLOCTEXT("NiagaraOpInfo", "Pow Name", "Pow");
+	Op->Description = NSLOCTEXT("NiagaraOpInfo", "Pow Desc", "Pow");
+	Op->Inputs.Add(FNiagaraOpInOutInfo(A, ENiagaraDataType::Vector, AText, AText, DefaultStr_VecOne));
+	Op->Inputs.Add(FNiagaraOpInOutInfo(B, ENiagaraDataType::Vector, BText, BText, DefaultStr_VecOne));
+	Op->Outputs.Add(FNiagaraOpInOutInfo(Result, ENiagaraDataType::Vector, ResultText, ResultText, DefaultStr_VecOne));
+	Op->OpDelegate.BindStatic(&INiagaraCompiler::Pow);
+
+	Op = &OpInfos.Add(Sign);
+	Op->Name = Sign;
+	Op->FriendlyName = NSLOCTEXT("NiagaraOpInfo", "Sign Name", "Sign");
+	Op->Description = NSLOCTEXT("NiagaraOpInfo", "Sign Desc", "Sign");
+	Op->Inputs.Add(FNiagaraOpInOutInfo(A, ENiagaraDataType::Vector, AText, AText, DefaultStr_VecOne));
+	Op->Outputs.Add(FNiagaraOpInOutInfo(Result, ENiagaraDataType::Vector, ResultText, ResultText, DefaultStr_VecOne));
+	Op->OpDelegate.BindStatic(&INiagaraCompiler::Sign);
+
+	Op = &OpInfos.Add(Step);
+	Op->Name = Step;
+	Op->FriendlyName = NSLOCTEXT("NiagaraOpInfo", "Step Name", "Step");
+	Op->Description = NSLOCTEXT("NiagaraOpInfo", "Step Desc", "Step");
+	Op->Inputs.Add(FNiagaraOpInOutInfo(A, ENiagaraDataType::Vector, AText, AText, DefaultStr_VecOne));
+	Op->Inputs.Add(FNiagaraOpInOutInfo(B, ENiagaraDataType::Vector, BText, BText, DefaultStr_VecOne));
+	Op->Outputs.Add(FNiagaraOpInOutInfo(Result, ENiagaraDataType::Vector, ResultText, ResultText, DefaultStr_VecOne));
+	Op->OpDelegate.BindStatic(&INiagaraCompiler::Step);
+
+	Op = &OpInfos.Add(Dot);
+	Op->Name = Dot;
+	Op->FriendlyName = NSLOCTEXT("NiagaraOpInfo", "Dot Name", "Dot");
+	Op->Description = NSLOCTEXT("NiagaraOpInfo", "Dot Desc", "Dot");
+	Op->Inputs.Add(FNiagaraOpInOutInfo(A, ENiagaraDataType::Vector, AText, AText, DefaultStr_VecOne));
+	Op->Inputs.Add(FNiagaraOpInOutInfo(B, ENiagaraDataType::Vector, BText, BText, DefaultStr_VecOne));
+	Op->Outputs.Add(FNiagaraOpInOutInfo(Result, ENiagaraDataType::Vector, ResultText, ResultText, DefaultStr_VecOne));
+	Op->OpDelegate.BindStatic(&INiagaraCompiler::Dot);
+
+	Op = &OpInfos.Add(Cross);
+	Op->Name = Cross;
+	Op->FriendlyName = NSLOCTEXT("NiagaraOpInfo", "Cross Name", "Cross");
+	Op->Description = NSLOCTEXT("NiagaraOpInfo", "Cross Desc", "Cross");
+	Op->Inputs.Add(FNiagaraOpInOutInfo(A, ENiagaraDataType::Vector, AText, AText, DefaultStr_VecOne));
+	Op->Inputs.Add(FNiagaraOpInOutInfo(B, ENiagaraDataType::Vector, BText, BText, DefaultStr_VecOne));
+	Op->Outputs.Add(FNiagaraOpInOutInfo(Result, ENiagaraDataType::Vector, ResultText, ResultText, DefaultStr_VecOne));
+	Op->OpDelegate.BindStatic(&INiagaraCompiler::Cross);
+
+	Op = &OpInfos.Add(Normalize);
+	Op->Name = Normalize;
+	Op->FriendlyName = NSLOCTEXT("NiagaraOpInfo", "Normalize Name", "Normalize");
+	Op->Description = NSLOCTEXT("NiagaraOpInfo", "Normalize Desc", "Normalize");
+	Op->Inputs.Add(FNiagaraOpInOutInfo(A, ENiagaraDataType::Vector, AText, AText, DefaultStr_VecOne));
+	Op->Outputs.Add(FNiagaraOpInOutInfo(Result, ENiagaraDataType::Vector, ResultText, ResultText, DefaultStr_VecOne));
+	Op->OpDelegate.BindStatic(&INiagaraCompiler::Normalize);
+
+	Op = &OpInfos.Add(Random);
+	Op->Name = Random;
+	Op->FriendlyName = NSLOCTEXT("NiagaraOpInfo", "Random Name", "Random");
+	Op->Description = NSLOCTEXT("NiagaraOpInfo", "Random Desc", "Random");
+	Op->Inputs.Add(FNiagaraOpInOutInfo(A, ENiagaraDataType::Vector, AText, AText, DefaultStr_VecOne));
+	Op->Outputs.Add(FNiagaraOpInOutInfo(Result, ENiagaraDataType::Vector, ResultText, ResultText, DefaultStr_VecOne));
+	Op->OpDelegate.BindStatic(&INiagaraCompiler::Random);
+
+	Op = &OpInfos.Add(Length);
+	Op->Name = Length;
+	Op->FriendlyName = NSLOCTEXT("NiagaraOpInfo", "Length Name", "Length");
+	Op->Description = NSLOCTEXT("NiagaraOpInfo", "Length Desc", "Length");
+	Op->Inputs.Add(FNiagaraOpInOutInfo(A, ENiagaraDataType::Vector, AText, AText, DefaultStr_VecOne));
+	Op->Outputs.Add(FNiagaraOpInOutInfo(Result, ENiagaraDataType::Vector, ResultText, ResultText, DefaultStr_VecOne));
+	Op->OpDelegate.BindStatic(&INiagaraCompiler::Length);
+
+	Op = &OpInfos.Add(Sin4);
+	Op->Name = Sin4;
+	Op->FriendlyName = NSLOCTEXT("NiagaraOpInfo", "Sin4 Name", "Sin4");
+	Op->Description = NSLOCTEXT("NiagaraOpInfo", "Sin4 Desc", "Sin4");
+	Op->Inputs.Add(FNiagaraOpInOutInfo(A, ENiagaraDataType::Vector, AText, AText, DefaultStr_VecZero));
+	Op->Outputs.Add(FNiagaraOpInOutInfo(Result, ENiagaraDataType::Vector, ResultText, ResultText, DefaultStr_VecZero));
+	Op->OpDelegate.BindStatic(&INiagaraCompiler::Sin4);
+
+	Op = &OpInfos.Add(Noise);
+	Op->Name = Noise;
+	Op->FriendlyName = NSLOCTEXT("NiagaraOpInfo", "Noise Name", "Noise");
+	Op->Description = NSLOCTEXT("NiagaraOpInfo", "Noise Desc", "Noise");
+	Op->Inputs.Add(FNiagaraOpInOutInfo(A, ENiagaraDataType::Vector, AText, AText, DefaultStr_VecOne));
+	Op->Outputs.Add(FNiagaraOpInOutInfo(Result, ENiagaraDataType::Vector, ResultText, ResultText, DefaultStr_VecOne));
+	Op->OpDelegate.BindStatic(&INiagaraCompiler::Noise);
+
+	Op = &OpInfos.Add(Scalar);
+	Op->Name = Scalar;
+	Op->FriendlyName = NSLOCTEXT("NiagaraOpInfo", "Scalar Name", "Scalar");
+	Op->Description = NSLOCTEXT("NiagaraOpInfo", "Scalar Desc", "Result is a vector with X in each of it's components.");
+	Op->Inputs.Add(FNiagaraOpInOutInfo(X, ENiagaraDataType::Scalar, XText, XText, DefaultStr_One));
+	Op->Outputs.Add(FNiagaraOpInOutInfo(Result, ENiagaraDataType::Vector, ResultText, ResultText, DefaultStr_VecOne));
+	Op->OpDelegate.BindStatic(&INiagaraCompiler::Scalar);
+
+	Op = &OpInfos.Add(Vector);
+	Op->Name = Vector;
+	Op->FriendlyName = NSLOCTEXT("NiagaraOpInfo", "Vector Name", "Vector");
+	Op->Description = NSLOCTEXT("NiagaraOpInfo", "Vector Desc", "Builds a constant vector.");
+	Op->Inputs.Add(FNiagaraOpInOutInfo(A, ENiagaraDataType::Vector, AText, AText, DefaultStr_VecOne));
+	Op->Outputs.Add(FNiagaraOpInOutInfo(Result, ENiagaraDataType::Vector, ResultText, ResultText, DefaultStr_VecOne));
+	Op->OpDelegate.BindStatic(&INiagaraCompiler::Vector);
+
+	Op = &OpInfos.Add(Matrix);
+	FName Row0(TEXT("Row0"));	FText Row0Text = NSLOCTEXT("NiagaraOpInfo", "First Matrix Row", "Row 0");
+	FName Row1(TEXT("Row1"));	FText Row1Text = NSLOCTEXT("NiagaraOpInfo", "Second Matrix Row", "Row 1");
+	FName Row2(TEXT("Row2"));	FText Row2Text = NSLOCTEXT("NiagaraOpInfo", "Third Matrix Row", "Row 2");
+	FName Row3(TEXT("Row3"));	FText Row3Text = NSLOCTEXT("NiagaraOpInfo", "Fourth Matrix Row", "Row 3");
+	FString Identity = TEXT("	1.0,0.0,0.0,0.0,0.0,1.0,0.0,0.0,0.0,0.0,1.0,0.0,0.0,0.0,0.0,1.0");
+	Op->Name = Matrix;
+	Op->FriendlyName = NSLOCTEXT("NiagaraOpInfo", "Matrix Name", "Matrix");
+	Op->Description = NSLOCTEXT("NiagaraOpInfo", "Matrix Desc", "Matrix");
+	Op->Inputs.Add(FNiagaraOpInOutInfo(Row0, ENiagaraDataType::Vector, Row0Text, Row0Text, TEXT("1.0,0.0,0.0,0.0")));
+	Op->Inputs.Add(FNiagaraOpInOutInfo(Row1, ENiagaraDataType::Vector, Row1Text, Row1Text, TEXT("0.0,1.0,0.0,0.0")));
+	Op->Inputs.Add(FNiagaraOpInOutInfo(Row2, ENiagaraDataType::Vector, Row2Text, Row2Text, TEXT("0.0,0.0,1.0,0.0")));
+	Op->Inputs.Add(FNiagaraOpInOutInfo(Row3, ENiagaraDataType::Vector, Row3Text, Row3Text, TEXT("0.0,0.0,0.0,1.0")));
+	Op->Outputs.Add(FNiagaraOpInOutInfo(Result, ENiagaraDataType::Matrix, ResultText, ResultText, Identity));
+	Op->OpDelegate.BindStatic(&INiagaraCompiler::Matrix);
+
+	Op = &OpInfos.Add(Compose);
+	Op->Name = Compose;
+	Op->FriendlyName = NSLOCTEXT("NiagaraOpInfo", "Compose Name", "Compose Vector XYZW");
+	Op->Description = NSLOCTEXT("NiagaraOpInfo", "Compose Desc", "Result = {{ A.x, B.y, C.z, D.w }}");
+	Op->Inputs.Add(FNiagaraOpInOutInfo(A, ENiagaraDataType::Vector, AText, AText, DefaultStr_VecOne));
+	Op->Inputs.Add(FNiagaraOpInOutInfo(B, ENiagaraDataType::Vector, BText, BText, DefaultStr_VecOne));
+	Op->Inputs.Add(FNiagaraOpInOutInfo(C, ENiagaraDataType::Vector, CText, CText, DefaultStr_VecOne));
+	Op->Inputs.Add(FNiagaraOpInOutInfo(D, ENiagaraDataType::Vector, DText, DText, DefaultStr_VecOne));
+	Op->Outputs.Add(FNiagaraOpInOutInfo(Result, ENiagaraDataType::Vector, ResultText, ResultText, DefaultStr_VecOne));
+	Op->OpDelegate.BindStatic(&INiagaraCompiler::Compose);
+
+	Op = &OpInfos.Add(ComposeX);
+	Op->Name = ComposeX;
+	Op->FriendlyName = NSLOCTEXT("NiagaraOpInfo", "ComposeX Name", "Compose Vector X");
+	Op->Description = NSLOCTEXT("NiagaraOpInfo", "ComposeX Desc", "Result = {{ A.x, B.x, C.x, D.x }}");
+	Op->Inputs.Add(FNiagaraOpInOutInfo(A, ENiagaraDataType::Vector, AText, AText, DefaultStr_VecOne));
+	Op->Inputs.Add(FNiagaraOpInOutInfo(B, ENiagaraDataType::Vector, BText, BText, DefaultStr_VecOne));
+	Op->Inputs.Add(FNiagaraOpInOutInfo(C, ENiagaraDataType::Vector, CText, CText, DefaultStr_VecOne));
+	Op->Inputs.Add(FNiagaraOpInOutInfo(D, ENiagaraDataType::Vector, DText, DText, DefaultStr_VecOne));
+	Op->Outputs.Add(FNiagaraOpInOutInfo(Result, ENiagaraDataType::Vector, ResultText, ResultText, DefaultStr_VecOne));
+	Op->OpDelegate.BindStatic(&INiagaraCompiler::ComposeX);
+
+	Op = &OpInfos.Add(ComposeY);
+	Op->Name = ComposeY;
+	Op->FriendlyName = NSLOCTEXT("NiagaraOpInfo", "ComposeY Name", "Compose Vector Y");
+	Op->Description = NSLOCTEXT("NiagaraOpInfo", "ComposeY Desc", "Result = {{ A.y, B.y, C.y, D.y }}");
+	Op->Inputs.Add(FNiagaraOpInOutInfo(A, ENiagaraDataType::Vector, AText, AText, DefaultStr_VecOne));
+	Op->Inputs.Add(FNiagaraOpInOutInfo(B, ENiagaraDataType::Vector, BText, BText, DefaultStr_VecOne));
+	Op->Inputs.Add(FNiagaraOpInOutInfo(C, ENiagaraDataType::Vector, CText, CText, DefaultStr_VecOne));
+	Op->Inputs.Add(FNiagaraOpInOutInfo(D, ENiagaraDataType::Vector, DText, DText, DefaultStr_VecOne));
+	Op->Outputs.Add(FNiagaraOpInOutInfo(Result, ENiagaraDataType::Vector, ResultText, ResultText, DefaultStr_VecOne));
+	Op->OpDelegate.BindStatic(&INiagaraCompiler::ComposeY);
+
+	Op = &OpInfos.Add(ComposeZ);
+	Op->Name = ComposeZ;
+	Op->FriendlyName = NSLOCTEXT("NiagaraOpInfo", "ComposeZ Name", "Compose Vector Z");
+	Op->Description = NSLOCTEXT("NiagaraOpInfo", "ComposeZ Desc", "Result = {{ A.z, B.z, C.z, D.z }}");
+	Op->Inputs.Add(FNiagaraOpInOutInfo(A, ENiagaraDataType::Vector, AText, AText, DefaultStr_VecOne));
+	Op->Inputs.Add(FNiagaraOpInOutInfo(B, ENiagaraDataType::Vector, BText, BText, DefaultStr_VecOne));
+	Op->Inputs.Add(FNiagaraOpInOutInfo(C, ENiagaraDataType::Vector, CText, CText, DefaultStr_VecOne));
+	Op->Inputs.Add(FNiagaraOpInOutInfo(D, ENiagaraDataType::Vector, DText, DText, DefaultStr_VecOne));
+	Op->Outputs.Add(FNiagaraOpInOutInfo(Result, ENiagaraDataType::Vector, ResultText, ResultText, DefaultStr_VecOne));
+	Op->OpDelegate.BindStatic(&INiagaraCompiler::ComposeZ);
+
+	Op = &OpInfos.Add(ComposeW);
+	Op->Name = ComposeW;
+	Op->FriendlyName = NSLOCTEXT("NiagaraOpInfo", "ComposeW Name", "Compose Vector W");
+	Op->Description = NSLOCTEXT("NiagaraOpInfo", "ComposeW Desc", "Result = {{ A.w, B.w, C.w, D.w }}");
+	Op->Inputs.Add(FNiagaraOpInOutInfo(A, ENiagaraDataType::Vector, AText, AText, DefaultStr_VecOne));
+	Op->Inputs.Add(FNiagaraOpInOutInfo(B, ENiagaraDataType::Vector, BText, BText, DefaultStr_VecOne));
+	Op->Inputs.Add(FNiagaraOpInOutInfo(C, ENiagaraDataType::Vector, CText, CText, DefaultStr_VecOne));
+	Op->Inputs.Add(FNiagaraOpInOutInfo(D, ENiagaraDataType::Vector, DText, DText, DefaultStr_VecOne));
+	Op->Outputs.Add(FNiagaraOpInOutInfo(Result, ENiagaraDataType::Vector, ResultText, ResultText, DefaultStr_VecOne));
+	Op->OpDelegate.BindStatic(&INiagaraCompiler::ComposeW);
+
+	Op = &OpInfos.Add(Split);
+	Op->Name = Split;
+	Op->FriendlyName = NSLOCTEXT("NiagaraOpInfo", "Split Name", "Compose");
+	Op->Description = NSLOCTEXT("NiagaraOpInfo", "Split Desc", "X = A.xxxx, Y = A.yyyy, Z = A.zzzz, W = A.wwww");
+	Op->Inputs.Add(FNiagaraOpInOutInfo(A, ENiagaraDataType::Vector, AText, AText, DefaultStr_VecOne));
+	Op->Outputs.Add(FNiagaraOpInOutInfo(X, ENiagaraDataType::Vector, XText, XText, DefaultStr_VecOne));
+	Op->Outputs.Add(FNiagaraOpInOutInfo(Y, ENiagaraDataType::Vector, YText, YText, DefaultStr_VecOne));
+	Op->Outputs.Add(FNiagaraOpInOutInfo(Z, ENiagaraDataType::Vector, ZText, ZText, DefaultStr_VecOne));
+	Op->Outputs.Add(FNiagaraOpInOutInfo(W, ENiagaraDataType::Vector, WText, WText, DefaultStr_VecOne));
+	Op->OpDelegate.BindStatic(&INiagaraCompiler::Split);
+
+	Op = &OpInfos.Add(GetX);
+	Op->Name = GetX;
+	Op->FriendlyName = NSLOCTEXT("NiagaraOpInfo", "GetX Name", "GetX");
+	Op->Description = NSLOCTEXT("NiagaraOpInfo", "GetX Desc", "X = A.xxxx");
+	Op->Inputs.Add(FNiagaraOpInOutInfo(A, ENiagaraDataType::Vector, AText, AText, DefaultStr_VecOne));
+	Op->Outputs.Add(FNiagaraOpInOutInfo(X, ENiagaraDataType::Vector, XText, XText, DefaultStr_VecOne));
+	Op->OpDelegate.BindStatic(&INiagaraCompiler::GetX);
+
+	Op = &OpInfos.Add(GetY);
+	Op->Name = GetY;
+	Op->FriendlyName = NSLOCTEXT("NiagaraOpInfo", "GetY Name", "GetY");
+	Op->Description = NSLOCTEXT("NiagaraOpInfo", "GetY Desc", "Y = A.yyyy");
+	Op->Inputs.Add(FNiagaraOpInOutInfo(A, ENiagaraDataType::Vector, AText, AText, DefaultStr_VecOne));
+	Op->Outputs.Add(FNiagaraOpInOutInfo(Y, ENiagaraDataType::Vector, YText, YText, DefaultStr_VecOne));
+	Op->OpDelegate.BindStatic(&INiagaraCompiler::GetY);
+
+	Op = &OpInfos.Add(GetZ);
+	Op->Name = GetZ;
+	Op->FriendlyName = NSLOCTEXT("NiagaraOpInfo", "GetZ Name", "GetZ");
+	Op->Description = NSLOCTEXT("NiagaraOpInfo", "GetZ Desc", "Z = A.zzzz");
+	Op->Inputs.Add(FNiagaraOpInOutInfo(A, ENiagaraDataType::Vector, AText, AText, DefaultStr_VecOne));
+	Op->Outputs.Add(FNiagaraOpInOutInfo(Z, ENiagaraDataType::Vector, ZText, ZText, DefaultStr_VecOne));
+	Op->OpDelegate.BindStatic(&INiagaraCompiler::GetZ);
+
+	Op = &OpInfos.Add(GetW);
+	Op->Name = GetW;
+	Op->FriendlyName = NSLOCTEXT("NiagaraOpInfo", "GetW Name", "GetW");
+	Op->Description = NSLOCTEXT("NiagaraOpInfo", "GetW Desc", "W = A.wwww");
+	Op->Inputs.Add(FNiagaraOpInOutInfo(A, ENiagaraDataType::Vector, AText, AText, DefaultStr_VecOne));
+	Op->Outputs.Add(FNiagaraOpInOutInfo(W, ENiagaraDataType::Vector, WText, WText, DefaultStr_VecOne));
+	Op->OpDelegate.BindStatic(&INiagaraCompiler::GetW);
+
+	Op = &OpInfos.Add(TransformVector);
+	Op->Name = TransformVector;
+	Op->FriendlyName = NSLOCTEXT("NiagaraOpInfo", "TransformVector Name", "Transform Vector");
+	Op->Description = NSLOCTEXT("NiagaraOpInfo", "TransformVector Desc", "Transforms vector B by matrix A");
+	Op->Inputs.Add(FNiagaraOpInOutInfo(A, ENiagaraDataType::Matrix, AText, AText, DefaultStr_VecOne));
+	Op->Inputs.Add(FNiagaraOpInOutInfo(B, ENiagaraDataType::Vector, BText, BText, DefaultStr_VecOne));
+	Op->Outputs.Add(FNiagaraOpInOutInfo(Result, ENiagaraDataType::Vector, ResultText, ResultText, DefaultStr_VecOne));
+	Op->OpDelegate.BindStatic(&INiagaraCompiler::TransformVector);
+
+	Op = &OpInfos.Add(Transpose);
+	Op->Name = Transpose;
+	Op->FriendlyName = NSLOCTEXT("NiagaraOpInfo", "Transpose Name", "Transpose");
+	Op->Description = NSLOCTEXT("NiagaraOpInfo", "Transpose Desc", "Result = Transpose of matrix A");
+	Op->Inputs.Add(FNiagaraOpInOutInfo(A, ENiagaraDataType::Matrix, AText, AText, DefaultStr_VecOne));
+	Op->Outputs.Add(FNiagaraOpInOutInfo(Result, ENiagaraDataType::Matrix, ResultText, ResultText, DefaultStr_VecOne));
+	Op->OpDelegate.BindStatic(&INiagaraCompiler::Transpose);
+
+	Op = &OpInfos.Add(Inverse);
+	Op->Name = Inverse;
+	Op->FriendlyName = NSLOCTEXT("NiagaraOpInfo", "Inverse Name", "Inverse");
+	Op->Description = NSLOCTEXT("NiagaraOpInfo", "Inverse Desc", "Result = Inverse of matrix A");
+	Op->Inputs.Add(FNiagaraOpInOutInfo(A, ENiagaraDataType::Matrix, AText, AText, DefaultStr_VecOne));
+	Op->Outputs.Add(FNiagaraOpInOutInfo(Result, ENiagaraDataType::Matrix, ResultText, ResultText, DefaultStr_VecOne));
+	Op->OpDelegate.BindStatic(&INiagaraCompiler::Inverse);
+
+#define NiagaraOp(OpName) if (GetOpInfo(OpName) == NULL)\
+	{\
+	UE_LOG(LogNiagaraCompiler, Fatal, TEXT("Couldn't find info for \"%s\". Have you added it to the NiagaraOpList but not provided the extra information in FNiagaraOpInfo::Init() ?"), TEXT(#OpName)); \
+	}
+	NiagaraOpList;
+#undef NiagaraOp
+
+#define NiagaraOp(OpName) (OpName == OpInfo.Key) || 
+
+	for (auto& OpInfo : OpInfos)
+	{
+		if (!(
+			NiagaraOpList
+			false)
+			)
+		{
+			UE_LOG(LogNiagaraCompiler, Fatal, TEXT("Found info about \"%s\" in FNiagaraOpInfo::Init() that did not have an entry in the NiagaraOpList. You must add it to both."), *OpInfo.Key.ToString());
+		}
+	}
+}
+
+const FNiagaraOpInfo* FNiagaraOpInfo::GetOpInfo(FName Name)
+{
+	if (GetOpInfoMap().Num() == 0)
+	{
+		Init();
+	}
+	return GetOpInfoMap().Find(Name);
+}
+
+TMap<FName, FNiagaraOpInfo>& FNiagaraOpInfo::GetOpInfoMap()
+{
+	static TMap<FName, FNiagaraOpInfo> OpInfos;
+
+	return OpInfos;
+}
+
+//////////////////////////////////////////////////////////////////////////
+
+//Debugging variant that checks inputs and outputs match expectations of FNiagaraOpInfo.
+#if DO_CHECK
+#define NiagaraOp(OpName) void INiagaraCompiler::##OpName(INiagaraCompiler* Compiler, TArray<TNiagaraExprPtr>& InputExpressions, TArray<TNiagaraExprPtr>& OutputExpressions)\
+{\
+	Compiler->CheckInputs(FNiagaraOpInfo::##OpName##, InputExpressions); \
+	return Compiler->OpName##_Internal(InputExpressions, OutputExpressions); \
+	Compiler->CheckOutputs(FNiagaraOpInfo::##OpName##, OutputExpressions); \
+}
+NiagaraOpList
+#undef NiagaraOp
+
+#else//DO_CHECK
+
+#define NiagaraOp(OpName) void INiagaraCompiler::##OpName(INiagaraCompiler* Compiler, TArray<TNiagaraExprPtr>& InputExpressions, TArray<TNiagaraExprPtr>& OutputExpressions)\
+{ \
+	return Compiler->OpName##_Internal(InputExpressions, OutputExpressions); \
+}
+NiagaraOpList
+#undef NiagaraOp
+
+#endif
+
diff --git a/Engine/Source/Editor/UnrealEd/Private/Niagara/NiagaraNodeConstant.cpp b/Engine/Source/Editor/UnrealEd/Private/Niagara/NiagaraNodeConstant.cpp
new file mode 100644
--- /dev/null
+++ b/Engine/Source/Editor/UnrealEd/Private/Niagara/NiagaraNodeConstant.cpp
@@ -0,0 +1,146 @@
+// Copyright 1998-2014 Epic Games, Inc. All Rights Reserved.
+
+#include "UnrealEd.h"
+#include "INiagaraCompiler.h"
+
+#define LOCTEXT_NAMESPACE "NiagaraNodeConst"
+
+UNiagaraNodeConstant::UNiagaraNodeConstant(const FObjectInitializer& ObjectInitializer)
+	: Super(ObjectInitializer)
+{
+	ConstName = NAME_None;
+	DataType = ENiagaraDataType::Vector;
+	bNeedsDefault = true;
+}
+
+void UNiagaraNodeConstant::AllocateDefaultPins()
+{
+	const UEdGraphSchema_Niagara* Schema = GetDefault<UEdGraphSchema_Niagara>();
+
+	if (DataType == ENiagaraDataType::Scalar)
+	{
+		if (bNeedsDefault)
+		{
+			CreatePin(EGPD_Input, Schema->PC_Float, TEXT(""), NULL, false, false, LOCTEXT("FloatPinName", "Default").ToString(), true);
+		}
+		CreatePin(EGPD_Output, Schema->PC_Vector, TEXT(""), NULL, false, false, ConstName.ToString());
+	}
+	else if (DataType == ENiagaraDataType::Vector)
+	{
+		if (bNeedsDefault)
+		{
+			CreatePin(EGPD_Input, Schema->PC_Vector, TEXT(""), NULL, false, false, LOCTEXT("VectorPinName", "Default").ToString(), true);
+		}
+		CreatePin(EGPD_Output, Schema->PC_Vector, TEXT(""), NULL, false, false, ConstName.ToString());
+	}
+	else if (DataType == ENiagaraDataType::Matrix)
+	{
+		if (bNeedsDefault)
+		{
+			CreatePin(EGPD_Input, Schema->PC_Vector, TEXT(""), NULL, false, false, LOCTEXT("MatrixPinName0", "Default Row0").ToString(), true);
+			CreatePin(EGPD_Input, Schema->PC_Vector, TEXT(""), NULL, false, false, LOCTEXT("MatrixPinName1", "Default Row1").ToString(), true);
+			CreatePin(EGPD_Input, Schema->PC_Vector, TEXT(""), NULL, false, false, LOCTEXT("MatrixPinName2", "Default Row2").ToString(), true);
+			CreatePin(EGPD_Input, Schema->PC_Vector, TEXT(""), NULL, false, false, LOCTEXT("MatrixPinName3", "Default Row3").ToString(), true);
+		}
+		CreatePin(EGPD_Output, Schema->PC_Matrix, TEXT(""), NULL, false, false, ConstName.ToString());
+	}
+}
+
+FText UNiagaraNodeConstant::GetNodeTitle(ENodeTitleType::Type TitleType) const
+{
+	if (ConstName != NAME_None)
+		return FText::FromName(ConstName);
+	else
+		return LOCTEXT("DefaultTitle", "External Constant");
+}
+
+FLinearColor UNiagaraNodeConstant::GetNodeTitleColor() const
+{
+	const UEdGraphSchema_Niagara* Schema = GetDefault<UEdGraphSchema_Niagara>();
+
+	if (Pins.Num() > 0 && Pins[0] != NULL)
+	{
+		return Schema->GetPinTypeColor(Pins[0]->PinType);
+	}
+	else
+	{
+		return Super::GetNodeTitleColor();
+	}
+}
+
+void UNiagaraNodeConstant::Compile(class INiagaraCompiler* Compiler, TArray<FNiagaraNodeResult>& Outputs)
+{
+	TArray<int32> OutputExpressions;
+	if (DataType == ENiagaraDataType::Scalar)
+	{
+		float Default = 1.0f;
+		UEdGraphPin* OutPin;
+		if (bNeedsDefault)
+		{
+			CastChecked<UEdGraphSchema_Niagara>(GetSchema())->GetPinDefaultValue(Pins[0], Default);
+			OutPin = Pins[1];
+		}
+		else
+		{
+			OutPin = Pins[0];
+		}
+		Outputs.Add(FNiagaraNodeResult(Compiler->GetExternalConstant(ConstName, Default), OutPin));
+	}
+	else if (DataType == ENiagaraDataType::Vector)
+	{
+		FVector4 Default = FVector4(1.0f, 1.0f, 1.0f, 1.0f);
+		UEdGraphPin* OutPin;
+		if (bNeedsDefault)
+		{
+			CastChecked<UEdGraphSchema_Niagara>(GetSchema())->GetPinDefaultValue(Pins[0], Default);
+			OutPin = Pins[1];
+		}
+		else
+		{
+			OutPin = Pins[0];
+		}
+		Outputs.Add(FNiagaraNodeResult(Compiler->GetExternalConstant(ConstName, Default), OutPin));
+	}
+	else if (DataType == ENiagaraDataType::Matrix)
+	{
+		FVector4 Default0 = FVector4(1.0f, 0.0f, 0.0f, 0.0f);
+		FVector4 Default1 = FVector4(0.0f, 1.0f, 0.0f, 0.0f);
+		FVector4 Default2 = FVector4(0.0f, 0.0f, 1.0f, 0.0f);
+		FVector4 Default3 = FVector4(0.0f, 0.0f, 0.0f, 1.0f);
+
+		UEdGraphPin* OutPin;
+		if (bNeedsDefault)
+		{
+			CastChecked<UEdGraphSchema_Niagara>(GetSchema())->GetPinDefaultValue(Pins[0], Default0);
+			CastChecked<UEdGraphSchema_Niagara>(GetSchema())->GetPinDefaultValue(Pins[1], Default1);
+			CastChecked<UEdGraphSchema_Niagara>(GetSchema())->GetPinDefaultValue(Pins[2], Default2);
+			CastChecked<UEdGraphSchema_Niagara>(GetSchema())->GetPinDefaultValue(Pins[3], Default3);
+			OutPin = Pins[4];
+		}
+		else
+		{
+			OutPin = Pins[0];
+		}
+		FMatrix Default;
+		Default.M[0][0] = Default0[0]; Default.M[0][1] = Default0[1]; Default.M[0][2] = Default0[2]; Default.M[0][3] = Default0[3];
+		Default.M[1][0] = Default1[0]; Default.M[1][1] = Default1[1]; Default.M[1][2] = Default1[2]; Default.M[1][3] = Default1[3];
+		Default.M[2][0] = Default2[0]; Default.M[2][1] = Default2[1]; Default.M[2][2] = Default2[2]; Default.M[2][3] = Default2[3];
+		Default.M[3][0] = Default3[0]; Default.M[3][1] = Default3[1]; Default.M[3][2] = Default3[2]; Default.M[3][3] = Default3[3];
+		Outputs.Add(FNiagaraNodeResult(Compiler->GetExternalConstant(ConstName, Default), OutPin));
+	}
+}
+
+void UNiagaraNodeConstant::PostEditChangeProperty(struct FPropertyChangedEvent& PropertyChangedEvent)
+{
+	static FName DataTypePropertyName(TEXT("DataType"));
+	if (PropertyChangedEvent.Property && PropertyChangedEvent.Property->GetFName() == DataTypePropertyName)
+	{
+		for (auto Pin : Pins)
+		{
+			Pin->BreakAllPinLinks();
+		}
+		Pins.Empty();
+		AllocateDefaultPins();
+	}
+}
+#undef LOCTEXT_NAMESPACE
\ No newline at end of file
diff --git a/Engine/Source/Editor/UnrealEd/Private/Niagara/NiagaraNodeGetAttr.cpp b/Engine/Source/Editor/UnrealEd/Private/Niagara/NiagaraNodeGetAttr.cpp
--- a/Engine/Source/Editor/UnrealEd/Private/Niagara/NiagaraNodeGetAttr.cpp
+++ b/Engine/Source/Editor/UnrealEd/Private/Niagara/NiagaraNodeGetAttr.cpp
@@ -1,6 +1,8 @@
 // Copyright 1998-2014 Epic Games, Inc. All Rights Reserved.
 
 #include "UnrealEd.h"
+#include "INiagaraCompiler.h"
+
 
 UNiagaraNodeGetAttr::UNiagaraNodeGetAttr(const FObjectInitializer& ObjectInitializer)
 	: Super(ObjectInitializer)
@@ -11,8 +13,7 @@ void UNiagaraNodeGetAttr::AllocateDefaultPins()
 {
 	const UEdGraphSchema_Niagara* Schema = GetDefault<UEdGraphSchema_Niagara>();
 
-	UScriptStruct* VectorStruct = FindObjectChecked<UScriptStruct>(UObject::StaticClass(), TEXT("Vector"));
-	CreatePin(EGPD_Output, Schema->PC_Struct, TEXT(""), VectorStruct, false, false, AttrName.ToString());
+	CreatePin(EGPD_Output, Schema->PC_Vector, TEXT(""), NULL, false, false, AttrName.ToString());
 }
 
 FText UNiagaraNodeGetAttr::GetNodeTitle(ENodeTitleType::Type TitleType) const
@@ -35,3 +36,9 @@ FLinearColor UNiagaraNodeGetAttr::GetNodeTitleColor() const
 		return Super::GetNodeTitleColor();
 	}
 }
+
+void UNiagaraNodeGetAttr::Compile(class INiagaraCompiler* Compiler, TArray<FNiagaraNodeResult>& Outputs)
+{
+	Outputs.Add(FNiagaraNodeResult(Compiler->GetAttribute(AttrName), Pins[0]));
+}
+
diff --git a/Engine/Source/Editor/UnrealEd/Private/Niagara/NiagaraNodeOp.cpp b/Engine/Source/Editor/UnrealEd/Private/Niagara/NiagaraNodeOp.cpp
--- a/Engine/Source/Editor/UnrealEd/Private/Niagara/NiagaraNodeOp.cpp
+++ b/Engine/Source/Editor/UnrealEd/Private/Niagara/NiagaraNodeOp.cpp
@@ -3,48 +3,93 @@
 #include "UnrealEd.h"
 #include "BlueprintGraphDefinitions.h"
 #include "GraphEditorSettings.h"
-#include "VectorVM.h"
 
 UNiagaraNodeOp::UNiagaraNodeOp(const FObjectInitializer& ObjectInitializer)
 	: Super(ObjectInitializer)
 {
 }
 
-const TCHAR* UNiagaraNodeOp::InPinNames[3] =
-{
-	TEXT("A"),
-	TEXT("B"),
-	TEXT("C")
-};
-
 void UNiagaraNodeOp::AllocateDefaultPins()
 {
 	const UEdGraphSchema_Niagara* Schema = GetDefault<UEdGraphSchema_Niagara>();
-	VectorVM::FVectorVMOpInfo const& Info = VectorVM::GetOpCodeInfo(OpIndex);
 
-	UScriptStruct* VectorStruct = FindObjectChecked<UScriptStruct>(UObject::StaticClass(), TEXT("Vector"));
-	for (int32 SrcIndex = 0; SrcIndex < 3; ++SrcIndex)
+	const FNiagaraOpInfo* OpInfo = FNiagaraOpInfo::GetOpInfo(OpName);
+	check(OpInfo);
+
+	for (int32 SrcIndex = 0; SrcIndex < OpInfo->Inputs.Num(); ++SrcIndex)
 	{
-		VectorVM::EOpSrc::Type PinType = Info.SrcTypes[SrcIndex];
-		if (PinType != VectorVM::EOpSrc::Invalid)
+		const FNiagaraOpInOutInfo& InOutInfo = OpInfo->Inputs[SrcIndex];
+		if (InOutInfo.DataType == ENiagaraDataType::Scalar)
+		{
+			UEdGraphPin* Pin = CreatePin(EGPD_Input, Schema->PC_Float, TEXT(""), NULL, false, false, InOutInfo.FriendlyName.ToString());
+			check(Pin);
+			Pin->bDefaultValueIsIgnored = false;
+			Pin->bDefaultValueIsReadOnly = false;
+			Pin->bNotConnectable = false;
+			Pin->DefaultValue = InOutInfo.Default;
+			Pin->AutogeneratedDefaultValue = InOutInfo.Default;
+			Pin->PinToolTip = InOutInfo.Description.ToString();
+		}
+		else if (InOutInfo.DataType == ENiagaraDataType::Vector)
 		{
-			UEdGraphPin* Pin = CreatePin(EGPD_Input, Schema->PC_Struct, TEXT(""), VectorStruct, false, false, InPinNames[SrcIndex]);
+			UEdGraphPin* Pin = CreatePin(EGPD_Input, Schema->PC_Vector, TEXT(""), NULL, false, false, InOutInfo.FriendlyName.ToString());
 			check(Pin);
 			Pin->bDefaultValueIsIgnored = false;
 			Pin->bDefaultValueIsReadOnly = false;
 			Pin->bNotConnectable = false;
-			Pin->DefaultValue = TEXT("1.0,1.0,1.0,1.0");
-			Pin->AutogeneratedDefaultValue = TEXT("1.0,1.0,1.0,1.0");
+			Pin->DefaultValue = InOutInfo.Default;
+			Pin->AutogeneratedDefaultValue = InOutInfo.Default;
+			Pin->PinToolTip = InOutInfo.Description.ToString();
+		}
+		else if (InOutInfo.DataType == ENiagaraDataType::Matrix)
+		{
+			UEdGraphPin* Pin = CreatePin(EGPD_Input, Schema->PC_Matrix, TEXT(""), NULL, false, false, InOutInfo.FriendlyName.ToString());
+			check(Pin);
+			Pin->bDefaultValueIsIgnored = false;
+			Pin->bDefaultValueIsReadOnly = false;
+			Pin->bNotConnectable = false;
+			Pin->DefaultValue = InOutInfo.Default;
+			Pin->AutogeneratedDefaultValue = InOutInfo.Default;
+			Pin->PinToolTip = InOutInfo.Description.ToString();
+		}
+	}
+	
+	for (int32 OutIdx = 0; OutIdx < OpInfo->Outputs.Num(); ++OutIdx)
+	{
+		const FNiagaraOpInOutInfo& InOutInfo = OpInfo->Outputs[OutIdx];
+		if (InOutInfo.DataType == ENiagaraDataType::Scalar)
+		{
+			UEdGraphPin* Pin = CreatePin(EGPD_Output, Schema->PC_Float, TEXT(""), NULL, false, false, InOutInfo.FriendlyName.ToString());
+			check(Pin);
+			Pin->PinToolTip = InOutInfo.Description.ToString();
+		}
+		else if (InOutInfo.DataType == ENiagaraDataType::Vector)
+		{
+			UEdGraphPin* Pin = CreatePin(EGPD_Output, Schema->PC_Vector, TEXT(""), NULL, false, false, InOutInfo.FriendlyName.ToString());
+			check(Pin);
+			Pin->PinToolTip = InOutInfo.Description.ToString();
+		}
+		else if (InOutInfo.DataType == ENiagaraDataType::Matrix)
+		{
+			UEdGraphPin* Pin = CreatePin(EGPD_Output, Schema->PC_Matrix, TEXT(""), NULL, false, false, InOutInfo.FriendlyName.ToString());
+			check(Pin);
+			Pin->PinToolTip = InOutInfo.Description.ToString();
 		}
 	}
-
-	CreatePin(EGPD_Output, Schema->PC_Struct, TEXT(""), VectorStruct, false, false, TEXT("Result"));
 }
 
 FText UNiagaraNodeOp::GetNodeTitle(ENodeTitleType::Type TitleType) const
 {
-	VectorVM::FVectorVMOpInfo const& Info = VectorVM::GetOpCodeInfo(OpIndex);
-	return FText::FromString(Info.FriendlyName);
+	const FNiagaraOpInfo* OpInfo = FNiagaraOpInfo::GetOpInfo(OpName);
+	check(OpInfo);
+	return OpInfo->FriendlyName;
+}
+
+FText UNiagaraNodeOp::GetTooltipText()const
+{
+	const FNiagaraOpInfo* OpInfo = FNiagaraOpInfo::GetOpInfo(OpName);
+	check(OpInfo);
+	return OpInfo->Description;
 }
 
 FLinearColor UNiagaraNodeOp::GetNodeTitleColor() const
@@ -52,4 +97,30 @@ FLinearColor UNiagaraNodeOp::GetNodeTitleColor() const
 	return GetDefault<UGraphEditorSettings>()->FunctionCallNodeTitleColor;
 }
 
+void UNiagaraNodeOp::Compile(class INiagaraCompiler* Compiler, TArray<FNiagaraNodeResult>& Outputs)
+{
+	const FNiagaraOpInfo* OpInfo = FNiagaraOpInfo::GetOpInfo(OpName);
+	check(OpInfo);
+	int32 NumInputs = OpInfo->Inputs.Num();
+	int32 NumOutputs = OpInfo->Outputs.Num();
 
+	TArray<TNiagaraExprPtr> Inputs;
+	for (int32 i = 0; i < NumInputs; ++i)
+	{
+		UEdGraphPin *Pin = Pins[i];
+		check(Pin->Direction == EGPD_Input);
+		Inputs.Add(Compiler->CompilePin(Pin));
+	}
+	
+	TArray<TNiagaraExprPtr> OutputExpressions;
+	OpInfo->OpDelegate.Execute(Compiler, Inputs, OutputExpressions);
+
+	check( OutputExpressions.Num() == OpInfo->Outputs.Num() );
+
+	for (int32 i = 0; i < NumOutputs; ++i)
+	{
+		UEdGraphPin *Pin = Pins[NumInputs + i];
+		check(Pin->Direction == EGPD_Output);
+		Outputs.Add(FNiagaraNodeResult(OutputExpressions[i], Pin));
+	}
+}
\ No newline at end of file
diff --git a/Engine/Source/Editor/UnrealEd/Private/Niagara/NiagaraNodeOutputUpdate.cpp b/Engine/Source/Editor/UnrealEd/Private/Niagara/NiagaraNodeOutputUpdate.cpp
--- a/Engine/Source/Editor/UnrealEd/Private/Niagara/NiagaraNodeOutputUpdate.cpp
+++ b/Engine/Source/Editor/UnrealEd/Private/Niagara/NiagaraNodeOutputUpdate.cpp
@@ -9,7 +9,6 @@ UNiagaraNodeOutputUpdate::UNiagaraNodeOutputUpdate(const FObjectInitializer& Obj
 {
 }
 
-
 void UNiagaraNodeOutputUpdate::AllocateDefaultPins()
 {
 	const UEdGraphSchema_Niagara* Schema = GetDefault<UEdGraphSchema_Niagara>();
@@ -19,15 +18,13 @@ void UNiagaraNodeOutputUpdate::AllocateDefaultPins()
 	TArray<FName> OutputNames;
 	Source->GetParticleAttributes(OutputNames);
 
-	UScriptStruct* VectorStruct = FindObjectChecked<UScriptStruct>(UObject::StaticClass(), TEXT("Vector"));
 	for (int32 i = 0; i < OutputNames.Num(); i++)
 	{
-		UEdGraphPin* Pin = CreatePin(EGPD_Input, Schema->PC_Struct, TEXT(""), VectorStruct, false, false, OutputNames[i].ToString());
+		UEdGraphPin* Pin = CreatePin(EGPD_Input, Schema->PC_Vector, TEXT(""), NULL, false, false, OutputNames[i].ToString());
 		Pin->bDefaultValueIsIgnored = true;
 	}
 }
 
-
 FText UNiagaraNodeOutputUpdate::GetNodeTitle(ENodeTitleType::Type TitleType) const
 {
 	return NSLOCTEXT("NiagaraNode", "Output", "Output");
@@ -37,3 +34,35 @@ FLinearColor UNiagaraNodeOutputUpdate::GetNodeTitleColor() const
 {
 	return GetDefault<UGraphEditorSettings>()->FunctionTerminatorNodeTitleColor;
 }
+
+void UNiagaraNodeOutputUpdate::Compile(class INiagaraCompiler* Compiler, TArray<FNiagaraNodeResult>& Outputs)
+{
+	TArray<FName> OutputNames;
+	Compiler->GetParticleAttributes(OutputNames);
+	int32 NumOutputs = OutputNames.Num();
+
+	//No output pins so we have 1 input pin per attribute.
+	check(OutputNames.Num() == Pins.Num());
+
+	TArray<TNiagaraExprPtr> InputExpressions;
+	for (int32 i = 0; i < NumOutputs; ++i)
+	{
+		UEdGraphPin* Pin = Pins[i];
+		check( Pin->Direction == EGPD_Input );
+
+		TNiagaraExprPtr Output = Compiler->CompilePin(Pin);
+		if (!Output.IsValid())
+		{
+			//The pin was not connected so pass through the initial value of the attribute.
+			Output = Compiler->GetAttribute(OutputNames[i]);
+		}
+
+		InputExpressions.Add(Output);
+	}
+
+	check(NumOutputs == InputExpressions.Num());
+	for (int32 i = 0; i < NumOutputs; ++i)
+	{
+		Compiler->Output(OutputNames[i], InputExpressions[i]);		
+	}
+}
\ No newline at end of file
diff --git a/Engine/Source/Editor/UnrealEd/Private/Niagara/NiagaraScriptSource.cpp b/Engine/Source/Editor/UnrealEd/Private/Niagara/NiagaraScriptSource.cpp
--- a/Engine/Source/Editor/UnrealEd/Private/Niagara/NiagaraScriptSource.cpp
+++ b/Engine/Source/Editor/UnrealEd/Private/Niagara/NiagaraScriptSource.cpp
@@ -53,11 +53,14 @@ void UNiagaraScriptSource::GetParticleAttributes(TArray<FName>& VectorOutputs)
 	VectorOutputs.Add(FName(TEXT("Particle Age")));
 }
 
-void UNiagaraScriptSource::GetEmitterAttributes(TArray<FName>& VectorInputs)
+void UNiagaraScriptSource::GetEmitterAttributes(TArray<FName>& VectorInputs, TArray<FName>& MatrixInputs)
 {
 	for (uint32 i=0; i < NiagaraConstants::NumBuiltinConstants; i++)
 	{
 		VectorInputs.Add(NiagaraConstants::ConstantNames[i]);
 	}
+
+	MatrixInputs.Empty();
+	MatrixInputs.Add(FName(TEXT("Emitter Transform")));
 }
 
diff --git a/Engine/Source/Editor/UnrealEd/Public/INiagaraCompiler.h b/Engine/Source/Editor/UnrealEd/Public/INiagaraCompiler.h
new file mode 100644
--- /dev/null
+++ b/Engine/Source/Editor/UnrealEd/Public/INiagaraCompiler.h
@@ -0,0 +1,48 @@
+// Copyright 1998-2014 Epic Games, Inc. All Rights Reserved.
+
+#pragma once
+
+#include "NiagaraEditorCommon.h"
+
+//Interface for Niagara compilers.
+class UNREALED_API INiagaraCompiler
+{
+public:
+
+	/** Compiles a script. */
+	virtual void CompileScript(class UNiagaraScript* InScript) = 0;
+
+	/** Traverses the pins input, recursively compiling them and then adds any expressions needed for the owner node.*/
+	virtual TNiagaraExprPtr CompilePin(UEdGraphPin* Pin) = 0;
+
+	/** Fills OutAttributes with the names of all attributes for this emitter. */
+	virtual void GetParticleAttributes(TArray<FName>& OutAttributes) = 0;
+	
+	/** Creates an expression which fetches the named attribute. */
+	virtual TNiagaraExprPtr GetAttribute(FName AttributeName) = 0;
+	/** Creates an expression which fetches the named constant. */
+	virtual TNiagaraExprPtr GetExternalConstant(FName ConstantName, float Default) = 0;
+	/** Creates an expression which fetches the named constant. */
+	virtual TNiagaraExprPtr GetExternalConstant(FName ConstantName, FVector4 Default) = 0;
+	/** Creates an expression which fetches the named constant. */
+	virtual TNiagaraExprPtr GetExternalConstant(FName ConstantName, const FMatrix& Default) = 0;
+	/** Creates an expression which writes the result of SourceExpression to the named attribute. */
+	virtual TNiagaraExprPtr Output(FName OutputName, TNiagaraExprPtr& SourceExpression) = 0;
+
+	/** Checks the validity of the input expressions against the informations about this operation in FNiagaraOpInfo. */
+	virtual void CheckInputs(FName OpName, TArray<TNiagaraExprPtr>& Inputs) = 0;
+	/** Checks the validity of the output expressions against the informations about this operation in FNiagaraOpInfo. */
+	virtual void CheckOutputs(FName OpName, TArray<TNiagaraExprPtr>& Outputs) = 0;
+
+	//Static function named after each op that allows access to internals via FNiagaraOpInfo.Func.
+#define NiagaraOp(OpName) static void OpName(INiagaraCompiler* Compiler, TArray<TNiagaraExprPtr>& InputExpressions, TArray<TNiagaraExprPtr>& OutputExpressions);
+	NiagaraOpList;
+#undef NiagaraOp
+
+protected:
+	//Internal virtual function for each op. 
+#define NiagaraOp(OpName) virtual void OpName##_Internal(TArray<TNiagaraExprPtr>& InputExpressions, TArray<TNiagaraExprPtr>& OutputExpressions) = 0;
+	NiagaraOpList;
+#undef NiagaraOp
+
+};
\ No newline at end of file
diff --git a/Engine/Source/Editor/UnrealEd/Public/NiagaraEditorCommon.h b/Engine/Source/Editor/UnrealEd/Public/NiagaraEditorCommon.h
new file mode 100644
--- /dev/null
+++ b/Engine/Source/Editor/UnrealEd/Public/NiagaraEditorCommon.h
@@ -0,0 +1,198 @@
+// Copyright 1998-2014 Epic Games, Inc. All Rights Reserved.
+
+#pragma once
+
+#include "Niagara/NiagaraCommon.h"
+
+//List of operations supported by the Niagara compiler(s) as visible from the outside.
+//Additional information about each operation is added in FNiagaraOpInfo::Init
+#define NiagaraOpList\
+	NiagaraOp(Add)\
+	NiagaraOp(Subtract)\
+	NiagaraOp(Multiply)\
+	NiagaraOp(MultiplyAdd)\
+	NiagaraOp(Lerp)\
+	NiagaraOp(Reciprocal)\
+	NiagaraOp(ReciprocalSqrt)\
+	NiagaraOp(Sqrt)\
+	NiagaraOp(Negate)\
+	NiagaraOp(Abs)\
+	NiagaraOp(Exp)\
+	NiagaraOp(Exp2)\
+	NiagaraOp(Log)\
+	NiagaraOp(LogBase2)\
+	NiagaraOp(Sin)\
+	NiagaraOp(Cos)\
+	NiagaraOp(Tan)\
+	NiagaraOp(ASin)\
+	NiagaraOp(ACos)\
+	NiagaraOp(ATan)\
+	NiagaraOp(ATan2)\
+	NiagaraOp(Ceil)\
+	NiagaraOp(Floor)\
+	NiagaraOp(Modulo)\
+	NiagaraOp(Fractional)\
+	NiagaraOp(Truncate)\
+	NiagaraOp(Clamp)\
+	NiagaraOp(Min)\
+	NiagaraOp(Max)\
+	NiagaraOp(Pow)\
+	NiagaraOp(Sign)\
+	NiagaraOp(Step)\
+	NiagaraOp(Dot)\
+	NiagaraOp(Cross)\
+	NiagaraOp(Normalize)\
+	NiagaraOp(Random)\
+	NiagaraOp(Length)\
+	NiagaraOp(Sin4)\
+	NiagaraOp(Noise)\
+	NiagaraOp(Scalar)\
+	NiagaraOp(Vector)\
+	NiagaraOp(Matrix)\
+	NiagaraOp(Compose)\
+	NiagaraOp(ComposeX)\
+	NiagaraOp(ComposeY)\
+	NiagaraOp(ComposeZ)\
+	NiagaraOp(ComposeW)\
+	NiagaraOp(Split)\
+	NiagaraOp(GetX)\
+	NiagaraOp(GetY)\
+	NiagaraOp(GetZ)\
+	NiagaraOp(GetW)\
+	NiagaraOp(TransformVector)\
+	NiagaraOp(Transpose)\
+	NiagaraOp(Inverse)
+
+enum class ENiagaraExpressionResultLocation
+{
+	InputData,
+	OutputData,
+	Constants,
+	Temporaries,
+	Unknown,
+};
+
+typedef TSharedPtr<class FNiagaraExpression> TNiagaraExprPtr;
+
+/** Base class for all Niagara Expresions. Intermediate data collected from the node graph used for compiling Niagara scripts. */
+class FNiagaraExpression : public TSharedFromThis<FNiagaraExpression>
+{
+public:
+	FNiagaraExpression(class FNiagaraCompiler* InCompiler, ENiagaraDataType InResultType)
+		: Compiler(InCompiler)
+		, ResultLocation(ENiagaraExpressionResultLocation::Unknown)
+		, ResultIndex(INDEX_NONE)
+		, ComponentIndex(INDEX_NONE)
+		, ResultType(InResultType)
+	{
+	}
+
+	virtual ~FNiagaraExpression()
+	{
+	}
+
+	/** Perform the work specific to the expression type. */
+	virtual void Process(){}
+
+	virtual void PostProcess()
+	{
+		SourceExpressions.Empty();
+	}
+
+	FORCEINLINE TNiagaraExprPtr GetSourceExpression(int32 i){ return StaticCastSharedPtr<FNiagaraExpression>(SourceExpressions[i]); }
+
+	/** Back pointer into the compiler. */
+	class FNiagaraCompiler* Compiler;
+
+	/** Location of this expressions result. */
+	ENiagaraExpressionResultLocation ResultLocation;
+
+	/** Index of this result within it's location. */
+	int32 ResultIndex;
+
+	/** The index of the component we're accessing. Typically this is unused and set to INDEX_NONE. Scalar ops will set and use it however. */
+	int32 ComponentIndex;
+
+	/** The data type of the result. */
+	ENiagaraDataType ResultType;
+
+	//Expressions that this one relies on. Will be processed before this one.
+	TArray<TNiagaraExprPtr> SourceExpressions;
+};
+
+DECLARE_DELEGATE_ThreeParams(FNiagaraOpDelegate, class INiagaraCompiler*, TArray<TNiagaraExprPtr>&, TArray<TNiagaraExprPtr> &);
+
+/** Information about an input or output of a Niagara operation node. */
+class UNREALED_API FNiagaraOpInOutInfo
+{
+public:
+	FName Name;
+	ENiagaraDataType DataType;
+	FText FriendlyName;
+	FText Description;
+	FString Default;
+
+	FNiagaraOpInOutInfo(FName InName, ENiagaraDataType InType, FText InFriendlyName, FText InDescription, FString InDefault)
+		: Name(InName)
+		, DataType(InType)
+		, FriendlyName(InFriendlyName)
+		, Description(InDescription)
+		, Default(InDefault)
+	{
+
+	}
+
+	FNiagaraOpInOutInfo(FName InName, ENiagaraDataType InType)
+		: Name(InName)
+		, DataType(InType)
+		, FriendlyName(FText::FromName(InName))
+		, Description(FText::FromName(InName))
+	{
+		switch (InType)
+		{
+		case ENiagaraDataType::Scalar: Default = TEXT("1.0f");	break;
+		case ENiagaraDataType::Vector: Default = TEXT("1.0f,1.0f,1.0f,1.0f");	break;
+		case ENiagaraDataType::Matrix: Default = TEXT("1.0f,0.0f,0.0f,0.0f,0.0f,1.0f,0.0f,0.0f,0.0f,0.0f,1.0f,0.0f,0.0f,0.0f,0.0f,1.0f");	break;
+		};
+	}
+};
+
+/** Information about a Niagara operation. */
+class UNREALED_API FNiagaraOpInfo
+{
+public:
+	FNiagaraOpInfo()
+	: Name(NAME_None)
+	, OpDelegate(NULL)
+	{}
+
+	FName Name;
+	FText FriendlyName;
+	FText Description;
+	TArray<FNiagaraOpInOutInfo> Inputs;
+	TArray<FNiagaraOpInOutInfo> Outputs;
+	FNiagaraOpDelegate OpDelegate;
+
+	static void Init();
+	static const FNiagaraOpInfo* GetOpInfo(FName Name);
+	static TMap<FName, FNiagaraOpInfo>& GetOpInfoMap();
+
+#define NiagaraOp(OpName) static const FName OpName;
+	NiagaraOpList
+#undef NiagaraOp
+};
+
+class FNiagaraNodeResult
+{
+public:
+	FNiagaraNodeResult(TNiagaraExprPtr InExpression, UEdGraphPin* Pin)
+		: Expression(InExpression)
+		, OutputPin(Pin)
+	{
+		check(InExpression.IsValid());
+		check(Pin != NULL);
+	}
+
+	TNiagaraExprPtr Expression;
+	UEdGraphPin* OutputPin;
+};
\ No newline at end of file
diff --git a/Engine/Source/Editor/UnrealEd/UnrealEd.Build.cs b/Engine/Source/Editor/UnrealEd/UnrealEd.Build.cs
--- a/Engine/Source/Editor/UnrealEd/UnrealEd.Build.cs
+++ b/Engine/Source/Editor/UnrealEd/UnrealEd.Build.cs
@@ -110,7 +110,6 @@ public class UnrealEd : ModuleRules
 				"SwarmInterface",
 				"TargetPlatform",
 				"TargetDeviceServices",
-				"VectorVM",
                 "EditorWidgets",
 				"GraphEditor",
 				"Kismet",
@@ -205,6 +204,7 @@ public class UnrealEd : ModuleRules
 				"UserFeedback",
              	"CollectionManager",
 				"BlueprintGraph",
+				"NiagaraEditor",
 			}
 			);
 
diff --git a/Engine/Source/Runtime/Core/Public/Math/UnrealMathNeon.h b/Engine/Source/Runtime/Core/Public/Math/UnrealMathNeon.h
--- a/Engine/Source/Runtime/Core/Public/Math/UnrealMathNeon.h
+++ b/Engine/Source/Runtime/Core/Public/Math/UnrealMathNeon.h
@@ -478,6 +478,19 @@ FORCEINLINE VectorRegister VectorBitwiseXor(const VectorRegister& Vec1, const Ve
 
 #define VectorSwizzle( Vec, X, Y, Z, W ) __builtin_shufflevector(Vec, Vec, X, Y, Z, W)
 
+/**
+* Creates a vector through selecting two components from each vector via a shuffle mask.
+*
+* @param Vec1		Source vector1
+* @param Vec2		Source vector2
+* @param X			Index for which component of Vector1 to use for X (literal 0-3)
+* @param Y			Index for which component to Vector1 to use for Y (literal 0-3)
+* @param Z			Index for which component to Vector2 to use for Z (literal 0-3)
+* @param W			Index for which component to Vector2 to use for W (literal 0-3)
+* @return			The swizzled vector
+*/
+#define VectorShuffle( Vec1, Vec2, X, Y, Z, W )	__builtin_shufflevector(Vec1, Vec2, X, Y, Z, W)
+
 /**
  * Calculates the cross product of two vectors (XYZ components). W is set to 0.
  *
diff --git a/Engine/Source/Runtime/Engine/Classes/Components/NiagaraComponent.h b/Engine/Source/Runtime/Engine/Classes/Components/NiagaraComponent.h
--- a/Engine/Source/Runtime/Engine/Classes/Components/NiagaraComponent.h
+++ b/Engine/Source/Runtime/Engine/Classes/Components/NiagaraComponent.h
@@ -36,6 +36,9 @@ public:
 	{
 		ParticleBuffers[CurrentBuffer].Reset(NumExpectedParticles * AttrMap.Num());
 		ParticleBuffers[CurrentBuffer].AddUninitialized(NumExpectedParticles * AttrMap.Num());
+
+		//Can we make this a ring buffer and cut memory use in half?
+
 		ParticleAllocation = NumExpectedParticles;
 	}
 
diff --git a/Engine/Source/Runtime/Engine/Classes/Engine/NiagaraConstants.h b/Engine/Source/Runtime/Engine/Classes/Engine/NiagaraConstants.h
--- a/Engine/Source/Runtime/Engine/Classes/Engine/NiagaraConstants.h
+++ b/Engine/Source/Runtime/Engine/Classes/Engine/NiagaraConstants.h
@@ -1,3 +1,6 @@
+// Copyright 1998-2014 Epic Games, Inc. All Rights Reserved.
+
+#pragma once
 
 /**  List of built-in constants
 *	The values of these enums correspond directly to indices for the compiler and interpreter
@@ -15,7 +18,7 @@ namespace NiagaraConstants
 	} ConstantDef;
 
 	FName ConstantNames[] = {
-		"DeltaTime",
+		"Delta Time",
 		"Emitter position",
 		"Emitter age",
 		"Emitter X Axis",
diff --git a/Engine/Source/Runtime/Engine/Classes/Engine/NiagaraScript.h b/Engine/Source/Runtime/Engine/Classes/Engine/NiagaraScript.h
--- a/Engine/Source/Runtime/Engine/Classes/Engine/NiagaraScript.h
+++ b/Engine/Source/Runtime/Engine/Classes/Engine/NiagaraScript.h
@@ -1,9 +1,10 @@
 // Copyright 1998-2014 Epic Games, Inc. All Rights Reserved.
 
 #pragma once
-#include "NiagaraScript.generated.h"
 
+#include "Niagara/NiagaraScriptConstantData.h"
 
+#include "NiagaraScript.generated.h"
 
 /** Runtime script for a Niagara system */
 UCLASS(MinimalAPI)
@@ -15,9 +16,9 @@ class UNiagaraScript : public UObject
 	UPROPERTY()
 	TArray<uint8> ByteCode;
 
-	/** The constant table required for execution. */
+	/** All the data for using constants in the script. */
 	UPROPERTY()
-	TArray<FVector4> ConstantTable;
+	FNiagaraScriptConstantData ConstantData;
 
 	/** Attributes used by this script. */
 	TArray<FName> Attributes;
diff --git a/Engine/Source/Runtime/Engine/Classes/Niagara/NiagaraCommon.h b/Engine/Source/Runtime/Engine/Classes/Niagara/NiagaraCommon.h
new file mode 100644
--- /dev/null
+++ b/Engine/Source/Runtime/Engine/Classes/Niagara/NiagaraCommon.h
@@ -0,0 +1,15 @@
+// Copyright 1998-2014 Epic Games, Inc. All Rights Reserved.
+
+#pragma once
+
+#include "NiagaraCommon.generated.h"
+
+UENUM()
+enum class ENiagaraDataType : uint8
+{
+	Scalar,
+	Vector,
+	Matrix,
+	Unknown,
+	Invalid,
+};
diff --git a/Engine/Source/Runtime/Engine/Classes/Niagara/NiagaraConstantSet.h b/Engine/Source/Runtime/Engine/Classes/Niagara/NiagaraConstantSet.h
new file mode 100644
--- /dev/null
+++ b/Engine/Source/Runtime/Engine/Classes/Niagara/NiagaraConstantSet.h
@@ -0,0 +1,193 @@
+// Copyright 1998-2014 Epic Games, Inc. All Rights Reserved.
+
+#pragma once
+
+#include "Niagara/NiagaraCommon.h"
+#include "NiagaraConstantSet.generated.h"
+
+USTRUCT()
+struct FNiagaraConstantMap
+{
+	GENERATED_USTRUCT_BODY()
+
+private:
+
+	TMap<FName, float> ScalarConstants;
+	TMap<FName, FVector4> VectorConstants;
+	TMap<FName, FMatrix> MatrixConstants;
+
+public:
+
+	void SetOrAdd(FName Name, float Sc)
+	{
+		ScalarConstants.FindOrAdd(Name) = Sc;
+	}
+
+	void SetOrAdd(FName Name, FVector4 Vc)
+	{
+		VectorConstants.FindOrAdd(Name) = Vc;
+	}
+
+	void SetOrAdd(FName Name, const FMatrix& Mc)
+	{
+		MatrixConstants.FindOrAdd(Name) = Mc;
+	}
+
+	float* FindScalar(FName InName){ return ScalarConstants.Find(InName); }
+	FVector4* FindVector(FName InName){ return VectorConstants.Find(InName); }
+	FMatrix* FindMatrix(FName InName){ return MatrixConstants.Find(InName); }
+	const float* FindScalar(FName InName)const { return ScalarConstants.Find(InName); }
+	const FVector4* FindVector(FName InName)const { return VectorConstants.Find(InName); }
+	const FMatrix* FindMatrix(FName InName)const { return MatrixConstants.Find(InName); }
+};
+
+USTRUCT()
+struct FNiagaraConstants
+{
+	GENERATED_USTRUCT_BODY()
+
+private:
+
+	UPROPERTY()
+	TArray<float> ScalarConstants;
+	UPROPERTY()
+	TArray<FName> ScalarNames;
+
+	UPROPERTY()
+	TArray<FVector4> VectorConstants;
+	UPROPERTY()
+	TArray<FName> VectorNames;
+
+	UPROPERTY()
+	TArray<FMatrix> MatrixConstants;
+	UPROPERTY()
+	TArray<FName> MatrixNames;
+
+public:
+
+	/** Fills the entire constants set into the constant table. */
+	void AppendToConstantsTable(TArray<FVector4>& ConstantsTable)const
+	{
+		int32 Idx = ConstantsTable.Num();
+		int32 NewIdx = 0;
+		ConstantsTable.AddUninitialized(ScalarTableSize() + VectorTableSize() + MatrixTableSize());
+		for (float Sc : ScalarConstants)
+		{
+			ConstantsTable[Idx + (NewIdx / 4)][NewIdx % 4] = Sc;
+			++NewIdx;
+			Idx = NewIdx % 4 == 0 ? Idx + 1 : Idx;
+		}
+
+		Idx += FMath::Min(1, NewIdx % 4);//Move to the next table entry if needed.
+		for (FVector4 Vc : VectorConstants)
+		{
+			ConstantsTable[Idx++] = Vc;
+		}
+		for (const FMatrix& Mc : MatrixConstants)
+		{
+			ConstantsTable[Idx] = FVector4(Mc.M[0][0], Mc.M[0][1], Mc.M[0][2], Mc.M[0][3]);
+			ConstantsTable[Idx + 1] = FVector4(Mc.M[1][0], Mc.M[1][1], Mc.M[1][2], Mc.M[1][3]);
+			ConstantsTable[Idx + 2] = FVector4(Mc.M[2][0], Mc.M[2][1], Mc.M[2][2], Mc.M[2][3]);
+			ConstantsTable[Idx + 3] = FVector4(Mc.M[3][0], Mc.M[3][1], Mc.M[3][2], Mc.M[3][3]);
+			Idx += 4;
+		}
+	}
+
+	/** 
+	Fills only selected constants into the table. In the order they appear in the array of passed names not the order they appear in the set. 
+	Checks the passed map for entires to superceed the default values in the set.
+	*/
+	void AppendToConstantsTable(TArray<FVector4>& ConstantsTable, const FNiagaraConstantMap& Externals)const
+	{
+		int32 Idx = ConstantsTable.Num();
+		int32 NewIdx = 0;
+		ConstantsTable.AddUninitialized(ScalarTableSize() + VectorTableSize() + MatrixTableSize());
+		for (int32 i = 0; i < ScalarNames.Num(); ++i)
+		{
+			FName ScName = ScalarNames[i];
+			const float* Scalar = Externals.FindScalar(ScName);
+			ConstantsTable[Idx + (NewIdx / 4)][NewIdx % 4] = Scalar ? *Scalar : ScalarConstants[i];
+			++NewIdx;
+			Idx = NewIdx % 4 == 0 ? Idx + 1 : Idx;
+		}
+
+		Idx += FMath::Min(1, NewIdx % 4);//Move to the next table entry if needed.
+		for (int32 i = 0; i < VectorNames.Num(); ++i)
+		{
+			FName VcName = VectorNames[i];
+			const FVector4* Vector = Externals.FindVector(VcName);
+			ConstantsTable[Idx++] = Vector ? *Vector : VectorConstants[i];
+		}
+		for (int32 i = 0; i < MatrixNames.Num(); ++i)
+		{
+			FName McName = MatrixNames[i];
+			const FMatrix* Mat = Externals.FindMatrix(McName);
+			FMatrix Mc = Mat ? *Mat : MatrixConstants[i];
+			ConstantsTable[Idx] = FVector4(Mc.M[0][0], Mc.M[0][1], Mc.M[0][2], Mc.M[0][3]);
+			ConstantsTable[Idx + 1] = FVector4(Mc.M[1][0], Mc.M[1][1], Mc.M[1][2], Mc.M[1][3]);
+			ConstantsTable[Idx + 2] = FVector4(Mc.M[2][0], Mc.M[2][1], Mc.M[2][2], Mc.M[2][3]);
+			ConstantsTable[Idx + 3] = FVector4(Mc.M[3][0], Mc.M[3][1], Mc.M[3][2], Mc.M[3][3]);
+		}
+	}
+
+	void SetOrAdd(FName Name, float Sc)
+	{
+		int32 Idx = ScalarNames.Find(Name);
+		if (Idx == INDEX_NONE)
+		{
+			ScalarNames.Add(Name);
+			ScalarConstants.Add(Sc);
+		}
+		else
+		{
+			ScalarConstants[Idx] = Sc;
+		}
+	}
+
+	void SetOrAdd(FName Name, FVector4 Vc)
+	{
+		int32 Idx = VectorNames.Find(Name);
+		if (Idx == INDEX_NONE)
+		{
+			VectorNames.Add(Name);
+			VectorConstants.Add(Vc);
+		}
+		else
+		{
+			VectorConstants[Idx] = Vc;
+		}
+	}
+
+	void SetOrAdd(FName Name, const FMatrix& Mc)
+	{
+		int32 Idx = MatrixNames.Find(Name);
+		if (Idx == INDEX_NONE)
+		{
+			MatrixNames.Add(Name);
+			MatrixConstants.Add(Mc);
+		}
+		else
+		{
+			MatrixConstants[Idx] = Mc;
+		}
+	}
+
+	FORCEINLINE int32 NumScalars()const { return ScalarNames.Num(); }
+	FORCEINLINE int32 NumVectors()const { return VectorNames.Num(); }
+	FORCEINLINE int32 NumMatrices()const { return MatrixNames.Num(); }
+	FORCEINLINE int32 ScalarTableSize()const { return(NumScalars() / 4) + FMath::Min(1,NumScalars() % 4); }
+	FORCEINLINE int32 VectorTableSize()const { return(NumVectors()); }
+	FORCEINLINE int32 MatrixTableSize()const { return(NumMatrices() * 4); }
+
+	/**	Returns the number of vectors these constants would use in a constants table. */
+	FORCEINLINE int32 GetTableSize()const { return ScalarTableSize() + VectorTableSize() + MatrixTableSize(); }
+
+	/** Return the first constant used for this constant in a table. */
+	FORCEINLINE int32 GetTableIndex_Scalar(FName InName)const{ return ScalarNames.Find(InName) / 4; }
+	FORCEINLINE int32 GetTableIndex_Vector(FName InName)const{ return ScalarTableSize() + VectorNames.Find(InName); }
+	FORCEINLINE int32 GetTableIndex_Matrix(FName InName)const{ return ScalarTableSize() + VectorTableSize() + MatrixNames.Find(InName) * 4; }
+
+	/** Return the absolute index of the passed scalar. Can be used to get the component index also. */
+	FORCEINLINE int32 GetAbsoluteIndex_Scalar(FName InName)const{ return ScalarNames.Find(InName); }
+};
+
diff --git a/Engine/Source/Runtime/Engine/Classes/Niagara/NiagaraScriptConstantData.h b/Engine/Source/Runtime/Engine/Classes/Niagara/NiagaraScriptConstantData.h
new file mode 100644
--- /dev/null
+++ b/Engine/Source/Runtime/Engine/Classes/Niagara/NiagaraScriptConstantData.h
@@ -0,0 +1,97 @@
+// Copyright 1998-2014 Epic Games, Inc. All Rights Reserved.
+
+#pragma once
+
+#include "NiagaraConstantSet.h"
+
+#include "NiagaraScriptConstantData.generated.h"
+
+USTRUCT()
+struct FNiagaraScriptConstantData
+{
+	GENERATED_USTRUCT_BODY()
+
+	/** The set of external constants for this script. */
+	UPROPERTY()
+	FNiagaraConstants ExternalConstants;
+
+	/** The set of internal constants for this script. */
+	UPROPERTY()
+	FNiagaraConstants InternalConstants;
+
+	/** Fill a constants table ready for use in an update or spawn. */
+	void FillConstantTable(const FNiagaraConstantMap& ExternalConstantsMap, TArray<FVector4>& OutConstantTable)const
+	{
+		//First up in the table comes the External constants in scalar->vector->matrix order.
+		//Only fills the constants actually used by the script.
+		OutConstantTable.Empty();
+		ExternalConstants.AppendToConstantsTable(OutConstantTable, ExternalConstantsMap);
+		//Next up add all the internal constants from the script.
+		InternalConstants.AppendToConstantsTable(OutConstantTable);
+	}
+
+	void SetOrAddInternal(FName Name, float Sc)
+	{
+		InternalConstants.SetOrAdd(Name, Sc);
+	}
+
+	void SetOrAddInternal(FName Name, FVector4 Vc)
+	{
+		InternalConstants.SetOrAdd(Name, Vc);
+	}
+
+	void SetOrAddInternal(FName Name, const FMatrix& Mc)
+	{
+		InternalConstants.SetOrAdd(Name, Mc);
+	}
+
+	void SetOrAddExternal(FName Name, float Sc)
+	{
+		ExternalConstants.SetOrAdd(Name, Sc);
+	}
+
+	void SetOrAddExternal(FName Name, FVector4 Vc)
+	{
+		ExternalConstants.SetOrAdd(Name, Vc);
+	}
+
+	void SetOrAddExternal(FName Name, const FMatrix& Mc)
+	{
+		ExternalConstants.SetOrAdd(Name, Mc);
+	}
+
+	/**
+	Calculates the index into a a constant table created from this constant data for the given constant name.
+	Obviously, assumes the constant data is complete. Any additions etc will invalidate previously calculated indexes.
+	*/
+	void GetTableIndex(FName Name, bool bInternal, int32& OutConstantIdx, int32& OutComponentIndex, ENiagaraDataType& OutType)const
+	{
+		int32 Base = bInternal ? ExternalConstants.GetTableSize() : 0;
+		int32 ConstIdx = INDEX_NONE;
+		int32 CompIdx = INDEX_NONE;
+		ENiagaraDataType Type;
+		const FNiagaraConstants& Consts = bInternal ? InternalConstants : ExternalConstants;
+
+		ConstIdx = Consts.GetAbsoluteIndex_Scalar(Name);
+		Type = ENiagaraDataType::Scalar;
+		if (ConstIdx == INDEX_NONE)
+		{
+			ConstIdx = Consts.GetTableIndex_Vector(Name);
+			Type = ENiagaraDataType::Vector;
+			if (ConstIdx == INDEX_NONE)
+			{
+				ConstIdx = Consts.GetTableIndex_Matrix(Name);
+				Type = ENiagaraDataType::Matrix;
+			}
+		}
+		else
+		{
+			CompIdx = ConstIdx % 4;
+			ConstIdx /= 4;
+		}
+
+		OutConstantIdx = Base + ConstIdx;
+		OutComponentIndex = CompIdx;
+		OutType = Type;
+	}
+};
diff --git a/Engine/Source/Runtime/Engine/Private/NiagaraComponent.cpp b/Engine/Source/Runtime/Engine/Private/NiagaraComponent.cpp
--- a/Engine/Source/Runtime/Engine/Private/NiagaraComponent.cpp
+++ b/Engine/Source/Runtime/Engine/Private/NiagaraComponent.cpp
@@ -141,17 +141,10 @@ public:
 		, Component(*InComponent)
 		, UpdateScript(*InComponent->UpdateScript)
 		, SpawnScript(InComponent->SpawnScript)
-		, ConstantTable(InComponent->UpdateScript->ConstantTable)
 		, SpawnRemainder(0.0f)
 		, CachedBounds(ForceInit)
 	{
 		check(InComponent->UpdateScript && UpdateScript.ByteCode.Num());
-
-		// make room for constants
-		while (ConstantTable.Num() < 10)	
-		{
-			ConstantTable.Add( FVector4(0.0f, 0.0f, 0.0f, 0.0f) );
-		}
 	}
 
 	void Tick(float DeltaSeconds)
@@ -197,12 +190,20 @@ public:
 
 	FNiagaraEmitterParticleData &GetData()	{ return Data; }
 
-	void SetConstant(int Idx, const FVector4 &Value)
+	void SetConstant(FName ConstantName, const float Value)
 	{
-		check(Idx < ConstantTable.Num());
-		ConstantTable[Idx] = Value;
+		Constants.SetOrAdd(ConstantName,Value);
 	}
 
+	void SetConstant(FName ConstantName, const FVector4 Value)
+	{
+		Constants.SetOrAdd(ConstantName, Value);
+	}
+
+	void SetConstant(FName ConstantName, const FMatrix& Value)
+	{
+		Constants.SetOrAdd(ConstantName, Value);
+	}
 private:
 	/** Temporary stuff for the prototype. */
 	float SpawnRate;
@@ -213,8 +214,8 @@ private:
 	UNiagaraScript& UpdateScript;
 	/** The particle spawn script. */
 	UNiagaraScript *SpawnScript;
-	/** Local constant table. */
-	TArray<FVector4> ConstantTable;
+	/** Local constant set. */
+	FNiagaraConstantMap Constants;
 	/** particle simulation data */
 	FNiagaraEmitterParticleData Data;
 	/** Keep partial particle spawns from last frame */
@@ -224,22 +225,6 @@ private:
 	/** Cached bounds. */
 	FBox CachedBounds;
 
-
-	/* builtin constants are setup in the tick function, but a script may have additional constants 
-	 derived from unconnected nodes, so we need to merge those into the simulation constant table */
-	void MergeScriptConstants(UNiagaraScript *Script)
-	{
-		// start at NumBuiltinConstants+1, because index 0 is always the fixed zero constant
-		for (int i = NiagaraConstants::NumBuiltinConstants + 1; i < Script->ConstantTable.Num(); i++)
-		{
-			if (ConstantTable.Num() <= i)
-			{
-				ConstantTable.Add(FVector4());
-			}
-			ConstantTable[i] = Script->ConstantTable[i];
-		}
-	}
-
 	/** Calc number to spawn */
 	int32 CalcNumToSpawn(float DeltaSeconds)
 	{
@@ -275,8 +260,9 @@ private:
 			OutputRegisters[AttrIndex] = (VectorRegister*)(Particles + AttrIndex * NumVectorsPerAttribute);
 		}
 
-		// copy script specific constants into the constant table
-		MergeScriptConstants(&UpdateScript);
+		//Fill constant table with required emitter constants and internal script constants.
+		TArray<FVector4> ConstantTable;
+		UpdateScript.ConstantData.FillConstantTable(Constants, ConstantTable);
 
 		VectorVM::Exec(
 			UpdateScript.ByteCode.GetData(),
@@ -316,8 +302,9 @@ private:
 				OutputRegisters[AttrIndex] = (VectorRegister*)(NewParticlesStart + AttrIndex * Data.GetParticleAllocation());
 			}
 
-			// copy script specific constants into the constant table
-			MergeScriptConstants(SpawnScript);
+			//Fill constant table with required emitter constants and internal script constants.
+			TArray<FVector4> ConstantTable;
+			SpawnScript->ConstantData.FillConstantTable(Constants, ConstantTable);
 
 			VectorVM::Exec(
 				SpawnScript->ByteCode.GetData(),
@@ -409,14 +396,25 @@ void UNiagaraComponent::TickComponent(float DeltaSeconds, enum ELevelTick TickTy
 	if (Simulation)
 	{
 		EmitterAge += DeltaSeconds;
-
-		Simulation->SetConstant(0, FVector4(0.0f, 0.0f, 0.0f, 0.0f) );	// zero constant
-		Simulation->SetConstant(1, FVector4(DeltaSeconds, DeltaSeconds, DeltaSeconds, DeltaSeconds));
-		Simulation->SetConstant(2, FVector4(ComponentToWorld.GetTranslation()));
-		Simulation->SetConstant(3, FVector4(EmitterAge, EmitterAge, EmitterAge, EmitterAge));
-		Simulation->SetConstant(4, FVector4(ComponentToWorld.GetUnitAxis(EAxis::X)) );
-		Simulation->SetConstant(5, FVector4(ComponentToWorld.GetUnitAxis(EAxis::Y)) );
-		Simulation->SetConstant(6, FVector4(ComponentToWorld.GetUnitAxis(EAxis::Z)) );
+	
+		//Todo, open this up to the UI and setting via code and BPs.
+		static FName Const_Zero(TEXT("ZERO"));
+		Simulation->SetConstant(Const_Zero, FVector4(0.0f, 0.0f, 0.0f, 0.0f));	// zero constant
+		static FName Const_DeltaTime(TEXT("Delta Time"));
+		Simulation->SetConstant(Const_DeltaTime, FVector4(DeltaSeconds, DeltaSeconds, DeltaSeconds, DeltaSeconds));
+		static FName Const_EmitterPos(TEXT("Emitter Position"));
+		Simulation->SetConstant(Const_EmitterPos, FVector4(ComponentToWorld.GetTranslation()));
+		static FName Const_Age(TEXT("Emitter Age"));
+		Simulation->SetConstant(Const_Age, FVector4(EmitterAge, EmitterAge, EmitterAge, EmitterAge));
+		static FName Const_EmitterX(TEXT("Emitter X Axis"));
+		Simulation->SetConstant(Const_EmitterX, FVector4(ComponentToWorld.GetUnitAxis(EAxis::X)));
+		static FName Const_EmitterY(TEXT("Emitter Y Axis"));
+		Simulation->SetConstant(Const_EmitterY, FVector4(ComponentToWorld.GetUnitAxis(EAxis::Y)));
+		static FName Const_EmitterZ(TEXT("Emitter Z Axis"));
+		Simulation->SetConstant(Const_EmitterZ, FVector4(ComponentToWorld.GetUnitAxis(EAxis::Z)));
+
+		static FName Const_EmitterTransform(TEXT("Emitter Transform"));
+		Simulation->SetConstant(Const_EmitterTransform, ComponentToWorld.ToMatrixWithScale());
 
 		Simulation->Tick(DeltaSeconds);
 		UpdateComponentToWorld();
diff --git a/Engine/Source/Runtime/VectorVM/Private/VectorVM.cpp b/Engine/Source/Runtime/VectorVM/Private/VectorVM.cpp
--- a/Engine/Source/Runtime/VectorVM/Private/VectorVM.cpp
+++ b/Engine/Source/Runtime/VectorVM/Private/VectorVM.cpp
@@ -8,14 +8,22 @@ IMPLEMENT_MODULE(FDefaultModuleImpl, VectorVM);
 
 DEFINE_LOG_CATEGORY_STATIC(LogVectorVM, All, All);
 
-#define SRCOP_RRR 0x00
-#define SRCOP_RRC 0x01
-#define SRCOP_RCR 0x02
-#define SRCOP_RCC 0x03
-#define SRCOP_CRC 0x05
-#define SRCOP_CCR 0x06
-#define SRCOP_CCC 0x07
-
+#define SRCOP_RRRR 0x00
+#define SRCOP_RRRC 0x01
+#define SRCOP_RRCR 0x02
+#define SRCOP_RRCC 0x03
+#define SRCOP_RCRR 0x04
+#define SRCOP_RCRC 0x05
+#define SRCOP_RCCR 0x06
+#define SRCOP_RCCC 0x07
+#define SRCOP_CRRR 0x08
+#define SRCOP_CRRC 0x09
+#define SRCOP_CRCR 0x0a
+#define SRCOP_CRCC 0x0b
+#define SRCOP_CCRR 0x0c
+#define SRCOP_CCRC 0x0d
+#define SRCOP_CCCR 0x0e
+#define SRCOP_CCCC 0x0f
 
 /**
  * Context information passed around during VM execution.
@@ -47,9 +55,9 @@ struct FVectorVMContext
 };
 
 /** Decode the next operation contained in the bytecode. */
-static FORCEINLINE VectorVM::EOp::Type DecodeOp(FVectorVMContext& Context)
+static FORCEINLINE VectorVM::EOp DecodeOp(FVectorVMContext& Context)
 {
-	return static_cast<VectorVM::EOp::Type>(*Context.Code++);
+	return static_cast<VectorVM::EOp>(*Context.Code++);
 }
 
 /** Decode a register from the bytecode. */
@@ -70,7 +78,78 @@ static FORCEINLINE uint8 DecodeSrcOperandTypes(FVectorVMContext& Context)
 	return *Context.Code++;
 }
 
-/** Base class for vector kernels with one dest and one src operand. */
+static FORCEINLINE uint8 DecodeByte(FVectorVMContext& Context)
+{
+	return *Context.Code++;
+}
+
+/** Handles reading of a constant. */
+struct FConstantHandler
+{
+	VectorRegister Constant;
+	FConstantHandler(FVectorVMContext& Context)
+		: Constant(DecodeConstant(Context))
+	{}
+	FORCEINLINE VectorRegister Get(){ return Constant; }
+};
+
+/** Handles reading of a register, advancing the pointer with each read. */
+struct FRegisterHandler
+{
+	VectorRegister* Register;
+	FRegisterHandler(FVectorVMContext& Context)
+		: Register(DecodeRegister(Context))
+	{}
+	FORCEINLINE VectorRegister Get(){ return *Register++; }
+};
+
+template<typename Kernel, typename Arg0Handler>
+FORCEINLINE void VectorUnaryLoop(FVectorVMContext& Context, VectorRegister* RESTRICT Dst, int32 NumVectors)
+{
+	Arg0Handler Arg0(Context);
+	for (int32 i = 0; i < NumVectors; ++i)
+	{
+		Kernel::DoKernel(Dst++, Arg0.Get());
+	}
+}
+
+template<typename Kernel, typename Arg0Handler, typename Arg1Handler>
+FORCEINLINE void VectorBinaryLoop(FVectorVMContext& Context, VectorRegister* RESTRICT Dst, int32 NumVectors)
+{
+	Arg0Handler Arg0(Context);
+	Arg1Handler Arg1(Context);
+	for (int32 i = 0; i < NumVectors; ++i)
+	{
+		Kernel::DoKernel(Dst++, Arg0.Get(), Arg1.Get());
+	}
+}
+
+template<typename Kernel, typename Arg0Handler, typename Arg1Handler, typename Arg2Handler>
+FORCEINLINE void VectorTrinaryLoop(FVectorVMContext& Context, VectorRegister* RESTRICT Dst, int32 NumVectors)
+{
+	Arg0Handler Arg0(Context);
+	Arg1Handler Arg1(Context);
+	Arg2Handler Arg2(Context);
+	for (int32 i = 0; i < NumVectors; ++i)
+	{
+		Kernel::DoKernel(Dst++, Arg0.Get(), Arg1.Get(), Arg2.Get());
+	}
+}
+
+template<typename Kernel, typename Arg0Handler, typename Arg1Handler, typename Arg2Handler, typename Arg3Handler>
+FORCEINLINE void VectorQuinaryLoop(FVectorVMContext& Context, VectorRegister* RESTRICT Dst, int32 NumVectors)
+{
+	Arg0Handler Arg0(Context);
+	Arg1Handler Arg1(Context);
+	Arg2Handler Arg2(Context);
+	Arg3Handler Arg3(Context);
+	for (int32 i = 0; i < NumVectors; ++i)
+	{
+		Kernel::DoKernel(Dst++, Arg0.Get(), Arg1.Get(), Arg2.Get(), Arg3.Get());
+	}
+}
+
+/** Base class of vector kernels with a single operand. */
 template <typename Kernel>
 struct TUnaryVectorKernel
 {
@@ -79,29 +158,16 @@ struct TUnaryVectorKernel
 		VectorRegister* RESTRICT Dst = DecodeRegister(Context);
 		uint32 SrcOpTypes = DecodeSrcOperandTypes(Context);
 		int32 NumVectors = Context.NumVectors;
-
-		if (SrcOpTypes == SRCOP_RRR)
-		{
-			VectorRegister* Src0 = DecodeRegister(Context);
-			for (int32 i = 0; i < NumVectors; ++i)
-			{
-				Kernel::DoKernel(Dst++,*Src0++);
-			}
-		}
-		else if (SrcOpTypes == SRCOP_RRC)
+		switch (SrcOpTypes)
 		{
-			VectorRegister Src0 = DecodeConstant(Context);
-			for (int32 i = 0; i < NumVectors; ++i)
-			{
-				Kernel::DoKernel(Dst++, Src0);
-			}
-		}
-
+		case SRCOP_RRRR: 	VectorUnaryLoop<Kernel, FRegisterHandler>(Context, Dst, NumVectors); break;
+		case SRCOP_RRRC:	VectorUnaryLoop<Kernel, FConstantHandler>(Context, Dst, NumVectors); break;
+		default: check(0); break;
+		};
 	}
 };
 
-
-/** Base class for vector kernels with one dest and two src operands. */
+/** Base class of Vector kernels with 2 operands. */
 template <typename Kernel>
 struct TBinaryVectorKernel
 {
@@ -109,48 +175,19 @@ struct TBinaryVectorKernel
 	{
 		VectorRegister* RESTRICT Dst = DecodeRegister(Context);
 		uint32 SrcOpTypes = DecodeSrcOperandTypes(Context);
-
 		int32 NumVectors = Context.NumVectors;
-		if (SrcOpTypes == SRCOP_RRR)
+		switch (SrcOpTypes)
 		{
-			VectorRegister* Src0 = DecodeRegister(Context);
-			VectorRegister* Src1 = DecodeRegister(Context);
-			for (int32 i = 0; i < NumVectors; ++i)
-			{
-				Kernel::DoKernel(Dst++, *Src0++, *Src1++);
-			}
-		}
-		else if (SrcOpTypes == SRCOP_RRC)
-		{
-			VectorRegister Src0 = DecodeConstant(Context);
-			VectorRegister *Src1 = DecodeRegister(Context);
-			for (int32 i = 0; i < NumVectors; ++i)
-			{
-				Kernel::DoKernel(Dst++, Src0, *Src1++);
-			}
-		}
-		else if (SrcOpTypes == SRCOP_RCR)
-		{
-			VectorRegister *Src0 = DecodeRegister(Context);
-			VectorRegister Src1 = DecodeConstant(Context);
-			for (int32 i = 0; i < NumVectors; ++i)
-			{
-				Kernel::DoKernel(Dst++, *Src0++, Src1);
-			}
-		}
-		else if (SrcOpTypes == SRCOP_RCC)
-		{
-			VectorRegister Src0 = DecodeConstant(Context);
-			VectorRegister Src1 = DecodeConstant(Context);
-			for (int32 i = 0; i < NumVectors; ++i)
-			{
-				Kernel::DoKernel(Dst++, Src0, Src1);
-			}
-		}
+		case SRCOP_RRRR: 	VectorBinaryLoop<Kernel, FRegisterHandler, FRegisterHandler>(Context, Dst, NumVectors); break;
+		case SRCOP_RRRC:	VectorBinaryLoop<Kernel, FConstantHandler, FRegisterHandler>(Context, Dst, NumVectors); break;
+		case SRCOP_RRCR: 	VectorBinaryLoop<Kernel, FRegisterHandler, FConstantHandler>(Context, Dst, NumVectors); break;
+		case SRCOP_RRCC:	VectorBinaryLoop<Kernel, FConstantHandler, FConstantHandler>(Context, Dst, NumVectors); break;
+		default: check(0); break;
+		};
 	}
 };
 
-/** Base class for vector kernels with one dest and three src operands. */
+/** Base class of Vector kernels with 3 operands. */
 template <typename Kernel>
 struct TTrinaryVectorKernel
 {
@@ -158,68 +195,51 @@ struct TTrinaryVectorKernel
 	{
 		VectorRegister* RESTRICT Dst = DecodeRegister(Context);
 		uint32 SrcOpTypes = DecodeSrcOperandTypes(Context);
-
 		int32 NumVectors = Context.NumVectors;
-		if (SrcOpTypes == SRCOP_RRR)
-		{
-			VectorRegister* Src0 = DecodeRegister(Context);
-			VectorRegister* Src1 = DecodeRegister(Context);
-			VectorRegister* Src2 = DecodeRegister(Context);
-			for (int32 i = 0; i < NumVectors; ++i)
-			{
-				Kernel::DoKernel(Dst++, *Src0++, *Src1++, *Src2++);
-			}
-		}
-		else if (SrcOpTypes == SRCOP_RRC)
-		{
-			VectorRegister Src0 = DecodeConstant(Context);
-			VectorRegister* Src1 = DecodeRegister(Context);
-			VectorRegister* Src2 = DecodeRegister(Context);
-			for (int32 i = 0; i < NumVectors; ++i)
-			{
-				Kernel::DoKernel(Dst++, Src0, *Src1++, *Src2++);
-			}
-		}
-		else if (SrcOpTypes == SRCOP_CCC)
-		{
-			VectorRegister Src0 = DecodeConstant(Context);
-			VectorRegister Src1 = DecodeConstant(Context);
-			VectorRegister Src2 = DecodeConstant(Context);
-			for (int32 i = 0; i < NumVectors; ++i)
-			{
-				Kernel::DoKernel(Dst++, Src0, Src1, Src2);
-			}
-		}
-		else if (SrcOpTypes == SRCOP_CCR)
+		switch (SrcOpTypes)
 		{
-			VectorRegister *Src0 = DecodeRegister(Context);
-			VectorRegister Src1 = DecodeConstant(Context);
-			VectorRegister Src2 = DecodeConstant(Context);
-			for (int32 i = 0; i < NumVectors; ++i)
-			{
-				Kernel::DoKernel(Dst++, *Src0++, Src1, Src2);
-			}
-		}
-		else if (SrcOpTypes == SRCOP_RCC)
-		{
-			VectorRegister Src0 = DecodeConstant(Context);
-			VectorRegister Src1 = DecodeConstant(Context);
-			VectorRegister *Src2 = DecodeRegister(Context);
-			for (int32 i = 0; i < NumVectors; ++i)
-			{
-				Kernel::DoKernel(Dst++, Src0, Src1, *Src2++);
-			}
-		}
-		else if (SrcOpTypes == SRCOP_RCR)
+		case SRCOP_RRRR: 	VectorTrinaryLoop<Kernel, FRegisterHandler, FRegisterHandler, FRegisterHandler>(Context, Dst, NumVectors); break;
+		case SRCOP_RRRC:	VectorTrinaryLoop<Kernel, FConstantHandler, FRegisterHandler, FRegisterHandler>(Context, Dst, NumVectors); break;
+		case SRCOP_RRCR: 	VectorTrinaryLoop<Kernel, FRegisterHandler, FConstantHandler, FRegisterHandler>(Context, Dst, NumVectors); break;
+		case SRCOP_RRCC:	VectorTrinaryLoop<Kernel, FConstantHandler, FConstantHandler, FRegisterHandler>(Context, Dst, NumVectors); break;
+		case SRCOP_RCRR: 	VectorTrinaryLoop<Kernel, FRegisterHandler, FRegisterHandler, FConstantHandler>(Context, Dst, NumVectors); break;
+		case SRCOP_RCRC:	VectorTrinaryLoop<Kernel, FConstantHandler, FRegisterHandler, FConstantHandler>(Context, Dst, NumVectors); break;
+		case SRCOP_RCCR: 	VectorTrinaryLoop<Kernel, FRegisterHandler, FConstantHandler, FConstantHandler>(Context, Dst, NumVectors); break;
+		case SRCOP_RCCC:	VectorTrinaryLoop<Kernel, FConstantHandler, FConstantHandler, FConstantHandler>(Context, Dst, NumVectors); break;
+		default: check(0); break;
+		};
+	}
+};
+
+/** Base class of Vector kernels with 4 operands. */
+template <typename Kernel>
+struct TQuinaryVectorKernel
+{
+	static void Exec(FVectorVMContext& Context)
+	{
+		VectorRegister* RESTRICT Dst = DecodeRegister(Context);
+		uint32 SrcOpTypes = DecodeSrcOperandTypes(Context);
+		int32 NumVectors = Context.NumVectors;
+		switch (SrcOpTypes)
 		{
-			VectorRegister *Src0 = DecodeRegister(Context);
-			VectorRegister Src1 = DecodeConstant(Context);
-			VectorRegister *Src2 = DecodeRegister(Context);
-			for (int32 i = 0; i < NumVectors; ++i)
-			{
-				Kernel::DoKernel(Dst++, *Src0++, Src1, *Src2++);
-			}
-		}
+		case SRCOP_RRRR: 	VectorQuinaryLoop<Kernel, FRegisterHandler, FRegisterHandler, FRegisterHandler, FRegisterHandler>(Context, Dst, NumVectors); break;
+		case SRCOP_RRRC:	VectorQuinaryLoop<Kernel, FConstantHandler, FRegisterHandler, FRegisterHandler, FRegisterHandler>(Context, Dst, NumVectors); break;
+		case SRCOP_RRCR: 	VectorQuinaryLoop<Kernel, FRegisterHandler, FConstantHandler, FRegisterHandler, FRegisterHandler>(Context, Dst, NumVectors); break;
+		case SRCOP_RRCC:	VectorQuinaryLoop<Kernel, FConstantHandler, FConstantHandler, FRegisterHandler, FRegisterHandler>(Context, Dst, NumVectors); break;
+		case SRCOP_RCRR: 	VectorQuinaryLoop<Kernel, FRegisterHandler, FRegisterHandler, FConstantHandler, FRegisterHandler>(Context, Dst, NumVectors); break;
+		case SRCOP_RCRC:	VectorQuinaryLoop<Kernel, FConstantHandler, FRegisterHandler, FConstantHandler, FRegisterHandler>(Context, Dst, NumVectors); break;
+		case SRCOP_RCCR: 	VectorQuinaryLoop<Kernel, FRegisterHandler, FConstantHandler, FConstantHandler, FRegisterHandler>(Context, Dst, NumVectors); break;
+		case SRCOP_RCCC:	VectorQuinaryLoop<Kernel, FConstantHandler, FConstantHandler, FConstantHandler, FRegisterHandler>(Context, Dst, NumVectors); break;
+		case SRCOP_CRRR: 	VectorQuinaryLoop<Kernel, FRegisterHandler, FRegisterHandler, FRegisterHandler, FConstantHandler>(Context, Dst, NumVectors); break;
+		case SRCOP_CRRC:	VectorQuinaryLoop<Kernel, FConstantHandler, FRegisterHandler, FRegisterHandler, FConstantHandler>(Context, Dst, NumVectors); break;
+		case SRCOP_CRCR: 	VectorQuinaryLoop<Kernel, FRegisterHandler, FConstantHandler, FRegisterHandler, FConstantHandler>(Context, Dst, NumVectors); break;
+		case SRCOP_CRCC:	VectorQuinaryLoop<Kernel, FConstantHandler, FConstantHandler, FRegisterHandler, FConstantHandler>(Context, Dst, NumVectors); break;
+		case SRCOP_CCRR: 	VectorQuinaryLoop<Kernel, FRegisterHandler, FRegisterHandler, FConstantHandler, FConstantHandler>(Context, Dst, NumVectors); break;
+		case SRCOP_CCRC:	VectorQuinaryLoop<Kernel, FConstantHandler, FRegisterHandler, FConstantHandler, FConstantHandler>(Context, Dst, NumVectors); break;
+		case SRCOP_CCCR: 	VectorQuinaryLoop<Kernel, FRegisterHandler, FConstantHandler, FConstantHandler, FConstantHandler>(Context, Dst, NumVectors); break;
+		case SRCOP_CCCC:	VectorQuinaryLoop<Kernel, FConstantHandler, FConstantHandler, FConstantHandler, FConstantHandler>(Context, Dst, NumVectors); break;
+		default: check(0); break;
+		};
 	}
 };
 
@@ -488,6 +508,34 @@ struct FVectorKernelNoise : public TUnaryVectorKernel<FVectorKernelNoise>
 
 VectorRegister FVectorKernelNoise::RandomTable[10][10][10];
 
+template<int32 Component>
+struct FVectorKernelSplat : public TUnaryVectorKernel<FVectorKernelSplat<Component>>
+{
+	static void FORCEINLINE DoKernel(VectorRegister* RESTRICT Dst, VectorRegister Src0)
+	{
+		*Dst = VectorReplicate(Src0, Component);
+	}
+};
+
+template<int32 Cmp0, int32 Cmp1, int32 Cmp2, int32 Cmp3>
+struct FVectorKernelCompose : public TQuinaryVectorKernel<FVectorKernelCompose<Cmp0, Cmp1, Cmp2, Cmp3>>
+{
+	static void FORCEINLINE DoKernel(VectorRegister* RESTRICT Dst, VectorRegister Src0, VectorRegister Src1, VectorRegister Src2, VectorRegister Src3)
+	{
+		//TODO - There's probably a faster way to do this.
+		VectorRegister Tmp0 = VectorShuffle(Src0, Src1, Cmp0, Cmp0, Cmp1, Cmp1);
+		VectorRegister Tmp1 = VectorShuffle(Src2, Src3, Cmp2, Cmp2, Cmp3, Cmp3);		
+		*Dst = VectorShuffle(Tmp0, Tmp1, 0, 2, 0, 2);
+	}
+};
+
+struct FVectorKernelOutput : public TUnaryVectorKernel<FVectorKernelOutput>
+{
+	static void FORCEINLINE DoKernel(VectorRegister* RESTRICT Dst, VectorRegister Src0)
+	{
+		VectorStoreAlignedStreamed(Src0, Dst);
+	}
+};
 
 void VectorVM::Init()
 {
@@ -550,7 +598,7 @@ void VectorVM::Exec(
 		// Setup execution context.
 		int32 VectorsThisChunk = FMath::Min<int32>(NumVectors, VectorsPerChunk);
 		FVectorVMContext Context(Code, RegisterTable, ConstantTable, VectorsThisChunk);
-		EOp::Type Op = EOp::done;
+		EOp Op = EOp::done;
 
 		// Execute VM on all vectors in this chunk.
 		do 
@@ -581,6 +629,16 @@ void VectorVM::Exec(
 			case EOp::normalize: FVectorKernelNormalize::Exec(Context); break;
 			case EOp::random: FVectorKernelRandom::Exec(Context); break;
 			case EOp::noise: FVectorKernelNoise::Exec(Context); break;
+			case EOp::splatx: FVectorKernelSplat<0>::Exec(Context); break;
+			case EOp::splaty: FVectorKernelSplat<1>::Exec(Context); break;
+			case EOp::splatz: FVectorKernelSplat<2>::Exec(Context); break;
+			case EOp::splatw: FVectorKernelSplat<3>::Exec(Context); break;
+			case EOp::compose: FVectorKernelCompose<0,1,2,3>::Exec(Context); break;
+			case EOp::composex: FVectorKernelCompose<0, 0, 0, 0>::Exec(Context); break;
+			case EOp::composey: FVectorKernelCompose<1, 1, 1, 1>::Exec(Context); break;
+			case EOp::composez: FVectorKernelCompose<2, 2, 2, 2>::Exec(Context); break;
+			case EOp::composew: FVectorKernelCompose<3, 3, 3, 3>::Exec(Context); break;
+			case EOp::output: FVectorKernelOutput::Exec(Context); break;
 
 			// Execution always terminates with a "done" opcode.
 			case EOp::done:
@@ -588,8 +646,8 @@ void VectorVM::Exec(
 
 			// Opcode not recognized / implemented.
 			default:
-				UE_LOG(LogVectorVM, Fatal, TEXT("Unknown op code 0x%02x"), (uint32)Op);
-				break;
+				UE_LOG(LogVectorVM, Error, TEXT("Unknown op code 0x%02x"), (uint32)Op);
+				return;//BAIL
 			}
 		} while (Op != EOp::done);
 
@@ -597,109 +655,17 @@ void VectorVM::Exec(
 	}
 }
 
-namespace VectorVM
-{
-	static FVectorVMOpInfo GOpInfo[] =
-	{
-		FVectorVMOpInfo(EOp::done, EOpFlags::None, EOpSrc::Invalid, EOpSrc::Invalid, EOpSrc::Invalid, TEXT("done")),
-
-		FVectorVMOpInfo(EOp::add, EOpFlags::Implemented | EOpFlags::Commutative, EOpSrc::Register, EOpSrc::Register, EOpSrc::Invalid, TEXT("Add")),
-		FVectorVMOpInfo(),
-
-		FVectorVMOpInfo(EOp::sub, EOpFlags::Implemented, EOpSrc::Register, EOpSrc::Register, EOpSrc::Invalid, TEXT("Sub")),
-		FVectorVMOpInfo(),
-
-		FVectorVMOpInfo(EOp::mul, EOpFlags::Implemented | EOpFlags::Commutative, EOpSrc::Register, EOpSrc::Register, EOpSrc::Invalid, TEXT("Multiply")),
-		FVectorVMOpInfo(),
-
-		FVectorVMOpInfo(EOp::mad,EOpFlags::Implemented|EOpFlags::Commutative,EOpSrc::Register,EOpSrc::Register,EOpSrc::Register,TEXT("Multply-Add")),
-		FVectorVMOpInfo(),
-		FVectorVMOpInfo(),
-		FVectorVMOpInfo(),
-		FVectorVMOpInfo(),
-		FVectorVMOpInfo(),
-
-		FVectorVMOpInfo(EOp::lerp,EOpFlags::Implemented,EOpSrc::Register,EOpSrc::Register,EOpSrc::Register,TEXT("Lerp")),
-		FVectorVMOpInfo(),
-		FVectorVMOpInfo(), 
-		FVectorVMOpInfo(),
-		FVectorVMOpInfo(),
-
-		FVectorVMOpInfo(EOp::rcp,EOpFlags::Implemented,EOpSrc::Register,EOpSrc::Invalid,EOpSrc::Invalid,TEXT("Reciprocal")),
-		FVectorVMOpInfo(EOp::rsq,EOpFlags::Implemented,EOpSrc::Register,EOpSrc::Invalid,EOpSrc::Invalid,TEXT("Reciprocal Sqrt")),
-		FVectorVMOpInfo(EOp::sqrt,EOpFlags::Implemented,EOpSrc::Register,EOpSrc::Invalid,EOpSrc::Invalid,TEXT("Sqrt")),
-		FVectorVMOpInfo(EOp::neg,EOpFlags::Implemented,EOpSrc::Register,EOpSrc::Invalid,EOpSrc::Invalid,TEXT("Negate")),
-		FVectorVMOpInfo(EOp::abs,EOpFlags::Implemented,EOpSrc::Register,EOpSrc::Invalid,EOpSrc::Invalid,TEXT("Absolute")),
-		FVectorVMOpInfo(EOp::exp,EOpFlags::None,EOpSrc::Register,EOpSrc::Invalid,EOpSrc::Invalid,TEXT("Exp")),
-		FVectorVMOpInfo(EOp::exp2,EOpFlags::None,EOpSrc::Register,EOpSrc::Invalid,EOpSrc::Invalid,TEXT("Exp2")),
-		FVectorVMOpInfo(EOp::log,EOpFlags::None,EOpSrc::Register,EOpSrc::Invalid,EOpSrc::Invalid,TEXT("Log")),
-		FVectorVMOpInfo(EOp::log2,EOpFlags::None,EOpSrc::Register,EOpSrc::Invalid,EOpSrc::Invalid,TEXT("Log base 2")),
-		FVectorVMOpInfo(EOp::sin,EOpFlags::Implemented,EOpSrc::Register,EOpSrc::Invalid,EOpSrc::Invalid,TEXT("Sin")),
-		FVectorVMOpInfo(),
-		FVectorVMOpInfo(EOp::cos, EOpFlags::None, EOpSrc::Register, EOpSrc::Invalid, EOpSrc::Invalid, TEXT("Cos")),
-		FVectorVMOpInfo(EOp::tan,EOpFlags::None,EOpSrc::Register,EOpSrc::Invalid,EOpSrc::Invalid,TEXT("Tan")),
-		FVectorVMOpInfo(EOp::asin,EOpFlags::None,EOpSrc::Register,EOpSrc::Invalid,EOpSrc::Invalid,TEXT("Arcsin")),
-		FVectorVMOpInfo(EOp::acos,EOpFlags::None,EOpSrc::Register,EOpSrc::Invalid,EOpSrc::Invalid,TEXT("Arccos")),
-		FVectorVMOpInfo(EOp::atan,EOpFlags::None,EOpSrc::Register,EOpSrc::Invalid,EOpSrc::Invalid,TEXT("Arctan")),
-		FVectorVMOpInfo(EOp::atan2,EOpFlags::None,EOpSrc::Register,EOpSrc::Register,EOpSrc::Invalid,TEXT("Arctan2")),
-		FVectorVMOpInfo(EOp::ceil,EOpFlags::None,EOpSrc::Register,EOpSrc::Invalid,EOpSrc::Invalid,TEXT("Round up")),
-		FVectorVMOpInfo(EOp::floor,EOpFlags::None,EOpSrc::Register,EOpSrc::Invalid,EOpSrc::Invalid,TEXT("Round down")),
-		FVectorVMOpInfo(EOp::fmod,EOpFlags::None,EOpSrc::Register,EOpSrc::Register,EOpSrc::Invalid,TEXT("Modulo")),
-		FVectorVMOpInfo(EOp::frac,EOpFlags::None,EOpSrc::Register,EOpSrc::Invalid,EOpSrc::Invalid,TEXT("Fractional")),
-		FVectorVMOpInfo(EOp::trunc,EOpFlags::None,EOpSrc::Register,EOpSrc::Invalid,EOpSrc::Invalid,TEXT("Truncate")),
-
-		FVectorVMOpInfo(EOp::clamp,EOpFlags::Implemented,EOpSrc::Register,EOpSrc::Register,EOpSrc::Register,TEXT("Clamp")),
-		FVectorVMOpInfo(),
-		FVectorVMOpInfo(),
-		FVectorVMOpInfo(),
-
-		FVectorVMOpInfo(EOp::min,EOpFlags::Implemented,EOpSrc::Register,EOpSrc::Register,EOpSrc::Invalid,TEXT("Min")),
-		FVectorVMOpInfo(),
-
-		FVectorVMOpInfo(EOp::max,EOpFlags::Implemented,EOpSrc::Register,EOpSrc::Register,EOpSrc::Invalid,TEXT("Max")),
-		FVectorVMOpInfo(),
-
-		FVectorVMOpInfo(EOp::pow,EOpFlags::Implemented,EOpSrc::Register,EOpSrc::Register,EOpSrc::Invalid,TEXT("Pow")),
-		FVectorVMOpInfo(),
-
-		FVectorVMOpInfo(EOp::add,EOpFlags::None,EOpSrc::Register,EOpSrc::Invalid,EOpSrc::Invalid,TEXT("Sign")),
-
-		FVectorVMOpInfo(EOp::step,EOpFlags::None,EOpSrc::Register,EOpSrc::Register,EOpSrc::Invalid,TEXT("Step")),
-		FVectorVMOpInfo(),
-
-		FVectorVMOpInfo(),
-
-		FVectorVMOpInfo(EOp::dot, EOpFlags::Implemented, EOpSrc::Register, EOpSrc::Register, EOpSrc::Invalid, TEXT("Dot Product")),
-		FVectorVMOpInfo(EOp::cross, EOpFlags::Implemented|EOpFlags::Commutative, EOpSrc::Register, EOpSrc::Register, EOpSrc::Invalid, TEXT("Cross Product")),
-		FVectorVMOpInfo(),
-
-		FVectorVMOpInfo(EOp::normalize, EOpFlags::Implemented, EOpSrc::Register, EOpSrc::Invalid, EOpSrc::Invalid, TEXT("Normalize")),
-		FVectorVMOpInfo(EOp::random, EOpFlags::Implemented, EOpSrc::Const, EOpSrc::Invalid, EOpSrc::Invalid, TEXT("Random")),
-
-		FVectorVMOpInfo(EOp::length, EOpFlags::Implemented, EOpSrc::Register, EOpSrc::Invalid, EOpSrc::Invalid, TEXT("Vector Length")),
-		FVectorVMOpInfo(),
-
-		FVectorVMOpInfo(EOp::sin4, EOpFlags::Implemented, EOpSrc::Register, EOpSrc::Invalid, EOpSrc::Invalid, TEXT("Sin4")),
-		FVectorVMOpInfo(),
-
-		FVectorVMOpInfo(),
-		FVectorVMOpInfo(),
-		FVectorVMOpInfo(),
-		FVectorVMOpInfo(EOp::noise, EOpFlags::Implemented, EOpSrc::Register, EOpSrc::Invalid, EOpSrc::Invalid, TEXT("Noise")),
-
-		FVectorVMOpInfo(EOp::add, EOpFlags::None, EOpSrc::Invalid, EOpSrc::Invalid, EOpSrc::Invalid, TEXT("invalid"))
-	};
-} // namespace VectorVM
-
-VectorVM::FVectorVMOpInfo const& VectorVM::GetOpCodeInfo(uint8 OpCodeIndex)
+uint8 VectorVM::GetNumOpCodes()
 {
-	return GOpInfo[FMath::Clamp<int32>(OpCodeIndex, 0, EOp::NumOpcodes)];
+	return (uint8)EOp::NumOpcodes;
 }
 
-
-uint8 VectorVM::GetNumOpCodes()
+uint8 VectorVM::CreateSrcOperandMask(bool bIsOp0Constant, bool bIsOp1Constant, bool bIsOp2Constant, bool bIsOp3Constant)
 {
-	return EOp::NumOpcodes;
+	return	(bIsOp0Constant ? (1 << 0) : 0) |
+			(bIsOp1Constant ? (1 << 1) : 0) |
+			(bIsOp2Constant ? (1 << 2) : 0) |
+			(bIsOp3Constant ? (1 << 3) : 0) ;
 }
 
 /*------------------------------------------------------------------------------
@@ -712,12 +678,12 @@ bool FVectorVMTest::RunTest(const FString& Parameters)
 {
 	uint8 TestCode[] =
 	{
-		VectorVM::EOp::mul, 0x00, SRCOP_RRR, 0x0 + VectorVM::NumTempRegisters, 0x0 + VectorVM::NumTempRegisters,       // mul r0, r8, r8
-		VectorVM::EOp::mad, 0x01, SRCOP_RRR, 0x01 + VectorVM::NumTempRegisters, 0x01 + VectorVM::NumTempRegisters, 0x00, // mad r1, r9, r9, r0
-		VectorVM::EOp::mad, 0x00, SRCOP_RRR, 0x02 + VectorVM::NumTempRegisters, 0x02 + VectorVM::NumTempRegisters, 0x01, // mad r0, r10, r10, r1
-		VectorVM::EOp::add,		0x01, SRCOP_RCR, 0x00, 0x01,       // addi r1, r0, c1
-		VectorVM::EOp::neg,		0x00, SRCOP_RRR, 0x01,             // neg r0, r1
-		VectorVM::EOp::clamp,	VectorVM::FirstOutputRegister, SRCOP_CCR, 0x00, 0x02, 0x03, // clampii r40, r0, c2, c3
+		(uint8)VectorVM::EOp::mul, 0x00, SRCOP_RRRR, 0x0 + VectorVM::NumTempRegisters, 0x0 + VectorVM::NumTempRegisters,       // mul r0, r8, r8
+		(uint8)VectorVM::EOp::mad, 0x01, SRCOP_RRRR, 0x01 + VectorVM::NumTempRegisters, 0x01 + VectorVM::NumTempRegisters, 0x00, // mad r1, r9, r9, r0
+		(uint8)VectorVM::EOp::mad, 0x00, SRCOP_RRRR, 0x02 + VectorVM::NumTempRegisters, 0x02 + VectorVM::NumTempRegisters, 0x01, // mad r0, r10, r10, r1
+		(uint8)VectorVM::EOp::add, 0x01, SRCOP_RRCR, 0x00, 0x01,       // addi r1, r0, c1
+		(uint8)VectorVM::EOp::neg, 0x00, SRCOP_RRRR, 0x01,             // neg r0, r1
+		(uint8)VectorVM::EOp::clamp, VectorVM::FirstOutputRegister, SRCOP_RCCR, 0x00, 0x02, 0x03, // clampii r40, r0, c2, c3
 		0x00 // terminator
 	};
 
diff --git a/Engine/Source/Runtime/VectorVM/Public/VectorVM.h b/Engine/Source/Runtime/VectorVM/Public/VectorVM.h
--- a/Engine/Source/Runtime/VectorVM/Public/VectorVM.h
+++ b/Engine/Source/Runtime/VectorVM/Public/VectorVM.h
@@ -17,118 +17,70 @@ namespace VectorVM
 		FirstInputRegister = NumTempRegisters,
 		FirstOutputRegister = FirstInputRegister + MaxInputRegisters,
 		FirstConstantRegister = FirstOutputRegister + MaxOutputRegisters,
-		MaxRegisters = NumTempRegisters + MaxInputRegisters + MaxOutputRegisters + MaxConstants
+		MaxRegisters = NumTempRegisters + MaxInputRegisters + MaxOutputRegisters + MaxConstants,
 	};
 
 	/** List of opcodes supported by the VM. */
-	namespace EOp
+	enum class EOp
 	{
-		enum Type
-		{
-			done = 0,
-			add = 1,
-			sub = 3,
-			mul = 5,
-			mad = 7,
-			lerp = 13,
-			rcp = 18,
-			rsq = 19,
-			sqrt = 20,
-			neg = 21,
-			abs = 22,
-			exp = 23,
-			exp2 = 24,
-			log = 25,
-			log2 = 26,
-			sin = 27,
-			cos = 29,
-			tan = 30,
-			asin = 31,
-			acos = 32,
-			atan = 33,
-			atan2 = 34,
-			ceil = 35,
-			floor = 36,
-			fmod = 37,
-			frac = 38,
-			trunc = 39,
-			clamp = 40,
-			min = 44,
-			max = 46,
-			pow = 48,
-			sign = 50,
-			step = 51,
-			dot = 54,
-			cross = 55,
-			normalize = 57,
-			random = 58,
-			length = 59,
-			sin4 = 61, 
-			noise = 66,
+			done,
+			add,
+			sub,
+			mul,
+			mad,
+			lerp,
+			rcp,
+			rsq,
+			sqrt,
+			neg,
+			abs,
+			exp,
+			exp2,
+			log,
+			log2,
+			sin,
+			cos,
+			tan,
+			asin,
+			acos,
+			atan,
+			atan2,
+			ceil,
+			floor,
+			fmod,
+			frac,
+			trunc,
+			clamp,
+			min,
+			max,
+			pow,
+			sign,
+			step,
+			dot,
+			cross,
+			normalize,
+			random,
+			length,
+			sin4,
+			noise,
+			splatx,
+			splaty,
+			splatz,
+			splatw,
+			compose,
+			composex,
+			composey,
+			composez,
+			composew,
+			output,
 			NumOpcodes
-		};
-	} // namespace EOp
-
-	/** Operand sources. */
-	namespace EOpSrc
-	{
-		enum Type
-		{
-			Invalid,
-			Register,
-			Const
-		};
-	}
-
-	namespace EOpFlags
-	{
-		enum
-		{
-			None = 0x0,
-			Implemented = 0x1,
-			Commutative = 0x2
-		};
-	}
-
-	/** Struct used to give info on particular op codes */
-	struct VECTORVM_API FVectorVMOpInfo
-	{
-		/** Base opcode. */
-		EOp::Type	BaseOpcode;
-		/** Flags. */
-		uint8		Flags;
-		/** Source operands. */
-		EOpSrc::Type SrcTypes[3];
-		/** Friendly name of op-code */
-		FString		FriendlyName;
-
-		/** Constructor */
-		FVectorVMOpInfo(
-			EOp::Type InBaseOpcode = EOp::done,
-			uint8 InFlags = EOpFlags::None,
-			EOpSrc::Type Src0 = EOpSrc::Type::Invalid,
-			EOpSrc::Type Src1 = EOpSrc::Type::Invalid,
-			EOpSrc::Type Src2 = EOpSrc::Type::Invalid,
-			const FString& InFriendlyName = TEXT("Undefined Opcode")
-			)
-			: BaseOpcode(InBaseOpcode)
-			, Flags(InFlags)
-			, FriendlyName(InFriendlyName)
-		{
-			SrcTypes[0] = Src0;
-			SrcTypes[1] = Src1;
-			SrcTypes[2] = Src2;
-		}
-
-		FORCEINLINE bool IsImplemented() const { return (Flags & EOpFlags::Implemented) != 0; }
-		FORCEINLINE bool IsCommutative() const { return (Flags & EOpFlags::Commutative) != 0; }
 	};
-	
-	/** Get info on a particular op-code */
-	VECTORVM_API FVectorVMOpInfo const& GetOpCodeInfo(uint8 OpCodeIndex);
+
 	/** Get total number of op-codes */
 	VECTORVM_API uint8 GetNumOpCodes();
 
+	VECTORVM_API uint8 CreateSrcOperandMask(bool bIsOp0Constant, bool bIsOp1Constant = false, bool bIsOp2Constant = false, bool bIsOp3Constant = false);
+
 	/**
 	 * Execute VectorVM bytecode.
 	 */
