Compile fix. Erroneous ## preprocessor concatenation that clang doesn't like but MS was fine with.

diff --git a/Engine/Source/Editor/UnrealEd/Private/Niagara/EdGraphSchema_Niagara.cpp b/Engine/Source/Editor/UnrealEd/Private/Niagara/EdGraphSchema_Niagara.cpp
--- a/Engine/Source/Editor/UnrealEd/Private/Niagara/EdGraphSchema_Niagara.cpp
+++ b/Engine/Source/Editor/UnrealEd/Private/Niagara/EdGraphSchema_Niagara.cpp
@@ -171,7 +171,7 @@ void UEdGraphSchema_Niagara::GetGraphContextActions(FGraphContextMenuBuilder& Co
 	}
 
 #define NiagaraOp(OPNAME) \
-		if(const FNiagaraOpInfo* OpInfo = FNiagaraOpInfo::GetOpInfo(FNiagaraOpInfo::##OPNAME##))\
+		if(const FNiagaraOpInfo* OpInfo = FNiagaraOpInfo::GetOpInfo(FNiagaraOpInfo::OPNAME))\
 		{\
 			TSharedPtr<FNiagaraSchemaAction_NewNode> AddOpAction = AddNewNodeAction(ContextMenuBuilder, TEXT("Add Operation"), OpInfo->FriendlyName, TEXT(""));\
 			UNiagaraNodeOp* OpNode = NewObject<UNiagaraNodeOp>(ContextMenuBuilder.OwnerOfTemporaries); \
