Compile fix. Missing #undef NiagaraOp

diff --git a/Engine/Source/Editor/UnrealEd/Private/Niagara/NiagaraEditorCommon.cpp b/Engine/Source/Editor/UnrealEd/Private/Niagara/NiagaraEditorCommon.cpp
--- a/Engine/Source/Editor/UnrealEd/Private/Niagara/NiagaraEditorCommon.cpp
+++ b/Engine/Source/Editor/UnrealEd/Private/Niagara/NiagaraEditorCommon.cpp
@@ -534,6 +534,7 @@ void FNiagaraOpInfo::Init()
 			UE_LOG(LogNiagaraCompiler, Fatal, TEXT("Found info about \"%s\" in FNiagaraOpInfo::Init() that did not have an entry in the NiagaraOpList. You must add it to both."), *OpInfo.Key.ToString());
 		}
 	}
+#undef NiagaraOp
 }
 
 const FNiagaraOpInfo* FNiagaraOpInfo::GetOpInfo(FName Name)
