First phase of deprecating TSubobjectPtr<> replacing TSubobjectPtr properties in the Engine with normal pointers and adding subobject accessors for derived classes/client code.

diff --git a/Engine/Source/Developer/AssetTools/Private/AssetTypeActions/AssetTypeActions_CameraAnim.cpp b/Engine/Source/Developer/AssetTools/Private/AssetTypeActions/AssetTypeActions_CameraAnim.cpp
--- a/Engine/Source/Developer/AssetTools/Private/AssetTypeActions/AssetTypeActions_CameraAnim.cpp
+++ b/Engine/Source/Developer/AssetTools/Private/AssetTypeActions/AssetTypeActions_CameraAnim.cpp
@@ -50,10 +50,10 @@ void FAssetTypeActions_CameraAnim::CreateCameraActorForCameraAnim(UCameraAnim* I
 	PreviewCamera.Get()->SetActorLabel(FText::Format(LOCTEXT("CamerAnimPreviewCameraName", "Preview Camera - {0}"), FText::FromName(InCameraAnim->GetFName())).ToString());
 
 	// copy data from the CamAnim to the CameraActor
-	check(PreviewCamera.Get()->CameraComponent);
+	check(PreviewCamera.Get()->GetCameraComponent());
 	PreviewCamera.Get()->PreviewedCameraAnim = InCameraAnim;
-	PreviewCamera.Get()->CameraComponent->FieldOfView = InCameraAnim->BaseFOV;
-	PreviewCamera.Get()->CameraComponent->PostProcessSettings = InCameraAnim->BasePostProcessSettings;
+	PreviewCamera.Get()->GetCameraComponent()->FieldOfView = InCameraAnim->BaseFOV;
+	PreviewCamera.Get()->GetCameraComponent()->PostProcessSettings = InCameraAnim->BasePostProcessSettings;
 }
 
 void FAssetTypeActions_CameraAnim::CreatePreviewPawnForCameraAnim(UCameraAnim* InCameraAnim)
@@ -257,4 +257,4 @@ void FAssetTypeActions_CameraAnim::OnMatineeEditorClosed(FEdMode* InEditorMode)
 	}
 }
 
-#undef LOCTEXT_NAMESPACE
\ No newline at end of file
+#undef LOCTEXT_NAMESPACE
diff --git a/Engine/Source/Developer/FunctionalTesting/Classes/FunctionalTest.h b/Engine/Source/Developer/FunctionalTesting/Classes/FunctionalTest.h
--- a/Engine/Source/Developer/FunctionalTesting/Classes/FunctionalTest.h
+++ b/Engine/Source/Developer/FunctionalTesting/Classes/FunctionalTest.h
@@ -29,8 +29,10 @@ class AFunctionalTest : public AActor
 
 	static const uint32 DefaultTimeLimit = 60;	// seconds
 
+private:
 	UPROPERTY()
-	TSubobjectPtr<UBillboardComponent> SpriteComponent;
+	UBillboardComponent* SpriteComponent;
+public:
 
 	UPROPERTY(BlueprintReadWrite, Category=FunctionalTesting)
 	TEnumAsByte<EFunctionalTestResult::Type> Result;
@@ -122,4 +124,8 @@ public:
 protected:
 	uint32 bIsRunning;
 	float TotalTime;
-};
\ No newline at end of file
+
+public:
+	/** Returns SpriteComponent subobject **/
+	FORCEINLINE UBillboardComponent* GetSpriteComponent() { return SpriteComponent; }
+};
diff --git a/Engine/Source/Developer/GameplayDebugger/Private/GameplayDebuggingComponent.cpp b/Engine/Source/Developer/GameplayDebugger/Private/GameplayDebuggingComponent.cpp
--- a/Engine/Source/Developer/GameplayDebugger/Private/GameplayDebuggingComponent.cpp
+++ b/Engine/Source/Developer/GameplayDebugger/Private/GameplayDebuggingComponent.cpp
@@ -479,9 +479,9 @@ void UGameplayDebuggingComponent::CollectPathData()
 	PathErrorString.Empty();
 	if (AAIController *MyController = Cast<AAIController>(MyPawn->Controller))
 	{
-		if ( MyController->NavComponent && MyController->NavComponent->HasValidPath())
+		if ( MyController->GetNavComponent() && MyController->GetNavComponent()->HasValidPath())
 		{
-			const FNavPathSharedPtr& NewPath = MyController->NavComponent->GetPath();
+			const FNavPathSharedPtr& NewPath = MyController->GetNavComponent()->GetPath();
 			if (!CurrentPath.HasSameObject(NewPath.Get()))
 			{
 				PathPoints.Reset();
diff --git a/Engine/Source/Developer/MovieSceneTools/Private/DirectorTrackEditor.cpp b/Engine/Source/Developer/MovieSceneTools/Private/DirectorTrackEditor.cpp
--- a/Engine/Source/Developer/MovieSceneTools/Private/DirectorTrackEditor.cpp
+++ b/Engine/Source/Developer/MovieSceneTools/Private/DirectorTrackEditor.cpp
@@ -364,7 +364,7 @@ void FShotSection::CalculateThumbnailWidthAndResize()
 	float SavedTime = Sequencer.Pin()->GetGlobalTime();
 	
 	Sequencer.Pin()->SetGlobalTime(Section->GetStartTime());
-	uint32 OutThumbnailWidth = AnimatableShotToolConstants::DirectorTrackHeight * Camera->CameraComponent->AspectRatio;
+	uint32 OutThumbnailWidth = AnimatableShotToolConstants::DirectorTrackHeight * Camera->GetCameraComponent()->AspectRatio;
 
 	// restore the global time
 	Sequencer.Pin()->SetGlobalTime(SavedTime);
diff --git a/Engine/Source/Editor/DetailCustomizations/Private/AmbientSoundDetails.cpp b/Engine/Source/Editor/DetailCustomizations/Private/AmbientSoundDetails.cpp
--- a/Engine/Source/Editor/DetailCustomizations/Private/AmbientSoundDetails.cpp
+++ b/Engine/Source/Editor/DetailCustomizations/Private/AmbientSoundDetails.cpp
@@ -106,7 +106,7 @@ bool FAmbientSoundDetails::IsEditSoundCueEnabled() const
 	if (AmbientSound.IsValid())
 	{
 		// Only sound cues are editable
-		USoundCue* SoundCue = Cast<USoundCue>(AmbientSound.Get()->AudioComponent->Sound);
+		USoundCue* SoundCue = Cast<USoundCue>(AmbientSound.Get()->GetAudioComponent()->Sound);
 		return (SoundCue != NULL);
 	}
 	return false;
@@ -116,7 +116,7 @@ FReply FAmbientSoundDetails::OnEditSoundCueClicked()
 {
 	if( AmbientSound.IsValid() )
 	{
-		USoundCue* SoundCue = Cast<USoundCue>(AmbientSound.Get()->AudioComponent->Sound);
+		USoundCue* SoundCue = Cast<USoundCue>(AmbientSound.Get()->GetAudioComponent()->Sound);
 		if (SoundCue)
 		{
 			FAssetEditorManager::Get().OpenEditorForAsset(SoundCue);
@@ -155,7 +155,7 @@ void FAmbientSoundDetails::CreateNewSoundCue( ESoundCueLayouts Layout )
 		AAmbientSound* AS = CastChecked<AAmbientSound>(AmbientSound.Get());
 
 		// First if the existing SoundCue is a child of the AmbientSound rename it off to oblivion so we can have a good name
-		USoundCue* SoundCue = Cast<USoundCue>(AmbientSound->AudioComponent->Sound);
+		USoundCue* SoundCue = Cast<USoundCue>(AmbientSound->GetAudioComponent()->Sound);
 		if (SoundCue && SoundCue->GetOuter() == AS)
 		{
 			SoundCue->Rename(*MakeUniqueObjectName( GetTransientPackage(), USoundCue::StaticClass() ).ToString(), GetTransientPackage(), REN_DontCreateRedirectors);
@@ -165,8 +165,8 @@ void FAmbientSoundDetails::CreateNewSoundCue( ESoundCueLayouts Layout )
 		USoundNode* PrevNode = NULL;
 
 		SoundCue = ConstructObject<USoundCue>(USoundCue::StaticClass(), AS, FName(*AS->GetInternalSoundCueName()));
-		AS->AudioComponent->Sound = SoundCue;
-		AS->AudioComponent->PostEditChange();
+		AS->GetAudioComponent()->Sound = SoundCue;
+		AS->GetAudioComponent()->PostEditChange();
 
 		if (Layout == SOUNDCUE_LAYOUT_RANDOM_LOOP || Layout == SOUNDCUE_LAYOUT_RANDOM_LOOP_WITH_DELAY)
 		{
@@ -264,7 +264,7 @@ FReply FAmbientSoundDetails::OnPlaySoundClicked()
 {
 	if( AmbientSound.IsValid() )
 	{
-		USoundBase* Sound = Cast<USoundBase>(AmbientSound.Get()->AudioComponent->Sound);
+		USoundBase* Sound = Cast<USoundBase>(AmbientSound.Get()->GetAudioComponent()->Sound);
 		if (Sound)
 		{
 			GEditor->PlayPreviewSound(Sound);
diff --git a/Engine/Source/Editor/DetailCustomizations/Private/ParticleSystemComponentDetails.cpp b/Engine/Source/Editor/DetailCustomizations/Private/ParticleSystemComponentDetails.cpp
--- a/Engine/Source/Editor/DetailCustomizations/Private/ParticleSystemComponentDetails.cpp
+++ b/Engine/Source/Editor/DetailCustomizations/Private/ParticleSystemComponentDetails.cpp
@@ -84,7 +84,7 @@ FReply FParticleSystemComponentDetails::OnResetEmitter()
 			{
 				if (AEmitter* Emitter = Cast<AEmitter>(SelectedObjects[Idx].Get()))
 				{
-					PSC = Emitter->ParticleSystemComponent;
+					PSC = Emitter->GetParticleSystemComponent();
 				}
 			}
 
diff --git a/Engine/Source/Editor/FoliageEdit/Private/FoliageEdMode.cpp b/Engine/Source/Editor/FoliageEdit/Private/FoliageEdMode.cpp
--- a/Engine/Source/Editor/FoliageEdit/Private/FoliageEdMode.cpp
+++ b/Engine/Source/Editor/FoliageEdit/Private/FoliageEdMode.cpp
@@ -1664,7 +1664,7 @@ void FEdModeFoliage::BakeFoliage(UFoliageType* Settings, bool bSelectedOnly)
 			UWorld* World = IFA->GetWorld();
 			check(World != nullptr);
 			AStaticMeshActor* SMA = World->SpawnActor<AStaticMeshActor>(Instance.Location, Instance.Rotation);
-			SMA->StaticMeshComponent->StaticMesh = Settings->GetStaticMesh();
+			SMA->GetStaticMeshComponent()->StaticMesh = Settings->GetStaticMesh();
 			SMA->GetRootComponent()->SetRelativeScale3D(Instance.DrawScale3D);
 			SMA->MarkComponentsRenderStateDirty();
 		}
diff --git a/Engine/Source/Editor/GeometryMode/Private/GeometryModifiers.cpp b/Engine/Source/Editor/GeometryMode/Private/GeometryModifiers.cpp
--- a/Engine/Source/Editor/GeometryMode/Private/GeometryModifiers.cpp
+++ b/Engine/Source/Editor/GeometryMode/Private/GeometryModifiers.cpp
@@ -1702,7 +1702,7 @@ static ABrush* ClipBrushAgainstPlane( const FPlane& InPlane, ABrush* InBrush)
 	check( ClippedBrush->GetClass() == InBrush->GetClass() );
 
 	ClippedBrush->Brush = new( InBrush->GetOuter(), NAME_None )UModel( FObjectInitializer(),NULL );
-	ClippedBrush->BrushComponent->Brush = ClippedBrush->Brush;
+	ClippedBrush->GetBrushComponent()->Brush = ClippedBrush->Brush;
 
 	GeometryClipping::BuildGiantAlignedBrush( *ClippedBrush, InPlane );
 
@@ -3156,4 +3156,4 @@ bool UGeomModifier_Weld::OnApply()
 	return true;
 }
 
-#undef LOCTEXT_NAMESPACE 
\ No newline at end of file
+#undef LOCTEXT_NAMESPACE 
diff --git a/Engine/Source/Editor/LandscapeEditor/Private/ActorFactoryLandscape.cpp b/Engine/Source/Editor/LandscapeEditor/Private/ActorFactoryLandscape.cpp
--- a/Engine/Source/Editor/LandscapeEditor/Private/ActorFactoryLandscape.cpp
+++ b/Engine/Source/Editor/LandscapeEditor/Private/ActorFactoryLandscape.cpp
@@ -45,13 +45,13 @@ ALandscapePlaceholder::ALandscapePlaceholder(const FObjectInitializer& ObjectIni
 	};
 	static FConstructorStatics ConstructorStatics;
 
-	TSubobjectPtr<USceneComponent> SceneComponent = ObjectInitializer.CreateDefaultSubobject<USceneComponent>(this, TEXT("RootComponent0"));
-	RootComponent = SceneComponent.Get();
+	USceneComponent* SceneComponent = ObjectInitializer.CreateDefaultSubobject<USceneComponent>(this, TEXT("RootComponent0"));
+	RootComponent = SceneComponent;
 	RootComponent->Mobility = EComponentMobility::Static;
 
 #if WITH_EDITORONLY_DATA
-	TSubobjectPtr<UBillboardComponent> SpriteComponent = ObjectInitializer.CreateEditorOnlyDefaultSubobject<UBillboardComponent>(this, TEXT("Sprite"));
-	if (SpriteComponent.IsValid())
+	UBillboardComponent* SpriteComponent = ObjectInitializer.CreateEditorOnlyDefaultSubobject<UBillboardComponent>(this, TEXT("Sprite"));
+	if (SpriteComponent)
 	{
 		SpriteComponent->Sprite = ConstructorStatics.TerrainTexture.Get();
 		SpriteComponent->RelativeScale3D = FVector(0.5f, 0.5f, 0.5f);
diff --git a/Engine/Source/Editor/LevelEditor/Private/LevelEditorActions.cpp b/Engine/Source/Editor/LevelEditor/Private/LevelEditorActions.cpp
--- a/Engine/Source/Editor/LevelEditor/Private/LevelEditorActions.cpp
+++ b/Engine/Source/Editor/LevelEditor/Private/LevelEditorActions.cpp
@@ -2297,7 +2297,7 @@ void FLevelEditorActionCallbacks::OnSaveBrushAsCollision()
 		UStaticMeshComponent* FoundStaticMeshComponent = NULL;
 		if( Actor->IsA(AStaticMeshActor::StaticClass()) )
 		{
-			FoundStaticMeshComponent = CastChecked<AStaticMeshActor>(Actor)->StaticMeshComponent;
+			FoundStaticMeshComponent = CastChecked<AStaticMeshActor>(Actor)->GetStaticMeshComponent();
 		}
 
 		UStaticMesh* FoundMesh = FoundStaticMeshComponent ? FoundStaticMeshComponent->StaticMesh : NULL;
diff --git a/Engine/Source/Editor/LevelEditor/Private/SLevelViewportToolBar.cpp b/Engine/Source/Editor/LevelEditor/Private/SLevelViewportToolBar.cpp
--- a/Engine/Source/Editor/LevelEditor/Private/SLevelViewportToolBar.cpp
+++ b/Engine/Source/Editor/LevelEditor/Private/SLevelViewportToolBar.cpp
@@ -1069,7 +1069,7 @@ void SLevelViewportToolBar::OnFOVValueChanged( float NewValue )
 		ACameraActor* CameraActor = Cast< ACameraActor >( ViewportClient.GetActiveActorLock().Get() );
 		if( CameraActor != NULL )
 		{
-			CameraActor->CameraComponent->FieldOfView = NewValue;
+			CameraActor->GetCameraComponent()->FieldOfView = NewValue;
 			bUpdateStoredFOV = false;
 		}
 	}
diff --git a/Engine/Source/Editor/Matinee/Private/Matinee.cpp b/Engine/Source/Editor/Matinee/Private/Matinee.cpp
--- a/Engine/Source/Editor/Matinee/Private/Matinee.cpp
+++ b/Engine/Source/Editor/Matinee/Private/Matinee.cpp
@@ -1988,7 +1988,7 @@ void FMatinee::OnClose()
 		// Set any manipulated cameras back to default frustum colours.
 		if (ACameraActor* Cam = Cast<ACameraActor>(MatineeActor->GroupInst[i]->GroupActor))
 		{
-			Cam->CameraComponent->RestoreFrustumColor();
+			Cam->GetCameraComponent()->RestoreFrustumColor();
 		}
 	}
 
diff --git a/Engine/Source/Editor/Matinee/Private/MatineeMenus.cpp b/Engine/Source/Editor/Matinee/Private/MatineeMenus.cpp
--- a/Engine/Source/Editor/Matinee/Private/MatineeMenus.cpp
+++ b/Engine/Source/Editor/Matinee/Private/MatineeMenus.cpp
@@ -1539,7 +1539,7 @@ void FMatinee::OnContextTrackExportAnimFBX()
 		// Someone could have hooked up an invalid actor.  In that case do nothing
 		if( SkelMeshActor )
 		{
-			SkelMesh = SkelMeshActor->SkeletalMeshComponent->SkeletalMesh;
+			SkelMesh = SkelMeshActor->GetSkeletalMeshComponent()->SkeletalMesh;
 		}
 		
 
@@ -1921,7 +1921,7 @@ void FMatinee::OnContextGroupExportAnimFBX()
 					// Someone could have hooked up an invalid actor.  In that case do nothing
 					if( SkelMeshActor )
 					{
-						SkelMeshComponent = SkelMeshActor->SkeletalMeshComponent;
+						SkelMeshComponent = SkelMeshActor->GetSkeletalMeshComponent();
 					}
 				}
 
diff --git a/Engine/Source/Editor/Matinee/Private/MatineeTools.cpp b/Engine/Source/Editor/Matinee/Private/MatineeTools.cpp
--- a/Engine/Source/Editor/Matinee/Private/MatineeTools.cpp
+++ b/Engine/Source/Editor/Matinee/Private/MatineeTools.cpp
@@ -173,9 +173,9 @@ void  FMatinee::UpdateCameraRecording (void)
 
 						const FTransform Transform(LevelVC->GetViewRotation(), LevelVC->GetViewLocation());
 						ACameraActor* NewCam = Cast<ACameraActor>(GEditor->AddActor( LevelVC->GetWorld()->GetCurrentLevel(), ACameraActor::StaticClass(), Transform ));
-						NewCam->CameraComponent->bConstrainAspectRatio = LevelVC->IsAspectRatioConstrained();
-						NewCam->CameraComponent->AspectRatio = LevelVC->AspectRatio;
-						NewCam->CameraComponent->FieldOfView = LevelVC->ViewFOV;
+						NewCam->GetCameraComponent()->bConstrainAspectRatio = LevelVC->IsAspectRatioConstrained();
+						NewCam->GetCameraComponent()->AspectRatio = LevelVC->AspectRatio;
+						NewCam->GetCameraComponent()->FieldOfView = LevelVC->ViewFOV;
 
 						//make new group for the camera
 						UInterpGroup* NewGroup = ConstructObject<UInterpGroup>( UInterpGroup::StaticClass(), IData, NAME_None, RF_Transactional );
@@ -3790,7 +3790,7 @@ void FMatinee::UpdateCamColours()
 		if(ACameraActor* Cam = Cast<ACameraActor>(Inst->GetGroupActor()))
 		{
 			const FColor OverrideColor = (Inst->GetGroupActor() == ViewedActor) ? ActiveCamColor : Inst->Group->GroupColor;
-			Cam->CameraComponent->OverrideFrustumColor(OverrideColor);
+			Cam->GetCameraComponent()->OverrideFrustumColor(OverrideColor);
 		}
 	}
 }
@@ -3866,8 +3866,8 @@ void FMatinee::UpdateLevelViewport(AActor* InActor, FLevelEditorViewportClient*
 		ACameraActor* CameraActor = Cast<ACameraActor>(InActor);
 		if (CameraActor)
 		{
-			CameraActor->CameraComponent->FieldOfView = InViewportClient->ViewFOV;
-			CameraActor->CameraComponent->AspectRatio = InViewportClient->AspectRatio;
+			CameraActor->GetCameraComponent()->FieldOfView = InViewportClient->ViewFOV;
+			CameraActor->GetCameraComponent()->AspectRatio = InViewportClient->AspectRatio;
 			CameraActor->SetActorLocation(InViewportClient->GetViewLocation(), false);
 			CameraActor->SetActorRotation(InViewportClient->GetViewRotation() );
 		}
@@ -3917,20 +3917,20 @@ void FMatinee::UpdateLevelViewport(AActor* InActor, FLevelEditorViewportClient*
 	{
 		// If the Camera's aspect ratio is zero, put a more reasonable default here - this at least stops it from crashing
 		// nb. the AspectRatio will be reported as a Map Check Warning
-		if( Cam->CameraComponent->AspectRatio == 0 )
+		if( Cam->GetCameraComponent()->AspectRatio == 0 )
 		{
 			InViewportClient->AspectRatio = 1.7f;
 		}
 		else
 		{
-			InViewportClient->AspectRatio = Cam->CameraComponent->AspectRatio;
+			InViewportClient->AspectRatio = Cam->GetCameraComponent()->AspectRatio;
 		}
 
 		//if this isn't the recording viewport OR (it is the recording viewport and it's playing or we're scrubbing the timeline
 		if ((InViewportClient != GetRecordingViewport()) || (MatineeActor && ( MatineeActor->bIsPlaying || (TrackWindow.IsValid() && TrackWindow->InterpEdVC->bGrabbingHandle ) ) ))
 		{
 			//don't stop the camera from zooming when not playing back
-			InViewportClient->ViewFOV = Cam->CameraComponent->FieldOfView;
+			InViewportClient->ViewFOV = Cam->GetCameraComponent()->FieldOfView;
 
 			// If there are selected actors, invalidate the viewports hit proxies, otherwise they won't be selectable afterwards
 			if ( InViewportClient->Viewport && GEditor->GetSelectedActorCount() > 0 )
diff --git a/Engine/Source/Editor/MeshPaint/Private/MeshPaintEdMode.cpp b/Engine/Source/Editor/MeshPaint/Private/MeshPaintEdMode.cpp
--- a/Engine/Source/Editor/MeshPaint/Private/MeshPaintEdMode.cpp
+++ b/Engine/Source/Editor/MeshPaint/Private/MeshPaintEdMode.cpp
@@ -2820,7 +2820,7 @@ void FEdModeMeshPaint::SaveSettingsForActor( AActor* InActor )
 		UStaticMeshComponent* StaticMeshComponent = NULL;
 		if( StaticMeshActor != NULL )
 		{
-			StaticMeshComponent = StaticMeshActor->StaticMeshComponent;
+			StaticMeshComponent = StaticMeshActor->GetStaticMeshComponent();
 		}
 
 		if( StaticMeshComponent != NULL )
@@ -4966,9 +4966,9 @@ void FEdModeMeshPaint::DuplicateTextureMaterialCombo()
 		{
 			int32 MaterialIndex = MeshData->SelectedMaterialIndex;
 			AStaticMeshActor* StaticMeshActor = Cast< AStaticMeshActor >( ActorBeingEdited.Get() );
-			if (StaticMeshActor != NULL && StaticMeshActor->StaticMeshComponent != NULL)
+			if (StaticMeshActor != NULL && StaticMeshActor->GetStaticMeshComponent() != NULL)
 			{
-				UMaterialInterface* MaterialToCheck = StaticMeshActor->StaticMeshComponent->GetMaterial(MaterialIndex);
+				UMaterialInterface* MaterialToCheck = StaticMeshActor->GetStaticMeshComponent()->GetMaterial(MaterialIndex);
 
 				bool bIsSourceTextureStreamedIn = SelectedTexture->IsFullyStreamedIn();
 
@@ -5060,7 +5060,7 @@ void FEdModeMeshPaint::DuplicateTextureMaterialCombo()
 				}
 
 				bool bMaterialChanged = false;
-				UStaticMeshComponent* SMComponent = StaticMeshActor->StaticMeshComponent;
+				UStaticMeshComponent* SMComponent = StaticMeshActor->GetStaticMeshComponent();
 				ClearStaticMeshTextureOverrides(SMComponent);
 
 				SMComponent->SetMaterial(MaterialIndex,NewMaterialInstance);
diff --git a/Engine/Source/Editor/MeshPaint/Private/SMeshPaint.cpp b/Engine/Source/Editor/MeshPaint/Private/SMeshPaint.cpp
--- a/Engine/Source/Editor/MeshPaint/Private/SMeshPaint.cpp
+++ b/Engine/Source/Editor/MeshPaint/Private/SMeshPaint.cpp
@@ -2326,7 +2326,7 @@ FReply SMeshPaint::SaveVertexPaintPackageButtonClicked()
 		AStaticMeshActor* StaticMeshActor = Cast< AStaticMeshActor >( Actor );
 		if( StaticMeshActor != NULL )
 		{
-			StaticMeshComponent = StaticMeshActor->StaticMeshComponent;
+			StaticMeshComponent = StaticMeshActor->GetStaticMeshComponent();
 		}
 
 		if( StaticMeshComponent != NULL && StaticMeshComponent->StaticMesh != NULL )
@@ -2451,7 +2451,7 @@ bool SMeshPaint::IsPushInstanceVertexColorsToMeshButtonEnabled() const
 			AStaticMeshActor* LeftStaticMeshActor = Cast< AStaticMeshActor >( LeftCompareActor );
 			if( LeftStaticMeshActor != NULL )
 			{
-				LeftStaticMeshComponent = LeftStaticMeshActor->StaticMeshComponent;
+				LeftStaticMeshComponent = LeftStaticMeshActor->GetStaticMeshComponent();
 			}
 
 			if( LeftStaticMeshComponent != NULL && LeftStaticMeshComponent->StaticMesh != NULL )
@@ -2464,7 +2464,7 @@ bool SMeshPaint::IsPushInstanceVertexColorsToMeshButtonEnabled() const
 					AStaticMeshActor* RightStaticMeshActor = Cast< AStaticMeshActor >( RightCompareActor );
 					if( RightStaticMeshActor != NULL )
 					{
-						RightStaticMeshComponent = RightStaticMeshActor->StaticMeshComponent;
+						RightStaticMeshComponent = RightStaticMeshActor->GetStaticMeshComponent();
 					}
 					
 					if(RightStaticMeshComponent != NULL && RightStaticMeshComponent->StaticMesh != NULL)
@@ -2494,7 +2494,7 @@ bool SMeshPaint::IsSaveVertexPaintPackageButtonEnabled() const
 		AStaticMeshActor* StaticMeshActor = Cast< AStaticMeshActor >( Actor );
 		if( StaticMeshActor != NULL )
 		{
-			StaticMeshComponent = StaticMeshActor->StaticMeshComponent;
+			StaticMeshComponent = StaticMeshActor->GetStaticMeshComponent();
 		}
 
 		if( StaticMeshComponent != NULL && StaticMeshComponent->StaticMesh != NULL )
diff --git a/Engine/Source/Editor/Persona/Private/AnimationEditorViewportClient.cpp b/Engine/Source/Editor/Persona/Private/AnimationEditorViewportClient.cpp
--- a/Engine/Source/Editor/Persona/Private/AnimationEditorViewportClient.cpp
+++ b/Engine/Source/Editor/Persona/Private/AnimationEditorViewportClient.cpp
@@ -1838,7 +1838,7 @@ TWeakObjectPtr<AWindDirectionalSource> FAnimationViewportClient::CreateWindActor
 
 	check(Wind.IsValid());
 	//initial wind strength value 
-	Wind->Component->Strength = PrevWindStrength;
+	Wind->GetComponent()->Strength = PrevWindStrength;
 	return Wind;
 }
 
@@ -1858,7 +1858,7 @@ void FAnimationViewportClient::EnableWindActor(bool bEnableWind)
 	{
 		PrevWindLocation = WindSourceActor->GetActorLocation();
 		PrevWindRotation = WindSourceActor->GetActorRotation();
-		PrevWindStrength = WindSourceActor->Component->Strength;
+		PrevWindStrength = WindSourceActor->GetComponent()->Strength;
 
 		if(World->DestroyActor(WindSourceActor.Get()))
 		{
@@ -1873,7 +1873,7 @@ void FAnimationViewportClient::SetWindStrength( float SliderPos )
 	if(WindSourceActor.IsValid())
 	{
 		//Clamp grid size slider value between 0 - 1
-		WindSourceActor->Component->Strength = FMath::Clamp<float>(SliderPos, 0.0f, 1.0f);
+		WindSourceActor->GetComponent()->Strength = FMath::Clamp<float>(SliderPos, 0.0f, 1.0f);
 		//to apply this new wind strength
 		WindSourceActor->UpdateComponentTransforms();
 	}
@@ -1883,7 +1883,7 @@ float FAnimationViewportClient::GetWindStrengthSliderValue() const
 {
 	if(WindSourceActor.IsValid())
 	{
-		return WindSourceActor->Component->Strength;
+		return WindSourceActor->GetComponent()->Strength;
 	}
 
 	return 0;
diff --git a/Engine/Source/Editor/UnrealEd/Private/AssetSelection.cpp b/Engine/Source/Editor/UnrealEd/Private/AssetSelection.cpp
--- a/Engine/Source/Editor/UnrealEd/Private/AssetSelection.cpp
+++ b/Engine/Source/Editor/UnrealEd/Private/AssetSelection.cpp
@@ -231,9 +231,9 @@ namespace AssetSelectionUtils
 					{
 						ActorInfo.bHaveStaticMesh = true;
 						AStaticMeshActor* StaticMeshActor = CastChecked<AStaticMeshActor>( CurrentActor );
-						if ( StaticMeshActor->StaticMeshComponent )
+						if ( StaticMeshActor->GetStaticMeshComponent() )
 						{
-							UStaticMesh* StaticMesh = StaticMeshActor->StaticMeshComponent->StaticMesh;
+							UStaticMesh* StaticMesh = StaticMeshActor->GetStaticMeshComponent()->StaticMesh;
 
 							ActorInfo.bAllSelectedStaticMeshesHaveCollisionModels &= ( (StaticMesh && StaticMesh->BodySetup) ? true : false );
 						}
diff --git a/Engine/Source/Editor/UnrealEd/Private/BSPOps.cpp b/Engine/Source/Editor/UnrealEd/Private/BSPOps.cpp
--- a/Engine/Source/Editor/UnrealEd/Private/BSPOps.cpp
+++ b/Engine/Source/Editor/UnrealEd/Private/BSPOps.cpp
@@ -514,7 +514,7 @@ void FBSPOps::csgPrepMovingBrush( ABrush* Actor )
 	check(Actor);
 //	UE_LOG(LogBSPOps, Log,  TEXT("Preparing brush %s"), *Actor->GetName() );  // moved here so that we can easily debug when an actor has lost parts of its brush
 
-	check(Actor->BrushComponent);
+	check(Actor->GetBrushComponent());
 	check(Actor->Brush);
 	check(Actor->Brush->RootOutside);
 
@@ -523,7 +523,7 @@ void FBSPOps::csgPrepMovingBrush( ABrush* Actor )
 	RebuildBrush(Actor->Brush);
 
 	// Make sure simplified collision is up to date.
-	Actor->BrushComponent->BuildSimpleBrushCollision();
+	Actor->GetBrushComponent()->BuildSimpleBrushCollision();
 }
 
 /**
@@ -533,14 +533,14 @@ void FBSPOps::csgPrepMovingBrush( ABrush* Actor )
 void FBSPOps::csgCopyBrush( ABrush* Dest, ABrush* Src, uint32 PolyFlags, EObjectFlags ResFlags, bool bNeedsPrep, bool bCopyPosRotScale, bool bAllowEmpty )
 {
 	check(Src);
-	check(Src->BrushComponent);
+	check(Src->GetBrushComponent());
 	check(Src->Brush);
 
 	// Handle empty brush.
 	if( !bAllowEmpty && !Src->Brush->Polys->Element.Num() )
 	{
 		Dest->Brush = NULL;
-		Dest->BrushComponent->Brush = NULL;
+		Dest->GetBrushComponent()->Brush = NULL;
 		return;
 	}
 
@@ -551,7 +551,7 @@ void FBSPOps::csgCopyBrush( ABrush* Dest, ABrush* Src, uint32 PolyFlags, EObject
 	check(Dest->Brush->Polys->Element.GetOwner()==Dest->Brush->Polys);
 	Dest->Brush->Polys->Element.AssignButKeepOwner(Src->Brush->Polys->Element);
 	check(Dest->Brush->Polys->Element.GetOwner()==Dest->Brush->Polys);
-	Dest->BrushComponent->Brush = Dest->Brush;
+	Dest->GetBrushComponent()->Brush = Dest->Brush;
 	if(Src->BrushBuilder != nullptr)
 	{
 		Dest->BrushBuilder = DuplicateObject<UBrushBuilder>(Src->BrushBuilder, Dest);
@@ -584,7 +584,7 @@ void FBSPOps::csgCopyBrush( ABrush* Dest, ABrush* Src, uint32 PolyFlags, EObject
 ABrush*	FBSPOps::csgAddOperation( ABrush* Actor, uint32 PolyFlags, EBrushType BrushType)
 {
 	check(Actor);
-	check(Actor->BrushComponent);
+	check(Actor->GetBrushComponent());
 	check(Actor->Brush);
 	check(Actor->Brush->Polys);
 	check(Actor->GetWorld());
@@ -1215,11 +1215,11 @@ void FBSPOps::RebuildBrush(UModel* Brush)
  */
 void FBSPOps::RotateBrushVerts(ABrush* Brush, const FRotator& Rotation, bool bClearComponents)
 {
-	if(Brush->BrushComponent->Brush && Brush->BrushComponent->Brush->Polys)
+	if(Brush->GetBrushComponent()->Brush && Brush->GetBrushComponent()->Brush->Polys)
 	{
-		for( int32 poly = 0 ; poly < Brush->BrushComponent->Brush->Polys->Element.Num() ; poly++ )
+		for( int32 poly = 0 ; poly < Brush->GetBrushComponent()->Brush->Polys->Element.Num() ; poly++ )
 		{
-			FPoly* Poly = &(Brush->BrushComponent->Brush->Polys->Element[poly]);
+			FPoly* Poly = &(Brush->GetBrushComponent()->Brush->Polys->Element[poly]);
 
 			// Rotate the vertices.
 			for( int32 vertex = 0 ; vertex < Poly->Vertices.Num() ; vertex++ )
@@ -1237,7 +1237,7 @@ void FBSPOps::RotateBrushVerts(ABrush* Brush, const FRotator& Rotation, bool bCl
 			Poly->Finalize(Brush,0);
 		}
 
-		Brush->BrushComponent->Brush->BuildBound();
+		Brush->GetBrushComponent()->Brush->BuildBound();
 
 		if( !Brush->IsStaticBrush() )
 		{
diff --git a/Engine/Source/Editor/UnrealEd/Private/Dialogs/SMeshMergingDialog.cpp b/Engine/Source/Editor/UnrealEd/Private/Dialogs/SMeshMergingDialog.cpp
--- a/Engine/Source/Editor/UnrealEd/Private/Dialogs/SMeshMergingDialog.cpp
+++ b/Engine/Source/Editor/UnrealEd/Private/Dialogs/SMeshMergingDialog.cpp
@@ -441,7 +441,7 @@ void SMeshMergingDialog::RunMerging()
 			FRotator MergedActorRotation(ForceInit);
 
 			AStaticMeshActor* MergedActor = World->SpawnActor<AStaticMeshActor>(MergedActorLocation, MergedActorRotation, Params);
-			MergedActor->StaticMeshComponent->StaticMesh = Cast<UStaticMesh>(AssetsToSync[0]);
+			MergedActor->GetStaticMeshComponent()->StaticMesh = Cast<UStaticMesh>(AssetsToSync[0]);
 			MergedActor->SetActorLabel(AssetsToSync[0]->GetName());
 
 			// Add source actors as children to merged actor and hide them
@@ -456,4 +456,4 @@ void SMeshMergingDialog::RunMerging()
 }
 
 
-#undef LOCTEXT_NAMESPACE
\ No newline at end of file
+#undef LOCTEXT_NAMESPACE
diff --git a/Engine/Source/Editor/UnrealEd/Private/Dialogs/SMeshProxyDialog.cpp b/Engine/Source/Editor/UnrealEd/Private/Dialogs/SMeshProxyDialog.cpp
--- a/Engine/Source/Editor/UnrealEd/Private/Dialogs/SMeshProxyDialog.cpp
+++ b/Engine/Source/Editor/UnrealEd/Private/Dialogs/SMeshProxyDialog.cpp
@@ -615,7 +615,7 @@ FReply SMeshProxyDialog::OnMergeClicked()
 			{
 				UWorld* World = Actors[0]->GetWorld();
 				AStaticMeshActor* WorldActor = Cast<AStaticMeshActor>(World->SpawnActor(AStaticMeshActor::StaticClass(), &ProxyLocation));
-				WorldActor->StaticMeshComponent->StaticMesh = ProxyMesh;
+				WorldActor->GetStaticMeshComponent()->StaticMesh = ProxyMesh;
 				WorldActor->SetActorLabel(ProxyMesh->GetName());
 			}
 		}
@@ -709,4 +709,4 @@ FReply SMeshProxyDialog::OnSelectPackageNameClicked()
 	return FReply::Handled();
 }
 
-#undef LOCTEXT_NAMESPACE
\ No newline at end of file
+#undef LOCTEXT_NAMESPACE
diff --git a/Engine/Source/Editor/UnrealEd/Private/Editor.cpp b/Engine/Source/Editor/UnrealEd/Private/Editor.cpp
--- a/Engine/Source/Editor/UnrealEd/Private/Editor.cpp
+++ b/Engine/Source/Editor/UnrealEd/Private/Editor.cpp
@@ -2129,9 +2129,9 @@ void UEditorEngine::ApplyDeltaToActor(AActor* InActor,
 	ABrush* Brush = Cast< ABrush >( InActor );
 	if( Brush )
 	{
-		if( Brush->BrushComponent && Brush->BrushComponent->Brush )
+		if( Brush->GetBrushComponent() && Brush->GetBrushComponent()->Brush )
 		{
-			Brush->BrushComponent->Brush->Polys->Element.ModifyAllItems();
+			Brush->GetBrushComponent()->Brush->Polys->Element.ModifyAllItems();
 		}
 	}
 
@@ -2277,7 +2277,7 @@ void UEditorEngine::ApplyDeltaToActor(AActor* InActor,
 				// Scale all of the polygons of the brush.
 				const FScaleMatrix matrix( FVector( ModifiedScale.X , ModifiedScale.Y, ModifiedScale.Z ) );
 				
-				if(Brush->BrushComponent->Brush && Brush->BrushComponent->Brush->Polys)
+				if(Brush->GetBrushComponent()->Brush && Brush->GetBrushComponent()->Brush->Polys)
 				{
 					// @todo UE4 : verify this code
 					// Brush Transform doesn't seem to change from this code path
@@ -2286,9 +2286,9 @@ void UEditorEngine::ApplyDeltaToActor(AActor* InActor,
 					// if this causes any issue, please contact LH
 					FTransform BrushActorToWorld = Brush->ActorToWorld();
 
-					for( int32 poly = 0 ; poly < Brush->BrushComponent->Brush->Polys->Element.Num() ; poly++ )
+					for( int32 poly = 0 ; poly < Brush->GetBrushComponent()->Brush->Polys->Element.Num() ; poly++ )
 					{
-						FPoly* Poly = &(Brush->BrushComponent->Brush->Polys->Element[poly]);
+						FPoly* Poly = &(Brush->GetBrushComponent()->Brush->Polys->Element[poly]);
 
 						FBox bboxBefore(0);
 						for( int32 vertex = 0 ; vertex < Poly->Vertices.Num() ; vertex++ )
@@ -2344,7 +2344,7 @@ void UEditorEngine::ApplyDeltaToActor(AActor* InActor,
 						Poly->Finalize((ABrush*)InActor,0);
 					}
 
-					Brush->BrushComponent->Brush->BuildBound();
+					Brush->GetBrushComponent()->Brush->BuildBound();
 
 					if( !Brush->IsStaticBrush() )
 					{
@@ -3778,17 +3778,17 @@ void UEditorEngine::ConvertActorsFromClass( UClass* FromClass, UClass* ToClass )
 			if( SMActor )
 			{
 				SourceActors.Add(Actor);
-				Info.GetFromActor(SMActor, SMActor->StaticMeshComponent);
+				Info.GetFromActor(SMActor, SMActor->GetStaticMeshComponent());
 			}
 			else if( FoliageActor )
 			{
 				SourceActors.Add(Actor);
-				Info.GetFromActor(FoliageActor, FoliageActor->StaticMeshComponent);
+				Info.GetFromActor(FoliageActor, FoliageActor->GetStaticMeshComponent());
 			}
 			else if ( bFromSkeletalMesh )
 			{
 				SourceActors.Add(Actor);
-				Info.GetFromActor(SKMActor, SKMActor->SkeletalMeshComponent);
+				Info.GetFromActor(SKMActor, SKMActor->GetSkeletalMeshComponent());
 			}
 
 			// Get the actor group if any
@@ -3832,7 +3832,7 @@ void UEditorEngine::ConvertActorsFromClass( UClass* FromClass, UClass* ToClass )
 				{					
 					AStaticMeshActor* SMActor = CastChecked<AStaticMeshActor>( World->SpawnActor( ToClass, &Info.Location, &Info.Rotation, SpawnInfo ) );
 					SMActor->UnregisterAllComponents();
-					Info.SetToActor(SMActor, SMActor->StaticMeshComponent);
+					Info.SetToActor(SMActor, SMActor->GetStaticMeshComponent());
 					SMActor->RegisterAllComponents();
 					GEditor->SelectActor( SMActor, true, false );
 					Actor = SMActor;
@@ -3841,7 +3841,7 @@ void UEditorEngine::ConvertActorsFromClass( UClass* FromClass, UClass* ToClass )
 				{					
 					AInteractiveFoliageActor* FoliageActor = World->SpawnActor<AInteractiveFoliageActor>( Info.Location, Info.Rotation, SpawnInfo );
 					FoliageActor->UnregisterAllComponents();
-					Info.SetToActor(FoliageActor, FoliageActor->StaticMeshComponent);
+					Info.SetToActor(FoliageActor, FoliageActor->GetStaticMeshComponent());
 					FoliageActor->RegisterAllComponents();
 					GEditor->SelectActor( FoliageActor, true, false );
 					Actor = FoliageActor;
@@ -3852,7 +3852,7 @@ void UEditorEngine::ConvertActorsFromClass( UClass* FromClass, UClass* ToClass )
 					// checked
 					ASkeletalMeshActor* SkeletalMeshActor = CastChecked<ASkeletalMeshActor>( World->SpawnActor( ToClass, &Info.Location, &Info.Rotation, SpawnInfo ));
 					SkeletalMeshActor->UnregisterAllComponents();
-					Info.SetToActor(SkeletalMeshActor, SkeletalMeshActor->SkeletalMeshComponent);
+					Info.SetToActor(SkeletalMeshActor, SkeletalMeshActor->GetSkeletalMeshComponent());
 					SkeletalMeshActor->RegisterAllComponents();
 					GEditor->SelectActor( SkeletalMeshActor, true, false );
 					Actor = SkeletalMeshActor;
@@ -4030,18 +4030,18 @@ bool UEditorEngine::SnapToSocket (AActor* ParentActor, AActor* ChildActor, const
 
 	ASkeletalMeshActor* ParentSkelMeshActor = Cast<ASkeletalMeshActor>(ParentActor);
 	if (ParentSkelMeshActor != NULL &&
-		ParentSkelMeshActor->SkeletalMeshComponent &&
-		ParentSkelMeshActor->SkeletalMeshComponent->SkeletalMesh)
+		ParentSkelMeshActor->GetSkeletalMeshComponent() &&
+		ParentSkelMeshActor->GetSkeletalMeshComponent()->SkeletalMesh)
 	{
-		bAttachedToSocket = AttachActorToComponent(ParentActor, ChildActor, ParentSkelMeshActor->SkeletalMeshComponent, SocketName);
+		bAttachedToSocket = AttachActorToComponent(ParentActor, ChildActor, ParentSkelMeshActor->GetSkeletalMeshComponent(), SocketName);
 	}
 
 	AStaticMeshActor* ParentStaticMeshActor = Cast<AStaticMeshActor>(ParentActor);
 	if (ParentStaticMeshActor != NULL &&
-		ParentStaticMeshActor->StaticMeshComponent &&
-		ParentStaticMeshActor->StaticMeshComponent->StaticMesh)
+		ParentStaticMeshActor->GetStaticMeshComponent() &&
+		ParentStaticMeshActor->GetStaticMeshComponent()->StaticMesh)
 	{
-		bAttachedToSocket = AttachActorToComponent(ChildActor, ParentStaticMeshActor->StaticMeshComponent, SocketName);
+		bAttachedToSocket = AttachActorToComponent(ChildActor, ParentStaticMeshActor->GetStaticMeshComponent(), SocketName);
 	}
 
 	// if the component is in a blueprint actor
@@ -5181,7 +5181,7 @@ static void CopyLightComponentProperties( const AActor& InOldActor, AActor& InNe
 	// using ULightComponent::StaticClass()->GetDefaultObject() will not work since each light actor sets default component properties differently.
 	ALight* OldActorDefaultObject = InOldActor.GetClass()->GetDefaultObject<ALight>();
 	check(OldActorDefaultObject);
-	UActorComponent* DefaultLightComponent = OldActorDefaultObject->LightComponent;
+	UActorComponent* DefaultLightComponent = OldActorDefaultObject->GetLightComponent();
 	check(DefaultLightComponent);
 
 	// The component we are copying from class
diff --git a/Engine/Source/Editor/UnrealEd/Private/EditorActor.cpp b/Engine/Source/Editor/UnrealEd/Private/EditorActor.cpp
--- a/Engine/Source/Editor/UnrealEd/Private/EditorActor.cpp
+++ b/Engine/Source/Editor/UnrealEd/Private/EditorActor.cpp
@@ -1681,9 +1681,9 @@ public:
 
 		if( OutStaticMeshActor.IsStaticMeshActor() )
 		{
-			if ( OutStaticMeshActor.StaticMeshActor->StaticMeshComponent )
+			if ( OutStaticMeshActor.StaticMeshActor->GetStaticMeshComponent() )
 			{
-				OutStaticMeshActor.StaticMesh = OutStaticMeshActor.StaticMeshActor->StaticMeshComponent->StaticMesh;
+				OutStaticMeshActor.StaticMesh = OutStaticMeshActor.StaticMeshActor->GetStaticMeshComponent()->StaticMesh;
 			}
 		}
 		return OutStaticMeshActor.HasStaticMesh();
@@ -1776,10 +1776,10 @@ void UUnrealEdEngine::edactSelectMatchingSkeletalMesh(bool bAllClasses)
 
 		// Look for SkelMeshActor
 		ASkeletalMeshActor* SkelMeshActor = Cast<ASkeletalMeshActor>(Actor);
-		if(SkelMeshActor && SkelMeshActor->SkeletalMeshComponent)
+		if(SkelMeshActor && SkelMeshActor->GetSkeletalMeshComponent())
 		{
 			bSelectSkelMeshActors = true;
-			SelectedMeshes.AddUnique(SkelMeshActor->SkeletalMeshComponent->SkeletalMesh);
+			SelectedMeshes.AddUnique(SkelMeshActor->GetSkeletalMeshComponent()->SkeletalMesh);
 			WorldList.AddUnique(Actor->GetWorld());			
 		}
 
@@ -1826,8 +1826,8 @@ void UUnrealEdEngine::edactSelectMatchingSkeletalMesh(bool bAllClasses)
 			{
 				ASkeletalMeshActor* SkelMeshActor = Cast<ASkeletalMeshActor>(Actor);
 				if( SkelMeshActor && 
-					SkelMeshActor->SkeletalMeshComponent && 
-					SelectedMeshes.Contains(SkelMeshActor->SkeletalMeshComponent->SkeletalMesh) )
+					SkelMeshActor->GetSkeletalMeshComponent() && 
+					SelectedMeshes.Contains(SkelMeshActor->GetSkeletalMeshComponent()->SkeletalMesh) )
 				{
 					bSelectActor = true;
 				}
@@ -1953,9 +1953,9 @@ void UUnrealEdEngine::edactSelectMatchingEmitter()
 
 		AEmitter* Emitter = Cast<AEmitter>( Actor );
 		
-		if ( Emitter && Emitter->ParticleSystemComponent && Emitter->ParticleSystemComponent->Template )
+		if ( Emitter && Emitter->GetParticleSystemComponent() && Emitter->GetParticleSystemComponent()->Template )
 		{
-			SelectedParticleSystemTemplates.AddUnique( Emitter->ParticleSystemComponent->Template );
+			SelectedParticleSystemTemplates.AddUnique( Emitter->GetParticleSystemComponent()->Template );
 			WorldList.AddUnique( Actor->GetWorld() );
 		}
 	}
@@ -1977,7 +1977,7 @@ void UUnrealEdEngine::edactSelectMatchingEmitter()
 		AEmitter* ActorAsEmitter = *ActorIterator;
 		if ( !ActorAsEmitter->IsHiddenEd() )
 		{
-			if ( ActorAsEmitter->ParticleSystemComponent && SelectedParticleSystemTemplates.Contains( ActorAsEmitter->ParticleSystemComponent->Template ) )
+			if ( ActorAsEmitter->GetParticleSystemComponent() && SelectedParticleSystemTemplates.Contains( ActorAsEmitter->GetParticleSystemComponent()->Template ) )
 			{
 				SelectActor( ActorAsEmitter, true, false );
 			}
diff --git a/Engine/Source/Editor/UnrealEd/Private/EditorExporters.cpp b/Engine/Source/Editor/UnrealEd/Private/EditorExporters.cpp
--- a/Engine/Source/Editor/UnrealEd/Private/EditorExporters.cpp
+++ b/Engine/Source/Editor/UnrealEd/Private/EditorExporters.cpp
@@ -609,9 +609,9 @@ bool ULevelExporterSTL::ExportText( const FExportObjectInnerContext* Context, UO
 		// Static meshes
 
 		AStaticMeshActor* Actor = Cast<AStaticMeshActor>(Level->Actors[iActor]);
-		if( Actor && ( !bSelectedOnly || Actor->IsSelected() ) && Actor->StaticMeshComponent->StaticMesh && Actor->StaticMeshComponent->StaticMesh->HasValidRenderData() )
+		if( Actor && ( !bSelectedOnly || Actor->IsSelected() ) && Actor->GetStaticMeshComponent()->StaticMesh && Actor->GetStaticMeshComponent()->StaticMesh->HasValidRenderData() )
 		{
-			FStaticMeshLODResources& LODModel = Actor->StaticMeshComponent->StaticMesh->RenderData->LODResources[0];
+			FStaticMeshLODResources& LODModel = Actor->GetStaticMeshComponent()->StaticMesh->RenderData->LODResources[0];
 			FIndexArrayView Indices = LODModel.IndexBuffer.GetArrayView();
 			int32 NumSections = LODModel.Sections.Num();
 			for (int32 SectionIndex = 0; SectionIndex < NumSections; ++SectionIndex)
@@ -3022,4 +3022,4 @@ namespace MaterialExportUtils
 		Material->PostEditChange();
 		return Material;
 	}
-}
\ No newline at end of file
+}
diff --git a/Engine/Source/Editor/UnrealEd/Private/EditorModes.cpp b/Engine/Source/Editor/UnrealEd/Private/EditorModes.cpp
--- a/Engine/Source/Editor/UnrealEd/Private/EditorModes.cpp
+++ b/Engine/Source/Editor/UnrealEd/Private/EditorModes.cpp
@@ -823,12 +823,12 @@ void FEdMode::DrawHUD(FEditorViewportClient* ViewportClient,FViewport* Viewport,
 
 			// Static mesh vertices
 			AStaticMeshActor* Actor = Cast<AStaticMeshActor>( SelectedActor );
-			if( Actor && Actor->StaticMeshComponent && Actor->StaticMeshComponent->StaticMesh
-				&& Actor->StaticMeshComponent->StaticMesh->RenderData )
+			if( Actor && Actor->GetStaticMeshComponent() && Actor->GetStaticMeshComponent()->StaticMesh
+				&& Actor->GetStaticMeshComponent()->StaticMesh->RenderData )
 			{
 				FTransform ActorToWorld = Actor->ActorToWorld();
 				Vertices.Empty();
-				const FPositionVertexBuffer& VertexBuffer = Actor->StaticMeshComponent->StaticMesh->RenderData->LODResources[0].PositionVertexBuffer;
+				const FPositionVertexBuffer& VertexBuffer = Actor->GetStaticMeshComponent()->StaticMesh->RenderData->LODResources[0].PositionVertexBuffer;
 				for( uint32 i = 0 ; i < VertexBuffer.GetNumVertices() ; i++ )
 				{
 					Vertices.AddUnique( ActorToWorld.TransformPosition( VertexBuffer.VertexPosition(i) ) );
diff --git a/Engine/Source/Editor/UnrealEd/Private/EditorObject.cpp b/Engine/Source/Editor/UnrealEd/Private/EditorObject.cpp
--- a/Engine/Source/Editor/UnrealEd/Private/EditorObject.cpp
+++ b/Engine/Source/Editor/UnrealEd/Private/EditorObject.cpp
@@ -549,8 +549,8 @@ static const TCHAR* ImportProperties(
 	{
 		check(GIsEditor);
 		ABrush* Actor = (ABrush*)DestData;
-		check(Actor->BrushComponent);
-		if( Actor->BrushComponent->Mobility == EComponentMobility::Static )
+		check(Actor->GetBrushComponent());
+		if( Actor->GetBrushComponent()->Mobility == EComponentMobility::Static )
 		{
 			// Prepare static brush.
 			Actor->SetNotForClientOrServer();
diff --git a/Engine/Source/Editor/UnrealEd/Private/EditorServer.cpp b/Engine/Source/Editor/UnrealEd/Private/EditorServer.cpp
--- a/Engine/Source/Editor/UnrealEd/Private/EditorServer.cpp
+++ b/Engine/Source/Editor/UnrealEd/Private/EditorServer.cpp
@@ -2750,13 +2750,13 @@ ULevel*  UEditorEngine::CreateTransLevelMoveBuffer( UWorld* InWorld )
 	// Spawn builder brush for the buffer level.
 	ABrush* BufferDefaultBrush = InWorld->SpawnActor<ABrush>( SpawnInfo );
 
-	check( BufferDefaultBrush->BrushComponent );
+	check( BufferDefaultBrush->GetBrushComponent() );
 	BufferDefaultBrush->Brush = CastChecked<UModel>(StaticFindObject(UModel::StaticClass(), BufferLevel->OwningWorld->GetOuter(), TEXT("Brush"), true), ECastCheckedType::NullAllowed);
 	if (!BufferDefaultBrush->Brush)
 	{
 		BufferDefaultBrush->Brush = new( InWorld, TEXT("Brush") )UModel( FObjectInitializer(), BufferDefaultBrush, 1 );
 	}
-	BufferDefaultBrush->BrushComponent->Brush = BufferDefaultBrush->Brush;
+	BufferDefaultBrush->GetBrushComponent()->Brush = BufferDefaultBrush->Brush;
 	BufferDefaultBrush->SetNotForClientOrServer();
 	BufferDefaultBrush->SetFlags( RF_Transactional );
 	BufferDefaultBrush->Brush->SetFlags( RF_Transactional );
@@ -3399,17 +3399,17 @@ bool UEditorEngine::Map_Check( UWorld* InWorld, const TCHAR* Str, FOutputDevice&
 			UMeshComponent*		MeshComponent		= NULL;
 			if( StaticMeshActor )
 			{
-				MeshComponent = StaticMeshActor->StaticMeshComponent;
+				MeshComponent = StaticMeshActor->GetStaticMeshComponent();
 			}
 			else if( SkeletalMeshActor )
 			{
-				MeshComponent = SkeletalMeshActor->SkeletalMeshComponent;
+				MeshComponent = SkeletalMeshActor->GetSkeletalMeshComponent();
 			}
 
 			// See whether there are lights that ended up with the same component. This was possible in earlier versions of the engine.
 			if( LightActor )
 			{
-				ULightComponent* LightComponent = LightActor->LightComponent;
+				ULightComponent* LightComponent = LightActor->GetLightComponent();
 				AActor* ExistingLightActor = LightGuidToActorMap.FindRef( LightComponent->LightGuid );
 				if( ExistingLightActor )
 				{
@@ -6192,10 +6192,10 @@ bool UEditorEngine::HandleRebuildVolumesCommand( const TCHAR* Str, FOutputDevice
 	{
 		AVolume* Volume = *It;
 
-		if(!Volume->IsTemplate() && Volume->BrushComponent.IsValid())
+		if(!Volume->IsTemplate() && Volume->GetBrushComponent())
 		{
 			UE_LOG(LogEditorServer, Log, TEXT("BSBC: %s"), *Volume->GetPathName() );
-			Volume->BrushComponent->BuildSimpleBrushCollision();
+			Volume->GetBrushComponent()->BuildSimpleBrushCollision();
 		}
 	}
 	return true;
diff --git a/Engine/Source/Editor/UnrealEd/Private/Factories/ActorFactory.cpp b/Engine/Source/Editor/UnrealEd/Private/Factories/ActorFactory.cpp
--- a/Engine/Source/Editor/UnrealEd/Private/Factories/ActorFactory.cpp
+++ b/Engine/Source/Editor/UnrealEd/Private/Factories/ActorFactory.cpp
@@ -255,7 +255,7 @@ void UActorFactoryStaticMesh::PostSpawnActor( UObject* Asset, AActor* NewActor)
 
 	// Change properties
 	AStaticMeshActor* StaticMeshActor = CastChecked<AStaticMeshActor>( NewActor );
-	UStaticMeshComponent* StaticMeshComponent = StaticMeshActor->StaticMeshComponent;
+	UStaticMeshComponent* StaticMeshComponent = StaticMeshActor->GetStaticMeshComponent();
 	check(StaticMeshComponent);
 
 	StaticMeshComponent->UnregisterComponent();
@@ -272,8 +272,8 @@ UObject* UActorFactoryStaticMesh::GetAssetFromActorInstance(AActor* Instance)
 	check(Instance->IsA(NewActorClass));
 	AStaticMeshActor* SMA = CastChecked<AStaticMeshActor>(Instance);
 
-	check(SMA->StaticMeshComponent);
-	return SMA->StaticMeshComponent->StaticMesh;
+	check(SMA->GetStaticMeshComponent());
+	return SMA->GetStaticMeshComponent()->StaticMesh;
 }
 
 void UActorFactoryStaticMesh::PostCreateBlueprint( UObject* Asset, AActor* CDO )
@@ -282,7 +282,7 @@ void UActorFactoryStaticMesh::PostCreateBlueprint( UObject* Asset, AActor* CDO )
 	{
 		UStaticMesh* StaticMesh = CastChecked<UStaticMesh>(Asset);
 		AStaticMeshActor* StaticMeshActor = CastChecked<AStaticMeshActor>(CDO);
-		UStaticMeshComponent* StaticMeshComponent = StaticMeshActor->StaticMeshComponent;
+		UStaticMeshComponent* StaticMeshComponent = StaticMeshActor->GetStaticMeshComponent();
 
 		StaticMeshComponent->StaticMesh = StaticMesh;
 		StaticMeshComponent->StaticMeshDerivedDataKey = StaticMesh->RenderData->DerivedDataKey;
@@ -475,7 +475,7 @@ void UActorFactoryEmitter::PostSpawnActor(UObject* Asset, AActor* NewActor)
 	GEditor->SetActorLabelUnique(NewActor, ParticleSystem->GetName());
 
 	// Term Component
-	NewEmitter->ParticleSystemComponent->UnregisterComponent();
+	NewEmitter->GetParticleSystemComponent()->UnregisterComponent();
 
 	// Change properties
 	NewEmitter->SetTemplate(ParticleSystem);
@@ -491,16 +491,16 @@ void UActorFactoryEmitter::PostSpawnActor(UObject* Asset, AActor* NewActor)
 	}
 
 	// Init Component
-	NewEmitter->ParticleSystemComponent->RegisterComponent();
+	NewEmitter->GetParticleSystemComponent()->RegisterComponent();
 }
 
 UObject* UActorFactoryEmitter::GetAssetFromActorInstance(AActor* Instance)
 {
 	check(Instance->IsA(NewActorClass));
 	AEmitter* Emitter = CastChecked<AEmitter>(Instance);
-	if (Emitter->ParticleSystemComponent.IsValid())
+	if (Emitter->GetParticleSystemComponent())
 	{
-		return Emitter->ParticleSystemComponent->Template;
+		return Emitter->GetParticleSystemComponent()->Template;
 	}
 	else
 	{
@@ -592,28 +592,28 @@ void UActorFactoryPhysicsAsset::PostSpawnActor(UObject* Asset, AActor* NewActor)
 	GEditor->SetActorLabelUnique(NewActor, PhysicsAsset->GetName());
 
 	// Term Component
-	NewSkelActor->SkeletalMeshComponent->UnregisterComponent();
+	NewSkelActor->GetSkeletalMeshComponent()->UnregisterComponent();
 
 	// Change properties
-	NewSkelActor->SkeletalMeshComponent->SkeletalMesh = UseSkelMesh;
+	NewSkelActor->GetSkeletalMeshComponent()->SkeletalMesh = UseSkelMesh;
 	if (NewSkelActor->GetWorld()->IsPlayInEditor())
 	{
 		NewSkelActor->ReplicatedMesh = UseSkelMesh;
 		NewSkelActor->ReplicatedPhysAsset = PhysicsAsset;
 	}
-	NewSkelActor->SkeletalMeshComponent->PhysicsAssetOverride = PhysicsAsset;
+	NewSkelActor->GetSkeletalMeshComponent()->PhysicsAssetOverride = PhysicsAsset;
 
 	// set physics setup
-	NewSkelActor->SkeletalMeshComponent->KinematicBonesUpdateType = EKinematicBonesUpdateToPhysics::SkipAllBones;
-	NewSkelActor->SkeletalMeshComponent->BodyInstance.bSimulatePhysics = true;
-	NewSkelActor->SkeletalMeshComponent->bBlendPhysics = true;
+	NewSkelActor->GetSkeletalMeshComponent()->KinematicBonesUpdateType = EKinematicBonesUpdateToPhysics::SkipAllBones;
+	NewSkelActor->GetSkeletalMeshComponent()->BodyInstance.bSimulatePhysics = true;
+	NewSkelActor->GetSkeletalMeshComponent()->bBlendPhysics = true;
 
 	NewSkelActor->bAlwaysRelevant = true;
 	NewSkelActor->bReplicateMovement = true;
 	NewSkelActor->SetReplicates(true);
 
 	// Init Component
-	NewSkelActor->SkeletalMeshComponent->RegisterComponent();
+	NewSkelActor->GetSkeletalMeshComponent()->RegisterComponent();
 }
 
 void UActorFactoryPhysicsAsset::PostCreateBlueprint( UObject* Asset, AActor* CDO )
@@ -628,14 +628,14 @@ void UActorFactoryPhysicsAsset::PostCreateBlueprint( UObject* Asset, AActor* CDO
 
 			USkeletalMesh* UseSkelMesh = PhysicsAsset->PreviewSkeletalMesh.Get();
 
-			SkeletalPhysicsActor->SkeletalMeshComponent->SkeletalMesh = UseSkelMesh;
-			SkeletalPhysicsActor->SkeletalMeshComponent->PhysicsAssetOverride = PhysicsAsset;
+			SkeletalPhysicsActor->GetSkeletalMeshComponent()->SkeletalMesh = UseSkelMesh;
+			SkeletalPhysicsActor->GetSkeletalMeshComponent()->PhysicsAssetOverride = PhysicsAsset;
 		}
 
 		// set physics setup
-		SkeletalPhysicsActor->SkeletalMeshComponent->KinematicBonesUpdateType = EKinematicBonesUpdateToPhysics::SkipAllBones;
-		SkeletalPhysicsActor->SkeletalMeshComponent->BodyInstance.bSimulatePhysics = true;
-		SkeletalPhysicsActor->SkeletalMeshComponent->bBlendPhysics = true;
+		SkeletalPhysicsActor->GetSkeletalMeshComponent()->KinematicBonesUpdateType = EKinematicBonesUpdateToPhysics::SkipAllBones;
+		SkeletalPhysicsActor->GetSkeletalMeshComponent()->BodyInstance.bSimulatePhysics = true;
+		SkeletalPhysicsActor->GetSkeletalMeshComponent()->bBlendPhysics = true;
 
 		SkeletalPhysicsActor->bAlwaysRelevant = true;
 		SkeletalPhysicsActor->bReplicateMovement = true;
@@ -771,7 +771,7 @@ void UActorFactoryAnimationAsset::PostSpawnActor( UObject* Asset, AActor* NewAct
 	UVertexAnimation* VertexAnimation = Cast<UVertexAnimation>(Asset);
 
 	ASkeletalMeshActor* NewSMActor = CastChecked<ASkeletalMeshActor>(NewActor);
-	USkeletalMeshComponent* NewSASComponent = (NewSMActor->SkeletalMeshComponent);
+	USkeletalMeshComponent* NewSASComponent = (NewSMActor->GetSkeletalMeshComponent());
 
 	if( NewSASComponent )
 	{
@@ -807,7 +807,7 @@ void UActorFactoryAnimationAsset::PostCreateBlueprint( UObject* Asset,  AActor*
 		UVertexAnimation* VertexAnimation = Cast<UVertexAnimation>(Asset);
 
 		ASkeletalMeshActor* SkeletalMeshActor = CastChecked<ASkeletalMeshActor>(CDO);
-		USkeletalMeshComponent* SkeletalComponent = (SkeletalMeshActor->SkeletalMeshComponent);
+		USkeletalMeshComponent* SkeletalComponent = (SkeletalMeshActor->GetSkeletalMeshComponent());
 		if (AnimationAsset)
 		{
 			SkeletalComponent->SetAnimationMode(EAnimationMode::Type::AnimationSingleNode);
@@ -960,20 +960,20 @@ void UActorFactorySkeletalMesh::PostSpawnActor( UObject* Asset, AActor* NewActor
 	GEditor->SetActorLabelUnique(NewActor, SkeletalMesh->GetName());
 
 	// Term Component
-	NewSMActor->SkeletalMeshComponent->UnregisterComponent();
+	NewSMActor->GetSkeletalMeshComponent()->UnregisterComponent();
 
 	// Change properties
-	NewSMActor->SkeletalMeshComponent->SkeletalMesh = SkeletalMesh;
+	NewSMActor->GetSkeletalMeshComponent()->SkeletalMesh = SkeletalMesh;
 	if (NewSMActor->GetWorld()->IsGameWorld())
 	{
 		NewSMActor->ReplicatedMesh = SkeletalMesh;
 	}
 
 	// Init Component
-	NewSMActor->SkeletalMeshComponent->RegisterComponent();
+	NewSMActor->GetSkeletalMeshComponent()->RegisterComponent();
 	if( AnimBlueprint )
 	{
-		NewSMActor->SkeletalMeshComponent->SetAnimInstanceClass(AnimBlueprint->GeneratedClass);
+		NewSMActor->GetSkeletalMeshComponent()->SetAnimInstanceClass(AnimBlueprint->GeneratedClass);
 	}
 }
 
@@ -985,8 +985,8 @@ void UActorFactorySkeletalMesh::PostCreateBlueprint( UObject* Asset, AActor* CDO
 		UAnimBlueprint* AnimBlueprint = Cast<UAnimBlueprint>(Asset);
 
 		ASkeletalMeshActor* SkeletalMeshActor = CastChecked<ASkeletalMeshActor>(CDO);
-		SkeletalMeshActor->SkeletalMeshComponent->SkeletalMesh = SkeletalMesh;
-		SkeletalMeshActor->SkeletalMeshComponent->AnimBlueprintGeneratedClass = AnimBlueprint ? Cast<UAnimBlueprintGeneratedClass>(AnimBlueprint->GeneratedClass) : NULL;
+		SkeletalMeshActor->GetSkeletalMeshComponent()->SkeletalMesh = SkeletalMesh;
+		SkeletalMeshActor->GetSkeletalMeshComponent()->AnimBlueprintGeneratedClass = AnimBlueprint ? Cast<UAnimBlueprintGeneratedClass>(AnimBlueprint->GeneratedClass) : NULL;
 	}
 }
 
@@ -1041,7 +1041,7 @@ void UActorFactoryAmbientSound::PostSpawnActor( UObject* Asset, AActor* NewActor
 	{
 		AAmbientSound* NewSound = CastChecked<AAmbientSound>( NewActor );
 		GEditor->SetActorLabelUnique(NewSound, AmbientSound->GetName());
-		NewSound->AudioComponent->SetSound(AmbientSound);
+		NewSound->GetAudioComponent()->SetSound(AmbientSound);
 	}
 }
 
@@ -1050,8 +1050,8 @@ UObject* UActorFactoryAmbientSound::GetAssetFromActorInstance(AActor* Instance)
 	check(Instance->IsA(NewActorClass));
 	AAmbientSound* SoundActor = CastChecked<AAmbientSound>(Instance);
 
-	check(SoundActor->AudioComponent);
-	return SoundActor->AudioComponent->Sound;
+	check(SoundActor->GetAudioComponent());
+	return SoundActor->GetAudioComponent()->Sound;
 }
 
 void UActorFactoryAmbientSound::PostCreateBlueprint( UObject* Asset, AActor* CDO )
@@ -1063,7 +1063,7 @@ void UActorFactoryAmbientSound::PostCreateBlueprint( UObject* Asset, AActor* CDO
 		if (AmbientSound != NULL)
 		{
 			AAmbientSound* NewSound = CastChecked<AAmbientSound>(CDO);
-			NewSound->AudioComponent->SetSound(AmbientSound);
+			NewSound->GetAudioComponent()->SetSound(AmbientSound);
 		}
 	}
 }
@@ -1475,13 +1475,13 @@ void UActorFactoryDestructible::PostSpawnActor( UObject* Asset, AActor* NewActor
 	GEditor->SetActorLabelUnique(NewActor, DestructibleMesh->GetName());
 
 	// Term Component
-	NewDestructibleActor->DestructibleComponent->UnregisterComponent();
+	NewDestructibleActor->GetDestructibleComponent()->UnregisterComponent();
 
 	// Change properties
-	NewDestructibleActor->DestructibleComponent->SetSkeletalMesh( DestructibleMesh );
+	NewDestructibleActor->GetDestructibleComponent()->SetSkeletalMesh( DestructibleMesh );
 
 	// Init Component
-	NewDestructibleActor->DestructibleComponent->RegisterComponent();
+	NewDestructibleActor->GetDestructibleComponent()->RegisterComponent();
 }
 
 UObject* UActorFactoryDestructible::GetAssetFromActorInstance(AActor* Instance)
@@ -1489,8 +1489,8 @@ UObject* UActorFactoryDestructible::GetAssetFromActorInstance(AActor* Instance)
 	check(Instance->IsA(NewActorClass));
 	ADestructibleActor* DA = CastChecked<ADestructibleActor>(Instance);
 
-	check(DA->DestructibleComponent);
-	return DA->DestructibleComponent->SkeletalMesh;
+	check(DA->GetDestructibleComponent());
+	return DA->GetDestructibleComponent()->SkeletalMesh;
 }
 
 void UActorFactoryDestructible::PostCreateBlueprint( UObject* Asset, AActor* CDO )
@@ -1500,7 +1500,7 @@ void UActorFactoryDestructible::PostCreateBlueprint( UObject* Asset, AActor* CDO
 		UDestructibleMesh* DestructibleMesh = CastChecked<UDestructibleMesh>(Asset);
 		ADestructibleActor* DestructibleActor = CastChecked<ADestructibleActor>(CDO);
 
-		DestructibleActor->DestructibleComponent->SetSkeletalMesh(DestructibleMesh);
+		DestructibleActor->GetDestructibleComponent()->SetSkeletalMesh(DestructibleMesh);
 	}
 }
 
@@ -1539,9 +1539,9 @@ void UActorFactoryVectorFieldVolume::PostSpawnActor( UObject* Asset, AActor* New
 
 	GEditor->SetActorLabelUnique(NewActor, VectorField->GetName());
 
-	if ( VectorFieldVolumeActor && VectorFieldVolumeActor->VectorFieldComponent )
+	if ( VectorFieldVolumeActor && VectorFieldVolumeActor->GetVectorFieldComponent() )
 	{
-		VectorFieldVolumeActor->VectorFieldComponent->VectorField = VectorField;
+		VectorFieldVolumeActor->GetVectorFieldComponent()->VectorField = VectorField;
 		VectorFieldVolumeActor->PostEditChange();
 	}
 }
@@ -1560,7 +1560,7 @@ void CreateBrushForVolumeActor( AVolume* NewActor, UBrushBuilder* BrushBuilder )
 		NewActor->PolyFlags = 0;
 		NewActor->Brush = new( NewActor, NAME_None, RF_Transactional )UModel(FObjectInitializer(), NULL, true );
 		NewActor->Brush->Polys = new( NewActor->Brush, NAME_None, RF_Transactional )UPolys(FObjectInitializer());
-		NewActor->BrushComponent->Brush = NewActor->Brush;
+		NewActor->GetBrushComponent()->Brush = NewActor->Brush;
 		if(BrushBuilder != nullptr)
 		{
 			NewActor->BrushBuilder = DuplicateObject<UBrushBuilder>(BrushBuilder, NewActor);
diff --git a/Engine/Source/Editor/UnrealEd/Private/Factories/EditorFactories.cpp b/Engine/Source/Editor/UnrealEd/Private/Factories/EditorFactories.cpp
--- a/Engine/Source/Editor/UnrealEd/Private/Factories/EditorFactories.cpp
+++ b/Engine/Source/Editor/UnrealEd/Private/Factories/EditorFactories.cpp
@@ -826,9 +826,9 @@ UObject* ULevelFactory::FactoryCreateText
 		// Copy brushes' model pointers over to their BrushComponent, to keep compatibility with old T3Ds.
 		if( Brush && bActorChanged )
 		{
-			if( Brush->BrushComponent ) // Should always be the case, but not asserting so that old broken content won't crash.
+			if( Brush->GetBrushComponent() ) // Should always be the case, but not asserting so that old broken content won't crash.
 			{
-				Brush->BrushComponent->Brush = Brush->Brush;
+				Brush->GetBrushComponent()->Brush = Brush->Brush;
 
 				// We need to avoid duplicating default/ builder brushes. This is done by destroying all brushes that are CSG_Active and are not
 				// the default brush in their respective levels.
diff --git a/Engine/Source/Editor/UnrealEd/Private/Fbx/FbxLightImport.cpp b/Engine/Source/Editor/UnrealEd/Private/Fbx/FbxLightImport.cpp
--- a/Engine/Source/Editor/UnrealEd/Private/Fbx/FbxLightImport.cpp
+++ b/Engine/Source/Editor/UnrealEd/Private/Fbx/FbxLightImport.cpp
@@ -43,7 +43,7 @@ ALight* FFbxImporter::CreateLight(FbxLight* InLight, UWorld* InWorld)
 		UnrealLight = InWorld->SpawnActor<ALight>( LightClass, SpawnInfo );
 		if (UnrealLight)
 		{
-			FillLightComponent(InLight,UnrealLight->LightComponent);	
+			FillLightComponent(InLight,UnrealLight->GetLightComponent());	
 		}
 		else
 		{
@@ -138,7 +138,7 @@ ACameraActor* FFbxImporter::CreateCamera(FbxCamera* InCamera, UWorld* InWorld)
 	UnrealCamera = InWorld->SpawnActor<ACameraActor>( SpawnInfo );
 	if (UnrealCamera)
 	{
-		UnrealCamera->CameraComponent->FieldOfView = InCamera->FieldOfView.Get();
+		UnrealCamera->GetCameraComponent()->FieldOfView = InCamera->FieldOfView.Get();
 	}
 	return UnrealCamera;
 }
diff --git a/Engine/Source/Editor/UnrealEd/Private/Fbx/FbxMainExport.cpp b/Engine/Source/Editor/UnrealEd/Private/Fbx/FbxMainExport.cpp
--- a/Engine/Source/Editor/UnrealEd/Private/Fbx/FbxMainExport.cpp
+++ b/Engine/Source/Editor/UnrealEd/Private/Fbx/FbxMainExport.cpp
@@ -277,7 +277,7 @@ void FFbxExporter::ExportLevelMesh( ULevel* InLevel, AMatineeActor* InMatineeAct
 			}
 			else if (Actor->IsA(AStaticMeshActor::StaticClass()))
 			{
-				ExportStaticMesh( Actor, CastChecked<AStaticMeshActor>(Actor)->StaticMeshComponent, InMatineeActor );
+				ExportStaticMesh( Actor, CastChecked<AStaticMeshActor>(Actor)->GetStaticMeshComponent(), InMatineeActor );
 			}
 			else if (Actor->IsA(ALandscapeProxy::StaticClass()))
 			{
@@ -306,14 +306,14 @@ void FFbxExporter::ExportLevelMesh( ULevel* InLevel, AMatineeActor* InMatineeAct
  */
 void FFbxExporter::ExportLight( ALight* Actor, AMatineeActor* InMatineeActor )
 {
-	if (Scene == NULL || Actor == NULL || !Actor->LightComponent.IsValid()) return;
+	if (Scene == NULL || Actor == NULL || !Actor->GetLightComponent()) return;
 
 	// Export the basic actor information.
 	FbxNode* FbxActor = ExportActor( Actor, InMatineeActor ); // this is the pivot node
 	// The real fbx light node
 	FbxNode* FbxLightNode = FbxActor->GetParent();
 
-	ULightComponent* BaseLight = Actor->LightComponent;
+	ULightComponent* BaseLight = Actor->GetLightComponent();
 
 	FString FbxNodeName = GetActorNodeName(Actor, InMatineeActor);
 
@@ -384,14 +384,14 @@ void FFbxExporter::ExportCamera( ACameraActor* Actor, AMatineeActor* InMatineeAc
 
 	// Export the view area information
 	Camera->ProjectionType.Set(FbxCamera::ePerspective);
-	Camera->SetAspect(FbxCamera::eFixedRatio, Actor->CameraComponent->AspectRatio, 1.0f);
-	Camera->FilmAspectRatio.Set(Actor->CameraComponent->AspectRatio);
-	Camera->SetApertureWidth(Actor->CameraComponent->AspectRatio * 0.612f); // 0.612f is a magic number from Maya that represents the ApertureHeight
+	Camera->SetAspect(FbxCamera::eFixedRatio, Actor->GetCameraComponent()->AspectRatio, 1.0f);
+	Camera->FilmAspectRatio.Set(Actor->GetCameraComponent()->AspectRatio);
+	Camera->SetApertureWidth(Actor->GetCameraComponent()->AspectRatio * 0.612f); // 0.612f is a magic number from Maya that represents the ApertureHeight
 	Camera->SetApertureMode(FbxCamera::eFocalLength);
-	Camera->FocalLength.Set(Camera->ComputeFocalLength(Actor->CameraComponent->FieldOfView));
+	Camera->FocalLength.Set(Camera->ComputeFocalLength(Actor->GetCameraComponent()->FieldOfView));
 	
 	// Add one user property for recording the AspectRatio animation
-	CreateAnimatableUserProperty(FbxCameraNode, Actor->CameraComponent->AspectRatio, "UE_AspectRatio", "UE_Matinee_Camera_AspectRatio");
+	CreateAnimatableUserProperty(FbxCameraNode, Actor->GetCameraComponent()->AspectRatio, "UE_AspectRatio", "UE_Matinee_Camera_AspectRatio");
 
 	// Push the near/far clip planes away, as the engine uses larger values than the default.
 	Camera->SetNearPlane(10.0f);
@@ -407,12 +407,12 @@ void FFbxExporter::ExportCamera( ACameraActor* Actor, AMatineeActor* InMatineeAc
  */
 void FFbxExporter::ExportBrush(ABrush* Actor, UModel* InModel, bool bConvertToStaticMesh )
 {
-	if (Scene == NULL || Actor == NULL || !Actor->BrushComponent.IsValid()) return;
+	if (Scene == NULL || Actor == NULL || !Actor->GetBrushComponent()) return;
 
  	if (!bConvertToStaticMesh)
  	{
  		// Retrieve the information structures, verifying the integrity of the data.
- 		UModel* Model = Actor->BrushComponent->Brush;
+ 		UModel* Model = Actor->GetBrushComponent()->Brush;
 
  		if (Model == NULL || Model->VertexBuffer.Vertices.Num() < 3 || Model->MaterialIndexBuffers.Num() == 0) return;
  
diff --git a/Engine/Source/Editor/UnrealEd/Private/Fbx/FbxMatineeImport.cpp b/Engine/Source/Editor/UnrealEd/Private/Fbx/FbxMatineeImport.cpp
--- a/Engine/Source/Editor/UnrealEd/Private/Fbx/FbxMatineeImport.cpp
+++ b/Engine/Source/Editor/UnrealEd/Private/Fbx/FbxMatineeImport.cpp
@@ -321,28 +321,28 @@ void FFbxImporter::ImportCamera(ACameraActor* Actor, UInterpGroupInst* MatineeGr
 	// Note: there is a pivot node between the Fbx camera Node and node attribute
 	FbxNode* FbxCameraNode = Camera->GetNode()->GetParent();
 	// Import the aspect ratio
-	Actor->CameraComponent->AspectRatio = Camera->FilmAspectRatio.Get(); // Assumes the FBX comes from Unreal or Maya
-	ImportAnimatedProperty(&Actor->CameraComponent->AspectRatio, TEXT("AspectRatio"), MatineeGroup, 
-				Actor->CameraComponent->AspectRatio, FbxCameraNode->FindProperty("UE_AspectRatio") );
+	Actor->GetCameraComponent()->AspectRatio = Camera->FilmAspectRatio.Get(); // Assumes the FBX comes from Unreal or Maya
+	ImportAnimatedProperty(&Actor->GetCameraComponent()->AspectRatio, TEXT("AspectRatio"), MatineeGroup, 
+				Actor->GetCameraComponent()->AspectRatio, FbxCameraNode->FindProperty("UE_AspectRatio") );
 
 	FbxPropertyT<FbxDouble> AperatureModeProperty;
 
 
 	if( Camera->FocalLength.IsValid() && Camera->GetApertureMode() == FbxCamera::eFocalLength )
 	{
-		Actor->CameraComponent->FieldOfView = Camera->ComputeFieldOfView(Camera->FocalLength.Get()); // Assumes the FBX comes from Unreal or Maya
+		Actor->GetCameraComponent()->FieldOfView = Camera->ComputeFieldOfView(Camera->FocalLength.Get()); // Assumes the FBX comes from Unreal or Maya
 
 		AperatureModeProperty = Camera->FocalLength;
 	}
 	else 
 	{
-		Actor->CameraComponent->FieldOfView = Camera->FieldOfView.Get();
+		Actor->GetCameraComponent()->FieldOfView = Camera->FieldOfView.Get();
 
 		AperatureModeProperty = Camera->FieldOfView;
 	}
 
 
-	ImportAnimatedProperty(&Actor->CameraComponent->FieldOfView, TEXT("FOVAngle"), MatineeGroup, AperatureModeProperty.Get(), AperatureModeProperty, true, Camera );
+	ImportAnimatedProperty(&Actor->GetCameraComponent()->FieldOfView, TEXT("FOVAngle"), MatineeGroup, AperatureModeProperty.Get(), AperatureModeProperty, true, Camera );
 }
 
 void FFbxImporter::ImportAnimatedProperty(float* Value, const TCHAR* ValueName, UInterpGroupInst* MatineeGroup, const float FbxValue, FbxProperty InProperty, bool bImportFOV, FbxCamera* Camera )
diff --git a/Engine/Source/Editor/UnrealEd/Private/GroupActor.cpp b/Engine/Source/Editor/UnrealEd/Private/GroupActor.cpp
--- a/Engine/Source/Editor/UnrealEd/Private/GroupActor.cpp
+++ b/Engine/Source/Editor/UnrealEd/Private/GroupActor.cpp
@@ -13,7 +13,7 @@ AGroupActor::AGroupActor(const FObjectInitializer& ObjectInitializer)
 {
 	bLocked = true;
 
-	TSubobjectPtr<USceneComponent> GroupComponent = ObjectInitializer.CreateDefaultSubobject<USceneComponent>(this, TEXT("GroupComponent"));
+	USceneComponent* GroupComponent = ObjectInitializer.CreateDefaultSubobject<USceneComponent>(this, TEXT("GroupComponent"));
 	RootComponent = GroupComponent;
 }
 
diff --git a/Engine/Source/Editor/UnrealEd/Private/LevelEditorViewport.cpp b/Engine/Source/Editor/UnrealEd/Private/LevelEditorViewport.cpp
--- a/Engine/Source/Editor/UnrealEd/Private/LevelEditorViewport.cpp
+++ b/Engine/Source/Editor/UnrealEd/Private/LevelEditorViewport.cpp
@@ -549,7 +549,7 @@ static bool AttemptApplyObjToActor( UObject* ObjToUse, AActor* ActorToApplyTo, i
 				if ( ASkeletalMeshActor* SkelMeshActor = Cast<ASkeletalMeshActor>(ActorToApplyTo) )
 				{
 					const FScopedTransaction Transaction( LOCTEXT( "DropSkelMeshOnObject", "Drop Skeletal Mesh On Object" ) );
-					USkeletalMeshComponent* SkelMeshComponent = SkelMeshActor->SkeletalMeshComponent;
+					USkeletalMeshComponent* SkelMeshComponent = SkelMeshActor->GetSkeletalMeshComponent();
 					SkelMeshComponent->Modify();
 					if ( DroppedObjAsSkeletalMesh )
 					{
@@ -583,7 +583,7 @@ static bool AttemptApplyObjToActor( UObject* ObjToUse, AActor* ActorToApplyTo, i
 					{
 						const FScopedTransaction Transaction(LOCTEXT("DropAnimBlueprintOnObject", "Drop Anim Blueprint On Object"));
 
-						USkeletalMeshComponent* SkelMeshComponent = SkelMeshActor->SkeletalMeshComponent;
+						USkeletalMeshComponent* SkelMeshComponent = SkelMeshActor->GetSkeletalMeshComponent();
 						// if anim blueprint skeleton and mesh skeleton does not match or component does not have any mesh, then change mesh
 						bool bShouldChangeMesh = (SkelMeshComponent->SkeletalMesh == NULL ||
 								!NeedsSkeleton->IsCompatible(SkelMeshComponent->SkeletalMesh->Skeleton));
@@ -636,7 +636,7 @@ static bool AttemptApplyObjToActor( UObject* ObjToUse, AActor* ActorToApplyTo, i
 					if(ASkeletalMeshActor* SkelMeshActor = Cast<ASkeletalMeshActor>(ActorToApplyTo))
 					{
 						const FScopedTransaction Transaction(LOCTEXT("DropAnimationOnObject", "Drop Animation On Object"));
-						USkeletalMeshComponent* SkelComponent = SkelMeshActor->SkeletalMeshComponent;
+						USkeletalMeshComponent* SkelComponent = SkelMeshActor->GetSkeletalMeshComponent();
 						SkelComponent->Modify();
 						// if asset skeleton and mesh skeleton does not match or component does not have any mesh, then change mesh
 						bool bShouldChangeMesh = SkelComponent->SkeletalMesh == NULL || 
diff --git a/Engine/Source/Editor/UnrealEd/Private/LevelViewportClickHandlers.cpp b/Engine/Source/Editor/UnrealEd/Private/LevelViewportClickHandlers.cpp
--- a/Engine/Source/Editor/UnrealEd/Private/LevelViewportClickHandlers.cpp
+++ b/Engine/Source/Editor/UnrealEd/Private/LevelViewportClickHandlers.cpp
@@ -78,8 +78,7 @@ namespace ClickHandlers
 
 			APointLight* Light = CastChecked<APointLight>(NewActor);
 			Light->SetMobility(EComponentMobility::Stationary);
-			UPointLightComponent* PointLightComponent = Cast<UPointLightComponent>( Light->LightComponent );
-
+			UPointLightComponent* PointLightComponent = Cast<UPointLightComponent>( Light->GetLightComponent() );
 			PointLightComponent->LightColor = PixelColor;
 		}
 	}
diff --git a/Engine/Source/Editor/UnrealEd/Private/Lightmass/Lightmass.cpp b/Engine/Source/Editor/UnrealEd/Private/Lightmass/Lightmass.cpp
--- a/Engine/Source/Editor/UnrealEd/Private/Lightmass/Lightmass.cpp
+++ b/Engine/Source/Editor/UnrealEd/Private/Lightmass/Lightmass.cpp
@@ -3324,7 +3324,7 @@ void FLightmassProcessor::ImportMeshAreaLightData()
 					FActorSpawnParameters SpawnInfo;
 					SpawnInfo.Owner = CurrentLevel->Actors[0];
 					AGeneratedMeshAreaLight* NewGeneratedLight = CurrentLevel->OwningWorld->SpawnActor<AGeneratedMeshAreaLight>(Position, Direction.Rotation());
-					USpotLightComponent* SpotComponent = CastChecked<USpotLightComponent>(NewGeneratedLight->LightComponent);
+					USpotLightComponent* SpotComponent = CastChecked<USpotLightComponent>(NewGeneratedLight->GetLightComponent());
 					// Unregister the component before we change its attributes
 					FComponentReregisterContext Reregister(SpotComponent);
 					// Setup spotlight properties to approximate a mesh area light
@@ -4043,4 +4043,4 @@ bool FLightmassProcessor::ImportTextureMapping(int32 Channel, FTextureMappingImp
 	return bResult;
 }
 
-#undef LOCTEXT_NAMESPACE
\ No newline at end of file
+#undef LOCTEXT_NAMESPACE
diff --git a/Engine/Source/Editor/UnrealEd/Private/Tests/EditorAutomationTests.cpp b/Engine/Source/Editor/UnrealEd/Private/Tests/EditorAutomationTests.cpp
--- a/Engine/Source/Editor/UnrealEd/Private/Tests/EditorAutomationTests.cpp
+++ b/Engine/Source/Editor/UnrealEd/Private/Tests/EditorAutomationTests.cpp
@@ -1326,7 +1326,7 @@ bool FCleanupConvertToValidation::Update()
 				{
 					AStaticMeshActor* StaticMeshActor = *ActorIt;
 
-					if ( StaticMeshActor->StaticMeshComponent->StaticMesh == GeneratedMesh )
+					if ( StaticMeshActor->GetStaticMeshComponent()->StaticMesh == GeneratedMesh )
 					{
 						TestWorld->DestroyActor(StaticMeshActor);
 					}
diff --git a/Engine/Source/Editor/UnrealEd/Private/Tests/EditorBuildPromotionTests.cpp b/Engine/Source/Editor/UnrealEd/Private/Tests/EditorBuildPromotionTests.cpp
--- a/Engine/Source/Editor/UnrealEd/Private/Tests/EditorBuildPromotionTests.cpp
+++ b/Engine/Source/Editor/UnrealEd/Private/Tests/EditorBuildPromotionTests.cpp
@@ -1078,7 +1078,7 @@ namespace EditorBuildPromotionTestUtils
 				AStaticMeshActor* PlacedMesh = Cast<AStaticMeshActor>(FActorFactoryAssetProxy::AddActorForAsset(DefaultMesh));
 				PlacedMesh->SetActorLocation(Location);
 
-				PlacedMesh->StaticMeshComponent->SetMaterial(0, Material);
+				PlacedMesh->GetStaticMeshComponent()->SetMaterial(0, Material);
 
 				return PlacedMesh;
 			}
diff --git a/Engine/Source/Editor/UnrealEd/Private/ThumbnailHelpers.cpp b/Engine/Source/Editor/UnrealEd/Private/ThumbnailHelpers.cpp
--- a/Engine/Source/Editor/UnrealEd/Private/ThumbnailHelpers.cpp
+++ b/Engine/Source/Editor/UnrealEd/Private/ThumbnailHelpers.cpp
@@ -233,14 +233,14 @@ FMaterialThumbnailScene::FMaterialThumbnailScene()
 	SpawnInfo.ObjectFlags = RF_Transient;
 	PreviewActor = GetWorld()->SpawnActor<AStaticMeshActor>( SpawnInfo );
 
-	PreviewActor->StaticMeshComponent->SetMobility(EComponentMobility::Movable);
+	PreviewActor->GetStaticMeshComponent()->SetMobility(EComponentMobility::Movable);
 	PreviewActor->SetActorEnableCollision(false);
 }
 
 void FMaterialThumbnailScene::SetMaterialInterface(UMaterialInterface* InMaterial)
 {
 	check(PreviewActor);
-	check(PreviewActor->StaticMeshComponent);
+	check(PreviewActor->GetStaticMeshComponent());
 
 	if ( InMaterial )
 	{
@@ -263,7 +263,7 @@ void FMaterialThumbnailScene::SetMaterialInterface(UMaterialInterface* InMateria
 					UStaticMesh* MeshToUse = Cast<UStaticMesh>(ThumbnailInfo->PreviewMesh.ResolveObject());
 					if ( MeshToUse )
 					{
-						PreviewActor->StaticMeshComponent->SetStaticMesh(MeshToUse);
+						PreviewActor->GetStaticMeshComponent()->SetStaticMesh(MeshToUse);
 						bFoundCustomMesh = true;
 					}
 				}
@@ -272,65 +272,65 @@ void FMaterialThumbnailScene::SetMaterialInterface(UMaterialInterface* InMateria
 				{
 					// Just use a plane if the mesh was not found
 					Transform.SetRotation(FQuat(FRotator(0, -90, 0)));
-					PreviewActor->StaticMeshComponent->SetStaticMesh(GUnrealEd->GetThumbnailManager()->EditorPlane);
+					PreviewActor->GetStaticMeshComponent()->SetStaticMesh(GUnrealEd->GetThumbnailManager()->EditorPlane);
 				}
 			}
 			break;
 
 		case TPT_Cube:
-			PreviewActor->StaticMeshComponent->SetStaticMesh(GUnrealEd->GetThumbnailManager()->EditorCube);
+			PreviewActor->GetStaticMeshComponent()->SetStaticMesh(GUnrealEd->GetThumbnailManager()->EditorCube);
 			break;
 
 		case TPT_Sphere:
 			// The sphere is a little big, scale it down to 256x256x256
 			Transform.SetScale3D( FVector(0.8f) );
-			PreviewActor->StaticMeshComponent->SetStaticMesh(GUnrealEd->GetThumbnailManager()->EditorSphere);
+			PreviewActor->GetStaticMeshComponent()->SetStaticMesh(GUnrealEd->GetThumbnailManager()->EditorSphere);
 			break;
 
 		case TPT_Cylinder:
-			PreviewActor->StaticMeshComponent->SetStaticMesh(GUnrealEd->GetThumbnailManager()->EditorCylinder);
+			PreviewActor->GetStaticMeshComponent()->SetStaticMesh(GUnrealEd->GetThumbnailManager()->EditorCylinder);
 			break;
 
 		case TPT_Plane:
 			// The plane needs to be rotated 90 degrees to face the camera
 			Transform.SetRotation(FQuat(FRotator(0, -90, 0)));
-			PreviewActor->StaticMeshComponent->SetStaticMesh(GUnrealEd->GetThumbnailManager()->EditorPlane);
+			PreviewActor->GetStaticMeshComponent()->SetStaticMesh(GUnrealEd->GetThumbnailManager()->EditorPlane);
 			break;
 
 		default:
 			check(0);
 		}
 
-		PreviewActor->StaticMeshComponent->SetRelativeTransform(Transform);
-		PreviewActor->StaticMeshComponent->UpdateBounds();
+		PreviewActor->GetStaticMeshComponent()->SetRelativeTransform(Transform);
+		PreviewActor->GetStaticMeshComponent()->UpdateBounds();
 
 		// Center the mesh at the world origin then offset to put it on top of the plane
-		const float BoundsZOffset = GetBoundsZOffset(PreviewActor->StaticMeshComponent->Bounds);
-		Transform.SetLocation(-PreviewActor->StaticMeshComponent->Bounds.Origin + FVector(0, 0, BoundsZOffset));
+		const float BoundsZOffset = GetBoundsZOffset(PreviewActor->GetStaticMeshComponent()->Bounds);
+		Transform.SetLocation(-PreviewActor->GetStaticMeshComponent()->Bounds.Origin + FVector(0, 0, BoundsZOffset));
 
-		PreviewActor->StaticMeshComponent->SetRelativeTransform(Transform);
+		PreviewActor->GetStaticMeshComponent()->SetRelativeTransform(Transform);
 	}
 
-	PreviewActor->StaticMeshComponent->SetMaterial(0, InMaterial);
-	PreviewActor->StaticMeshComponent->RecreateRenderState_Concurrent();
+	PreviewActor->GetStaticMeshComponent()->SetMaterial(0, InMaterial);
+	PreviewActor->GetStaticMeshComponent()->RecreateRenderState_Concurrent();
 }
 
 void FMaterialThumbnailScene::GetViewMatrixParameters(const float InFOVDegrees, FVector& OutOrigin, float& OutOrbitPitch, float& OutOrbitYaw, float& OutOrbitZoom) const
 {
 	check(PreviewActor);
-	check(PreviewActor->StaticMeshComponent);
-	check(PreviewActor->StaticMeshComponent->GetMaterial(0));
+	check(PreviewActor->GetStaticMeshComponent());
+	check(PreviewActor->GetStaticMeshComponent()->GetMaterial(0));
 
 	// Fit the mesh in the view using the following formula
 	// tan(HalfFOV) = Width/TargetCameraDistance
 	const float HalfFOVRadians = FMath::DegreesToRadians<float>(InFOVDegrees) * 0.5f;
 	// Add extra size to view slightly outside of the bounds to compensate for perspective
 	const float BoundsMultiplier = 1.15f;
-	const float HalfMeshSize = PreviewActor->StaticMeshComponent->Bounds.SphereRadius * BoundsMultiplier;
-	const float BoundsZOffset = GetBoundsZOffset(PreviewActor->StaticMeshComponent->Bounds);
+	const float HalfMeshSize = PreviewActor->GetStaticMeshComponent()->Bounds.SphereRadius * BoundsMultiplier;
+	const float BoundsZOffset = GetBoundsZOffset(PreviewActor->GetStaticMeshComponent()->Bounds);
 	const float TargetDistance = HalfMeshSize / FMath::Tan(HalfFOVRadians);
 
-	USceneThumbnailInfo* ThumbnailInfo = Cast<USceneThumbnailInfo>(PreviewActor->StaticMeshComponent->GetMaterial(0)->ThumbnailInfo);
+	USceneThumbnailInfo* ThumbnailInfo = Cast<USceneThumbnailInfo>(PreviewActor->GetStaticMeshComponent()->GetMaterial(0)->ThumbnailInfo);
 	if ( ThumbnailInfo )
 	{
 		if ( TargetDistance + ThumbnailInfo->OrbitZoom < 0 )
@@ -373,35 +373,35 @@ FSkeletalMeshThumbnailScene::FSkeletalMeshThumbnailScene()
 
 void FSkeletalMeshThumbnailScene::SetSkeletalMesh(USkeletalMesh* InSkeletalMesh)
 {
-	PreviewActor->SkeletalMeshComponent->Materials.Empty();
-	PreviewActor->SkeletalMeshComponent->SetSkeletalMesh(InSkeletalMesh);
+	PreviewActor->GetSkeletalMeshComponent()->Materials.Empty();
+	PreviewActor->GetSkeletalMeshComponent()->SetSkeletalMesh(InSkeletalMesh);
 
 	if ( InSkeletalMesh )
 	{
 		FTransform MeshTransform = FTransform::Identity;
 
 		PreviewActor->SetActorLocation(FVector(0,0,0), false);
-		PreviewActor->SkeletalMeshComponent->UpdateBounds();
+		PreviewActor->GetSkeletalMeshComponent()->UpdateBounds();
 
 		// Center the mesh at the world origin then offset to put it on top of the plane
-		const float BoundsZOffset = GetBoundsZOffset(PreviewActor->SkeletalMeshComponent->Bounds);
-		PreviewActor->SetActorLocation( -PreviewActor->SkeletalMeshComponent->Bounds.Origin + FVector(0, 0, BoundsZOffset), false );
-		PreviewActor->SkeletalMeshComponent->RecreateRenderState_Concurrent();
+		const float BoundsZOffset = GetBoundsZOffset(PreviewActor->GetSkeletalMeshComponent()->Bounds);
+		PreviewActor->SetActorLocation( -PreviewActor->GetSkeletalMeshComponent()->Bounds.Origin + FVector(0, 0, BoundsZOffset), false );
+		PreviewActor->GetSkeletalMeshComponent()->RecreateRenderState_Concurrent();
 	}
 }
 
 void FSkeletalMeshThumbnailScene::GetViewMatrixParameters(const float InFOVDegrees, FVector& OutOrigin, float& OutOrbitPitch, float& OutOrbitYaw, float& OutOrbitZoom) const
 {
-	check(PreviewActor->SkeletalMeshComponent);
-	check(PreviewActor->SkeletalMeshComponent->SkeletalMesh);
+	check(PreviewActor->GetSkeletalMeshComponent());
+	check(PreviewActor->GetSkeletalMeshComponent()->SkeletalMesh);
 
 	const float HalfFOVRadians = FMath::DegreesToRadians<float>(InFOVDegrees) * 0.5f;
 	// No need to add extra size to view slightly outside of the sphere to compensate for perspective since skeletal meshes already buffer bounds.
-	const float HalfMeshSize = PreviewActor->SkeletalMeshComponent->Bounds.SphereRadius; 
-	const float BoundsZOffset = GetBoundsZOffset(PreviewActor->SkeletalMeshComponent->Bounds);
+	const float HalfMeshSize = PreviewActor->GetSkeletalMeshComponent()->Bounds.SphereRadius; 
+	const float BoundsZOffset = GetBoundsZOffset(PreviewActor->GetSkeletalMeshComponent()->Bounds);
 	const float TargetDistance = HalfMeshSize / FMath::Tan(HalfFOVRadians);
 
-	USceneThumbnailInfo* ThumbnailInfo = Cast<USceneThumbnailInfo>(PreviewActor->SkeletalMeshComponent->SkeletalMesh->ThumbnailInfo);
+	USceneThumbnailInfo* ThumbnailInfo = Cast<USceneThumbnailInfo>(PreviewActor->GetSkeletalMeshComponent()->SkeletalMesh->ThumbnailInfo);
 	if ( ThumbnailInfo )
 	{
 		if ( TargetDistance + ThumbnailInfo->OrbitZoom < 0 )
@@ -440,41 +440,41 @@ FStaticMeshThumbnailScene::FStaticMeshThumbnailScene()
 	SpawnInfo.ObjectFlags = RF_Transient;
 	PreviewActor = GetWorld()->SpawnActor<AStaticMeshActor>( SpawnInfo );
 
-	PreviewActor->StaticMeshComponent->SetMobility(EComponentMobility::Movable);
+	PreviewActor->GetStaticMeshComponent()->SetMobility(EComponentMobility::Movable);
 	PreviewActor->SetActorEnableCollision(false);
 }
 
 void FStaticMeshThumbnailScene::SetStaticMesh(UStaticMesh* StaticMesh)
 {
-	PreviewActor->StaticMeshComponent->SetStaticMesh(StaticMesh);
+	PreviewActor->GetStaticMeshComponent()->SetStaticMesh(StaticMesh);
 
 	if ( StaticMesh )
 	{
 		FTransform MeshTransform = FTransform::Identity;
 
 		PreviewActor->SetActorLocation(FVector(0,0,0), false);
-		PreviewActor->StaticMeshComponent->UpdateBounds();
+		PreviewActor->GetStaticMeshComponent()->UpdateBounds();
 
 		// Center the mesh at the world origin then offset to put it on top of the plane
-		const float BoundsZOffset = GetBoundsZOffset(PreviewActor->StaticMeshComponent->Bounds);
-		PreviewActor->SetActorLocation( -PreviewActor->StaticMeshComponent->Bounds.Origin + FVector(0, 0, BoundsZOffset), false );
-		PreviewActor->StaticMeshComponent->RecreateRenderState_Concurrent();
+		const float BoundsZOffset = GetBoundsZOffset(PreviewActor->GetStaticMeshComponent()->Bounds);
+		PreviewActor->SetActorLocation( -PreviewActor->GetStaticMeshComponent()->Bounds.Origin + FVector(0, 0, BoundsZOffset), false );
+		PreviewActor->GetStaticMeshComponent()->RecreateRenderState_Concurrent();
 	}
 }
 
 void FStaticMeshThumbnailScene::GetViewMatrixParameters(const float InFOVDegrees, FVector& OutOrigin, float& OutOrbitPitch, float& OutOrbitYaw, float& OutOrbitZoom) const
 {
 	check(PreviewActor);
-	check(PreviewActor->StaticMeshComponent);
-	check(PreviewActor->StaticMeshComponent->StaticMesh);
+	check(PreviewActor->GetStaticMeshComponent());
+	check(PreviewActor->GetStaticMeshComponent()->StaticMesh);
 
 	const float HalfFOVRadians = FMath::DegreesToRadians<float>(InFOVDegrees) * 0.5f;
 	// Add extra size to view slightly outside of the sphere to compensate for perspective
-	const float HalfMeshSize = PreviewActor->StaticMeshComponent->Bounds.SphereRadius * 1.15;
-	const float BoundsZOffset = GetBoundsZOffset(PreviewActor->StaticMeshComponent->Bounds);
+	const float HalfMeshSize = PreviewActor->GetStaticMeshComponent()->Bounds.SphereRadius * 1.15;
+	const float BoundsZOffset = GetBoundsZOffset(PreviewActor->GetStaticMeshComponent()->Bounds);
 	const float TargetDistance = HalfMeshSize / FMath::Tan(HalfFOVRadians);
 
-	USceneThumbnailInfo* ThumbnailInfo = Cast<USceneThumbnailInfo>(PreviewActor->StaticMeshComponent->StaticMesh->ThumbnailInfo);
+	USceneThumbnailInfo* ThumbnailInfo = Cast<USceneThumbnailInfo>(PreviewActor->GetStaticMeshComponent()->StaticMesh->ThumbnailInfo);
 	if ( ThumbnailInfo )
 	{
 		if ( TargetDistance + ThumbnailInfo->OrbitZoom < 0 )
@@ -516,7 +516,7 @@ FAnimationSequenceThumbnailScene::FAnimationSequenceThumbnailScene()
 
 bool FAnimationSequenceThumbnailScene::SetAnimation(UAnimSequenceBase* InAnimation)
 {
-	PreviewActor->SkeletalMeshComponent->Materials.Empty();
+	PreviewActor->GetSkeletalMeshComponent()->Materials.Empty();
 
 	bool bSetSucessfully = false;
 
@@ -528,7 +528,7 @@ bool FAnimationSequenceThumbnailScene::SetAnimation(UAnimSequenceBase* InAnimati
 		{
 			USkeletalMesh* PreviewSkeletalMesh = Skeleton->GetAssetPreviewMesh(InAnimation);
 
-			PreviewActor->SkeletalMeshComponent->SetSkeletalMesh(PreviewSkeletalMesh);
+			PreviewActor->GetSkeletalMeshComponent()->SetSkeletalMesh(PreviewSkeletalMesh);
 
 			if (PreviewSkeletalMesh)
 			{
@@ -537,36 +537,36 @@ bool FAnimationSequenceThumbnailScene::SetAnimation(UAnimSequenceBase* InAnimati
 				// Handle posing the mesh at the middle of the animation
 				const float AnimPosition = InAnimation->SequenceLength / 2.f;
 
-				PreviewActor->SkeletalMeshComponent->PlayAnimation(InAnimation,false);
-				PreviewActor->SkeletalMeshComponent->Stop();
-				PreviewActor->SkeletalMeshComponent->SetPosition(AnimPosition, false);
+				PreviewActor->GetSkeletalMeshComponent()->PlayAnimation(InAnimation,false);
+				PreviewActor->GetSkeletalMeshComponent()->Stop();
+				PreviewActor->GetSkeletalMeshComponent()->SetPosition(AnimPosition, false);
 
-				UAnimSingleNodeInstance* SingleNodeInstance = PreviewActor->SkeletalMeshComponent->GetSingleNodeInstance();
+				UAnimSingleNodeInstance* SingleNodeInstance = PreviewActor->GetSkeletalMeshComponent()->GetSingleNodeInstance();
 				if (SingleNodeInstance)
 				{
 					SingleNodeInstance->UpdateMontageWeightForTimeSkip(AnimPosition);
 				}
 
-				PreviewActor->SkeletalMeshComponent->RefreshBoneTransforms(nullptr);
+				PreviewActor->GetSkeletalMeshComponent()->RefreshBoneTransforms(nullptr);
 
 				FTransform MeshTransform = FTransform::Identity;
 
 				PreviewActor->SetActorLocation(FVector(0, 0, 0), false);
-				PreviewActor->SkeletalMeshComponent->UpdateBounds();
+				PreviewActor->GetSkeletalMeshComponent()->UpdateBounds();
 
 				// Center the mesh at the world origin then offset to put it on top of the plane
-				const float BoundsZOffset = GetBoundsZOffset(PreviewActor->SkeletalMeshComponent->Bounds);
-				PreviewActor->SetActorLocation(-PreviewActor->SkeletalMeshComponent->Bounds.Origin + FVector(0, 0, BoundsZOffset), false);
-				PreviewActor->SkeletalMeshComponent->RecreateRenderState_Concurrent();
+				const float BoundsZOffset = GetBoundsZOffset(PreviewActor->GetSkeletalMeshComponent()->Bounds);
+				PreviewActor->SetActorLocation(-PreviewActor->GetSkeletalMeshComponent()->Bounds.Origin + FVector(0, 0, BoundsZOffset), false);
+				PreviewActor->GetSkeletalMeshComponent()->RecreateRenderState_Concurrent();
 			}
 		}
 	}
 	
 	if(!bSetSucessfully)
 	{
-		CleanupComponentChildren(PreviewActor->SkeletalMeshComponent);
-		PreviewActor->SkeletalMeshComponent->SetAnimation(NULL);
-		PreviewActor->SkeletalMeshComponent->SetSkeletalMesh(nullptr);
+		CleanupComponentChildren(PreviewActor->GetSkeletalMeshComponent());
+		PreviewActor->GetSkeletalMeshComponent()->SetAnimation(NULL);
+		PreviewActor->GetSkeletalMeshComponent()->SetSkeletalMesh(nullptr);
 	}
 
 	return bSetSucessfully;
@@ -589,13 +589,13 @@ void FAnimationSequenceThumbnailScene::CleanupComponentChildren(USceneComponent*
 void FAnimationSequenceThumbnailScene::GetViewMatrixParameters(const float InFOVDegrees, FVector& OutOrigin, float& OutOrbitPitch, float& OutOrbitYaw, float& OutOrbitZoom) const
 {
 	check(PreviewAnimation);
-	check(PreviewActor->SkeletalMeshComponent);
-	check(PreviewActor->SkeletalMeshComponent->SkeletalMesh);
+	check(PreviewActor->GetSkeletalMeshComponent());
+	check(PreviewActor->GetSkeletalMeshComponent()->SkeletalMesh);
 
 	const float HalfFOVRadians = FMath::DegreesToRadians<float>(InFOVDegrees) * 0.5f;
 	// No need to add extra size to view slightly outside of the sphere to compensate for perspective since skeletal meshes already buffer bounds.
-	const float HalfMeshSize = PreviewActor->SkeletalMeshComponent->Bounds.SphereRadius;
-	const float BoundsZOffset = GetBoundsZOffset(PreviewActor->SkeletalMeshComponent->Bounds);
+	const float HalfMeshSize = PreviewActor->GetSkeletalMeshComponent()->Bounds.SphereRadius;
+	const float BoundsZOffset = GetBoundsZOffset(PreviewActor->GetSkeletalMeshComponent()->Bounds);
 	const float TargetDistance = HalfMeshSize / FMath::Tan(HalfFOVRadians);
 
 	USceneThumbnailInfo* ThumbnailInfo = Cast<USceneThumbnailInfo>(PreviewAnimation->ThumbnailInfo);
@@ -640,7 +640,7 @@ FBlendSpaceThumbnailScene::FBlendSpaceThumbnailScene()
 
 bool FBlendSpaceThumbnailScene::SetBlendSpace(class UBlendSpaceBase* InBlendSpace)
 {
-	PreviewActor->SkeletalMeshComponent->Materials.Empty();
+	PreviewActor->GetSkeletalMeshComponent()->Materials.Empty();
 
 	bool bSetSucessfully = false;
 
@@ -652,17 +652,17 @@ bool FBlendSpaceThumbnailScene::SetBlendSpace(class UBlendSpaceBase* InBlendSpac
 		{
 			USkeletalMesh* PreviewSkeletalMesh = Skeleton->GetAssetPreviewMesh(InBlendSpace);
 
-			PreviewActor->SkeletalMeshComponent->SetSkeletalMesh(PreviewSkeletalMesh);
+			PreviewActor->GetSkeletalMeshComponent()->SetSkeletalMesh(PreviewSkeletalMesh);
 
 			if (PreviewSkeletalMesh)
 			{
 				bSetSucessfully = true;
 
 				// Handle posing the mesh at the middle of the animation
-				PreviewActor->SkeletalMeshComponent->PlayAnimation(InBlendSpace, false);
-				PreviewActor->SkeletalMeshComponent->Stop();
+				PreviewActor->GetSkeletalMeshComponent()->PlayAnimation(InBlendSpace, false);
+				PreviewActor->GetSkeletalMeshComponent()->Stop();
 
-				UAnimSingleNodeInstance* AnimInstance = PreviewActor->SkeletalMeshComponent->GetSingleNodeInstance();
+				UAnimSingleNodeInstance* AnimInstance = PreviewActor->GetSkeletalMeshComponent()->GetSingleNodeInstance();
 				if (AnimInstance)
 				{
 					FVector BlendInput(0.f);
@@ -673,26 +673,26 @@ bool FBlendSpaceThumbnailScene::SetBlendSpace(class UBlendSpaceBase* InBlendSpac
 					}
 					AnimInstance->UpdateBlendspaceSamples(BlendInput);
 				}
-				PreviewActor->SkeletalMeshComponent->RefreshBoneTransforms(nullptr);
+				PreviewActor->GetSkeletalMeshComponent()->RefreshBoneTransforms(nullptr);
 
 				FTransform MeshTransform = FTransform::Identity;
 
 				PreviewActor->SetActorLocation(FVector(0, 0, 0), false);
-				PreviewActor->SkeletalMeshComponent->UpdateBounds();
+				PreviewActor->GetSkeletalMeshComponent()->UpdateBounds();
 
 				// Center the mesh at the world origin then offset to put it on top of the plane
-				const float BoundsZOffset = GetBoundsZOffset(PreviewActor->SkeletalMeshComponent->Bounds);
-				PreviewActor->SetActorLocation(-PreviewActor->SkeletalMeshComponent->Bounds.Origin + FVector(0, 0, BoundsZOffset), false);
-				PreviewActor->SkeletalMeshComponent->RecreateRenderState_Concurrent();
+				const float BoundsZOffset = GetBoundsZOffset(PreviewActor->GetSkeletalMeshComponent()->Bounds);
+				PreviewActor->SetActorLocation(-PreviewActor->GetSkeletalMeshComponent()->Bounds.Origin + FVector(0, 0, BoundsZOffset), false);
+				PreviewActor->GetSkeletalMeshComponent()->RecreateRenderState_Concurrent();
 			}
 		}
 	}
 
 	if (!bSetSucessfully)
 	{
-		CleanupComponentChildren(PreviewActor->SkeletalMeshComponent);
-		PreviewActor->SkeletalMeshComponent->SetAnimation(NULL);
-		PreviewActor->SkeletalMeshComponent->SetSkeletalMesh(nullptr);
+		CleanupComponentChildren(PreviewActor->GetSkeletalMeshComponent());
+		PreviewActor->GetSkeletalMeshComponent()->SetAnimation(NULL);
+		PreviewActor->GetSkeletalMeshComponent()->SetSkeletalMesh(nullptr);
 	}
 
 	return bSetSucessfully;
@@ -715,13 +715,13 @@ void FBlendSpaceThumbnailScene::CleanupComponentChildren(USceneComponent* Compon
 void FBlendSpaceThumbnailScene::GetViewMatrixParameters(const float InFOVDegrees, FVector& OutOrigin, float& OutOrbitPitch, float& OutOrbitYaw, float& OutOrbitZoom) const
 {
 	check(PreviewAnimation);
-	check(PreviewActor->SkeletalMeshComponent);
-	check(PreviewActor->SkeletalMeshComponent->SkeletalMesh);
+	check(PreviewActor->GetSkeletalMeshComponent());
+	check(PreviewActor->GetSkeletalMeshComponent()->SkeletalMesh);
 
 	const float HalfFOVRadians = FMath::DegreesToRadians<float>(InFOVDegrees) * 0.5f;
 	// No need to add extra size to view slightly outside of the sphere to compensate for perspective since skeletal meshes already buffer bounds.
-	const float HalfMeshSize = PreviewActor->SkeletalMeshComponent->Bounds.SphereRadius;
-	const float BoundsZOffset = GetBoundsZOffset(PreviewActor->SkeletalMeshComponent->Bounds);
+	const float HalfMeshSize = PreviewActor->GetSkeletalMeshComponent()->Bounds.SphereRadius;
+	const float BoundsZOffset = GetBoundsZOffset(PreviewActor->GetSkeletalMeshComponent()->Bounds);
 	const float TargetDistance = HalfMeshSize / FMath::Tan(HalfFOVRadians);
 
 	USceneThumbnailInfo* ThumbnailInfo = Cast<USceneThumbnailInfo>(PreviewAnimation->ThumbnailInfo);
@@ -766,7 +766,7 @@ FAnimBlueprintThumbnailScene::FAnimBlueprintThumbnailScene()
 
 bool FAnimBlueprintThumbnailScene::SetAnimBlueprint(class UAnimBlueprint* InBlueprint)
 {
-	PreviewActor->SkeletalMeshComponent->Materials.Empty();
+	PreviewActor->GetSkeletalMeshComponent()->Materials.Empty();
 
 	bool bSetSucessfully = false;
 
@@ -778,32 +778,32 @@ bool FAnimBlueprintThumbnailScene::SetAnimBlueprint(class UAnimBlueprint* InBlue
 		{
 			USkeletalMesh* PreviewSkeletalMesh = Skeleton->GetPreviewMesh(true);
 
-			PreviewActor->SkeletalMeshComponent->SetSkeletalMesh(PreviewSkeletalMesh);
+			PreviewActor->GetSkeletalMeshComponent()->SetSkeletalMesh(PreviewSkeletalMesh);
 
 			if (PreviewSkeletalMesh)
 			{
 				bSetSucessfully = true;
 
-				PreviewActor->SkeletalMeshComponent->SetAnimInstanceClass(InBlueprint->GeneratedClass);
+				PreviewActor->GetSkeletalMeshComponent()->SetAnimInstanceClass(InBlueprint->GeneratedClass);
 
 				FTransform MeshTransform = FTransform::Identity;
 
 				PreviewActor->SetActorLocation(FVector(0, 0, 0), false);
-				PreviewActor->SkeletalMeshComponent->UpdateBounds();
+				PreviewActor->GetSkeletalMeshComponent()->UpdateBounds();
 
 				// Center the mesh at the world origin then offset to put it on top of the plane
-				const float BoundsZOffset = GetBoundsZOffset(PreviewActor->SkeletalMeshComponent->Bounds);
-				PreviewActor->SetActorLocation(-PreviewActor->SkeletalMeshComponent->Bounds.Origin + FVector(0, 0, BoundsZOffset), false);
-				PreviewActor->SkeletalMeshComponent->RecreateRenderState_Concurrent();
+				const float BoundsZOffset = GetBoundsZOffset(PreviewActor->GetSkeletalMeshComponent()->Bounds);
+				PreviewActor->SetActorLocation(-PreviewActor->GetSkeletalMeshComponent()->Bounds.Origin + FVector(0, 0, BoundsZOffset), false);
+				PreviewActor->GetSkeletalMeshComponent()->RecreateRenderState_Concurrent();
 			}
 		}
 	}
 
 	if (!bSetSucessfully)
 	{
-		CleanupComponentChildren(PreviewActor->SkeletalMeshComponent);
-		PreviewActor->SkeletalMeshComponent->SetSkeletalMesh(nullptr);
-		PreviewActor->SkeletalMeshComponent->SetAnimInstanceClass(nullptr);
+		CleanupComponentChildren(PreviewActor->GetSkeletalMeshComponent());
+		PreviewActor->GetSkeletalMeshComponent()->SetSkeletalMesh(nullptr);
+		PreviewActor->GetSkeletalMeshComponent()->SetAnimInstanceClass(nullptr);
 	}
 
 	return bSetSucessfully;
@@ -826,13 +826,13 @@ void FAnimBlueprintThumbnailScene::CleanupComponentChildren(USceneComponent* Com
 void FAnimBlueprintThumbnailScene::GetViewMatrixParameters(const float InFOVDegrees, FVector& OutOrigin, float& OutOrbitPitch, float& OutOrbitYaw, float& OutOrbitZoom) const
 {
 	check(PreviewBlueprint);
-	check(PreviewActor->SkeletalMeshComponent);
-	check(PreviewActor->SkeletalMeshComponent->SkeletalMesh);
+	check(PreviewActor->GetSkeletalMeshComponent());
+	check(PreviewActor->GetSkeletalMeshComponent()->SkeletalMesh);
 
 	const float HalfFOVRadians = FMath::DegreesToRadians<float>(InFOVDegrees) * 0.5f;
 	// No need to add extra size to view slightly outside of the sphere to compensate for perspective since skeletal meshes already buffer bounds.
-	const float HalfMeshSize = PreviewActor->SkeletalMeshComponent->Bounds.SphereRadius;
-	const float BoundsZOffset = GetBoundsZOffset(PreviewActor->SkeletalMeshComponent->Bounds);
+	const float HalfMeshSize = PreviewActor->GetSkeletalMeshComponent()->Bounds.SphereRadius;
+	const float BoundsZOffset = GetBoundsZOffset(PreviewActor->GetSkeletalMeshComponent()->Bounds);
 	const float TargetDistance = HalfMeshSize / FMath::Tan(HalfFOVRadians);
 
 	USceneThumbnailInfo* ThumbnailInfo = Cast<USceneThumbnailInfo>(PreviewBlueprint->ThumbnailInfo);
diff --git a/Engine/Source/Editor/UnrealEd/Private/UnrealEdSrv.cpp b/Engine/Source/Editor/UnrealEd/Private/UnrealEdSrv.cpp
--- a/Engine/Source/Editor/UnrealEd/Private/UnrealEdSrv.cpp
+++ b/Engine/Source/Editor/UnrealEd/Private/UnrealEdSrv.cpp
@@ -2858,11 +2858,11 @@ bool UUnrealEdEngine::Exec_Actor( UWorld* InWorld, const TCHAR* Str, FOutputDevi
 					{
 						if (Actor->IsA(AStaticMeshActor::StaticClass()))
 						{
-							Exporter->ExportStaticMesh(Actor, CastChecked<AStaticMeshActor>(Actor)->StaticMeshComponent, NULL);
+							Exporter->ExportStaticMesh(Actor, CastChecked<AStaticMeshActor>(Actor)->GetStaticMeshComponent(), NULL);
 						}
 						else if (Actor->IsA(ASkeletalMeshActor::StaticClass()))
 						{
-							Exporter->ExportSkeletalMesh(Actor, CastChecked<ASkeletalMeshActor>(Actor)->SkeletalMeshComponent);
+							Exporter->ExportSkeletalMesh(Actor, CastChecked<ASkeletalMeshActor>(Actor)->GetSkeletalMeshComponent());
 						}
 						else if (Actor->IsA(ABrush::StaticClass()))
 						{
diff --git a/Engine/Source/Editor/WorldBrowser/Private/Tiles/WorldTileCollectionModel.cpp b/Engine/Source/Editor/WorldBrowser/Private/Tiles/WorldTileCollectionModel.cpp
--- a/Engine/Source/Editor/WorldBrowser/Private/Tiles/WorldTileCollectionModel.cpp
+++ b/Engine/Source/Editor/WorldBrowser/Private/Tiles/WorldTileCollectionModel.cpp
@@ -1983,7 +1983,7 @@ bool FWorldTileCollectionModel::GenerateLODLevels(FLevelModelList InLevelList, i
 				FVector Location = AssetsToSpawnTransform[AssetIdx].GetLocation();
 				FRotator Rotation(ForceInit);
 				AStaticMeshActor* MeshActor = LODWorld->SpawnActor<AStaticMeshActor>(Location, Rotation);
-				MeshActor->StaticMeshComponent->StaticMesh = AssetsToSpawn[AssetIdx];
+				MeshActor->GetStaticMeshComponent()->StaticMesh = AssetsToSpawn[AssetIdx];
 				MeshActor->SetActorLabel(AssetsToSpawn[AssetIdx]->GetName());
 			}
 		
diff --git a/Engine/Source/Programs/UnrealHeaderTool/Private/HeaderParser.cpp b/Engine/Source/Programs/UnrealHeaderTool/Private/HeaderParser.cpp
--- a/Engine/Source/Programs/UnrealHeaderTool/Private/HeaderParser.cpp
+++ b/Engine/Source/Programs/UnrealHeaderTool/Private/HeaderParser.cpp
@@ -3061,13 +3061,7 @@ bool FHeaderParser::GetVarType
 					}
 					else if (bIsSubobjectPtrTemplate)
 					{
-						Flags |= CPF_SubobjectReference;
-
-						if (((Flags & CPF_Edit) && (Flags & CPF_EditConst) == 0) ||
-							  ((Flags & CPF_BlueprintVisible) && (Flags & CPF_BlueprintReadOnly) == 0))
-						{
-							FError::Throwf(TEXT("%s: Subobject properties can't be editable (use VisibleAnywhere or BlueprintReadOnly instead)."), VarType.Identifier);
-						}
+						Flags |= CPF_SubobjectReference | CPF_InstancedReference;
 					}
 
 					bExpectStar = false;
@@ -3189,6 +3183,15 @@ bool FHeaderParser::GetVarType
 				FError::Throwf(TEXT("Unrecognized type '%s'"), VarType.Identifier );
 			}
 		}
+
+		if ((Flags & CPF_InstancedReference) && CurrentAccessSpecifier == ACCESS_Private && VariableCategory == EVariableCategory::Member)
+		{
+			if (((Flags & CPF_Edit) && (Flags & CPF_EditConst) == 0) ||
+				((Flags & CPF_BlueprintVisible) && (Flags & CPF_BlueprintReadOnly) == 0))
+			{
+				FError::Throwf(TEXT("%s: Subobject (instanced) properties can't be editable (use VisibleAnywhere or BlueprintReadOnly instead)."), VarType.Identifier);
+			}
+		}
 	}
 
 	if (VariableCategory != EVariableCategory::Member)
