PIE client/server window position/size tries to remain within the confines of the desktop area

#ttp 325376 - EDITOR: SHOOTERGAME: New client/server position outside of monitor

#branch UE4

#change Modified the code so that the window now attempts to position itself within the desktop area, and if the desired window size is too large of the desktop then it shrinks itself to fit. It tiles the windows horizontally right until they reach the width of the desktop, then it moves down 'a line' and repeats until it reaches the bottom right of the desktop where it then jump to the top left and repeats.

#change Changed GetWindowBorderSize to add on the title bar size to the margin if so desired

#note The project's Saved\\Config\\ directory contains files like PIEGameUserSettings#, these save the preferred position of the PIE window. If you notice that the window spawns at a good location, then jumps to another it's likely that it's due to a saved position being present in these configs. Please try deleting these and reproing steps before reporting any issues.

diff --git a/Engine/Source/Editor/UnrealEd/Private/PlayLevel.cpp b/Engine/Source/Editor/UnrealEd/Private/PlayLevel.cpp
--- a/Engine/Source/Editor/UnrealEd/Private/PlayLevel.cpp
+++ b/Engine/Source/Editor/UnrealEd/Private/PlayLevel.cpp
@@ -651,31 +651,79 @@ void UEditorEngine::PlaySessionSingleStepped()
 	FEditorDelegates::SingleStepPIE.Broadcast(bIsSimulatingInEditor);
 }
 
-void AdvanceWindowPositionsForNextPIEWindow(int32 &WinX, int32 &WinY, FIntPoint WinSize)
+/* fits the window position to make sure it falls within the confines of the desktop */
+void FitWindowPositionToWorkArea(FIntPoint &WinPos, FIntPoint &WinSize, const FMargin &WinPadding)
 {
-	const FIntPoint WindowPadding(16, 16);
-
-	WinX += WinSize.X + WindowPadding.X;
+	const int32 HorzPad = WinPadding.GetTotalSpaceAlong<Orient_Horizontal>();
+	const int32 VertPad = WinPadding.GetTotalSpaceAlong<Orient_Vertical>();
+	FIntPoint TotalSize( WinSize.X + HorzPad, WinSize.Y + VertPad );
 
 	FDisplayMetrics DisplayMetrics;
 	FSlateApplication::Get().GetDisplayMetrics(DisplayMetrics);
-	FSlateRect PreferredWorkArea(DisplayMetrics.VirtualDisplayRect.Left, DisplayMetrics.VirtualDisplayRect.Top, DisplayMetrics.VirtualDisplayRect.Right, DisplayMetrics.VirtualDisplayRect.Bottom);
-	// if no more windows fit horizontally, place them in a new row
-	if ((WinX + WindowPadding.X + WinSize.X) > PreferredWorkArea.Right)
+
+	// Limit the size, to make sure it fits within the desktop area
 	{
-		WinX = PreferredWorkArea.Left + WindowPadding.X;
-		WinY += WindowPadding.Y;
+		FIntPoint NewWinSize;
+		NewWinSize.X = FMath::Min(TotalSize.X, DisplayMetrics.VirtualDisplayRect.Right - DisplayMetrics.VirtualDisplayRect.Left );
+		NewWinSize.Y = FMath::Min(TotalSize.Y, DisplayMetrics.VirtualDisplayRect.Bottom - DisplayMetrics.VirtualDisplayRect.Top );
+		if( NewWinSize != TotalSize )
+		{
+			TotalSize = NewWinSize;
+			WinSize.X = NewWinSize.X - HorzPad;
+			WinSize.Y = NewWinSize.Y - VertPad;
+		}
 	}
 
+	const FSlateRect PreferredWorkArea( DisplayMetrics.VirtualDisplayRect.Left, 
+										DisplayMetrics.VirtualDisplayRect.Top, 
+										DisplayMetrics.VirtualDisplayRect.Right - TotalSize.X, 
+										DisplayMetrics.VirtualDisplayRect.Bottom - TotalSize.Y );
+
+	// if no more windows fit horizontally, place them in a new row
+	if (WinPos.X > PreferredWorkArea.Right)
+	{
+		WinPos.X = PreferredWorkArea.Left;
+		WinPos.Y += TotalSize.Y;
+		if (WinPos.Y > PreferredWorkArea.Bottom)
+		{
+			WinPos.Y = PreferredWorkArea.Top;
+		}
+	}
+	
 	// if no more rows fit vertically, stack windows on top of each other
-	if (WinY + WindowPadding.Y + WinSize.Y >= PreferredWorkArea.Bottom)
+	else if (WinPos.Y > PreferredWorkArea.Bottom)
 	{
-		WinX += WindowPadding.X;
-		WinY = PreferredWorkArea.Top + WindowPadding.Y;
+		WinPos.Y = PreferredWorkArea.Top;
+		WinPos.X += TotalSize.X;
+		if (WinPos.X > PreferredWorkArea.Right)
+		{
+			WinPos.X = PreferredWorkArea.Left;
+		}
 	}
+
+	// Clamp values to make sure they fall within the desktop area
+	WinPos.X = FMath::Clamp(WinPos.X, (int32)PreferredWorkArea.Left, (int32)PreferredWorkArea.Right);
+	WinPos.Y = FMath::Clamp(WinPos.Y, (int32)PreferredWorkArea.Top, (int32)PreferredWorkArea.Bottom);
 }
 
-/* returns the size of the window depending of the net mode. */
+/* advances the windows position to the next location and fits */
+void AdvanceWindowPositionForNextPIEWindow(FIntPoint &WinPos, const FIntPoint &WinSize, const FMargin &WinPadding, bool bVertical)
+{
+	const int32 HorzPad = WinPadding.GetTotalSpaceAlong<Orient_Horizontal>();
+	const int32 VertPad = WinPadding.GetTotalSpaceAlong<Orient_Vertical>();
+	const FIntPoint TotalSize( WinSize.X + HorzPad, WinSize.Y + VertPad );
+
+	if(bVertical)
+	{
+		WinPos.Y += TotalSize.Y;
+	}
+	else
+	{
+		WinPos.X += TotalSize.X;
+	}
+}
+
+/* returns the size of the window depending on the net mode. */
 void GetWindowSizeForInstanceType(FIntPoint &WindowSize, const ULevelEditorPlaySettings* PlayInSettings)
 {
 	if (PlayInSettings->PlayNetMode == PIE_Standalone)
@@ -690,15 +738,29 @@ void GetWindowSizeForInstanceType(FIntPoint &WindowSize, const ULevelEditorPlayS
 	}
 }
 
+/* sets the size of the window depending on the net mode */
+void SetWindowSizeForInstanceType(const FIntPoint &WindowSize, ULevelEditorPlaySettings* PlayInSettings)
+{
+	if (PlayInSettings->PlayNetMode == PIE_Standalone)
+	{
+		PlayInSettings->StandaloneWindowWidth = WindowSize.X;
+		PlayInSettings->StandaloneWindowHeight = WindowSize.Y;
+	}
+	else
+	{
+		PlayInSettings->ClientWindowWidth = WindowSize.X;
+		PlayInSettings->ClientWindowHeight = WindowSize.Y;
+	}
+}
+
 /** 
  * Generate the command line for pie instance. Window position, size etc. 
  *
- * @param	WinX			Window X position. This will contain the X position to use for the next window. (Not changed for dedicated server window).
- * @param	WinY			Window Y position. This will contain the X position to use for the next window. (Not changed for dedicated server window).
+ * @param	WinPos			Window position. This will contain the X & Y position to use for the next window. (Not changed for dedicated server window).
  * @param	InstanceNum		PIE instance index.
  * @param	bIsDedicatedServer	Is this instance a dedicate server. true if so else false.
  */
-FString GenerateCmdLineForNextPieInstance(int32 &WinX, int32 &WinY, int32 &InstanceNum, bool bIsDedicatedServer)
+FString GenerateCmdLineForNextPieInstance(FIntPoint &WinPos, int32 &InstanceNum, bool bIsDedicatedServer)
 {
 	const ULevelEditorPlaySettings* PlayInSettings = GetDefault<ULevelEditorPlaySettings>();
 	// Get GameSettings INI override
@@ -717,27 +779,31 @@ FString GenerateCmdLineForNextPieInstance(int32 &WinX, int32 &WinY, int32 &Insta
 	}
 	else
 	{
-		// Adjust window by WindowBorderSize and title bar so it doesn't appear off-screen
-		FMargin WindowBorderSize(16.0f, 80.0f);
+		// Default to what we expect the border to be (on windows at least) to prevent it occurring offscreen if TLW call fails
+		FMargin WindowBorderSize(8.0f, 30.0f, 8.0f, 8.0f);
 		TSharedPtr<SWindow> TopLevelWindow = FSlateApplication::Get().GetActiveTopLevelWindow();
 
 		if (TopLevelWindow.IsValid())
 		{
-			WindowBorderSize = TopLevelWindow->GetWindowBorderSize();
+			WindowBorderSize = TopLevelWindow->GetWindowBorderSize(true);
 		}
 
-		// Listen server or clients: specify default win position and SAVEWINPOS so the final positions are saved
-		// in order to preserve PIE networking window setup
-		int32 WindowXPos = WinX + WindowBorderSize.Left;
-		int32 WindowYPos = WinY + WindowBorderSize.Top;
-		CmdLine += FString::Printf(TEXT("WinX=%d WinY=%d SAVEWINPOS=1"), WindowXPos, WindowYPos);
-
 		// Get the size of the window based on the type
 		FIntPoint WinSize(0,0);
 		GetWindowSizeForInstanceType(WinSize, PlayInSettings);
 
-		// Advance window
-		AdvanceWindowPositionsForNextPIEWindow(WinX, WinY, WinSize);
+		// Make sure the window is going to fit where we want it
+		FitWindowPositionToWorkArea(WinPos, WinSize, WindowBorderSize);
+
+		// Set the size, incase it was modified
+		SetWindowSizeForInstanceType(WinSize, GetMutableDefault<ULevelEditorPlaySettings>());
+
+		// Listen server or clients: specify default win position and SAVEWINPOS so the final positions are saved
+		// in order to preserve PIE networking window setup
+		CmdLine += FString::Printf(TEXT("WinX=%d WinY=%d SAVEWINPOS=1"), WinPos.X + (int32)WindowBorderSize.Left, WinPos.Y + (int32)WindowBorderSize.Top);
+
+		// Advance window for next PIE instance...
+		AdvanceWindowPositionForNextPIEWindow(WinPos, WinSize, WindowBorderSize, false);
 	}
 	
 	return CmdLine;
@@ -761,9 +827,22 @@ void GetMultipleInstancePositions(int32 index, int32 &LastX, int32 &LastY)
 		PlayInSettings->NewWindowPosition = FIntPoint(LastX, LastY);
 	}
 
+	FIntPoint WinPos(LastX, LastY);
+
+	// Get the size of the window based on the type
 	FIntPoint WinSize(0, 0);
 	GetWindowSizeForInstanceType(WinSize, PlayInSettings);
-	AdvanceWindowPositionsForNextPIEWindow(LastX, LastY, WinSize);
+
+	// Advance window and make sure the window is going to fit where we want it
+	const FMargin WinPadding(16, 16);
+	AdvanceWindowPositionForNextPIEWindow(WinPos, WinSize, WinPadding, false);
+	FitWindowPositionToWorkArea(WinPos, WinSize, WinPadding);
+
+	// Set the size, incase it was modified
+	SetWindowSizeForInstanceType(WinSize, PlayInSettings);
+
+	LastX = WinPos.X;
+	LastY = WinPos.Y;
 }
 
 void UEditorEngine::StartQueuedPlayMapRequest()
@@ -779,10 +858,9 @@ void UEditorEngine::StartQueuedPlayMapRequest()
 	{
 		int32 NumClients = 0;
 
-		// If we start to the right of the editor work area, call advance it will find the next place we can place a new instance window
-		FSlateRect PreferredWorkArea = FSlateApplication::Get().GetPreferredWorkArea();		
+		// If we start to the right of the editor work area, call FitToWorkArea and it will find the next place we can place a new instance window if that's not preferable.
+		const FSlateRect PreferredWorkArea = FSlateApplication::Get().GetPreferredWorkArea();		
 		FIntPoint WinPosition((int32)PreferredWorkArea.Right, (int32)PreferredWorkArea.Top);
-		AdvanceWindowPositionsForNextPIEWindow(WinPosition.X, WinPosition.Y, FIntPoint(0,0));
 
 		// We'll need to spawn a server if we're playing outside the editor or the editor wants to run as a client
 		if (bPlayOnLocalPcSession || PlayInSettings->PlayNetMode == PIE_Client)
@@ -991,7 +1069,7 @@ void UEditorEngine::PlayStandaloneLocalPc(FString MapNameOverride, FIntPoint* Wi
 	FString CmdLine;
 	if (WindowPos != NULL)	// If WindowPos == NULL, we're just launching one instance
 	{
-		CmdLine = GenerateCmdLineForNextPieInstance(WindowPos->X, WindowPos->Y, PIENum, bIsServer && PlayInSettings->PlayNetDedicated);
+		CmdLine = GenerateCmdLineForNextPieInstance(*WindowPos, PIENum, bIsServer && PlayInSettings->PlayNetDedicated);
 	}
 	
 	const FString URLParms = bIsServer && !PlayInSettings->PlayNetDedicated ? TEXT("?Listen") : FString();
diff --git a/Engine/Source/Runtime/SlateCore/Private/Widgets/SWindow.cpp b/Engine/Source/Runtime/SlateCore/Private/Widgets/SWindow.cpp
--- a/Engine/Source/Runtime/SlateCore/Private/Widgets/SWindow.cpp
+++ b/Engine/Source/Runtime/SlateCore/Private/Widgets/SWindow.cpp
@@ -458,7 +458,7 @@ void SWindow::ConstructWindowInternals( const bool bCreateTitleBar )
 
 					+ SVerticalBox::Slot()
 					
-						.Padding(TAttribute<FMargin>( this, &SWindow::GetWindowBorderSize ))
+						.Padding(TAttribute<FMargin>::Create(TAttribute<FMargin>::FGetter::CreateSP(this, &SWindow::GetWindowBorderSize, false)))
 						[
 							MainWindowArea
 						]
@@ -671,7 +671,7 @@ FSlateRect SWindow::GetClippingRectangleInWindow() const
 }
 
 
-FMargin SWindow::GetWindowBorderSize() const
+FMargin SWindow::GetWindowBorderSize( bool bIncTitleBar ) const
 {
 // Mac didn't want a window border, and consoles don't either, so only do this in Windows
 
@@ -681,8 +681,13 @@ FMargin SWindow::GetWindowBorderSize() const
 #if PLATFORM_WINDOWS // || PLATFORM_LINUX
 	if (NativeWindow.IsValid() && NativeWindow->IsMaximized())
 	{
-		int32 OSWindowBorderSize = NativeWindow->GetWindowBorderSize();
-		return FMargin(OSWindowBorderSize);
+		FMargin BorderSize(NativeWindow->GetWindowBorderSize());
+		if(bIncTitleBar)
+		{
+			// Add title bar size (whether it's visible or not)
+			BorderSize.Top += GetSystemMetrics(SM_CYCAPTION);
+		}
+		return BorderSize;
 	}
 
 	return SWindowDefs::WindowBorderSize;
diff --git a/Engine/Source/Runtime/SlateCore/Public/Widgets/SWindow.h b/Engine/Source/Runtime/SlateCore/Public/Widgets/SWindow.h
--- a/Engine/Source/Runtime/SlateCore/Public/Widgets/SWindow.h
+++ b/Engine/Source/Runtime/SlateCore/Public/Widgets/SWindow.h
@@ -257,7 +257,7 @@ public:
 	FSlateRect GetClippingRectangleInWindow() const;
 
 	/** Returns the margins used for the window border. This varies based on whether it's maximized or not. */
-	FMargin GetWindowBorderSize() const;
+	FMargin GetWindowBorderSize( bool bIncTitleBar = false ) const;
 
 	/** Relocate the window to a screenspace position specified by NewPosition */
 	void MoveWindowTo( FVector2D NewPosition );
