Hide the Player start class from settings as it is too advanced

diff --git a/Engine/Source/Editor/UnrealEd/Classes/Settings/LevelEditorPlaySettings.h b/Engine/Source/Editor/UnrealEd/Classes/Settings/LevelEditorPlaySettings.h
--- a/Engine/Source/Editor/UnrealEd/Classes/Settings/LevelEditorPlaySettings.h
+++ b/Engine/Source/Editor/UnrealEd/Classes/Settings/LevelEditorPlaySettings.h
@@ -115,7 +115,7 @@ class UNREALED_API ULevelEditorPlaySettings
 public:
 
 	/** The PlayerStart class used when spawning the player at the current camera location. */
-	UPROPERTY(globalconfig, EditAnywhere, Category=General, meta=(ToolTip=""))
+	UPROPERTY(globalconfig, Category=General, meta=(ToolTip=""))
 	FString PlayFromHerePlayerStartClassName;
 
 public:
