Fixed issue where pressing F to focus to a BSP without selected polys, edges or verts points the camera at the world origin.
#ttp 349580 - LIVE: EDITOR: "F" Focus Hotkey not functioning in Geometry Edit Mode
#branch UE4
#proj Editor.UnrealEd

diff --git a/Engine/Source/Editor/UnrealEd/Private/EditorServer.cpp b/Engine/Source/Editor/UnrealEd/Private/EditorServer.cpp
--- a/Engine/Source/Editor/UnrealEd/Private/EditorServer.cpp
+++ b/Engine/Source/Editor/UnrealEd/Private/EditorServer.cpp
@@ -4200,34 +4200,6 @@ void UEditorEngine::MoveViewportCamerasToActor(const TArray<AActor*> &Actors, bo
 				const FBox DefaultSizeBox(Actor->GetActorLocation() - DefaultExtent, Actor->GetActorLocation() + DefaultExtent);
 				BoundingBox += DefaultSizeBox;
 			}
-			else if( Actor->IsA<ABrush>() && GLevelEditorModeTools().IsModeActive( FBuiltinEditorModes::EM_Geometry ) )
-			{
-				FEdModeGeometry* GeometryMode = GLevelEditorModeTools().GetActiveModeTyped<FEdModeGeometry>(FBuiltinEditorModes::EM_Geometry);
-
-				TArray<FGeomVertex*> SelectedVertices;
-				GeometryMode->GetSelectedVertices( SelectedVertices );
-				for( FGeomVertex* Vertex : SelectedVertices )
-				{
-					BoundingBox += Vertex->GetWidgetLocation();
-				}
-				
-				TArray<FGeomPoly*> SelectedPolys;
-				GeometryMode->GetSelectedPolygons( SelectedPolys );
-				for( FGeomPoly* Poly : SelectedPolys )
-				{
-					BoundingBox += Poly->GetWidgetLocation();
-				}
-
-				TArray<FGeomEdge*> SelectedEdges;
-				GeometryMode->GetSelectedEdges( SelectedEdges );
-				for( FGeomEdge* Edge : SelectedEdges )
-				{
-					BoundingBox += Edge->GetWidgetLocation();
-				}
-
-				// Zoom out a little bit so you can see the selection
-				BoundingBox = BoundingBox.ExpandBy( 25 );
-			}
 			else
 			{
 				TArray<UPrimitiveComponent*> Components;
@@ -4262,6 +4234,42 @@ void UEditorEngine::MoveViewportCamerasToActor(const TArray<AActor*> &Actors, bo
 						}
 					}
 				}
+
+				if (Actor->IsA<ABrush>() && GLevelEditorModeTools().IsModeActive(FBuiltinEditorModes::EM_Geometry))
+				{
+					FEdModeGeometry* GeometryMode = GLevelEditorModeTools().GetActiveModeTyped<FEdModeGeometry>(FBuiltinEditorModes::EM_Geometry);
+
+					TArray<FGeomVertex*> SelectedVertices;
+					TArray<FGeomPoly*> SelectedPolys;
+					TArray<FGeomEdge*> SelectedEdges;
+
+					GeometryMode->GetSelectedVertices(SelectedVertices);
+					GeometryMode->GetSelectedPolygons(SelectedPolys);
+					GeometryMode->GetSelectedEdges(SelectedEdges);
+
+					if (SelectedVertices.Num() + SelectedPolys.Num() + SelectedEdges.Num() > 0)
+					{
+						BoundingBox.Init();
+
+						for (FGeomVertex* Vertex : SelectedVertices)
+						{
+							BoundingBox += Vertex->GetWidgetLocation();
+						}
+
+						for (FGeomPoly* Poly : SelectedPolys)
+						{
+							BoundingBox += Poly->GetWidgetLocation();
+						}
+
+						for (FGeomEdge* Edge : SelectedEdges)
+						{
+							BoundingBox += Edge->GetWidgetLocation();
+						}
+
+						// Zoom out a little bit so you can see the selection
+						BoundingBox = BoundingBox.ExpandBy(25);
+					}
+				}
 			}
 
 			NumActiveActors++;
