Further fix for the crash when renaming a level asset that was recently removed from the level browser.
#ttp 347330 - CRASH: Attempting to rename a level asset that was recently removed from the level browser can result in a crash
#branch UE4
#proj Editor.UnrealEd

diff --git a/Engine/Source/Editor/UnrealEd/Private/PackageTools.cpp b/Engine/Source/Editor/UnrealEd/Private/PackageTools.cpp
--- a/Engine/Source/Editor/UnrealEd/Private/PackageTools.cpp
+++ b/Engine/Source/Editor/UnrealEd/Private/PackageTools.cpp
@@ -327,6 +327,7 @@ namespace PackageTools
 				GWarn->StatusUpdate( PackageIndex, PackagesToUnload.Num(), FText::Format(NSLOCTEXT("UnrealEd", "Unloadingf", "Unloading {0}..."), FText::FromString(PackageBeingUnloaded->GetName()) ) );
 
 				PackageBeingUnloaded->bHasBeenFullyLoaded = false;
+				PackageBeingUnloaded->ClearFlags(RF_WasLoaded);
 				if ( PackageBeingUnloaded->PackageFlags & PKG_ContainsScript )
 				{
 					bScriptPackageWasUnloaded = true;
