Fixed bug in which the capture region was ignored when capturing high res screenshots.
#branch UE4
#proj Editor.UnrealEd

diff --git a/Engine/Source/Editor/UnrealEd/Private/EditorViewportClient.cpp b/Engine/Source/Editor/UnrealEd/Private/EditorViewportClient.cpp
--- a/Engine/Source/Editor/UnrealEd/Private/EditorViewportClient.cpp
+++ b/Engine/Source/Editor/UnrealEd/Private/EditorViewportClient.cpp
@@ -3928,6 +3928,26 @@ void FEditorViewportClient::ProcessScreenShots(FViewport* InViewport)
 				bWriteAlpha = HighResScreenshotConfig.MergeMaskIntoAlpha(Bitmap);
 			}
 
+			// Clip the bitmap to just the capture region if valid
+			if (!SourceRect.IsEmpty())
+			{
+				FColor* const Data = Bitmap.GetData();
+				const int32 OldWidth = BitmapSize.X;
+				const int32 OldHeight = BitmapSize.Y;
+				const int32 NewWidth = SourceRect.Width();
+				const int32 NewHeight = SourceRect.Height();
+				const int32 CaptureTopRow = SourceRect.Min.Y;
+				const int32 CaptureLeftColumn = SourceRect.Min.X;
+
+				for (int32 Row = 0; Row < NewHeight; Row++)
+				{
+					FMemory::Memmove(Data + Row * NewWidth, Data + (Row + CaptureTopRow) * OldWidth + CaptureLeftColumn, NewWidth * sizeof(*Data));
+				}
+
+				Bitmap.RemoveAt(NewWidth * NewHeight, OldWidth * OldHeight - NewWidth * NewHeight, false);
+				BitmapSize = FIntPoint(NewWidth, NewHeight);
+			}
+
 			// Set full alpha on the bitmap
 			if (!bWriteAlpha)
 			{
