Performance test code cleanup.

diff --git a/Engine/Source/Editor/UnrealEd/Private/Tests/AutomationEditorCommon.cpp b/Engine/Source/Editor/UnrealEd/Private/Tests/AutomationEditorCommon.cpp
--- a/Engine/Source/Editor/UnrealEd/Private/Tests/AutomationEditorCommon.cpp
+++ b/Engine/Source/Editor/UnrealEd/Private/Tests/AutomationEditorCommon.cpp
@@ -283,58 +283,56 @@ namespace AutomationEditorCommonUtils
 	}
 }
 
+
 /**
 * Writes a number to a text file.
-* @param FolderNameForTest is the folder that has the same name as the test. (For Example: "Performance").
-* @param FolderNameForBeingTested is the name for the thing that is being tested. (For Example: "MapName").
-* @param FileName is the name of the file with an extension
-* @param NumberToWriteToFile is the float number that is expected to be written to the file.
+* @param InTestName is the folder that has the same name as the test. (For Example: "Performance").
+* @param InItemBeingTested is the name for the thing that is being tested. (For Example: "MapName").
+* @param InFileName is the name of the file with an extension
+* @param InNumberToBeWritten is the float number that is expected to be written to the file.
 * @param Delimiter is the delimiter to be used. TEXT(",")
 */
-void WriteToTextFile(const FString& FolderNameForTest ,const FString& FolderNameForBeingTested, const FString& FileName, const float& NumberToWriteToFile, const FString& Delimiter)
+void WriteToTextFile(const FString& InTestName ,const FString& InItemBeingTested, const FString& InFileName, const float& InNumberToBeWritten, const FString& Delimiter)
 {
 	//Performance file locations and setups.
-	FString FileSaveLocation = FPaths::Combine(*FPaths::AutomationLogDir(), *FolderNameForTest, *FolderNameForBeingTested, *FileName);
-
-	//Variables that hold the content from the text files
-	FString ReadFromTextFile;
+	FString FileSaveLocation = FPaths::Combine(*FPaths::AutomationLogDir(), *InTestName, *InItemBeingTested, *InFileName);
 
-	//Log out to a text file the Duration.
-	FString CurrentNumberToWrite = FString::Printf(TEXT("%f"), NumberToWriteToFile);
 	if (FPaths::FileExists(FileSaveLocation))
 	{
-		FFileHelper::LoadFileToString(ReadFromTextFile, *FileSaveLocation);
-		FString FileSetup = ReadFromTextFile + Delimiter + CurrentNumberToWrite;
+		//The text files existing content.
+		FString TextFileContents;
+
+		//Write to the text file the combined contents from the text file with the number to write.
+		FFileHelper::LoadFileToString(TextFileContents, *FileSaveLocation);
+		FString FileSetup = TextFileContents + Delimiter + FString::SanitizeFloat(InNumberToBeWritten);
 		FFileHelper::SaveStringToFile(FileSetup, *FileSaveLocation);
 		return;
 	}
 
-	FFileHelper::SaveStringToFile(CurrentNumberToWrite, *FileSaveLocation);
+	FFileHelper::SaveStringToFile(FString::SanitizeFloat(InNumberToBeWritten), *FileSaveLocation);
 }
 
 /**
 * Returns the sum of the numbers available in an array of float.
-* @param ArrayofNumbers is the name of the array intended to be used.
+* @param InFloatArray is the name of the array intended to be used.
 * @param bisAveragedInstead will return the average of the available numbers instead of the sum.
 */
-float TotalFromFloatArray(const TArray<float>& ArrayofNumbers, bool bisAveragedInstead)
+float TotalFromFloatArray(const TArray<float>& InFloatArray, bool bisAveragedInstead)
 {
 	//Total Value holds the sum of all the numbers available in the array.
 	float TotalValue = 0;
 
-	//Loop through the array.
-	for (int32 I = 0; I < ArrayofNumbers.Num(); ++I)
+	//Get the sum of the array.
+	for (int32 I = 0; I < InFloatArray.Num(); ++I)
 	{
-		//If the character in the array is a number then we'll want to add it to our total.
-		TotalValue += ArrayofNumbers[I];
+		TotalValue += InFloatArray[I];
 	}
 
 	//If bAverageInstead equals true then only the average is returned.
-	//Otherwise return the total value.
 	if (bisAveragedInstead)
 	{
-		UE_LOG(LogEditorAutomationTests, VeryVerbose, TEXT("Average value of the Array is %f"), (TotalValue / ArrayofNumbers.Num()));
-		return (TotalValue / ArrayofNumbers.Num());
+		UE_LOG(LogEditorAutomationTests, VeryVerbose, TEXT("Average value of the Array is %f"), (TotalValue / InFloatArray.Num()));
+		return (TotalValue / InFloatArray.Num());
 	}
 
 	UE_LOG(LogEditorAutomationTests, VeryVerbose, TEXT("Total Value of the Array is %f"), TotalValue);
@@ -343,44 +341,153 @@ float TotalFromFloatArray(const TArray<float>& ArrayofNumbers, bool bisAveragedI
 
 /**
 * Returns the largest value from an array of float numbers.
-* @param ArrayofNumbers is the name of the array intended to be used.
+* @param InFloatArray is the name of the array intended to be used.
 */
-float LargetValueInFloatArray(const TArray<float>& ArrayOfNumbers)
+float LargetValueInFloatArray(const TArray<float>& InFloatArray)
 {
 	//Total Value holds the sum of all the numbers available in the array.
 	float LargestValue = 0;
 
-	//Loop through the array.
-	for (int32 I = 0; I < ArrayOfNumbers.Num(); ++I)
+	//Find the largest value
+	for (int32 I = 0; I < InFloatArray.Num(); ++I)
 	{
-		if (LargestValue < ArrayOfNumbers[I])
+		if (LargestValue < InFloatArray[I])
 		{
-			LargestValue = ArrayOfNumbers[I];
+			LargestValue = InFloatArray[I];
 		}
 	}
 	UE_LOG(LogEditorAutomationTests, VeryVerbose, TEXT("The Largest value of the array is %f"), LargestValue);
 	return LargestValue;
 }
 
-TArray<FString> CreateArrayFromFile(const FString& FullFileLocation)
+/**
+* Returns the contents of a text file as an array of FString.
+* @param InFileLocation is the location of the file.
+*/
+TArray<FString> CreateArrayFromFile(const FString& InFileLocation)
 {
 	FString RawData;
 	TArray<FString> DataArray;
 
-	if (FPaths::FileExists(*FullFileLocation))
+	if (FPaths::FileExists(*InFileLocation))
 	{
-		UE_LOG(LogEditorAutomationTests, VeryVerbose, TEXT("Loading and parsing the data from '%s' into an array."), *FullFileLocation);
-		FFileHelper::LoadFileToString(RawData, *FullFileLocation);
+		UE_LOG(LogEditorAutomationTests, VeryVerbose, TEXT("Loading and parsing the data from '%s' into an array."), *InFileLocation);
+		FFileHelper::LoadFileToString(RawData, *InFileLocation);
 		RawData.ParseIntoArray(&DataArray, TEXT(","), false);
 
 		return DataArray;
 	}
 
+	UE_LOG(LogEditorAutomationTests, Warning, TEXT("Unable to create an array.  '%s' does not exist."), *InFileLocation);
 	RawData = TEXT("0");
 	DataArray.Add(RawData);
+	
 	return DataArray;
 }
 
+bool IsArchiveWriteable(const FString InFilePath, const FArchive* InArchiveName)
+{
+	if (!InArchiveName)
+	{
+		UE_LOG(LogEditorAutomationTests, Error, TEXT("Failed to write to the csv file: %s"), *FPaths::ConvertRelativePathToFull(InFilePath));
+		return false;
+	}
+	return true;
+}
+
+/**
+* Dumps the information held within the EditorPerfCaptureParameters struct into a CSV file.
+* @param EditorPerfStats is the name of the struct that holds the needed performance information.
+*/
+void EditorPerfDump(EditorPerfCaptureParameters& EditorPerfStats)
+{
+	UE_LOG(LogEditorAutomationTests, Log, TEXT("Begin generating the editor performance charts."));
+
+	//The file location where to save the data.
+	FString DataFileLocation = FPaths::Combine(*FPaths::AutomationLogDir(), TEXT("Performance"), *EditorPerfStats.MapName);
+
+	//Get the map load time (in seconds) from the text file that is created when the load map latent command is ran.
+	EditorPerfStats.MapLoadTime = 0;
+	FString MapLoadTimeFileLocation = FPaths::Combine(*DataFileLocation, TEXT("RAWMapLoadTime.txt"));
+	if (FPaths::FileExists(*MapLoadTimeFileLocation))
+	{
+		TArray<FString> SavedMapLoadTimes = CreateArrayFromFile(MapLoadTimeFileLocation);
+		EditorPerfStats.MapLoadTime = FCString::Atof(*SavedMapLoadTimes.Last());
+	}
+
+	//Filename for the RAW csv which holds the data gathered from a single test ran.
+	FString RAWCSVFilePath = FString::Printf(TEXT("%s/RAW_%s_%s.csv"), *DataFileLocation, *EditorPerfStats.MapName, *FDateTime::Now().ToString());
+
+	//Filename for the pretty csv file.
+	FString PerfCSVFilePath = FString::Printf(TEXT("%s/%s_Performance.csv"), *DataFileLocation, *EditorPerfStats.MapName);
+
+	//Create the raw csv and then add the title row it.
+	FArchive* RAWCSVArchive = IFileManager::Get().CreateFileWriter(*RAWCSVFilePath);
+	FString RAWCSVLine = (TEXT("Map Name, Changelist, Time Stamp, Map Load Time, Average FPS, Frame Time, Used Physical Memory, Used Virtual Memory, Used Peak Physical, Used Peak Virtual, Available Physical Memory, Available Virtual Memory\n"));
+	RAWCSVArchive->Serialize(TCHAR_TO_ANSI(*RAWCSVLine), RAWCSVLine.Len());
+
+	//Dump the stats from each run to the raw csv file and then close it.
+	for (int32 I = 0; I < EditorPerfStats.TimeStamp.Num(); I++)
+	{
+		//If the raw file isn't available to write to then we'll fail back this test.
+		if (IsArchiveWriteable(RAWCSVFilePath, RAWCSVArchive))
+		{
+			RAWCSVLine = FString::Printf(TEXT("%s,%s,%s,%.3f,%.1f,%.1f,%.0f,%.0f,%.0f,%.0f,%.0f,%.0f%s"), *EditorPerfStats.MapName, *GEngineVersion.ToString(EVersionComponent::Changelist), *EditorPerfStats.TimeStamp[I].ToString(), EditorPerfStats.MapLoadTime, EditorPerfStats.AverageFPS[I], EditorPerfStats.AverageFrameTime[I], EditorPerfStats.UsedPhysical[I], EditorPerfStats.UsedVirtual[I], EditorPerfStats.PeakUsedPhysical[I], EditorPerfStats.PeakUsedVirtual[I], EditorPerfStats.AvailablePhysical[I], EditorPerfStats.AvailableVirtual[I], LINE_TERMINATOR);
+			RAWCSVArchive->Serialize(TCHAR_TO_ANSI(*RAWCSVLine), RAWCSVLine.Len());
+		}
+	}
+	RAWCSVArchive->Close();
+
+	//Get the final pretty data for the Performance csv file.
+	float AverageFPS = TotalFromFloatArray(EditorPerfStats.AverageFPS, true);
+	float AverageFrameTime = TotalFromFloatArray(EditorPerfStats.AverageFrameTime, true);
+	float MemoryUsedPhysical = TotalFromFloatArray(EditorPerfStats.UsedPhysical, true);
+	float MemoryAvailPhysAvg = TotalFromFloatArray(EditorPerfStats.AvailablePhysical, true);
+	float MemoryAvailVirtualAvg = TotalFromFloatArray(EditorPerfStats.AvailableVirtual, true);
+	float MemoryUsedVirtualAvg = TotalFromFloatArray(EditorPerfStats.UsedVirtual, true);
+	float MemoryUsedPeak = LargetValueInFloatArray(EditorPerfStats.PeakUsedPhysical);
+	float MemoryUsedPeakVirtual = LargetValueInFloatArray(EditorPerfStats.PeakUsedVirtual);
+
+	//TestRunDuration is the length of time the test lasted in ticks.
+	FTimespan TestRunDuration = (EditorPerfStats.TimeStamp.Last().GetTicks() - EditorPerfStats.TimeStamp[0].GetTicks()) + ETimespan::TicksPerSecond;
+
+	//The performance csv file will be created if it didn't exist prior to the start of this test.
+	if (!FPaths::FileExists(*PerfCSVFilePath))
+	{
+		FArchive* FinalCSVArchive = IFileManager::Get().CreateFileWriter(*PerfCSVFilePath);
+		if (IsArchiveWriteable(PerfCSVFilePath, FinalCSVArchive))
+		{
+			FString FinalCSVLine = (TEXT("Date, Map Name, Changelist, Test Run Time , Map Load Time, Average FPS, Average MS, Used Physical KB, Used Virtual KB, Used Peak Physcial KB, Used Peak Virtual KB, Available Physical KB, Available Virtual KB\n"));
+			FinalCSVArchive->Serialize(TCHAR_TO_ANSI(*FinalCSVLine), FinalCSVLine.Len());
+			FinalCSVArchive->Close();
+		}
+	}
+
+	//Load the existing performance csv so that it doesn't get saved over and lost.
+	FString OldPerformanceCSVFile;
+	FFileHelper::LoadFileToString(OldPerformanceCSVFile, *PerfCSVFilePath);
+	FArchive* FinalCSVArchive = IFileManager::Get().CreateFileWriter(*PerfCSVFilePath);
+	if (IsArchiveWriteable(PerfCSVFilePath, FinalCSVArchive))
+	{
+		//Dump the old performance csv file data to the new csv file.
+		FinalCSVArchive->Serialize(TCHAR_TO_ANSI(*OldPerformanceCSVFile), OldPerformanceCSVFile.Len());
+
+		//Dump the pretty stats to the Performance CSV file and then close it so we can edit it while the engine is still running.
+		FString FinalCSVLine = FString::Printf(TEXT("%s,%s,%s,%.0f,%.3f,%.1f,%.1f,%.0f,%.0f,%.0f,%.0f,%.0f,%.0f%s"), *FDateTime::Now().ToString(), *EditorPerfStats.MapName, *GEngineVersion.ToString(EVersionComponent::Changelist), TestRunDuration.GetTotalSeconds(), EditorPerfStats.MapLoadTime, AverageFPS, AverageFrameTime, MemoryUsedPhysical, MemoryUsedVirtualAvg, MemoryUsedPeak, MemoryUsedPeakVirtual, MemoryAvailPhysAvg, MemoryAvailVirtualAvg, LINE_TERMINATOR);
+		FinalCSVArchive->Serialize(TCHAR_TO_ANSI(*FinalCSVLine), FinalCSVLine.Len());
+		FinalCSVArchive->Close();
+	}
+
+	//Display the test results to the user.
+	UE_LOG(LogEditorAutomationTests, Display, TEXT("AVG FPS: '%.1f'"), AverageFPS);
+	UE_LOG(LogEditorAutomationTests, Display, TEXT("AVG Frame Time: '%.1f' ms"), AverageFrameTime);
+	UE_LOG(LogEditorAutomationTests, Display, TEXT("AVG Used Physical Memory: '%.0f' kb"), MemoryUsedPhysical);
+	UE_LOG(LogEditorAutomationTests, Display, TEXT("AVG Used Virtual Memory: '%.0f' kb"), MemoryUsedVirtualAvg);
+	UE_LOG(LogEditorAutomationTests, Display, TEXT("Performance csv file is located here: %s"), *FPaths::ConvertRelativePathToFull(PerfCSVFilePath));
+	UE_LOG(LogEditorAutomationTests, Log, TEXT("Performance csv file is located here: %s"), *FPaths::ConvertRelativePathToFull(PerfCSVFilePath));
+	UE_LOG(LogEditorAutomationTests, Log, TEXT("Raw performance csv file is located here: %s"), *FPaths::ConvertRelativePathToFull(RAWCSVFilePath));
+}
+
 
 ///////////////////////////////////////////////////////////////////////
 // Common Latent commands
@@ -524,7 +631,7 @@ bool FEditorPerfCaptureCommand::Update()
 		//Query OS for used Peak Used virtual memory
 		EditorPerfStats.PeakUsedVirtual.Add((float)MemoryStats.PeakUsedVirtual / 1024);
 
-		//Capturing the time stamp.
+		//Capture the time stamp.
 		FString FormatedTimeStamp = FString::Printf(TEXT("%04i-%02i-%02i %02i:%02i:%02i.%03i"), CurrentDateAndTime.GetYear(), CurrentDateAndTime.GetMonth(), CurrentDateAndTime.GetDay(), CurrentDateAndTime.GetHour(), CurrentDateAndTime.GetMinute(), CurrentDateAndTime.GetSecond(), CurrentDateAndTime.GetMillisecond());
 		EditorPerfStats.FormattedTimeStamp.Add(FormatedTimeStamp);
 		EditorPerfStats.TimeStamp.Add(CurrentDateAndTime);
@@ -532,124 +639,8 @@ bool FEditorPerfCaptureCommand::Update()
 		return false;
 	}
 
-	//Combine performance data into one chart.
-	ADD_LATENT_AUTOMATION_COMMAND(FGenerateEditorPerfCaptureCommand(EditorPerfStats));
-
-	return true;
-}
-
-bool FGenerateEditorPerfCaptureCommand::Update()
-{
-	UE_LOG(LogEditorAutomationTests, Log, TEXT("Begin generating the editor performance charts."));
-
-	//Get the current changelist number
-	FString Changelist = GEngineVersion.ToString(EVersionComponent::Changelist);
-
-	//The locations to save the csv files.
-	FString FileLocation = FPaths::Combine(*FPaths::AutomationLogDir(), TEXT("Performance"), *EditorPerfStats.MapName);
-
-	//If a map load time file exists then grab the last load time from it.
-	FString MapLoadTimeFileLocation = FPaths::Combine(*FileLocation, TEXT("RAWMapLoadTime.txt"));
-	EditorPerfStats.MapLoadTime = 0;
-	if (FPaths::FileExists(*MapLoadTimeFileLocation))
-	{
-		TArray<FString> SavedMapLoadTimes = CreateArrayFromFile(MapLoadTimeFileLocation);
-		EditorPerfStats.MapLoadTime = FCString::Atof(*SavedMapLoadTimes.Last());
-	}
-
-	//RAW csv holds the data gathered from a single test ran.
-	//Final csv only holds the final information for each full test run.
-	FString FinalCSVFilename = FString::Printf(TEXT("%s/%s_Performance.csv"), *FileLocation, *EditorPerfStats.MapName);
-	FString RAWCSVFilename = FString::Printf(TEXT("%s/RAW_%s_%s.csv"), *FileLocation, *EditorPerfStats.MapName, *FDateTime::Now().ToString());
-
-	//Create the .csv file that will hold the raw data from the text files.
-	FArchive* RAWCSVFile = IFileManager::Get().CreateFileWriter(*RAWCSVFilename);
-	if (!RAWCSVFile)
-	{
-		UE_LOG(LogEditorAutomationTests, Error, TEXT("Failed to create the csv file: %s"), *RAWCSVFilename);
-		return true;
-	}
-
-	//Add the top title row to the raw csv file
-	FString RAWCSVLine = (TEXT("Map Name, Changelist, Time Stamp, Map Load Time, Average FPS, Frame Time, Used Physical Memory, Used Virtual Memory, Used Peak Physical, Used Peak Virtual, Available Physical Memory, Available Virtual Memory\n"));
-	RAWCSVFile->Serialize(TCHAR_TO_ANSI(*RAWCSVLine), RAWCSVLine.Len());
-	
-	//Loop through the text files and add them to the raw csv file.
-	for (int32 I = 0; I < EditorPerfStats.TimeStamp.Num(); I++)
-	{
-		//If the raw file isn't available to write to then we'll fail back this test.
-		if (!RAWCSVFile)
-		{
-			UE_LOG(LogEditorAutomationTests, Error, TEXT("Failed to write to the csv file: '%s'.  \nIt may have been opened or deleted."), *RAWCSVFilename);
-			return true;
-		}
-
-		//Create the line that will be written and then write it to the csv file.
-		RAWCSVLine = FString::Printf(TEXT("%s,%s,%s,%.3f,%.1f,%.1f,%.0f,%.0f,%.0f,%.0f,%.0f,%.0f%s"), *EditorPerfStats.MapName, *Changelist, *EditorPerfStats.TimeStamp[I].ToString(), EditorPerfStats.MapLoadTime, EditorPerfStats.AverageFPS[I], EditorPerfStats.AverageFrameTime[I], EditorPerfStats.UsedPhysical[I], EditorPerfStats.UsedVirtual[I], EditorPerfStats.PeakUsedPhysical[I], EditorPerfStats.PeakUsedVirtual[I], EditorPerfStats.AvailablePhysical[I], EditorPerfStats.AvailableVirtual[I], LINE_TERMINATOR);
-		RAWCSVFile->Serialize(TCHAR_TO_ANSI(*RAWCSVLine), RAWCSVLine.Len());
-	}
-	//Close the raw csv file.
-	RAWCSVFile->Close();
-	
-	//Get the final data for the Performance csv file.
-	//This gets the averaged numbers.
-	float FPSAvg = TotalFromFloatArray(EditorPerfStats.AverageFPS, true);
-	float FrameTimeAvg = TotalFromFloatArray(EditorPerfStats.AverageFrameTime, true);
-	float MemoryUsedPhysical = TotalFromFloatArray(EditorPerfStats.UsedPhysical, true);
-	float MemoryAvailPhysAvg = TotalFromFloatArray(EditorPerfStats.AvailablePhysical, true);
-	float MemoryAvailVirtualAvg = TotalFromFloatArray(EditorPerfStats.AvailableVirtual, true);
-	float MemoryUsedVirtualAvg = TotalFromFloatArray(EditorPerfStats.UsedVirtual, true);
-	float MemoryUsedPeak = LargetValueInFloatArray(EditorPerfStats.PeakUsedPhysical);
-	float MemoryUsedPeakVirtual = LargetValueInFloatArray(EditorPerfStats.PeakUsedVirtual);
-
-	//This is how long the test lasted.
-	FTimespan TestRunDuration = (EditorPerfStats.TimeStamp.Last().GetTicks() - EditorPerfStats.TimeStamp[0].GetTicks()) + ETimespan::TicksPerSecond;
-	
-	//If the Performance csv file does not have a size then it must be new.
-	//This adds the title row to the performance csv file.
-	if (IFileManager::Get().FileSize(*FinalCSVFilename) <= 0)
-	{
-		FArchive* FinalCSVFile = IFileManager::Get().CreateFileWriter(*FinalCSVFilename);
-		//If the csv file can't be created or written to then we'll fail back this test.
-		if (!FinalCSVFile)
-		{
-			UE_LOG(LogEditorAutomationTests, Error, TEXT("Failed to create the performance csv file: '%s'."), *FPaths::ConvertRelativePathToFull(FinalCSVFilename));
-			return true;
-		}
-		
-		//Create the .csv file that will hold the raw data from the text files.
-		FString FinalCSVLine = (TEXT("Date, Map Name, Changelist, Test Run Time , Map Load Time, Average FPS, Average MS, Used Physical KB, Used Virtual KB, Used Peak Physcial KB, Used Peak Virtual KB, Available Physical KB, Available Virtual KB\n"));
-		FinalCSVFile->Serialize(TCHAR_TO_ANSI(*FinalCSVLine), FinalCSVLine.Len());
-		FinalCSVFile->Close();
-	}
-
-	//Load the currently existing performance csv data.
-	FString OldPerformanceCSVFile;
-	FFileHelper::LoadFileToString(OldPerformanceCSVFile, *FinalCSVFilename);
-	FArchive* FinalCSVFile = IFileManager::Get().CreateFileWriter(*FinalCSVFilename);
-	//If the performance file isn't available to write to then we'll fail back this test.
-	if (!FinalCSVFile)
-	{
-		UE_LOG(LogEditorAutomationTests, Error, TEXT("Failed to write to the Performance csv file: '%s'.  \nIt may have been opened or deleted while in use."), *FPaths::ConvertRelativePathToFull(FinalCSVFilename));
-		return true;
-	}
-	
-	//Write the old performance csv file data to the new csv file.
-	FinalCSVFile->Serialize(TCHAR_TO_ANSI(*OldPerformanceCSVFile), OldPerformanceCSVFile.Len());
-
-	//Write the stats to the Performance CSV file and then close it so we can edit it while the engine is still running.
-	FString FinalCSVLine = FString::Printf(TEXT("%s,%s,%s,%.0f,%.3f,%.1f,%.1f,%.0f,%.0f,%.0f,%.0f,%.0f,%.0f%s"), *FDateTime::Now().ToString(), *EditorPerfStats.MapName, *Changelist, TestRunDuration.GetTotalSeconds(), EditorPerfStats.MapLoadTime, FPSAvg, FrameTimeAvg, MemoryUsedPhysical, MemoryUsedVirtualAvg, MemoryUsedPeak, MemoryUsedPeakVirtual, MemoryAvailPhysAvg, MemoryAvailVirtualAvg, LINE_TERMINATOR);
-	FinalCSVFile->Serialize(TCHAR_TO_ANSI(*FinalCSVLine), FinalCSVLine.Len());
-	FinalCSVFile->Close();
-
-	//Display the test results to the user.
-	UE_LOG(LogEditorAutomationTests, Display, TEXT("AVG FPS: '%.1f'"), FPSAvg);
-	UE_LOG(LogEditorAutomationTests, Display, TEXT("AVG Frame Time: '%.1f' ms"), FrameTimeAvg);
-	UE_LOG(LogEditorAutomationTests, Display, TEXT("AVG Used Physical Memory: '%.0f' kb"), MemoryUsedPhysical);
-	UE_LOG(LogEditorAutomationTests, Display, TEXT("AVG Used Virtual Memory: '%.0f' kb"), MemoryUsedVirtualAvg);
-	UE_LOG(LogEditorAutomationTests, Display, TEXT("Performance csv file is located here: %s"), *FPaths::ConvertRelativePathToFull(FinalCSVFilename));
-	UE_LOG(LogEditorAutomationTests, Log, TEXT("Performance csv file is located here: %s"), *FPaths::ConvertRelativePathToFull(FinalCSVFilename));
-	UE_LOG(LogEditorAutomationTests, Log, TEXT("Raw performance csv file is located here: %s"), *FPaths::ConvertRelativePathToFull(RAWCSVFilename));
+	//Dump the performance data in a csv file.
+	EditorPerfDump(EditorPerfStats);
 
 	return true;
 }
@@ -661,18 +652,15 @@ bool FEditorLoadMap::Update()
 {
 	//Get the base filename for the map that will be used.
 	FString ShortMapName = FPaths::GetBaseFilename(MapName);
-	
-	double MapLoadStartTime = 0.0f;
-	double MapLoadTime = 0.0f;
 
-	//Get the current number of seconds.  This will be used to track how long it took to load the map.
-	MapLoadStartTime = FPlatformTime::Seconds();
+	//Get the current number of seconds before loading the map.
+	double MapLoadStartTime = FPlatformTime::Seconds();
 
 	//Load the map
 	FEditorAutomationTestUtilities::LoadMap(MapName);
 
 	//This is the time it took to load the map in the editor.
-	MapLoadTime = FPlatformTime::Seconds() - MapLoadStartTime;
+	double MapLoadTime = FPlatformTime::Seconds() - MapLoadStartTime;
 
 	//Gets the main frame module to get the name of our current level.
 	const IMainFrameModule& MainFrameModule = FModuleManager::GetModuleChecked< IMainFrameModule >("MainFrame");
diff --git a/Engine/Source/Editor/UnrealEd/Public/Tests/AutomationEditorCommon.h b/Engine/Source/Editor/UnrealEd/Public/Tests/AutomationEditorCommon.h
--- a/Engine/Source/Editor/UnrealEd/Public/Tests/AutomationEditorCommon.h
+++ b/Engine/Source/Editor/UnrealEd/Public/Tests/AutomationEditorCommon.h
@@ -95,6 +95,7 @@ namespace AutomationEditorCommonUtils
 	void ApplyCustomFactorySettings(UFactory* InFactory, const TArray<FImportFactorySettingValues>& FactorySettings);
 }
 
+//Struct used to hold the data for the Editor Performance test.
 struct EditorPerfCaptureParameters
 {
 	//Basic Test Info
@@ -103,7 +104,6 @@ struct EditorPerfCaptureParameters
 
 	//Saved Performance Values
 	float MapLoadTime;
-	TArray<FTimespan> TimeBetweenCaptures;
 	TArray<float> AverageFPS;
 	TArray<float> AverageFrameTime;
 	TArray<float> UsedPhysical;
@@ -152,35 +152,17 @@ DEFINE_LATENT_AUTOMATION_COMMAND_ONE_PARAMETER(FStartPIECommand, bool, bSimulate
 */
 DEFINE_LATENT_AUTOMATION_COMMAND(FEndPlayMapCommand);
 
-/**
-* This will display the average FPS and Memory numbers over a duration of time.
-*/
-//DEFINE_LATENT_AUTOMATION_COMMAND_ONE_PARAMETER(FEditorPerformanceCommand, float, Duration);
-
 /**
 * This will capture the average FPS and Memory numbers over a duration of time.
 * @param The name of the EditorPerfCaptureParameters struct that will be used to hold the data.
 */
 DEFINE_LATENT_AUTOMATION_COMMAND_ONE_PARAMETER(FEditorPerfCaptureCommand, EditorPerfCaptureParameters, EditorPerfStats);
 
-/**
-* Generates the editor performance capture chart.
-* @param The name of the EditorPerfCaptureParameters struct that will be used to hold the data.
-*/
-DEFINE_LATENT_AUTOMATION_COMMAND_ONE_PARAMETER(FGenerateEditorPerfCaptureCommand, EditorPerfCaptureParameters, EditorPerfStats);
-
 /**
 * Loads a map
 */
 DEFINE_LATENT_AUTOMATION_COMMAND_ONE_PARAMETER(FEditorLoadMap, FString, MapName);
 
-/**
-* Generates the editor performance capture chart.
-* @param Input the name of the map.
-*/
-//DEFINE_LATENT_AUTOMATION_COMMAND_ONE_PARAMETER(FGenerateEditorPerformanceCharts, FString, MapName);
-
-
 /**
 * Waits for shaders to finish compiling before moving on to the next thing.
 */
