nicer UI

diff --git a/Engine/Source/Editor/UnrealEd/Private/SEditorViewportViewMenu.cpp b/Engine/Source/Editor/UnrealEd/Private/SEditorViewportViewMenu.cpp
--- a/Engine/Source/Editor/UnrealEd/Private/SEditorViewportViewMenu.cpp
+++ b/Engine/Source/Editor/UnrealEd/Private/SEditorViewportViewMenu.cpp
@@ -220,11 +220,11 @@ TSharedRef<SWidget> SEditorViewportViewMenu::GenerateViewMenuContent() const
 					Menu.AddMenuEntry( BaseViewportCommands.FixedExposure3m, NAME_None, LOCTEXT("FixedExposure3m", "Fixed at Log -3 = Brightness 1/8x") );
 					Menu.AddMenuEntry( BaseViewportCommands.FixedExposure2m, NAME_None, LOCTEXT("FixedExposure2m", "Fixed at Log -2 = Brightness 1/4x") );
 					Menu.AddMenuEntry( BaseViewportCommands.FixedExposure1m, NAME_None, LOCTEXT("FixedExposure1m", "Fixed at Log -1 = Brightness 1/2x") );
-					Menu.AddMenuEntry( BaseViewportCommands.FixedExposure0, NAME_None, LOCTEXT("FixedExposure0", "Fixed at Log 0 = Brightness 1x") );
-					Menu.AddMenuEntry( BaseViewportCommands.FixedExposure1p, NAME_None, LOCTEXT("FixedExposure1p", "Fixed at Log 1 = Brightness 2x") );
-					Menu.AddMenuEntry( BaseViewportCommands.FixedExposure2p, NAME_None, LOCTEXT("FixedExposure2p", "Fixed at Log 2 = Brightness 4x") );
-					Menu.AddMenuEntry( BaseViewportCommands.FixedExposure3p, NAME_None, LOCTEXT("FixedExposure3p", "Fixed at Log 3 = Brightness 8x") );
-					Menu.AddMenuEntry( BaseViewportCommands.FixedExposure4p, NAME_None, LOCTEXT("FixedExposure4p", "Fixed at Log 4 = Brightness 16x") );
+					Menu.AddMenuEntry( BaseViewportCommands.FixedExposure0,  NAME_None, LOCTEXT("FixedExposure0",  "Fixed at Log  0") );
+					Menu.AddMenuEntry( BaseViewportCommands.FixedExposure1p, NAME_None, LOCTEXT("FixedExposure1p", "Fixed at Log +1 = Brightness 2x") );
+					Menu.AddMenuEntry( BaseViewportCommands.FixedExposure2p, NAME_None, LOCTEXT("FixedExposure2p", "Fixed at Log +2 = Brightness 4x") );
+					Menu.AddMenuEntry( BaseViewportCommands.FixedExposure3p, NAME_None, LOCTEXT("FixedExposure3p", "Fixed at Log +3 = Brightness 8x") );
+					Menu.AddMenuEntry( BaseViewportCommands.FixedExposure4p, NAME_None, LOCTEXT("FixedExposure4p", "Fixed at Log +4 = Brightness 16x") );
 				}
 			};
 
