PostDuplicate for Blueprints will correct self variable node references.

#ttp 349451 - CRITICAL: Regression: BP: Some blueprints will give a "Target must have a connection error when duplicated.

diff --git a/Engine/Source/Editor/UnrealEd/Private/Kismet2/BlueprintEditorUtils.cpp b/Engine/Source/Editor/UnrealEd/Private/Kismet2/BlueprintEditorUtils.cpp
--- a/Engine/Source/Editor/UnrealEd/Private/Kismet2/BlueprintEditorUtils.cpp
+++ b/Engine/Source/Editor/UnrealEd/Private/Kismet2/BlueprintEditorUtils.cpp
@@ -1402,6 +1402,7 @@ void FBlueprintEditorUtils::PostDuplicateBlueprint(UBlueprint* Blueprint)
 			IKismetCompilerInterface& Compiler = FModuleManager::LoadModuleChecked<IKismetCompilerInterface>(KISMET_COMPILER_MODULENAME);
 			FCompilerResultsLog Results;
 			FKismetCompilerOptions CompileOptions;
+			CompileOptions.CompileType = EKismetCompileType::SkeletonOnly;
 			CompileOptions.bIsDuplicationInstigated = true;
 
 			//SCS can change structure of the class
@@ -1485,51 +1486,55 @@ void FBlueprintEditorUtils::PostDuplicateBlueprint(UBlueprint* Blueprint)
 				Compiler.CompileBlueprint(Blueprint, CompileOptions, Results);
 			}
 
-			FArchiveReplaceObjectRef<UObject> ReplaceTemplateRefs(NewBPGC, OldToNewMap, /*bNullPrivateRefs=*/ false, /*bIgnoreOuterRef=*/ false, /*bIgnoreArchetypeRef=*/ false);
-
-			// Now propagate the values from the old CDO to the new one
-			check(Blueprint->SkeletonGeneratedClass != NULL);
-
-			UObject* NewCDO = Blueprint->GeneratedClass->GetDefaultObject();
-			check(NewCDO != NULL);
-			UEditorEngine::CopyPropertiesForUnrelatedObjects(OldCDO, NewCDO);
-		}
-
-		// Create a new blueprint guid
-		Blueprint->GenerateNewGuid();
+			// Create a new blueprint guid
+			Blueprint->GenerateNewGuid();
 
-		// Give all nodes a new Guid
-		TArray< UEdGraphNode* > AllGraphNodes;
-		GetAllNodesOfClass(Blueprint, AllGraphNodes);
-		for(auto& Node : AllGraphNodes)
-		{
-			Node->CreateNewGuid();
-
-			// Some variable nodes must be fixed up on duplicate, this cannot wait for individual 
-			// node calls to PostDuplicate because it happens after compilation and will still result in errors
-			if(UK2Node_Variable* VariableNode = Cast<UK2Node_Variable>(Node))
+			// Give all nodes a new Guid
+			TArray< UEdGraphNode* > AllGraphNodes;
+			GetAllNodesOfClass(Blueprint, AllGraphNodes);
+			for(auto& Node : AllGraphNodes)
 			{
-				// Self context variable nodes need to be updated with the new Blueprint class
-				if(VariableNode->VariableReference.IsSelfContext())
+				Node->CreateNewGuid();
+
+				// Some variable nodes must be fixed up on duplicate, this cannot wait for individual 
+				// node calls to PostDuplicate because it happens after compilation and will still result in errors
+				if(UK2Node_Variable* VariableNode = Cast<UK2Node_Variable>(Node))
 				{
-					const UEdGraphSchema_K2* K2Schema = GetDefault<UEdGraphSchema_K2>();
-					if(UEdGraphPin* SelfPin = K2Schema->FindSelfPin(*VariableNode, EGPD_Input))
+					// Self context variable nodes need to be updated with the new Blueprint class
+					if(VariableNode->VariableReference.IsSelfContext())
 					{
-						UClass* TargetClass = nullptr;
-
-						if(UProperty* Property = VariableNode->VariableReference.ResolveMember<UProperty>(VariableNode))
-						{
-							TargetClass = Property->GetOwnerClass()->GetAuthoritativeClass();
-						}
-						else
+						const UEdGraphSchema_K2* K2Schema = GetDefault<UEdGraphSchema_K2>();
+						if(UEdGraphPin* SelfPin = K2Schema->FindSelfPin(*VariableNode, EGPD_Input))
 						{
-							TargetClass = Blueprint->SkeletonGeneratedClass->GetAuthoritativeClass();
-						}
+							UClass* TargetClass = nullptr;
 
-						SelfPin->PinType.PinSubCategoryObject = TargetClass;
+							if(UProperty* Property = VariableNode->VariableReference.ResolveMember<UProperty>(VariableNode))
+							{
+								TargetClass = Property->GetOwnerClass()->GetAuthoritativeClass();
+							}
+							else
+							{
+								TargetClass = Blueprint->SkeletonGeneratedClass->GetAuthoritativeClass();
+							}
+
+							SelfPin->PinType.PinSubCategoryObject = TargetClass;
+						}
 					}
 				}
 			}
+
+			// Needs a full compile to handle the ArchiveReplaceObjectRef
+			CompileOptions.CompileType = EKismetCompileType::Full;
+			Compiler.CompileBlueprint(Blueprint, CompileOptions, Results);
+
+			FArchiveReplaceObjectRef<UObject> ReplaceTemplateRefs(NewBPGC, OldToNewMap, /*bNullPrivateRefs=*/ false, /*bIgnoreOuterRef=*/ false, /*bIgnoreArchetypeRef=*/ false);
+
+			// Now propagate the values from the old CDO to the new one
+			check(Blueprint->SkeletonGeneratedClass != NULL);
+
+			UObject* NewCDO = Blueprint->GeneratedClass->GetDefaultObject();
+			check(NewCDO != NULL);
+			UEditorEngine::CopyPropertiesForUnrelatedObjects(OldCDO, NewCDO);
 		}
 
 		// And compile again to make sure they go into the generated class, get cleaned up, etc...
