Added comment for forward axis set up when importing.

diff --git a/Engine/Source/Editor/UnrealEd/Private/Fbx/FbxMainImport.cpp b/Engine/Source/Editor/UnrealEd/Private/Fbx/FbxMainImport.cpp
--- a/Engine/Source/Editor/UnrealEd/Private/Fbx/FbxMainImport.cpp
+++ b/Engine/Source/Editor/UnrealEd/Private/Fbx/FbxMainImport.cpp
@@ -742,10 +742,16 @@ bool FFbxImporter::ImportFromFile(const FString& Filename, const FString& Type)
 			{
 				if (GetImportOptions()->bConvertScene)
 				{
+					// we use -Y as forward axis here when we import. This is odd considering our forward axis is technically +X
+					// but this is to mimic Maya/Max behavior where if you make a model facing +X facing, 
+					// when you import that mesh, you want +X facing in engine. 
+					// only thing that doesn't work is hand flipping because Max/Maya is RHS but UE is LHS
+					// On the positive note, we now have import transform set up you can do to rotate mesh if you don't like default setting
 					FbxAxisSystem::EFrontVector FrontVector = (FbxAxisSystem::EFrontVector) - FbxAxisSystem::eParityOdd;
 					const FbxAxisSystem UnrealZUp(FbxAxisSystem::eZAxis, FrontVector, FbxAxisSystem::eRightHanded);
+					const FbxAxisSystem SourceSetup = Scene->GetGlobalSettings().GetAxisSystem();
 
-					if(Scene->GetGlobalSettings().GetAxisSystem() != UnrealZUp)
+					if(SourceSetup != UnrealZUp)
 					{
 						// Converts the FBX data to Z-up, X-forward, Y-left.  Unreal is the same except with Y-right, 
 						// but the conversion to left-handed coordinates is not working properly
