Apply preview feature level to PIE world.

diff --git a/Engine/Source/Editor/UnrealEd/Private/PlayLevel.cpp b/Engine/Source/Editor/UnrealEd/Private/PlayLevel.cpp
--- a/Engine/Source/Editor/UnrealEd/Private/PlayLevel.cpp
+++ b/Engine/Source/Editor/UnrealEd/Private/PlayLevel.cpp
@@ -3087,7 +3087,7 @@ UWorld* UEditorEngine::CreatePIEWorldByDuplication(FWorldContext &WorldContext,
 
 	GPlayInEditorID = -1;
 	check( NewPIEWorld );
-
+	NewPIEWorld->FeatureLevel = EditorWorld->FeatureLevel;
 	PostCreatePIEWorld(NewPIEWorld);
 
 	// After loading the map, reset these so that things continue as normal
