EditorAssetAutomationTests.cpp Fixed an issue where QAGame was hardcoded as the location for the exported files.

diff --git a/Engine/Source/Editor/UnrealEd/Private/Tests/EditorAssetAutomationTests.cpp b/Engine/Source/Editor/UnrealEd/Private/Tests/EditorAssetAutomationTests.cpp
--- a/Engine/Source/Editor/UnrealEd/Private/Tests/EditorAssetAutomationTests.cpp
+++ b/Engine/Source/Editor/UnrealEd/Private/Tests/EditorAssetAutomationTests.cpp
@@ -874,7 +874,7 @@ namespace ImportExportAssetHelper
 			if (ActiveWindow.IsValid())
 			{
 				FString ScreenshotName;
-				const FString TestName = FString::Printf(TEXT("ImportExportTest/%s"), *ImportedAsset->GetName());
+				const FString TestName = FString::Printf(TEXT("AssetImportExport/Screenshots/%s"), *ImportedAsset->GetName());
 				AutomationCommon::GetScreenshotPath(TestName, ScreenshotName, false);
 
 				TSharedRef<SWidget> WindowRef = ActiveWindow.ToSharedRef();
@@ -916,7 +916,8 @@ namespace ImportExportAssetHelper
 				}
 
 				//Export the asset
-				const FString ExportPath = FString::Printf(TEXT("../../../QAGame/Saved/Exports/%s.%s"), *ImportedAsset->GetName(), *Extension);
+				const FString ExportAssetName = FString::Printf(TEXT("%s.%s"), *ImportedAsset->GetName(), *Extension);
+				FString ExportPath = FPaths::Combine(*FPaths::AutomationDir(), TEXT("AssetImportExport"), TEXT("Exported"),*ExportAssetName);
 				UExporter* ExporterToUse = UExporter::FindExporter(ImportedAsset, *Extension);
 				
 				UExporter::FExportToFileParams Params;
