Fixed EditorSettingsViewer not building

This was due to no .Build.cs files referencing it, so it had become stale and was causing crashes.

diff --git a/Engine/Source/Editor/UnrealEd/UnrealEd.Build.cs b/Engine/Source/Editor/UnrealEd/UnrealEd.Build.cs
--- a/Engine/Source/Editor/UnrealEd/UnrealEd.Build.cs
+++ b/Engine/Source/Editor/UnrealEd/UnrealEd.Build.cs
@@ -95,7 +95,6 @@ public class UnrealEd : ModuleRules
 				"MaterialEditor",
 				"MessageLog",
 				"NetworkFileSystem",
-				"ProjectSettingsViewer",
 				"PropertyEditor",
 				"Projects",
 				"RawMesh",
@@ -138,6 +137,7 @@ public class UnrealEd : ModuleRules
 				"CurveTableEditor",
 				"DataTableEditor",
 				"DestructibleMeshEditor",
+				"EditorSettingsViewer",
 				"LandscapeEditor",
 				"KismetCompiler",
 				"DetailCustomizations",
@@ -171,6 +171,7 @@ public class UnrealEd : ModuleRules
 				"NiagaraEditor",
 				"MeshUtilities",
 				"GameProjectGeneration",
+				"ProjectSettingsViewer",
 				"ProjectTargetPlatformEditor",
 				"PListEditor",
                 "Documentation",
