Unify labelling for blend spaces across different create menus

#TTP 338341 - NUX: Inconsistent naming of Blendspace types

diff --git a/Engine/Source/Editor/UnrealEd/Private/AnimationEditorUtils.cpp b/Engine/Source/Editor/UnrealEd/Private/AnimationEditorUtils.cpp
--- a/Engine/Source/Editor/UnrealEd/Private/AnimationEditorUtils.cpp
+++ b/Engine/Source/Editor/UnrealEd/Private/AnimationEditorUtils.cpp
@@ -401,21 +401,21 @@ namespace AnimationEditorUtils
 		MenuBuilder.BeginSection("CreateBlendSpace", LOCTEXT("CreateBlendSpaceMenuHeading", "Blend Spaces"));
 		{
 			MenuBuilder.AddMenuEntry(
-				LOCTEXT("SkeletalMesh_New1DBlendspace", "1D Blend Space"),
-				LOCTEXT("SkeletalMesh_New1DBlendspaceTooltip", "Creates a 1D Blend Space using the selected skeleton."),
-				FSlateIcon(FEditorStyle::GetStyleSetName(), "ClassIcon.BlendSpace1D"),
+				LOCTEXT("SkeletalMesh_New2DBlendspace", "Blend Space"),
+				LOCTEXT("SkeletalMesh_New2DBlendspaceTooltip", "Creates a Blend Space using the selected skeleton."),
+				FSlateIcon(FEditorStyle::GetStyleSetName(), "ClassIcon.BlendSpace"),
 				FUIAction(
-					FExecuteAction::CreateStatic(&ExecuteNewAnimAsset<UBlendSpaceFactory1D, UBlendSpace1D>, Skeletons, FString("_BlendSpace1D"), AssetCreated, bInContentBrowser),
+					FExecuteAction::CreateStatic(&ExecuteNewAnimAsset<UBlendSpaceFactoryNew, UBlendSpace>, Skeletons, FString("_BlendSpace"), AssetCreated, bInContentBrowser),
 					FCanExecuteAction()
 					)
 				);
 
 			MenuBuilder.AddMenuEntry(
-				LOCTEXT("SkeletalMesh_New2DBlendspace", "2D Blend Space"),
-				LOCTEXT("SkeletalMesh_New2DBlendspaceTooltip", "Creates a 2D Blend Space using the selected skeleton."),
-				FSlateIcon(FEditorStyle::GetStyleSetName(), "ClassIcon.BlendSpace"),
+				LOCTEXT("SkeletalMesh_New1DBlendspace", "Blend Space 1D"),
+				LOCTEXT("SkeletalMesh_New1DBlendspaceTooltip", "Creates a 1D Blend Space using the selected skeleton."),
+				FSlateIcon(FEditorStyle::GetStyleSetName(), "ClassIcon.BlendSpace1D"),
 				FUIAction(
-					FExecuteAction::CreateStatic(&ExecuteNewAnimAsset<UBlendSpaceFactoryNew, UBlendSpace>, Skeletons, FString("_BlendSpace2D"), AssetCreated, bInContentBrowser),
+					FExecuteAction::CreateStatic(&ExecuteNewAnimAsset<UBlendSpaceFactory1D, UBlendSpace1D>, Skeletons, FString("_BlendSpace1D"), AssetCreated, bInContentBrowser),
 					FCanExecuteAction()
 					)
 				);
@@ -425,21 +425,21 @@ namespace AnimationEditorUtils
 		MenuBuilder.BeginSection("CreateAimOffset", LOCTEXT("CreateAimOffsetMenuHeading", "Aim Offsets"));
 		{
 			MenuBuilder.AddMenuEntry(
-				LOCTEXT("SkeletalMesh_New1DAimOffset", "1D Aim Offset"),
-				LOCTEXT("SkeletalMesh_New1DAimOffsetTooltip", "Creates a 1D Aim Offset blendspace using the selected skeleton."),
+				LOCTEXT("SkeletalMesh_New2DAimOffset", "Aim Offset"),
+				LOCTEXT("SkeletalMesh_New2DAimOffsetTooltip", "Creates a Aim Offset blendspace using the selected skeleton."),
 				FSlateIcon(),
 				FUIAction(
-					FExecuteAction::CreateStatic(&ExecuteNewAnimAsset<UAimOffsetBlendSpaceFactory1D, UAimOffsetBlendSpace1D>, Skeletons, FString("_AimOffset1D"), AssetCreated, bInContentBrowser),
+					FExecuteAction::CreateStatic(&ExecuteNewAnimAsset<UAimOffsetBlendSpaceFactoryNew, UAimOffsetBlendSpace>, Skeletons, FString("_AimOffset2D"), AssetCreated, bInContentBrowser),
 					FCanExecuteAction()
 					)
 				);
 
 			MenuBuilder.AddMenuEntry(
-				LOCTEXT("SkeletalMesh_New2DAimOffset", "2D Aim Offset"),
-				LOCTEXT("SkeletalMesh_New2DAimOffsetTooltip", "Creates a 2D Aim Offset blendspace using the selected skeleton."),
+				LOCTEXT("SkeletalMesh_New1DAimOffset", "Aim Offset 1D"),
+				LOCTEXT("SkeletalMesh_New1DAimOffsetTooltip", "Creates a 1D Aim Offset blendspace using the selected skeleton."),
 				FSlateIcon(),
 				FUIAction(
-					FExecuteAction::CreateStatic(&ExecuteNewAnimAsset<UAimOffsetBlendSpaceFactoryNew, UAimOffsetBlendSpace>, Skeletons, FString("_AimOffset2D"), AssetCreated, bInContentBrowser),
+					FExecuteAction::CreateStatic(&ExecuteNewAnimAsset<UAimOffsetBlendSpaceFactory1D, UAimOffsetBlendSpace1D>, Skeletons, FString("_AimOffset1D"), AssetCreated, bInContentBrowser),
 					FCanExecuteAction()
 					)
 				);
