Moved the majority of the performance test code to hte EditorPerformanceTest.cpp

diff --git a/Engine/Source/Editor/UnrealEd/Private/Tests/AutomationEditorCommon.cpp b/Engine/Source/Editor/UnrealEd/Private/Tests/AutomationEditorCommon.cpp
--- a/Engine/Source/Editor/UnrealEd/Private/Tests/AutomationEditorCommon.cpp
+++ b/Engine/Source/Editor/UnrealEd/Private/Tests/AutomationEditorCommon.cpp
@@ -281,214 +281,124 @@ namespace AutomationEditorCommonUtils
 			InFactory->ConfigureProperties();
 		}
 	}
-}
-
 
-/**
-* Writes a number to a text file.
-* @param InTestName is the folder that has the same name as the test. (For Example: "Performance").
-* @param InItemBeingTested is the name for the thing that is being tested. (For Example: "MapName").
-* @param InFileName is the name of the file with an extension
-* @param InNumberToBeWritten is the float number that is expected to be written to the file.
-* @param Delimiter is the delimiter to be used. TEXT(",")
-*/
-void WriteToTextFile(const FString& InTestName ,const FString& InItemBeingTested, const FString& InFileName, const float& InNumberToBeWritten, const FString& Delimiter)
-{
-	//Performance file locations and setups.
-	FString FileSaveLocation = FPaths::Combine(*FPaths::AutomationLogDir(), *InTestName, *InItemBeingTested, *InFileName);
-
-	if (FPaths::FileExists(FileSaveLocation))
+	/**
+	* Writes a number to a text file.
+	* @param InTestName is the folder that has the same name as the test. (For Example: "Performance").
+	* @param InItemBeingTested is the name for the thing that is being tested. (For Example: "MapName").
+	* @param InFileName is the name of the file with an extension
+	* @param InNumberToBeWritten is the float number that is expected to be written to the file.
+	* @param Delimiter is the delimiter to be used. TEXT(",")
+	*/
+	void WriteToTextFile(const FString& InTestName, const FString& InItemBeingTested, const FString& InFileName, const float& InNumberToBeWritten, const FString& Delimiter)
 	{
-		//The text files existing content.
-		FString TextFileContents;
-
-		//Write to the text file the combined contents from the text file with the number to write.
-		FFileHelper::LoadFileToString(TextFileContents, *FileSaveLocation);
-		FString FileSetup = TextFileContents + Delimiter + FString::SanitizeFloat(InNumberToBeWritten);
-		FFileHelper::SaveStringToFile(FileSetup, *FileSaveLocation);
-		return;
-	}
-
-	FFileHelper::SaveStringToFile(FString::SanitizeFloat(InNumberToBeWritten), *FileSaveLocation);
-}
+		//Performance file locations and setups.
+		FString FileSaveLocation = FPaths::Combine(*FPaths::AutomationLogDir(), *InTestName, *InItemBeingTested, *InFileName);
 
-/**
-* Returns the sum of the numbers available in an array of float.
-* @param InFloatArray is the name of the array intended to be used.
-* @param bisAveragedInstead will return the average of the available numbers instead of the sum.
-*/
-float TotalFromFloatArray(const TArray<float>& InFloatArray, bool bisAveragedInstead)
-{
-	//Total Value holds the sum of all the numbers available in the array.
-	float TotalValue = 0;
+		if (FPaths::FileExists(FileSaveLocation))
+		{
+			//The text files existing content.
+			FString TextFileContents;
+
+			//Write to the text file the combined contents from the text file with the number to write.
+			FFileHelper::LoadFileToString(TextFileContents, *FileSaveLocation);
+			FString FileSetup = TextFileContents + Delimiter + FString::SanitizeFloat(InNumberToBeWritten);
+			FFileHelper::SaveStringToFile(FileSetup, *FileSaveLocation);
+			return;
+		}
 
-	//Get the sum of the array.
-	for (int32 I = 0; I < InFloatArray.Num(); ++I)
-	{
-		TotalValue += InFloatArray[I];
+		FFileHelper::SaveStringToFile(FString::SanitizeFloat(InNumberToBeWritten), *FileSaveLocation);
 	}
 
-	//If bAverageInstead equals true then only the average is returned.
-	if (bisAveragedInstead)
+	/**
+	* Returns the sum of the numbers available in an array of float.
+	* @param InFloatArray is the name of the array intended to be used.
+	* @param bisAveragedInstead will return the average of the available numbers instead of the sum.
+	*/
+	float TotalFromFloatArray(const TArray<float>& InFloatArray, bool bisAveragedInstead)
 	{
-		UE_LOG(LogEditorAutomationTests, VeryVerbose, TEXT("Average value of the Array is %f"), (TotalValue / InFloatArray.Num()));
-		return (TotalValue / InFloatArray.Num());
-	}
-
-	UE_LOG(LogEditorAutomationTests, VeryVerbose, TEXT("Total Value of the Array is %f"), TotalValue);
-	return TotalValue;
-}
-
-/**
-* Returns the largest value from an array of float numbers.
-* @param InFloatArray is the name of the array intended to be used.
-*/
-float LargetValueInFloatArray(const TArray<float>& InFloatArray)
-{
-	//Total Value holds the sum of all the numbers available in the array.
-	float LargestValue = 0;
+		//Total Value holds the sum of all the numbers available in the array.
+		float TotalValue = 0;
 
-	//Find the largest value
-	for (int32 I = 0; I < InFloatArray.Num(); ++I)
-	{
-		if (LargestValue < InFloatArray[I])
+		//Get the sum of the array.
+		for (int32 I = 0; I < InFloatArray.Num(); ++I)
 		{
-			LargestValue = InFloatArray[I];
+			TotalValue += InFloatArray[I];
 		}
-	}
-	UE_LOG(LogEditorAutomationTests, VeryVerbose, TEXT("The Largest value of the array is %f"), LargestValue);
-	return LargestValue;
-}
-
-/**
-* Returns the contents of a text file as an array of FString.
-* @param InFileLocation is the location of the file.
-*/
-TArray<FString> CreateArrayFromFile(const FString& InFileLocation)
-{
-	FString RawData;
-	TArray<FString> DataArray;
 
-	if (FPaths::FileExists(*InFileLocation))
-	{
-		UE_LOG(LogEditorAutomationTests, VeryVerbose, TEXT("Loading and parsing the data from '%s' into an array."), *InFileLocation);
-		FFileHelper::LoadFileToString(RawData, *InFileLocation);
-		RawData.ParseIntoArray(&DataArray, TEXT(","), false);
+		//If bAverageInstead equals true then only the average is returned.
+		if (bisAveragedInstead)
+		{
+			UE_LOG(LogEditorAutomationTests, VeryVerbose, TEXT("Average value of the Array is %f"), (TotalValue / InFloatArray.Num()));
+			return (TotalValue / InFloatArray.Num());
+		}
 
-		return DataArray;
+		UE_LOG(LogEditorAutomationTests, VeryVerbose, TEXT("Total Value of the Array is %f"), TotalValue);
+		return TotalValue;
 	}
 
-	UE_LOG(LogEditorAutomationTests, Warning, TEXT("Unable to create an array.  '%s' does not exist."), *InFileLocation);
-	RawData = TEXT("0");
-	DataArray.Add(RawData);
-	
-	return DataArray;
-}
-
-bool IsArchiveWriteable(const FString InFilePath, const FArchive* InArchiveName)
-{
-	if (!InArchiveName)
+	/**
+	* Returns the largest value from an array of float numbers.
+	* @param InFloatArray is the name of the array intended to be used.
+	*/
+	float LargetValueInFloatArray(const TArray<float>& InFloatArray)
 	{
-		UE_LOG(LogEditorAutomationTests, Error, TEXT("Failed to write to the csv file: %s"), *FPaths::ConvertRelativePathToFull(InFilePath));
-		return false;
-	}
-	return true;
-}
-
-/**
-* Dumps the information held within the EditorPerfCaptureParameters struct into a CSV file.
-* @param EditorPerfStats is the name of the struct that holds the needed performance information.
-*/
-void EditorPerfDump(EditorPerfCaptureParameters& EditorPerfStats)
-{
-	UE_LOG(LogEditorAutomationTests, Log, TEXT("Begin generating the editor performance charts."));
+		//Total Value holds the sum of all the numbers available in the array.
+		float LargestValue = 0;
 
-	//The file location where to save the data.
-	FString DataFileLocation = FPaths::Combine(*FPaths::AutomationLogDir(), TEXT("Performance"), *EditorPerfStats.MapName);
+		//Find the largest value
+		for (int32 I = 0; I < InFloatArray.Num(); ++I)
+		{
+			if (LargestValue < InFloatArray[I])
+			{
+				LargestValue = InFloatArray[I];
+			}
+		}
+		UE_LOG(LogEditorAutomationTests, VeryVerbose, TEXT("The Largest value of the array is %f"), LargestValue);
+		return LargestValue;
+	}
 
-	//Get the map load time (in seconds) from the text file that is created when the load map latent command is ran.
-	EditorPerfStats.MapLoadTime = 0;
-	FString MapLoadTimeFileLocation = FPaths::Combine(*DataFileLocation, TEXT("RAWMapLoadTime.txt"));
-	if (FPaths::FileExists(*MapLoadTimeFileLocation))
+	/**
+	* Returns the contents of a text file as an array of FString.
+	* @param InFileLocation is the location of the file.
+	*/
+	TArray<FString> CreateArrayFromFile(const FString& InFileLocation)
 	{
-		TArray<FString> SavedMapLoadTimes = CreateArrayFromFile(MapLoadTimeFileLocation);
-		EditorPerfStats.MapLoadTime = FCString::Atof(*SavedMapLoadTimes.Last());
-	}
+		FString RawData;
+		TArray<FString> DataArray;
 
-	//Filename for the RAW csv which holds the data gathered from a single test ran.
-	FString RAWCSVFilePath = FString::Printf(TEXT("%s/RAW_%s_%s.csv"), *DataFileLocation, *EditorPerfStats.MapName, *FDateTime::Now().ToString());
+		if (FPaths::FileExists(*InFileLocation))
+		{
+			UE_LOG(LogEditorAutomationTests, VeryVerbose, TEXT("Loading and parsing the data from '%s' into an array."), *InFileLocation);
+			FFileHelper::LoadFileToString(RawData, *InFileLocation);
+			RawData.ParseIntoArray(&DataArray, TEXT(","), false);
 
-	//Filename for the pretty csv file.
-	FString PerfCSVFilePath = FString::Printf(TEXT("%s/%s_Performance.csv"), *DataFileLocation, *EditorPerfStats.MapName);
+			return DataArray;
+		}
 
-	//Create the raw csv and then add the title row it.
-	FArchive* RAWCSVArchive = IFileManager::Get().CreateFileWriter(*RAWCSVFilePath);
-	FString RAWCSVLine = (TEXT("Map Name, Changelist, Time Stamp, Map Load Time, Average FPS, Frame Time, Used Physical Memory, Used Virtual Memory, Used Peak Physical, Used Peak Virtual, Available Physical Memory, Available Virtual Memory\n"));
-	RAWCSVArchive->Serialize(TCHAR_TO_ANSI(*RAWCSVLine), RAWCSVLine.Len());
+		UE_LOG(LogEditorAutomationTests, Warning, TEXT("Unable to create an array.  '%s' does not exist."), *InFileLocation);
+		RawData = TEXT("0");
+		DataArray.Add(RawData);
 
-	//Dump the stats from each run to the raw csv file and then close it.
-	for (int32 I = 0; I < EditorPerfStats.TimeStamp.Num(); I++)
-	{
-		//If the raw file isn't available to write to then we'll fail back this test.
-		if (IsArchiveWriteable(RAWCSVFilePath, RAWCSVArchive))
-		{
-			RAWCSVLine = FString::Printf(TEXT("%s,%s,%s,%.3f,%.1f,%.1f,%.0f,%.0f,%.0f,%.0f,%.0f,%.0f%s"), *EditorPerfStats.MapName, *GEngineVersion.ToString(EVersionComponent::Changelist), *EditorPerfStats.TimeStamp[I].ToString(), EditorPerfStats.MapLoadTime, EditorPerfStats.AverageFPS[I], EditorPerfStats.AverageFrameTime[I], EditorPerfStats.UsedPhysical[I], EditorPerfStats.UsedVirtual[I], EditorPerfStats.PeakUsedPhysical[I], EditorPerfStats.PeakUsedVirtual[I], EditorPerfStats.AvailablePhysical[I], EditorPerfStats.AvailableVirtual[I], LINE_TERMINATOR);
-			RAWCSVArchive->Serialize(TCHAR_TO_ANSI(*RAWCSVLine), RAWCSVLine.Len());
-		}
+		return DataArray;
 	}
-	RAWCSVArchive->Close();
-
-	//Get the final pretty data for the Performance csv file.
-	float AverageFPS = TotalFromFloatArray(EditorPerfStats.AverageFPS, true);
-	float AverageFrameTime = TotalFromFloatArray(EditorPerfStats.AverageFrameTime, true);
-	float MemoryUsedPhysical = TotalFromFloatArray(EditorPerfStats.UsedPhysical, true);
-	float MemoryAvailPhysAvg = TotalFromFloatArray(EditorPerfStats.AvailablePhysical, true);
-	float MemoryAvailVirtualAvg = TotalFromFloatArray(EditorPerfStats.AvailableVirtual, true);
-	float MemoryUsedVirtualAvg = TotalFromFloatArray(EditorPerfStats.UsedVirtual, true);
-	float MemoryUsedPeak = LargetValueInFloatArray(EditorPerfStats.PeakUsedPhysical);
-	float MemoryUsedPeakVirtual = LargetValueInFloatArray(EditorPerfStats.PeakUsedVirtual);
-
-	//TestRunDuration is the length of time the test lasted in ticks.
-	FTimespan TestRunDuration = (EditorPerfStats.TimeStamp.Last().GetTicks() - EditorPerfStats.TimeStamp[0].GetTicks()) + ETimespan::TicksPerSecond;
-
-	//The performance csv file will be created if it didn't exist prior to the start of this test.
-	if (!FPaths::FileExists(*PerfCSVFilePath))
+
+	/**
+	* Returns true if the archive/file can be written to otherwise false..
+	* @param InFilePath - is the location of the file.
+	* @param InArchiveName - is the name of the arhcive to be used.
+	*/
+	bool IsArchiveWriteable(const FString InFilePath, const FArchive* InArchiveName)
 	{
-		FArchive* FinalCSVArchive = IFileManager::Get().CreateFileWriter(*PerfCSVFilePath);
-		if (IsArchiveWriteable(PerfCSVFilePath, FinalCSVArchive))
+		if (!InArchiveName)
 		{
-			FString FinalCSVLine = (TEXT("Date, Map Name, Changelist, Test Run Time , Map Load Time, Average FPS, Average MS, Used Physical KB, Used Virtual KB, Used Peak Physcial KB, Used Peak Virtual KB, Available Physical KB, Available Virtual KB\n"));
-			FinalCSVArchive->Serialize(TCHAR_TO_ANSI(*FinalCSVLine), FinalCSVLine.Len());
-			FinalCSVArchive->Close();
+			UE_LOG(LogEditorAutomationTests, Error, TEXT("Failed to write to the csv file: %s"), *FPaths::ConvertRelativePathToFull(InFilePath));
+			return false;
 		}
+		return true;
 	}
-
-	//Load the existing performance csv so that it doesn't get saved over and lost.
-	FString OldPerformanceCSVFile;
-	FFileHelper::LoadFileToString(OldPerformanceCSVFile, *PerfCSVFilePath);
-	FArchive* FinalCSVArchive = IFileManager::Get().CreateFileWriter(*PerfCSVFilePath);
-	if (IsArchiveWriteable(PerfCSVFilePath, FinalCSVArchive))
-	{
-		//Dump the old performance csv file data to the new csv file.
-		FinalCSVArchive->Serialize(TCHAR_TO_ANSI(*OldPerformanceCSVFile), OldPerformanceCSVFile.Len());
-
-		//Dump the pretty stats to the Performance CSV file and then close it so we can edit it while the engine is still running.
-		FString FinalCSVLine = FString::Printf(TEXT("%s,%s,%s,%.0f,%.3f,%.1f,%.1f,%.0f,%.0f,%.0f,%.0f,%.0f,%.0f%s"), *FDateTime::Now().ToString(), *EditorPerfStats.MapName, *GEngineVersion.ToString(EVersionComponent::Changelist), TestRunDuration.GetTotalSeconds(), EditorPerfStats.MapLoadTime, AverageFPS, AverageFrameTime, MemoryUsedPhysical, MemoryUsedVirtualAvg, MemoryUsedPeak, MemoryUsedPeakVirtual, MemoryAvailPhysAvg, MemoryAvailVirtualAvg, LINE_TERMINATOR);
-		FinalCSVArchive->Serialize(TCHAR_TO_ANSI(*FinalCSVLine), FinalCSVLine.Len());
-		FinalCSVArchive->Close();
-	}
-
-	//Display the test results to the user.
-	UE_LOG(LogEditorAutomationTests, Display, TEXT("AVG FPS: '%.1f'"), AverageFPS);
-	UE_LOG(LogEditorAutomationTests, Display, TEXT("AVG Frame Time: '%.1f' ms"), AverageFrameTime);
-	UE_LOG(LogEditorAutomationTests, Display, TEXT("AVG Used Physical Memory: '%.0f' kb"), MemoryUsedPhysical);
-	UE_LOG(LogEditorAutomationTests, Display, TEXT("AVG Used Virtual Memory: '%.0f' kb"), MemoryUsedVirtualAvg);
-	UE_LOG(LogEditorAutomationTests, Display, TEXT("Performance csv file is located here: %s"), *FPaths::ConvertRelativePathToFull(PerfCSVFilePath));
-	UE_LOG(LogEditorAutomationTests, Log, TEXT("Performance csv file is located here: %s"), *FPaths::ConvertRelativePathToFull(PerfCSVFilePath));
-	UE_LOG(LogEditorAutomationTests, Log, TEXT("Raw performance csv file is located here: %s"), *FPaths::ConvertRelativePathToFull(RAWCSVFilePath));
 }
 
-
 ///////////////////////////////////////////////////////////////////////
 // Common Latent commands
 
@@ -574,77 +484,6 @@ bool FEndPlayMapCommand::Update()
 	return true;
 }
 
-/**
-* This command grabs the FPS and Memory stats for the current editor session.
-*/
-bool FEditorPerfCaptureCommand::Update()
-{
-	//Capture the current time stamp and format it to YYYY-MM-DD HH:MM:SS.mmm.
-	FDateTime CurrentDateAndTime = FDateTime::Now();
-
-	//This is how long it has been since the last run through.
-	FTimespan ElapsedTime = 0;
-	FTimespan TimeBetweenCaptures;
-
-	//We want to only capture data every whole second.
-	if (EditorPerfStats.TimeStamp.Num() > 0)
-	{
-		TimeBetweenCaptures = CurrentDateAndTime.GetTicks() - EditorPerfStats.TimeStamp.Last().GetTicks();
-		if (TimeBetweenCaptures.GetTicks() < ETimespan::TicksPerSecond)
-		{
-			return false;
-		}
-		
-		ElapsedTime = CurrentDateAndTime.GetTicks() - EditorPerfStats.TimeStamp[0].GetTicks();
-	}
-
-	if (ElapsedTime.GetTotalSeconds() <= EditorPerfStats.TestDuration)
-	{
-		//Find the Average FPS
-		//Clamp to avoid huge averages at startup or after hitches
-		const float CurrentFPS = 1.0f / FSlateApplication::Get().GetAverageDeltaTime();
-		const float ClampedFPS = (CurrentFPS < 0.0f || CurrentFPS > 4000.0f) ? 0.0f : CurrentFPS;
-		EditorPerfStats.AverageFPS.Add(ClampedFPS);
-
-		//Find the Frame Time in ms.
-		//Clamp to avoid huge averages at startup or after hitches
-		const float AverageMS = FSlateApplication::Get().GetAverageDeltaTime() * 1000.0f;
-		const float ClampedMS = (AverageMS < 0.0f || AverageMS > 4000.0f) ? 0.0f : AverageMS;
-		EditorPerfStats.AverageFrameTime.Add(ClampedMS);
-
-		//Query OS for process memory used.
-		FPlatformMemoryStats MemoryStats = FPlatformMemory::GetStats();
-		EditorPerfStats.UsedPhysical.Add((float)MemoryStats.UsedPhysical / 1024);
-
-		//Query OS for available physical memory
-		EditorPerfStats.AvailablePhysical.Add((float)MemoryStats.AvailablePhysical / 1024);
-
-		//Query OS for available virtual memory
-		EditorPerfStats.AvailableVirtual.Add((float)MemoryStats.AvailableVirtual / 1024);
-
-		//Query OS for used virtual memory
-		EditorPerfStats.UsedVirtual.Add((float)MemoryStats.UsedVirtual / 1024);
-
-		//Query OS for used Peak Used physical memory
-		EditorPerfStats.PeakUsedPhysical.Add((float)MemoryStats.PeakUsedPhysical / 1024);
-
-		//Query OS for used Peak Used virtual memory
-		EditorPerfStats.PeakUsedVirtual.Add((float)MemoryStats.PeakUsedVirtual / 1024);
-
-		//Capture the time stamp.
-		FString FormatedTimeStamp = FString::Printf(TEXT("%04i-%02i-%02i %02i:%02i:%02i.%03i"), CurrentDateAndTime.GetYear(), CurrentDateAndTime.GetMonth(), CurrentDateAndTime.GetDay(), CurrentDateAndTime.GetHour(), CurrentDateAndTime.GetMinute(), CurrentDateAndTime.GetSecond(), CurrentDateAndTime.GetMillisecond());
-		EditorPerfStats.FormattedTimeStamp.Add(FormatedTimeStamp);
-		EditorPerfStats.TimeStamp.Add(CurrentDateAndTime);
-
-		return false;
-	}
-
-	//Dump the performance data in a csv file.
-	EditorPerfDump(EditorPerfStats);
-
-	return true;
-}
-
 /**
 * This this command loads a map into the editor.
 */
@@ -669,7 +508,7 @@ bool FEditorLoadMap::Update()
 	UE_LOG(LogEditorAutomationTests, Log, TEXT("%s has been loaded."), *ShortMapName);
 
 	//Log out to a text file the time it takes to load the map.
-	WriteToTextFile(TEXT("Performance"), LoadedMapName, TEXT("RAWMapLoadTime.txt"), MapLoadTime, TEXT(","));
+	AutomationEditorCommonUtils::WriteToTextFile(TEXT("Performance"), LoadedMapName, TEXT("RAWMapLoadTime.txt"), MapLoadTime, TEXT(","));
 	
 	UE_LOG(LogEditorAutomationTests, Display, TEXT("%s took %.3f to load."), *LoadedMapName, MapLoadTime);
 	
diff --git a/Engine/Source/Editor/UnrealEd/Private/Tests/EditorPerformanceTests.cpp b/Engine/Source/Editor/UnrealEd/Private/Tests/EditorPerformanceTests.cpp
--- a/Engine/Source/Editor/UnrealEd/Private/Tests/EditorPerformanceTests.cpp
+++ b/Engine/Source/Editor/UnrealEd/Private/Tests/EditorPerformanceTests.cpp
@@ -6,8 +6,218 @@
 #include "AutomationEditorCommon.h"
 #include "AutomationCommon.h"
 #include "IMainFrameModule.h"
+#include "FileManagerGeneric.h"
 
 
+//////////////////////////////////////////////////////////////////////////
+//Struct used to hold the data for the Editor Performance test.
+struct EditorPerfCaptureParameters
+{
+	//Basic Test Info
+	FString MapName;
+	int32 TestDuration;
+
+	//Saved Performance Values
+	float MapLoadTime;
+	TArray<float> AverageFPS;
+	TArray<float> AverageFrameTime;
+	TArray<float> UsedPhysical;
+	TArray<float> AvailablePhysical;
+	TArray<float> AvailableVirtual;
+	TArray<float> UsedVirtual;
+	TArray<float> PeakUsedVirtual;
+	TArray<float> PeakUsedPhysical;
+	TArray<FDateTime> TimeStamp;
+	TArray<FString> FormattedTimeStamp;
+
+	EditorPerfCaptureParameters()
+		: MapName(TEXT("None"))
+		, TestDuration(60)
+		, MapLoadTime(0)
+	{
+	}
+};
+
+//////////////////////////////////////////////////////////////////////////
+//Editor Performance Functions
+
+/**
+* Dumps the information held within the EditorPerfCaptureParameters struct into a CSV file.
+* @param EditorPerfStats is the name of the struct that holds the needed performance information.
+*/
+void EditorPerfDump(EditorPerfCaptureParameters& EditorPerfStats)
+{
+	UE_LOG(LogEditorAutomationTests, Log, TEXT("Begin generating the editor performance charts."));
+
+	//The file location where to save the data.
+	FString DataFileLocation = FPaths::Combine(*FPaths::AutomationLogDir(), TEXT("Performance"), *EditorPerfStats.MapName);
+
+	//Get the map load time (in seconds) from the text file that is created when the load map latent command is ran.
+	EditorPerfStats.MapLoadTime = 0;
+	FString MapLoadTimeFileLocation = FPaths::Combine(*DataFileLocation, TEXT("RAWMapLoadTime.txt"));
+	if (FPaths::FileExists(*MapLoadTimeFileLocation))
+	{
+		TArray<FString> SavedMapLoadTimes = AutomationEditorCommonUtils::CreateArrayFromFile(MapLoadTimeFileLocation);
+		EditorPerfStats.MapLoadTime = FCString::Atof(*SavedMapLoadTimes.Last());
+	}
+	
+	//Filename for the RAW csv which holds the data gathered from a single test ran.
+	FString RAWCSVFilePath = FString::Printf(TEXT("%s/RAW_%s_%s.csv"), *DataFileLocation, *EditorPerfStats.MapName, *FDateTime::Now().ToString());
+
+	//Filename for the pretty csv file.
+	FString PerfCSVFilePath = FString::Printf(TEXT("%s/%s_Performance.csv"), *DataFileLocation, *EditorPerfStats.MapName);
+
+	//Create the raw csv and then add the title row it.
+	FArchive* RAWCSVArchive = IFileManager::Get().CreateFileWriter(*RAWCSVFilePath);
+	FString RAWCSVLine = (TEXT("Map Name, Changelist, Time Stamp, Map Load Time, Average FPS, Frame Time, Used Physical Memory, Used Virtual Memory, Used Peak Physical, Used Peak Virtual, Available Physical Memory, Available Virtual Memory\n"));
+	RAWCSVArchive->Serialize(TCHAR_TO_ANSI(*RAWCSVLine), RAWCSVLine.Len());
+
+	//Dump the stats from each run to the raw csv file and then close it.
+	for (int32 I = 0; I < EditorPerfStats.TimeStamp.Num(); I++)
+	{
+		//If the raw file isn't available to write to then we'll fail back this test.
+		if (AutomationEditorCommonUtils::IsArchiveWriteable(RAWCSVFilePath, RAWCSVArchive))
+		{
+			RAWCSVLine = FString::Printf(TEXT("%s,%s,%s,%.3f,%.1f,%.1f,%.0f,%.0f,%.0f,%.0f,%.0f,%.0f%s"), *EditorPerfStats.MapName, *GEngineVersion.ToString(EVersionComponent::Changelist), *EditorPerfStats.TimeStamp[I].ToString(), EditorPerfStats.MapLoadTime, EditorPerfStats.AverageFPS[I], EditorPerfStats.AverageFrameTime[I], EditorPerfStats.UsedPhysical[I], EditorPerfStats.UsedVirtual[I], EditorPerfStats.PeakUsedPhysical[I], EditorPerfStats.PeakUsedVirtual[I], EditorPerfStats.AvailablePhysical[I], EditorPerfStats.AvailableVirtual[I], LINE_TERMINATOR);
+			RAWCSVArchive->Serialize(TCHAR_TO_ANSI(*RAWCSVLine), RAWCSVLine.Len());
+		}
+	}
+	RAWCSVArchive->Close();
+
+	//Get the final pretty data for the Performance csv file.
+	float AverageFPS = AutomationEditorCommonUtils::TotalFromFloatArray(EditorPerfStats.AverageFPS, true);
+	float AverageFrameTime = AutomationEditorCommonUtils::TotalFromFloatArray(EditorPerfStats.AverageFrameTime, true);
+	float MemoryUsedPhysical = AutomationEditorCommonUtils::TotalFromFloatArray(EditorPerfStats.UsedPhysical, true);
+	float MemoryAvailPhysAvg = AutomationEditorCommonUtils::TotalFromFloatArray(EditorPerfStats.AvailablePhysical, true);
+	float MemoryAvailVirtualAvg = AutomationEditorCommonUtils::TotalFromFloatArray(EditorPerfStats.AvailableVirtual, true);
+	float MemoryUsedVirtualAvg = AutomationEditorCommonUtils::TotalFromFloatArray(EditorPerfStats.UsedVirtual, true);
+	float MemoryUsedPeak = AutomationEditorCommonUtils::LargetValueInFloatArray(EditorPerfStats.PeakUsedPhysical);
+	float MemoryUsedPeakVirtual = AutomationEditorCommonUtils::LargetValueInFloatArray(EditorPerfStats.PeakUsedVirtual);
+
+	//TestRunDuration is the length of time the test lasted in ticks.
+	FTimespan TestRunDuration = (EditorPerfStats.TimeStamp.Last().GetTicks() - EditorPerfStats.TimeStamp[0].GetTicks()) + ETimespan::TicksPerSecond;
+
+	//The performance csv file will be created if it didn't exist prior to the start of this test.
+	if (!FPaths::FileExists(*PerfCSVFilePath))
+	{
+		FArchive* FinalCSVArchive = IFileManager::Get().CreateFileWriter(*PerfCSVFilePath);
+		if (AutomationEditorCommonUtils::IsArchiveWriteable(PerfCSVFilePath, FinalCSVArchive))
+		{
+			FString FinalCSVLine = (TEXT("Date, Map Name, Changelist, Test Run Time , Map Load Time, Average FPS, Average MS, Used Physical KB, Used Virtual KB, Used Peak Physcial KB, Used Peak Virtual KB, Available Physical KB, Available Virtual KB\n"));
+			FinalCSVArchive->Serialize(TCHAR_TO_ANSI(*FinalCSVLine), FinalCSVLine.Len());
+			FinalCSVArchive->Close();
+		}
+	}
+
+	//Load the existing performance csv so that it doesn't get saved over and lost.
+	FString OldPerformanceCSVFile;
+	FFileHelper::LoadFileToString(OldPerformanceCSVFile, *PerfCSVFilePath);
+	FArchive* FinalCSVArchive = IFileManager::Get().CreateFileWriter(*PerfCSVFilePath);
+	if (AutomationEditorCommonUtils::IsArchiveWriteable(PerfCSVFilePath, FinalCSVArchive))
+	{
+		//Dump the old performance csv file data to the new csv file.
+		FinalCSVArchive->Serialize(TCHAR_TO_ANSI(*OldPerformanceCSVFile), OldPerformanceCSVFile.Len());
+
+		//Dump the pretty stats to the Performance CSV file and then close it so we can edit it while the engine is still running.
+		FString FinalCSVLine = FString::Printf(TEXT("%s,%s,%s,%.0f,%.3f,%.1f,%.1f,%.0f,%.0f,%.0f,%.0f,%.0f,%.0f%s"), *FDateTime::Now().ToString(), *EditorPerfStats.MapName, *GEngineVersion.ToString(EVersionComponent::Changelist), TestRunDuration.GetTotalSeconds(), EditorPerfStats.MapLoadTime, AverageFPS, AverageFrameTime, MemoryUsedPhysical, MemoryUsedVirtualAvg, MemoryUsedPeak, MemoryUsedPeakVirtual, MemoryAvailPhysAvg, MemoryAvailVirtualAvg, LINE_TERMINATOR);
+		FinalCSVArchive->Serialize(TCHAR_TO_ANSI(*FinalCSVLine), FinalCSVLine.Len());
+		FinalCSVArchive->Close();
+	}
+
+	//Display the test results to the user.
+	UE_LOG(LogEditorAutomationTests, Display, TEXT("AVG FPS: '%.1f'"), AverageFPS);
+	UE_LOG(LogEditorAutomationTests, Display, TEXT("AVG Frame Time: '%.1f' ms"), AverageFrameTime);
+	UE_LOG(LogEditorAutomationTests, Display, TEXT("AVG Used Physical Memory: '%.0f' kb"), MemoryUsedPhysical);
+	UE_LOG(LogEditorAutomationTests, Display, TEXT("AVG Used Virtual Memory: '%.0f' kb"), MemoryUsedVirtualAvg);
+	UE_LOG(LogEditorAutomationTests, Display, TEXT("Performance csv file is located here: %s"), *FPaths::ConvertRelativePathToFull(PerfCSVFilePath));
+	UE_LOG(LogEditorAutomationTests, Log, TEXT("Performance csv file is located here: %s"), *FPaths::ConvertRelativePathToFull(PerfCSVFilePath));
+	UE_LOG(LogEditorAutomationTests, Log, TEXT("Raw performance csv file is located here: %s"), *FPaths::ConvertRelativePathToFull(RAWCSVFilePath));
+}
+
+
+
+//////////////////////////////////////////////////////////////////////////
+//Editor Performance Latent Commands
+
+/**
+* This will capture the average FPS and Memory numbers over a duration of time.
+* @param The name of the EditorPerfCaptureParameters struct that will be used to hold the data.
+*/
+DEFINE_LATENT_AUTOMATION_COMMAND_ONE_PARAMETER(FEditorPerfCaptureCommand, EditorPerfCaptureParameters, EditorPerfStats);
+
+/**
+* This command grabs the FPS and Memory stats for the current editor session.
+*/
+bool FEditorPerfCaptureCommand::Update()
+{
+	//Capture the current time stamp and format it to YYYY-MM-DD HH:MM:SS.mmm.
+	FDateTime CurrentDateAndTime = FDateTime::Now();
+
+	//This is how long it has been since the last run through.
+	FTimespan ElapsedTime = 0;
+	FTimespan TimeBetweenCaptures;
+
+	//We want to only capture data every whole second.
+	if (EditorPerfStats.TimeStamp.Num() > 0)
+	{
+		TimeBetweenCaptures = CurrentDateAndTime.GetTicks() - EditorPerfStats.TimeStamp.Last().GetTicks();
+		if (TimeBetweenCaptures.GetTicks() < ETimespan::TicksPerSecond)
+		{
+			return false;
+		}
+
+		ElapsedTime = CurrentDateAndTime.GetTicks() - EditorPerfStats.TimeStamp[0].GetTicks();
+	}
+
+	if (ElapsedTime.GetTotalSeconds() <= EditorPerfStats.TestDuration)
+	{
+		//Find the Average FPS
+		//Clamp to avoid huge averages at startup or after hitches
+		const float CurrentFPS = 1.0f / FSlateApplication::Get().GetAverageDeltaTime();
+		const float ClampedFPS = (CurrentFPS < 0.0f || CurrentFPS > 4000.0f) ? 0.0f : CurrentFPS;
+		EditorPerfStats.AverageFPS.Add(ClampedFPS);
+
+		//Find the Frame Time in ms.
+		//Clamp to avoid huge averages at startup or after hitches
+		const float AverageMS = FSlateApplication::Get().GetAverageDeltaTime() * 1000.0f;
+		const float ClampedMS = (AverageMS < 0.0f || AverageMS > 4000.0f) ? 0.0f : AverageMS;
+		EditorPerfStats.AverageFrameTime.Add(ClampedMS);
+
+		//Query OS for process memory used.
+		FPlatformMemoryStats MemoryStats = FPlatformMemory::GetStats();
+		EditorPerfStats.UsedPhysical.Add((float)MemoryStats.UsedPhysical / 1024);
+
+		//Query OS for available physical memory
+		EditorPerfStats.AvailablePhysical.Add((float)MemoryStats.AvailablePhysical / 1024);
+
+		//Query OS for available virtual memory
+		EditorPerfStats.AvailableVirtual.Add((float)MemoryStats.AvailableVirtual / 1024);
+
+		//Query OS for used virtual memory
+		EditorPerfStats.UsedVirtual.Add((float)MemoryStats.UsedVirtual / 1024);
+
+		//Query OS for used Peak Used physical memory
+		EditorPerfStats.PeakUsedPhysical.Add((float)MemoryStats.PeakUsedPhysical / 1024);
+
+		//Query OS for used Peak Used virtual memory
+		EditorPerfStats.PeakUsedVirtual.Add((float)MemoryStats.PeakUsedVirtual / 1024);
+
+		//Capture the time stamp.
+		FString FormatedTimeStamp = FString::Printf(TEXT("%04i-%02i-%02i %02i:%02i:%02i.%03i"), CurrentDateAndTime.GetYear(), CurrentDateAndTime.GetMonth(), CurrentDateAndTime.GetDay(), CurrentDateAndTime.GetHour(), CurrentDateAndTime.GetMinute(), CurrentDateAndTime.GetSecond(), CurrentDateAndTime.GetMillisecond());
+		EditorPerfStats.FormattedTimeStamp.Add(FormatedTimeStamp);
+		EditorPerfStats.TimeStamp.Add(CurrentDateAndTime);
+
+		return false;
+	}
+
+	//Dump the performance data in a csv file.
+	EditorPerfDump(EditorPerfStats);
+
+	return true;
+}
+
+//////////////////////////////////////////////////////////////////////////
+//Editor Performance Tests
 
 /**
 * Map Performance in Editor tests
diff --git a/Engine/Source/Editor/UnrealEd/Public/Tests/AutomationEditorCommon.h b/Engine/Source/Editor/UnrealEd/Public/Tests/AutomationEditorCommon.h
--- a/Engine/Source/Editor/UnrealEd/Public/Tests/AutomationEditorCommon.h
+++ b/Engine/Source/Editor/UnrealEd/Public/Tests/AutomationEditorCommon.h
@@ -93,35 +93,49 @@ namespace AutomationEditorCommonUtils
 	* @param FactorySettings - An array of custom settings to apply to the factory
 	*/
 	void ApplyCustomFactorySettings(UFactory* InFactory, const TArray<FImportFactorySettingValues>& FactorySettings);
+
+	/**
+	* Writes a number to a text file.
+	*
+	* @param InTestName - is the folder that has the same name as the test. (For Example: "Performance").
+	* @param InItemBeingTested - is the name for the thing that is being tested. (For Example: "MapName").
+	* @param InFileName - is the name of the file with an extension
+	* @param InNumberToBeWritten - is the float number that is expected to be written to the file.
+	* @param Delimiter - is the delimiter to be used. TEXT(",")
+	*/
+	void WriteToTextFile(const FString& InTestName, const FString& InItemBeingTested, const FString& InFileName, const float& InNumberToBeWritten, const FString& Delimiter);
+
+	/**
+	* Returns the sum of the numbers available in an array of float.
+
+	* @param InFloatArray - is the name of the array intended to be used.
+	* @param bisAveragedInstead - will return the average of the available numbers instead of the sum.
+	*/
+	float TotalFromFloatArray(const TArray<float>& InFloatArray, bool bisAveragedInstead);
+
+	/**
+	* Returns the largest value from an array of float numbers.
+
+	* @param InFloatArray - is the name of the array intended to be used.
+	*/
+	float LargetValueInFloatArray(const TArray<float>& InFloatArray);
+
+	/**
+	* Returns the contents of a text file as an array of FString.
+
+	* @param InFileLocation - is the location of the file.
+	*/
+	TArray<FString> CreateArrayFromFile(const FString& InFileLocation);
+
+	/**
+	* Returns true if the archive/file can be written to otherwise false.
+
+	* @param InFilePath - is the location of the file.
+	* @param InArchiveName - is the name of the archive to be used.
+	*/
+	bool IsArchiveWriteable(const FString InFilePath, const FArchive* InArchiveName);
 }
 
-//Struct used to hold the data for the Editor Performance test.
-struct EditorPerfCaptureParameters
-{
-	//Basic Test Info
-	FString MapName;
-	int32 TestDuration;
-
-	//Saved Performance Values
-	float MapLoadTime;
-	TArray<float> AverageFPS;
-	TArray<float> AverageFrameTime;
-	TArray<float> UsedPhysical;
-	TArray<float> AvailablePhysical;
-	TArray<float> AvailableVirtual;
-	TArray<float> UsedVirtual;
-	TArray<float> PeakUsedVirtual;
-	TArray<float> PeakUsedPhysical;
-	TArray<FDateTime> TimeStamp;
-	TArray<FString> FormattedTimeStamp;
-
-	EditorPerfCaptureParameters()
-		: MapName(TEXT("None"))
-		, TestDuration(60)
-		, MapLoadTime(0)
-	{
-	}
-};
 
 //////////////////////////////////////////////////////////////////////////
 //Common latent commands used for automated editor testing.
@@ -152,12 +166,6 @@ DEFINE_LATENT_AUTOMATION_COMMAND_ONE_PARAMETER(FStartPIECommand, bool, bSimulate
 */
 DEFINE_LATENT_AUTOMATION_COMMAND(FEndPlayMapCommand);
 
-/**
-* This will capture the average FPS and Memory numbers over a duration of time.
-* @param The name of the EditorPerfCaptureParameters struct that will be used to hold the data.
-*/
-DEFINE_LATENT_AUTOMATION_COMMAND_ONE_PARAMETER(FEditorPerfCaptureCommand, EditorPerfCaptureParameters, EditorPerfStats);
-
 /**
 * Loads a map
 */
