Fixed crash when removing "current" streaming level in Levels browser

diff --git a/Engine/Source/Editor/UnrealEd/Private/EditorLevelUtils.cpp b/Engine/Source/Editor/UnrealEd/Private/EditorLevelUtils.cpp
--- a/Engine/Source/Editor/UnrealEd/Private/EditorLevelUtils.cpp
+++ b/Engine/Source/Editor/UnrealEd/Private/EditorLevelUtils.cpp
@@ -399,23 +399,9 @@ namespace EditorLevelUtils
 		const bool bRemoveSuccessful		= PrivateRemoveLevelFromWorld( InLevel );
 		if ( bRemoveSuccessful )
 		{
-			if ( OwningWorld )
+			if (bRemovingCurrentLevel)
 			{
-				// remove all group actors from the world in the level we are removing
-				// otherwise, this will cause group actors to not be garbage collected
-				for (int32 GroupIndex = OwningWorld->ActiveGroupActors.Num()-1; GroupIndex >= 0; --GroupIndex)
-				{
-					AGroupActor* GroupActor = Cast<AGroupActor>(OwningWorld->ActiveGroupActors[GroupIndex]);
-					if (GroupActor && GroupActor->IsInLevel(InLevel))
-					{
-						OwningWorld->ActiveGroupActors.RemoveAt(GroupIndex);
-					}
-				}
-			}
-
-			if(bRemovingCurrentLevel)
-			{
-				MakeLevelCurrent(InLevel->OwningWorld->PersistentLevel);
+				MakeLevelCurrent(OwningWorld->PersistentLevel);
 			}
 
 			// Redraw the main editor viewports.
@@ -506,6 +492,17 @@ namespace EditorLevelUtils
 		int32 NumFailedDestroyedAttempts = 0;
 		bool bDestroyedActor = false;
 
+		// remove all group actors from the world in the level we are removing
+		// otherwise, this will cause group actors to not be garbage collected
+		for (int32 GroupIndex = World->ActiveGroupActors.Num()-1; GroupIndex >= 0; --GroupIndex)
+		{
+			AGroupActor* GroupActor = Cast<AGroupActor>(World->ActiveGroupActors[GroupIndex]);
+			if (GroupActor && GroupActor->IsInLevel(InLevel))
+			{
+				World->ActiveGroupActors.RemoveAt(GroupIndex);
+			}
+		}
+
 		for(int32 ActorIndex = 0; ActorIndex < InLevel->Actors.Num(); ++ActorIndex)
 		{
 			AActor* ActorToRemove = InLevel->Actors[ActorIndex];
