Changed "Assets were open when the editor was last closed" notification so that it times out, and fades out instantly when an option is selected.

diff --git a/Engine/Source/Editor/MainFrame/Private/Frame/MainFrameHandler.h b/Engine/Source/Editor/MainFrame/Private/Frame/MainFrameHandler.h
--- a/Engine/Source/Editor/MainFrame/Private/Frame/MainFrameHandler.h
+++ b/Engine/Source/Editor/MainFrame/Private/Frame/MainFrameHandler.h
@@ -205,7 +205,7 @@ public:
 			LevelEditor.FocusViewport();
 
 			// Restore any assets we had open. Note we don't do this on immersive PIE as its annoying to the user.
-			FAssetEditorManager::Get().RestorePreviouslyOpenAssets();
+			FAssetEditorManager::Get().RequestRestorePreviouslyOpenAssets();
 		}
 	}
 
diff --git a/Engine/Source/Editor/UnrealEd/Private/Toolkits/AssetEditorManager.cpp b/Engine/Source/Editor/UnrealEd/Private/Toolkits/AssetEditorManager.cpp
--- a/Engine/Source/Editor/UnrealEd/Private/Toolkits/AssetEditorManager.cpp
+++ b/Engine/Source/Editor/UnrealEd/Private/Toolkits/AssetEditorManager.cpp
@@ -32,6 +32,7 @@ FAssetEditorManager& FAssetEditorManager::Get()
 
 FAssetEditorManager::FAssetEditorManager()
 	: bSavingOnShutdown(false)
+	, bRequestRestorePreviouslyOpenAssets(false)
 {
 	// Message bus to receive requests to load assets
 	MessageEndpoint = FMessageEndpoint::Builder("FAssetEditorManager")
@@ -454,11 +455,22 @@ void FAssetEditorManager::OpenEditorForAsset(const FString& AssetPathName)
 
 bool FAssetEditorManager::HandleTicker( float DeltaTime )
 {
+	if (bRequestRestorePreviouslyOpenAssets)
+	{
+		RestorePreviouslyOpenAssets();
+		bRequestRestorePreviouslyOpenAssets = false;
+	}
 	MessageEndpoint->ProcessInbox();
 
 	return true;
 }
 
+void FAssetEditorManager::RequestRestorePreviouslyOpenAssets()
+{
+	// We defer the restore so that we guarantee that it happens once initialization is complete
+	bRequestRestorePreviouslyOpenAssets = true;
+}
+
 void FAssetEditorManager::RestorePreviouslyOpenAssets()
 {
 	TArray<FString> OpenAssets;
@@ -540,8 +552,9 @@ void FAssetEditorManager::SpawnRestorePreviouslyOpenAssetsNotification(const boo
 		SNotificationItem::CS_None
 		));
 
-	// We will be keeping track of this ourselves
-	Info.bFireAndForget = false;
+	// We will let the notification expire automatically after 10 seconds
+	Info.bFireAndForget = true;
+	Info.ExpireDuration = 10.0f;
 
 	// We want the auto-save to be subtle
 	Info.bUseLargeFont = false;
@@ -557,22 +570,24 @@ void FAssetEditorManager::SpawnRestorePreviouslyOpenAssetsNotification(const boo
 	}
 
 	// Close any existing notification
+	TSharedPtr<SNotificationItem> RestorePreviouslyOpenAssetsNotification = RestorePreviouslyOpenAssetsNotificationPtr.Pin();
 	if(RestorePreviouslyOpenAssetsNotification.IsValid())
 	{
 		RestorePreviouslyOpenAssetsNotification->ExpireAndFadeout();
-		RestorePreviouslyOpenAssetsNotification.Reset();
 	}
 
-	RestorePreviouslyOpenAssetsNotification = FSlateNotificationManager::Get().AddNotification(Info);
+	RestorePreviouslyOpenAssetsNotificationPtr = FSlateNotificationManager::Get().AddNotification(Info);
 }
 
 void FAssetEditorManager::OnConfirmRestorePreviouslyOpenAssets(TArray<FString> AssetsToOpen)
 {
 	// Close any existing notification
-	if(RestorePreviouslyOpenAssetsNotification.IsValid())
+	TSharedPtr<SNotificationItem> RestorePreviouslyOpenAssetsNotification = RestorePreviouslyOpenAssetsNotificationPtr.Pin();
+	if (RestorePreviouslyOpenAssetsNotification.IsValid())
 	{
+		RestorePreviouslyOpenAssetsNotification->SetExpireDuration(0.0f);
+		RestorePreviouslyOpenAssetsNotification->SetFadeOutDuration(0.5f);
 		RestorePreviouslyOpenAssetsNotification->ExpireAndFadeout();
-		RestorePreviouslyOpenAssetsNotification.Reset();
 
 		// If the user suppressed the notification for future sessions, make sure this is reflected in their settings
 		// Note: We do that inside this if statement so that we only do it if we were showing a UI they could interact with
@@ -591,10 +606,12 @@ void FAssetEditorManager::OnConfirmRestorePreviouslyOpenAssets(TArray<FString> A
 void FAssetEditorManager::OnCancelRestorePreviouslyOpenAssets()
 {
 	// Close any existing notification
-	if(RestorePreviouslyOpenAssetsNotification.IsValid())
+	TSharedPtr<SNotificationItem> RestorePreviouslyOpenAssetsNotification = RestorePreviouslyOpenAssetsNotificationPtr.Pin();
+	if (RestorePreviouslyOpenAssetsNotification.IsValid())
 	{
+		RestorePreviouslyOpenAssetsNotification->SetExpireDuration(0.0f);
+		RestorePreviouslyOpenAssetsNotification->SetFadeOutDuration(0.5f);
 		RestorePreviouslyOpenAssetsNotification->ExpireAndFadeout();
-		RestorePreviouslyOpenAssetsNotification.Reset();
 	}
 }
 
diff --git a/Engine/Source/Editor/UnrealEd/Public/Toolkits/AssetEditorManager.h b/Engine/Source/Editor/UnrealEd/Public/Toolkits/AssetEditorManager.h
--- a/Engine/Source/Editor/UnrealEd/Public/Toolkits/AssetEditorManager.h
+++ b/Engine/Source/Editor/UnrealEd/Public/Toolkits/AssetEditorManager.h
@@ -93,8 +93,8 @@ public:
 	DECLARE_EVENT_OneParam(FAssetEditorManager, FAssetEditorOpenEvent, UObject*);
 	FAssetEditorOpenEvent& OnAssetEditorOpened() { return AssetEditorOpenedEvent; }
 
-	/** Restore the assets that were previously open when the editor was last closed */
-	void RestorePreviouslyOpenAssets();
+	/** Request notification to restore the assets that were previously open when the editor was last closed */
+	void RequestRestorePreviouslyOpenAssets();
 
 private:
 
@@ -125,6 +125,9 @@ private:
 	/** Saves a list of open asset editors so they can be restored on editor restart */
 	void SaveOpenAssetEditors(bool bOnShutdown);
 
+	/** Restore the assets that were previously open when the editor was last closed */
+	void RestorePreviouslyOpenAssets();
+
 private:
 
 	/** struct used by OpenedEditorTimes map to store editor names and times */
@@ -174,8 +177,11 @@ private:
 	/** Flag whether we are currently shutting down */
 	bool bSavingOnShutdown;
 
+	/** Flag whether there has been a request to notify whether to restore previously open assets */
+	bool bRequestRestorePreviouslyOpenAssets;
+
 	/** A pointer to the notification used by RestorePreviouslyOpenAssets */
-	TSharedPtr<SNotificationItem> RestorePreviouslyOpenAssetsNotification;
+	TWeakPtr<SNotificationItem> RestorePreviouslyOpenAssetsNotificationPtr;
 };
 
 
diff --git a/Engine/Source/Runtime/Slate/Private/Widgets/Notifications/SNotificationList.cpp b/Engine/Source/Runtime/Slate/Private/Widgets/Notifications/SNotificationList.cpp
--- a/Engine/Source/Runtime/Slate/Private/Widgets/Notifications/SNotificationList.cpp
+++ b/Engine/Source/Runtime/Slate/Private/Widgets/Notifications/SNotificationList.cpp
@@ -510,6 +510,18 @@ public:
 		ExpireDuration = Duration;
 	}
 
+	/** Sets the FadeInDuration */
+	virtual void SetFadeInDuration(float Duration)
+	{
+		FadeInDuration = Duration;
+	}
+
+	/** Sets the FadeOutDuration */
+	virtual void SetFadeOutDuration(float Duration)
+	{
+		FadeOutDuration = Duration;
+	}
+
 protected:
 
 	/* Used to determine whether the button is visible */
diff --git a/Engine/Source/Runtime/Slate/Public/Widgets/Notifications/SNotificationList.h b/Engine/Source/Runtime/Slate/Public/Widgets/Notifications/SNotificationList.h
--- a/Engine/Source/Runtime/Slate/Public/Widgets/Notifications/SNotificationList.h
+++ b/Engine/Source/Runtime/Slate/Public/Widgets/Notifications/SNotificationList.h
@@ -30,6 +30,12 @@ public:
 	/** Sets the ExpireDuration */
 	virtual void SetExpireDuration(float ExpireDuration) = 0;
 
+	/** Sets the FadeInDuration */
+	virtual void SetFadeInDuration(float FadeInDuration) = 0;
+
+	/** Sets the FadeOutDuration */
+	virtual void SetFadeOutDuration(float FadeOutDuration) = 0;
+
 	/** Gets the visibility state of the throbber, success, and fail images */
 	virtual ECompletionState GetCompletionState() const = 0;
 
