Fixed crash when removing a currently loaded sub level from the level browser after it's saved.

diff --git a/Engine/Source/Editor/UnrealEd/Private/EditorLevelUtils.cpp b/Engine/Source/Editor/UnrealEd/Private/EditorLevelUtils.cpp
--- a/Engine/Source/Editor/UnrealEd/Private/EditorLevelUtils.cpp
+++ b/Engine/Source/Editor/UnrealEd/Private/EditorLevelUtils.cpp
@@ -404,6 +404,8 @@ namespace EditorLevelUtils
 				MakeLevelCurrent(OwningWorld->PersistentLevel);
 			}
 
+			EditorDestroyLevel(InLevel);
+
 			// Redraw the main editor viewports.
 			FEditorSupportDelegates::RedrawAllViewports.Broadcast();
 
@@ -434,17 +436,17 @@ namespace EditorLevelUtils
 	/**
 	* Removes a level from the world.  Returns true if the level was removed successfully.
 	*
-	* @param	Level		The level to remove from the world.
+	* @param	InLevel		The level to remove from the world.
 	* @return				true if the level was removed successfully, false otherwise.
 	*/
-	bool PrivateRemoveLevelFromWorld(ULevel* Level)
+	bool PrivateRemoveLevelFromWorld(ULevel* InLevel)
 	{
-		if ( !Level || Level->IsPersistentLevel() )
+		if ( !InLevel || InLevel->IsPersistentLevel() )
 		{
 			return false;
 		}
 
-		if ( FLevelUtils::IsLevelLocked(Level) )
+		if ( FLevelUtils::IsLevelLocked(InLevel) )
 		{
 			FMessageDialog::Open( EAppMsgType::Ok, NSLOCTEXT("UnrealEd", "Error_OperationDisallowedOnLockedLevelRemoveLevelFromWorld", "RemoveLevelFromWorld: The requested operation could not be completed because the level is locked.") );
 			return false;
@@ -452,10 +454,10 @@ namespace EditorLevelUtils
 
 		int32 StreamingLevelIndex = INDEX_NONE;
 
-		for( int32 LevelIndex = 0 ; LevelIndex < Level->OwningWorld->StreamingLevels.Num() ; ++LevelIndex )
+		for( int32 LevelIndex = 0 ; LevelIndex < InLevel->OwningWorld->StreamingLevels.Num() ; ++LevelIndex )
 		{
-			ULevelStreaming* StreamingLevel = Level->OwningWorld->StreamingLevels[ LevelIndex ];
-			if( StreamingLevel && StreamingLevel->GetLoadedLevel() == Level )
+			ULevelStreaming* StreamingLevel = InLevel->OwningWorld->StreamingLevels[ LevelIndex ];
+			if( StreamingLevel && StreamingLevel->GetLoadedLevel() == InLevel )
 			{
 				StreamingLevelIndex = LevelIndex;
 				break;
@@ -464,24 +466,16 @@ namespace EditorLevelUtils
 
 		if (StreamingLevelIndex != INDEX_NONE)
 		{
-			Level->OwningWorld->StreamingLevels[StreamingLevelIndex]->MarkPendingKill();
-			Level->OwningWorld->StreamingLevels.RemoveAt( StreamingLevelIndex );
-			Level->OwningWorld->RefreshStreamingLevels();
+			InLevel->OwningWorld->StreamingLevels[StreamingLevelIndex]->MarkPendingKill();
+			InLevel->OwningWorld->StreamingLevels.RemoveAt( StreamingLevelIndex );
+			InLevel->OwningWorld->RefreshStreamingLevels();
 		}
-		else if (Level->bIsVisible)
+		else if (InLevel->bIsVisible)
 		{
-			Level->OwningWorld->RemoveFromWorld(Level);
-			check(Level->bIsVisible == false);
+			InLevel->OwningWorld->RemoveFromWorld(InLevel);
+			check(InLevel->bIsVisible == false);
 		}
 
-		return EditorDestroyLevel(Level);
-	}
-	
-	bool EditorDestroyLevel( ULevel* InLevel )
-	{
-		check(InLevel);
-		check(!InLevel->IsPersistentLevel());
-
 		InLevel->ReleaseRenderingResources();
 
 		IStreamingManager::Get().RemoveLevel( InLevel );
@@ -523,12 +517,22 @@ namespace EditorLevelUtils
 			UE_LOG(LogLevelTools, Log, TEXT("Failed to destroy %d actors after attempting to destroy level!"), NumFailedDestroyedAttempts);
 		}
 
+		World->MarkPackageDirty();
+		World->BroadcastLevelsChanged();
+
+		return true;
+	}
+
+	bool EditorDestroyLevel(ULevel* InLevel)
+	{
+		UWorld* World = InLevel->OwningWorld;
+
 		InLevel->GetOuter()->MarkPendingKill();
 		InLevel->MarkPendingKill();
- 		InLevel->GetOuter()->ClearFlags(RF_Public|RF_Standalone);
+		InLevel->GetOuter()->ClearFlags(RF_Public | RF_Standalone);
 
 		UPackage* Package = Cast<UPackage>(InLevel->GetOutermost());
-		// Destroying actors in the world will have set the package dirty flag - clear it here so it can be unloaded successfully
+		// We want to unconditionally destroy the level, so clear the dirty flag here so it can be unloaded successfully
 		Package->SetDirtyFlag(false);
 
 		TArray<UPackage*> Packages;
@@ -538,11 +542,9 @@ namespace EditorLevelUtils
 			FFormatNamedArguments Args;
 			Args.Add(TEXT("Package"), FText::FromString(Package->GetName()));
 			FMessageDialog::Open(EAppMsgType::Ok, FText::Format(LOCTEXT("UnloadPackagesFail", "Unable to unload package '{Package}'."), Args));
+			return false;
 		}
 
-		World->MarkPackageDirty();
-		World->BroadcastLevelsChanged();
-
 		return true;
 	}
 
diff --git a/Engine/Source/Editor/UnrealEd/Public/EditorLevelUtils.h b/Engine/Source/Editor/UnrealEd/Public/EditorLevelUtils.h
--- a/Engine/Source/Editor/UnrealEd/Public/EditorLevelUtils.h
+++ b/Engine/Source/Editor/UnrealEd/Public/EditorLevelUtils.h
@@ -110,7 +110,7 @@ namespace EditorLevelUtils
 	bool PrivateRemoveLevelFromWorld( ULevel* Level );
 
 	/** 
-	 * Completely removes the level from the world.
+	 * Completely removes the level from the world, unloads its package and forces garbage collection.
 	 *
 	 * @note: This function doesn't remove the associated streaming level.
 	 *
