Missing class icons using the "default" icon in Blueprints (and Scene Outliner) now appear.

#ttp 349222 - BP: MY BLUEPRINT: Some variables missing icons

diff --git a/Engine/Source/Editor/UnrealEd/Private/ClassIconFinder.cpp b/Engine/Source/Editor/UnrealEd/Private/ClassIconFinder.cpp
--- a/Engine/Source/Editor/UnrealEd/Private/ClassIconFinder.cpp
+++ b/Engine/Source/Editor/UnrealEd/Private/ClassIconFinder.cpp
@@ -147,7 +147,7 @@ FName FClassIconFinder::FindIconNameImpl(const UClass* InClass, const FName& InD
 	if( Brush == NULL )
 	{
 		// If we didn't supply an override name for the default icon use default class icon.
-		if( InDefaultName == "" )
+		if( InDefaultName == NAME_None )
 		{
 			BrushName = *FString::Printf( TEXT( "%s.Default" ), StyleRoot);
 		}
