ReinstanceOnlyWhenNecessary:
- mechanism to distinguish if instances of recompiled BPGC must be reinstanced
- replacing PB class without reinstancing.

diff --git a/Engine/Config/BaseEngine.ini b/Engine/Config/BaseEngine.ini
--- a/Engine/Config/BaseEngine.ini
+++ b/Engine/Config/BaseEngine.ini
@@ -1664,6 +1664,7 @@ bTurnOffEditorConstructionScript=false
 PrintStringDuration=2.0
 bUseLocalGraphVariables=false
 bPersistentUberGraphFrame=false
+bReinstanceOnlyWhenNecessary=false
 bDisplaysLayout=false
 
 [/Script/Engine.Blueprint]
diff --git a/Engine/Source/Editor/KismetCompiler/Private/KismetCompiler.cpp b/Engine/Source/Editor/KismetCompiler/Private/KismetCompiler.cpp
--- a/Engine/Source/Editor/KismetCompiler/Private/KismetCompiler.cpp
+++ b/Engine/Source/Editor/KismetCompiler/Private/KismetCompiler.cpp
@@ -47,6 +47,7 @@ DEFINE_STAT(EKismetCompilerStats_ChooseTerminalScope);
 DEFINE_STAT(EKismetCompilerStats_CleanAndSanitizeClass);
 DEFINE_STAT(EKismetCompilerStats_CreateClassVariables);
 DEFINE_STAT(EKismetCompilerStats_BindAndLinkClass);
+DEFINE_STAT(EKismetCompilerStats_ChecksumCDO);
 		
 //////////////////////////////////////////////////////////////////////////
 // FKismetCompilerContext
@@ -3590,6 +3591,14 @@ void FKismetCompilerContext::Compile()
 		bool Result = ValidateGeneratedClass(NewClass);
 		// TODO What do we do if validation fails?
 	}
+
+	if (bIsFullCompile)
+	{
+		BP_SCOPED_COMPILER_EVENT_STAT(EKismetCompilerStats_ChecksumCDO);
+		UObject* NewCDO = NewClass->GetDefaultObject(false);
+		FArchiveObjectCrc32 CrcArchive;
+		Blueprint->CrcPreviousCompiledCDO = NewCDO ? CrcArchive.Crc32(NewCDO) : 0;
+	}
 }
 
 bool FKismetCompilerContext::ValidateGeneratedClass(UBlueprintGeneratedClass* Class)
diff --git a/Engine/Source/Editor/KismetCompiler/Private/KismetCompilerModule.cpp b/Engine/Source/Editor/KismetCompiler/Private/KismetCompilerModule.cpp
--- a/Engine/Source/Editor/KismetCompiler/Private/KismetCompilerModule.cpp
+++ b/Engine/Source/Editor/KismetCompiler/Private/KismetCompilerModule.cpp
@@ -150,6 +150,8 @@ void FKismet2CompilerModule::CompileBlueprint(class UBlueprint* Blueprint, const
 			// There were errors.  Compile the generated class to have function stubs
 			Blueprint->Status = BS_Error;
 
+			static const FBoolConfigValueHelper ReinstanceOnlyWhenNecessary(TEXT("Kismet"), TEXT("bReinstanceOnlyWhenNecessary"), GEngineIni);
+
 			// Reinstance objects here, so we can preserve their memory layouts to reinstance them again
 			if( ParentReinstancer != NULL )
 			{
@@ -157,7 +159,7 @@ void FKismet2CompilerModule::CompileBlueprint(class UBlueprint* Blueprint, const
 
 				if(!Blueprint->bIsRegeneratingOnLoad)
 				{
-					ParentReinstancer->ReinstanceObjects();
+					ParentReinstancer->ReinstanceObjects(!ReinstanceOnlyWhenNecessary);
 				}
 			}
 
@@ -174,7 +176,7 @@ void FKismet2CompilerModule::CompileBlueprint(class UBlueprint* Blueprint, const
 			StubReinstancer.UpdateBytecodeReferences();
 			if( !Blueprint->bIsRegeneratingOnLoad )
 			{
-				StubReinstancer.ReinstanceObjects();
+				StubReinstancer.ReinstanceObjects(!ReinstanceOnlyWhenNecessary);
 			}
 		}
 	}
diff --git a/Engine/Source/Editor/KismetCompiler/Public/EdGraphCompilerUtilities.h b/Engine/Source/Editor/KismetCompiler/Public/EdGraphCompilerUtilities.h
--- a/Engine/Source/Editor/KismetCompiler/Public/EdGraphCompilerUtilities.h
+++ b/Engine/Source/Editor/KismetCompiler/Public/EdGraphCompilerUtilities.h
@@ -27,6 +27,7 @@ DECLARE_CYCLE_STAT_EXTERN(TEXT("Choose Terminal Scope"), EKismetCompilerStats_Ch
 DECLARE_CYCLE_STAT_EXTERN(TEXT("Clean and Sanitize Class"), EKismetCompilerStats_CleanAndSanitizeClass, STATGROUP_KismetCompiler, );
 DECLARE_CYCLE_STAT_EXTERN(TEXT("Create Class Properties"), EKismetCompilerStats_CreateClassVariables, STATGROUP_KismetCompiler, );
 DECLARE_CYCLE_STAT_EXTERN(TEXT("Bind and Link Class"), EKismetCompilerStats_BindAndLinkClass, STATGROUP_KismetCompiler, );
+DECLARE_CYCLE_STAT_EXTERN(TEXT("Calculate checksum of CDO"), EKismetCompilerStats_ChecksumCDO, STATGROUP_KismetCompiler, );
 
 //////////////////////////////////////////////////////////////////////////
 
diff --git a/Engine/Source/Editor/UnrealEd/Private/Kismet2/Kismet2.cpp b/Engine/Source/Editor/UnrealEd/Private/Kismet2/Kismet2.cpp
--- a/Engine/Source/Editor/UnrealEd/Private/Kismet2/Kismet2.cpp
+++ b/Engine/Source/Editor/UnrealEd/Private/Kismet2/Kismet2.cpp
@@ -347,10 +347,10 @@ void FKismetEditorUtilities::CompileBlueprint(UBlueprint* BlueprintObj, bool bIs
 	bool const bStartedWithUnsavedChanges = (BlueprintPackage != NULL) ? BlueprintPackage->IsDirty() : true;
 #if WITH_EDITOR
 	// Do not want to run this code without the editor present nor when running commandlets.
-	if(GEditor && GIsEditor && !IsRunningCommandlet())
+	if (GEditor && GIsEditor && !IsRunningCommandlet())
 	{
 		// We do not want to regenerate a search Guid during loads, nothing has changed in the Blueprint and it is cached elsewhere
-		if(!bIsRegeneratingOnLoad)
+		if (!bIsRegeneratingOnLoad)
 		{
 			BlueprintObj->SearchGuid = FGuid::NewGuid();
 			FFindInBlueprintSearchManager::Get().AddOrUpdateBlueprintSearchMetadata(BlueprintObj);
@@ -409,7 +409,8 @@ void FKismetEditorUtilities::CompileBlueprint(UBlueprint* BlueprintObj, bool bIs
 		}
 
 		// Replace instances of this class
-		ReinstanceHelper.ReinstanceObjects();
+		static const FBoolConfigValueHelper ReinstanceOnlyWhenNecessary(TEXT("Kismet"), TEXT("bReinstanceOnlyWhenNecessary"), GEngineIni);
+		ReinstanceHelper.ReinstanceObjects(!ReinstanceOnlyWhenNecessary);
 
 		if (!bSkipGarbageCollection)
 		{
@@ -423,7 +424,8 @@ void FKismetEditorUtilities::CompileBlueprint(UBlueprint* BlueprintObj, bool bIs
 		// ReinstanceHelper.VerifyReplacement();
 	}
 
-	{ BP_SCOPED_COMPILER_EVENT_NAME(TEXT("Notify Blueprint Changed"));
+	{ 
+		BP_SCOPED_COMPILER_EVENT_NAME(TEXT("Notify Blueprint Changed"));
 
 		// Blueprint has changed, broadcast notify
 		BlueprintObj->BroadcastChanged();
@@ -440,7 +442,8 @@ void FKismetEditorUtilities::CompileBlueprint(UBlueprint* BlueprintObj, bool bIs
 		BlueprintObj->NewVariables[VarIndex].DefaultValue.Empty();
 	}
 
-	{ BP_SCOPED_COMPILER_EVENT_NAME(TEXT("Refresh Dependent Blueprints"));
+	{ 
+		BP_SCOPED_COMPILER_EVENT_NAME(TEXT("Refresh Dependent Blueprints"));
 
 		TArray<UBlueprint*> DependentBPs;
 		FBlueprintEditorUtils::GetDependentBlueprints(BlueprintObj, DependentBPs);
diff --git a/Engine/Source/Editor/UnrealEd/Private/Kismet2/KismetReinstanceUtilities.cpp b/Engine/Source/Editor/UnrealEd/Private/Kismet2/KismetReinstanceUtilities.cpp
--- a/Engine/Source/Editor/UnrealEd/Private/Kismet2/KismetReinstanceUtilities.cpp
+++ b/Engine/Source/Editor/UnrealEd/Private/Kismet2/KismetReinstanceUtilities.cpp
@@ -7,12 +7,14 @@
 #include "BlueprintEditorUtils.h"
 #include "Layers/ILayers.h"
 #include "ComponentInstanceDataCache.h"
+#include "Engine/DynamicBlueprintBinding.h"
 
 DEFINE_STAT(EKismetReinstancerStats_ReplaceInstancesOfClass);
 DEFINE_STAT(EKismetReinstancerStats_FindReferencers);
 DEFINE_STAT(EKismetReinstancerStats_ReplaceReferences);
 DEFINE_STAT(EKismetReinstancerStats_UpdateBytecodeReferences);
 DEFINE_STAT(EKismetReinstancerStats_RecompileChildClasses);
+DEFINE_STAT(EKismetReinstancerStats_ReplaceClassNoReinsancing);
 
 /////////////////////////////////////////////////////////////////////////////////
 // FBlueprintCompileReinstancer
@@ -23,6 +25,7 @@ FBlueprintCompileReinstancer::FBlueprintCompileReinstancer(UClass* InClassToRein
 	, OriginalCDO(NULL)
 	, bHasReinstanced(false)
 	, bSkipGarbageCollection(bSkipGC)
+	, ClassToReinstanceDefaultValuesCRC(0)
 {
 	if( InClassToReinstance != NULL )
 	{
@@ -106,6 +109,7 @@ FBlueprintCompileReinstancer::FBlueprintCompileReinstancer(UClass* InClassToRein
 		check(GeneratingBP || GIsAutomationTesting);
 		if(GeneratingBP)
 		{
+			ClassToReinstanceDefaultValuesCRC = GeneratingBP->CrcPreviousCompiledCDO;
 			FBlueprintEditorUtils::GetDependentBlueprints(GeneratingBP, Dependencies);
 		}
 	}
@@ -156,7 +160,7 @@ FBlueprintCompileReinstancer::~FBlueprintCompileReinstancer()
 	}
 }
 
-void FBlueprintCompileReinstancer::ReinstanceObjects()
+void FBlueprintCompileReinstancer::ReinstanceObjects(bool bAlwaysReinstance)
 {
 	BP_SCOPED_COMPILER_EVENT_NAME(TEXT("Reinstance Objects"));
 
@@ -167,11 +171,51 @@ void FBlueprintCompileReinstancer::ReinstanceObjects()
 	}
 	bHasReinstanced = true;
 
-	if( ClassToReinstance && DuplicatedClass )
+	if (ClassToReinstance && DuplicatedClass)
 	{
-		ReplaceInstancesOfClass(DuplicatedClass, ClassToReinstance, OriginalCDO);
+		bool bShouldReinstance = true;
+		if (!bAlwaysReinstance)
+		{
+			BP_SCOPED_COMPILER_EVENT_STAT(EKismetReinstancerStats_ReplaceClassNoReinsancing);
+
+			auto BPClassA = Cast<const UBlueprintGeneratedClass>(DuplicatedClass);
+			auto BPClassB = Cast<const UBlueprintGeneratedClass>(ClassToReinstance);
+			auto BP = Cast<const UBlueprint>(ClassToReinstance->ClassGeneratedBy);
+			const bool bTheSameDefaultValues = BP && ClassToReinstanceDefaultValuesCRC && (BP->CrcPreviousCompiledCDO == ClassToReinstanceDefaultValuesCRC);
+			const bool bTheSame = bTheSameDefaultValues && BPClassA && BPClassB && FStructUtils::TheSameLayout(BPClassA, BPClassB, true);
+			if (bTheSame)
+			{
+				UE_LOG(LogBlueprint, Log, TEXT("BlueprintCompileReinstancer: class '%s' is replaced without reinstancing (the optimized way)."), *GetPathNameSafe(ClassToReinstance));
+
+				TArray<UObject*> ObjectsToReplace;
+				GetObjectsOfClass(DuplicatedClass, ObjectsToReplace, false);
+
+				const bool bIsActor = ClassToReinstance->IsChildOf<AActor>();
+				for (auto Obj : ObjectsToReplace)
+				{
+					if (!Obj->IsTemplate() && !Obj->IsPendingKill())
+					{
+						Obj->SetClass(ClassToReinstance);
+						if (bIsActor)
+						{
+							auto Actor = CastChecked<AActor>(Obj);
+							Actor->ReregisterAllComponents();
+							Actor->RerunConstructionScripts();
+						}
+					}
+				}
+
+				bShouldReinstance = false;
+			}
+		}
+
+		if (bShouldReinstance)
+		{
+			ReplaceInstancesOfClass(DuplicatedClass, ClassToReinstance, OriginalCDO);
+		}
 	
-		{ BP_SCOPED_COMPILER_EVENT_STAT(EKismetReinstancerStats_RecompileChildClasses);
+		{ 
+			BP_SCOPED_COMPILER_EVENT_STAT(EKismetReinstancerStats_RecompileChildClasses);
 
 			// Reparent all dependent blueprints, and recompile to ensure that they get reinstanced with the new memory layout
 			for( auto ChildBP = Children.CreateIterator(); ChildBP; ++ChildBP)
diff --git a/Engine/Source/Editor/UnrealEd/Public/Kismet2/KismetReinstanceUtilities.h b/Engine/Source/Editor/UnrealEd/Public/Kismet2/KismetReinstanceUtilities.h
--- a/Engine/Source/Editor/UnrealEd/Public/Kismet2/KismetReinstanceUtilities.h
+++ b/Engine/Source/Editor/UnrealEd/Public/Kismet2/KismetReinstanceUtilities.h
@@ -15,6 +15,7 @@ DECLARE_CYCLE_STAT_EXTERN(TEXT("Find Referencers"), EKismetReinstancerStats_Find
 DECLARE_CYCLE_STAT_EXTERN(TEXT("Replace References"), EKismetReinstancerStats_ReplaceReferences, STATGROUP_KismetReinstancer, );
 DECLARE_CYCLE_STAT_EXTERN(TEXT("Update Bytecode References"), EKismetReinstancerStats_UpdateBytecodeReferences, STATGROUP_KismetReinstancer, );
 DECLARE_CYCLE_STAT_EXTERN(TEXT("Recompile Child Classes"), EKismetReinstancerStats_RecompileChildClasses, STATGROUP_KismetReinstancer, );
+DECLARE_CYCLE_STAT_EXTERN(TEXT("Replace Classes Without Reinstancing"), EKismetReinstancerStats_ReplaceClassNoReinsancing, STATGROUP_KismetReinstancer, );
 
 class UNREALED_API FBlueprintCompileReinstancer
 {
@@ -45,6 +46,8 @@ protected:
 	/** Don't call GC */
 	bool bSkipGarbageCollection;
 
+	uint32 ClassToReinstanceDefaultValuesCRC;
+
 public:
 	virtual ~FBlueprintCompileReinstancer();
 
@@ -58,7 +61,7 @@ public:
 	void GenerateFieldMappings(TMap<UObject*, UObject*>& FieldMapping);
 
 	/** Reinstances all objects in the ObjectReinstancingMap */
-	void ReinstanceObjects();
+	void ReinstanceObjects(bool bAlwaysReinstance = true);
 
 	/** Updates references to properties and functions of the class that has in the bytecode of dependent blueprints */
 	void UpdateBytecodeReferences();
diff --git a/Engine/Source/Runtime/Engine/Classes/Engine/Blueprint.h b/Engine/Source/Runtime/Engine/Classes/Engine/Blueprint.h
--- a/Engine/Source/Runtime/Engine/Classes/Engine/Blueprint.h
+++ b/Engine/Source/Runtime/Engine/Classes/Engine/Blueprint.h
@@ -392,6 +392,7 @@ class ENGINE_API UBlueprint : public UBlueprintCore
 #endif // WITH_EDITORONLY_DATA
 
 public:
+
 	/** Broadcasts a notification whenever the blueprint has changed. */
 	DECLARE_EVENT_OneParam( UBlueprint, FChangedEvent, class UBlueprint* );
 	FChangedEvent& OnChanged() { return ChangedEvent; }
@@ -414,9 +415,11 @@ public:
 	/** The blueprint is currently compiled */
 	UPROPERTY(transient)
 	uint32 bBeingCompiled:1;
-#endif // WITH_EDITORONLY_DATA
 
-#if WITH_EDITORONLY_DATA
+	/** CRC for CDO calculated right after the latest compilation */
+	UPROPERTY(transient, duplicatetransient)
+	uint32 CrcPreviousCompiledCDO;
+
 	bool bCachedDependenciesUpToDate;
 	TSet<TWeakObjectPtr<UBlueprint>> CachedDependencies;
 
diff --git a/Engine/Source/Runtime/Engine/Classes/Engine/BlueprintGeneratedClass.h b/Engine/Source/Runtime/Engine/Classes/Engine/BlueprintGeneratedClass.h
--- a/Engine/Source/Runtime/Engine/Classes/Engine/BlueprintGeneratedClass.h
+++ b/Engine/Source/Runtime/Engine/Classes/Engine/BlueprintGeneratedClass.h
@@ -161,11 +161,11 @@ private:
 struct FStructUtils
 {
 private:
-	static bool ArePropertiesTheSame(const UProperty* A, const UProperty* B);
+	static bool ArePropertiesTheSame(const UProperty* A, const UProperty* B, bool bCheckPropertiesNames);
 
 public:
 	// does structures have exactly the same memory layout
-	ENGINE_API static bool TheSameLayout(const UStruct* StructA, const UStruct* StructB);
+	ENGINE_API static bool TheSameLayout(const UStruct* StructA, const UStruct* StructB, bool bCheckPropertiesNames = false);
 };
 
 USTRUCT()
diff --git a/Engine/Source/Runtime/Engine/Private/BlueprintGeneratedClass.cpp b/Engine/Source/Runtime/Engine/Private/BlueprintGeneratedClass.cpp
--- a/Engine/Source/Runtime/Engine/Private/BlueprintGeneratedClass.cpp
+++ b/Engine/Source/Runtime/Engine/Private/BlueprintGeneratedClass.cpp
@@ -77,7 +77,7 @@ UClass* UBlueprintGeneratedClass::GetAuthoritativeClass()
 	return (GeneratingBP->GeneratedClass != NULL) ? GeneratingBP->GeneratedClass : this;
 }
 
-bool FStructUtils::ArePropertiesTheSame(const UProperty* A, const UProperty* B)
+bool FStructUtils::ArePropertiesTheSame(const UProperty* A, const UProperty* B, bool bCheckPropertiesNames)
 {
 	if (A == B)
 	{
@@ -89,6 +89,11 @@ bool FStructUtils::ArePropertiesTheSame(const UProperty* A, const UProperty* B)
 		return false;
 	}
 
+	if (bCheckPropertiesNames && (A->GetFName() != B->GetFName()))
+	{
+		return false;
+	}
+
 	if (A->GetSize() != B->GetSize())
 	{
 		return false;
@@ -107,7 +112,7 @@ bool FStructUtils::ArePropertiesTheSame(const UProperty* A, const UProperty* B)
 	return true;
 }
 
-bool FStructUtils::TheSameLayout(const UStruct* StructA, const UStruct* StructB)
+bool FStructUtils::TheSameLayout(const UStruct* StructA, const UStruct* StructB, bool bCheckPropertiesNames)
 {
 	bool bResult = false;
 	if (StructA && StructB)
@@ -116,15 +121,11 @@ bool FStructUtils::TheSameLayout(const UStruct* StructA, const UStruct* StructB)
 		const UProperty* PropertyB = StructB->PropertyLink;
 
 		bResult = true;
-		while (bResult)
+		while (bResult && (PropertyA != PropertyB))
 		{
-			bResult = ArePropertiesTheSame(PropertyA, PropertyB);
+			bResult = ArePropertiesTheSame(PropertyA, PropertyB, bCheckPropertiesNames);
 			PropertyA = PropertyA ? PropertyA->PropertyLinkNext : NULL;
 			PropertyB = PropertyB ? PropertyB->PropertyLinkNext : NULL;
-			if (PropertyA == PropertyB)
-			{
-				break;
-			}
 		}
 	}
 	return bResult;
