Eliminate malformed html in Blueprints API due to < and > symbols

diff --git a/Engine/Source/Editor/UnrealEd/Private/Commandlets/GenerateBlueprintAPICommandlet.cpp b/Engine/Source/Editor/UnrealEd/Private/Commandlets/GenerateBlueprintAPICommandlet.cpp
--- a/Engine/Source/Editor/UnrealEd/Private/Commandlets/GenerateBlueprintAPICommandlet.cpp
+++ b/Engine/Source/Editor/UnrealEd/Private/Commandlets/GenerateBlueprintAPICommandlet.cpp
@@ -233,6 +233,9 @@ FString MakeJsonString( const FString& String )
 	OutString = OutString.Replace(TEXT("\xF7"), TEXT("&divide;"));
 	OutString = OutString.Replace(TEXT("\x2022"), TEXT("&middot;"));
 
+	OutString = OutString.Replace(TEXT("<"), TEXT("&lt;"));
+	OutString = OutString.Replace(TEXT(">"), TEXT("&gt;"));
+
 	return OutString;
 }
 
