Fixed copying array and struct properties from PIE to Editor actors, so that references to other PIE actors are fixed up to point to the corresponding Editor counterpart.

diff --git a/Engine/Source/Editor/UnrealEd/Private/Editor.cpp b/Engine/Source/Editor/UnrealEd/Private/Editor.cpp
--- a/Engine/Source/Editor/UnrealEd/Private/Editor.cpp
+++ b/Engine/Source/Editor/UnrealEd/Private/Editor.cpp
@@ -7075,7 +7075,7 @@ namespace EditorUtilities
 	}
 
 
-	void CopySingleActorProperty( UObject* const InSourceObject, UObject* const InTargetObject, UProperty* const InProperty )
+	void CopySingleActorPropertyRecursive(const void* const InSourcePtr, void* const InTargetPtr, UObject* const InTargetObject, UProperty* const InProperty)
 	{
 		// Properties that are *object* properties are tricky
 		// Sometimes the object will be a reference to a PIE-world object, and copying that reference back to an actor CDO asset is not a good idea
@@ -7084,48 +7084,91 @@ namespace EditorUtilities
 		bool bNeedsGenericCopy = true;
 		if( UObjectPropertyBase* const ObjectProperty = Cast<UObjectPropertyBase>(InProperty) )
 		{
-			UObject* const SourceObjectPropertyValue = ObjectProperty->GetObjectPropertyValue_InContainer(InSourceObject);
-			if( SourceObjectPropertyValue && SourceObjectPropertyValue->GetOutermost()->PackageFlags & PKG_PlayInEditor )
+			const int32 PropertyArrayDim = InProperty->ArrayDim;
+			for (int32 ArrayIndex = 0; ArrayIndex < PropertyArrayDim; ArrayIndex++)
 			{
-				// Not all the code paths below actually copy the object, but even if they don't we need to claim that they
-				// did, as copying a reference to an object in a PIE world leads to crashes
-				bNeedsGenericCopy = false;
-
-				// REFERENCE an existing actor in the editor world from a REFERENCE in the PIE world
-				if( SourceObjectPropertyValue->IsA(AActor::StaticClass()) )
+				UObject* const SourceObjectPropertyValue = ObjectProperty->GetObjectPropertyValue_InContainer(InSourcePtr, ArrayIndex);
+				if (SourceObjectPropertyValue && SourceObjectPropertyValue->GetOutermost()->PackageFlags & PKG_PlayInEditor)
 				{
-					// We can try and fix-up an actor reference from the PIE world to instead be the version from the persistent world
-					AActor* const EditorWorldActor = GetEditorWorldCounterpartActor(Cast<AActor>(SourceObjectPropertyValue));
-					if( EditorWorldActor )
+					// Not all the code paths below actually copy the object, but even if they don't we need to claim that they
+					// did, as copying a reference to an object in a PIE world leads to crashes
+					bNeedsGenericCopy = false;
+
+					// REFERENCE an existing actor in the editor world from a REFERENCE in the PIE world
+					if (SourceObjectPropertyValue->IsA(AActor::StaticClass()))
 					{
-						ObjectProperty->SetObjectPropertyValue_InContainer(InTargetObject, EditorWorldActor);
+						// We can try and fix-up an actor reference from the PIE world to instead be the version from the persistent world
+						AActor* const EditorWorldActor = GetEditorWorldCounterpartActor(Cast<AActor>(SourceObjectPropertyValue));
+						if (EditorWorldActor)
+						{
+							ObjectProperty->SetObjectPropertyValue_InContainer(InTargetPtr, EditorWorldActor, ArrayIndex);
+						}
 					}
-				}
-				// REFERENCE an existing actor component in the editor world from a REFERENCE in the PIE world
-				else if( SourceObjectPropertyValue->IsA(UActorComponent::StaticClass()) && InTargetObject->IsA(AActor::StaticClass()) )
-				{
-					AActor* const TargetActor = Cast<AActor>(InTargetObject);
-					TArray<UActorComponent*> TargetComponents;
-					TargetActor->GetComponents(TargetComponents);
-
-					// We can try and fix-up an actor component reference from the PIE world to instead be the version from the persistent world
-					int32 TargetComponentIndex = 0;
-					UActorComponent* const EditorWorldComponent = FindMatchingComponentInstance(Cast<UActorComponent>(SourceObjectPropertyValue), TargetActor, TargetComponents, TargetComponentIndex);
-					if(EditorWorldComponent)
+					// REFERENCE an existing actor component in the editor world from a REFERENCE in the PIE world
+					else if (SourceObjectPropertyValue->IsA(UActorComponent::StaticClass()) && InTargetObject->IsA(AActor::StaticClass()))
 					{
-						ObjectProperty->SetObjectPropertyValue_InContainer(InTargetObject, EditorWorldComponent);
+						AActor* const TargetActor = Cast<AActor>(InTargetObject);
+						TArray<UActorComponent*> TargetComponents;
+						TargetActor->GetComponents(TargetComponents);
+
+						// We can try and fix-up an actor component reference from the PIE world to instead be the version from the persistent world
+						int32 TargetComponentIndex = 0;
+						UActorComponent* const EditorWorldComponent = FindMatchingComponentInstance(Cast<UActorComponent>(SourceObjectPropertyValue), TargetActor, TargetComponents, TargetComponentIndex);
+						if (EditorWorldComponent)
+						{
+							ObjectProperty->SetObjectPropertyValue_InContainer(InTargetPtr, EditorWorldComponent, ArrayIndex);
+						}
 					}
 				}
 			}
 		}
+		else if (UStructProperty* const StructProperty = Cast<UStructProperty>(InProperty))
+		{
+			const int32 PropertyArrayDim = InProperty->ArrayDim;
+			for (int32 ArrayIndex = 0; ArrayIndex < PropertyArrayDim; ArrayIndex++)
+			{
+				const void* const SourcePtr = StructProperty->ContainerPtrToValuePtr<void>(InSourcePtr, ArrayIndex);
+				void* const TargetPtr = StructProperty->ContainerPtrToValuePtr<void>(InTargetPtr, ArrayIndex);
+
+				for (TFieldIterator<UProperty> It(StructProperty->Struct); It; ++It)
+				{
+					UProperty* const InnerProperty = *It;
+					CopySingleActorPropertyRecursive(SourcePtr, TargetPtr, InTargetObject, InnerProperty);
+				}
+			}
+
+			bNeedsGenericCopy = false;
+		}
+		else if (UArrayProperty* const ArrayProperty = Cast<UArrayProperty>(InProperty))
+		{
+			check(InProperty->ArrayDim == 1);
+			FScriptArrayHelper SourceArrayHelper(ArrayProperty, ArrayProperty->ContainerPtrToValuePtr<void>(InSourcePtr));
+			FScriptArrayHelper TargetArrayHelper(ArrayProperty, ArrayProperty->ContainerPtrToValuePtr<void>(InTargetPtr));
+
+			UProperty* InnerProperty = ArrayProperty->Inner;
+			int32 Num = SourceArrayHelper.Num();
+
+			TargetArrayHelper.EmptyAndAddUninitializedValues(Num);
+
+			for (int32 Index = 0; Index < Num; Index++)
+			{
+				CopySingleActorPropertyRecursive(SourceArrayHelper.GetRawPtr(Index), TargetArrayHelper.GetRawPtr(Index), InTargetObject, InnerProperty);
+			}
+
+			bNeedsGenericCopy = false;
+		}
 		
 		// Handle copying properties that either aren't an object, or aren't part of the PIE world
 		if( bNeedsGenericCopy )
 		{
-			InProperty->CopyCompleteValue_InContainer(InTargetObject, InSourceObject);
+			InProperty->CopyCompleteValue_InContainer(InTargetPtr, InSourcePtr);
 		}
 	}
 
+	void CopySingleActorProperty(const UObject* const InSourceObject, UObject* const InTargetObject, UProperty* const InProperty)
+	{
+		CopySingleActorPropertyRecursive(InSourceObject, InTargetObject, InTargetObject, InProperty);
+	}
 
 	int32 CopyActorProperties( AActor* SourceActor, AActor* TargetActor, const ECopyOptions::Type Options )
 	{
