Changed some shader compilation code so that it processes all shadermaps even if they aren't being applied in ProcessAsyncResults.
Improved -iterate.

diff --git a/Engine/Source/Editor/UnrealEd/Classes/CookOnTheSide/CookOnTheFlyServer.h b/Engine/Source/Editor/UnrealEd/Classes/CookOnTheSide/CookOnTheFlyServer.h
--- a/Engine/Source/Editor/UnrealEd/Classes/CookOnTheSide/CookOnTheFlyServer.h
+++ b/Engine/Source/Editor/UnrealEd/Classes/CookOnTheSide/CookOnTheFlyServer.h
@@ -21,12 +21,21 @@ enum class ECookInitializationFlags
 	Unversioned = 0x8,				// save the cooked packages without a version number
 	AutoTick = 0x10,				// enable ticking (only works in the editor)
 	AsyncSave = 0x20,				// save packages async
-	LeakTest = 0x40,				// test for uobject leaks after each level load
 	IncludeServerMaps = 0x80,		// should we include the server maps when cooking
 	GenerateStreamingInstallManifest = 0x100,  // should we generate streaming install manifest
 };
 ENUM_CLASS_FLAGS(ECookInitializationFlags);
 
+enum class ECookByTheBookOptions
+{
+	None = 0x0,					// no flags
+	CookAll = 0x1,				// cook all maps and content in the content directory
+	MapsOnly = 0x2,				// cook only maps
+	NoDevContent = 0x4,			// don't include dev content
+	LeakTest = 0x8,				// test for uobject leaks after each level load
+};
+ENUM_CLASS_FLAGS(ECookByTheBookOptions);
+
 UENUM()
 namespace ECookMode
 {
@@ -572,6 +581,8 @@ private:
 			CookStartTime( 0.0 )
 		{ }
 
+		/** Should we test for UObject leaks */
+		bool bLeakTest;
 		/** Should we generate streaming install manifests (only valid option in cook by the book) */
 		bool bGenerateStreamingInstallManifests;
 		/** Is cook by the book currently running */
@@ -660,7 +671,10 @@ public:
 	 * Start a cook by the book session
 	 * Cook on the fly can't run at the same time as cook by the book
 	 */
-	void StartCookByTheBook(const TArray<ITargetPlatform*>& TargetPlatforms, const TArray<FString>& CookMaps, const TArray<FString>& CookDirectories, const TArray<FString>& CookCultures, const TArray<FString>& IniMapSections );
+	void StartCookByTheBook(const TArray<ITargetPlatform*>& TargetPlatforms, 
+		const TArray<FString>& CookMaps, const TArray<FString>& CookDirectories, 
+		const TArray<FString>& CookCultures, const TArray<FString>& IniMapSections, 
+		ECookByTheBookOptions CookOptions = ECookByTheBookOptions::None );
 
 	/**
 	 * Queue a cook by the book cancel (you might want to do this instead of calling cancel directly so that you don't have to be in the game thread when canceling
@@ -741,6 +755,12 @@ private:
 	void CollectFilesToCook(TArray<FString>& FilesInPath, 
 		const TArray<FString>& CookMaps, const TArray<FString>& CookDirectories, const TArray<FString>& CookCultures, 
 		const TArray<FString>& IniMapSections, bool bCookAll, bool bMapsOnly, bool bNoDev);
+
+	/**
+	 * AddFileToCook add file to cook list 
+	 */
+	void AddFileToCook( TArray<FString>& InOutFilesToCook, const FString &InFilename ) const;
+
 	/**
 	 * Call back from the TickCookOnTheSide when a cook by the book finishes (when started form StartCookByTheBook)
 	 */
@@ -756,20 +776,6 @@ private:
 		return CurrentCookMode == ECookMode::CookByTheBookFromTheEditor || CurrentCookMode == ECookMode::CookByTheBook; 
 	}
 
-	/**
-	 * GetDependencies
-	 * 
-	 * @param Packages List of packages to use as the root set for dependency checking
-	 * @param Found return value, all objects which package is dependent on
-	 */
-	void GetDependencies( const TSet<UPackage*>& Packages, TSet<UObject*>& Found);
-	/**
-	 * GenerateManifestInfo
-	 * generate the manifest information for a given package
-	 *
-	 * @param Package package to generate manifest information for
-	 */
-	void GenerateManifestInfo( UPackage* Package, const TArray<FName>& TargetPlatformNames );
 
 	//////////////////////////////////////////////////////////////////////////
 	// cook on the fly specific functions
@@ -803,7 +809,59 @@ private:
 	 */
 	bool ShouldCook(const FString& InFileName, const FName& InPlatformName);
 
+	/**
+	 * GetDependencies
+	 * 
+	 * @param Packages List of packages to use as the root set for dependency checking
+	 * @param Found return value, all objects which package is dependent on
+	 */
+	void GetDependencies( const TSet<UPackage*>& Packages, TSet<UObject*>& Found);
+	/**
+	 * GenerateManifestInfo
+	 * generate the manifest information for a given package
+	 *
+	 * @param Package package to generate manifest information for
+	 */
+	void GenerateManifestInfo( UPackage* Package, const TArray<FName>& TargetPlatformNames );
+
+	/**
+	 * GetCurrentIniVersionStrings gets the current ini version strings for compare against previous cook
+	 * 
+	 * @param IniVersionStrings return list of the important current ini version strings
+	 * @return false if function fails (should assume all platforms are out of date)
+	 */
+	bool GetCurrentIniVersionStrings( const ITargetPlatform* TargetPlatform, TArray<FString> &IniVersionStrings ) const;
 
+	/**
+	 * GetCookedIniVersionStrings gets the ini version strings used in previous cook for specified target platform
+	 * 
+	 * @param IniVersionStrings return list of the previous cooks ini version strings
+	 * @return false if function fails to find the ini version strings
+	 */
+	bool GetCookedIniVersionStrings( const ITargetPlatform* TargetPlatform, TArray<FString>& IniVersionStrings ) const;
+
+	/**
+	 * Checks if important ini settings have changed since last cook for each target platform 
+	 * 
+	 * @param TargetPlatforms to check if out of date
+	 * @param OutOfDateTargetPlatforms return list of out of date target platforms which should be cleaned
+	 */
+	bool IniSettingsOutOfDate( const TArray<ITargetPlatform*>& TargetPlatforms, TArray<ITargetPlatform*>& OutOfDateTargetPlatforms ) const;
+
+	/**
+	 * SaveIniVersionStrings save out the ini version strings for all platforms so next cook can use them
+	 * -iterate uses this to figure out if content needs to be recooked
+	 * 
+	 * @param TargetPlatforms the target platforms to generate ini version info for
+	 */
+	void SaveIniVersionStrings( const TArray<ITargetPlatform*>& TargetPlatforms ) const;
+
+	/**
+	 * IsCookFlagSet
+	 * 
+	 * @param InCookFlag used to check against the current cook flags
+	 * @return true if the cook flag is set false otherwise
+	 */
 	bool IsCookFlagSet( const ECookInitializationFlags& InCookFlags ) const 
 	{
 		return (CookFlags & InCookFlags) != ECookInitializationFlags::None;
diff --git a/Engine/Source/Editor/UnrealEd/Private/Commandlets/CookCommandlet.cpp b/Engine/Source/Editor/UnrealEd/Private/Commandlets/CookCommandlet.cpp
--- a/Engine/Source/Editor/UnrealEd/Private/Commandlets/CookCommandlet.cpp
+++ b/Engine/Source/Editor/UnrealEd/Private/Commandlets/CookCommandlet.cpp
@@ -1114,12 +1114,21 @@ bool UCookCommandlet::NewCook( const TArray<ITargetPlatform*>& Platforms, TArray
 	//////////////////////////////////////////////////////////////////////////
 	// start cook by the book 
 
+
+
+	ECookByTheBookOptions CookOptions = ECookByTheBookOptions::None;
+
+	CookOptions |= bLeakTest ? ECookByTheBookOptions::LeakTest : ECookByTheBookOptions::None; 
+	CookOptions |= bCookAll ? ECookByTheBookOptions::CookAll : ECookByTheBookOptions::None;
+	CookOptions |= Switches.Contains(TEXT("MAPSONLY")) ? ECookByTheBookOptions::MapsOnly : ECookByTheBookOptions::None;
+	CookOptions |= Switches.Contains(TEXT("NODEV")) ? ECookByTheBookOptions::NoDevContent : ECookByTheBookOptions::None;
+
 	for ( const auto& MapName : CmdLineMapEntries )
 	{
 		MapList.Add( MapName );
 	}
 
-	CookOnTheFlyServer->StartCookByTheBook(Platforms, MapList, CmdLineDirEntries, CmdLineCultEntries, CmdLineIniSections );
+	CookOnTheFlyServer->StartCookByTheBook(Platforms, MapList, CmdLineDirEntries, CmdLineCultEntries, CmdLineIniSections, CookOptions );
 
 	// Garbage collection should happen when either
 	//	1. We have cooked a map
diff --git a/Engine/Source/Editor/UnrealEd/Private/CookOnTheFlyServer.cpp b/Engine/Source/Editor/UnrealEd/Private/CookOnTheFlyServer.cpp
--- a/Engine/Source/Editor/UnrealEd/Private/CookOnTheFlyServer.cpp
+++ b/Engine/Source/Editor/UnrealEd/Private/CookOnTheFlyServer.cpp
@@ -35,7 +35,7 @@
 DEFINE_LOG_CATEGORY_STATIC(LogCookOnTheFly, Log, All);
 
 #define DEBUG_COOKONTHEFLY 0
-#define OUTPUT_TIMING 1
+#define OUTPUT_TIMING 0
 
 #if OUTPUT_TIMING
 
@@ -1054,12 +1054,9 @@ uint32 UCookOnTheFlyServer::TickCookOnTheSide( const float TimeSlice, uint32 &Co
 				}
 
 				if ( bIsAllDataCached == false )
+				{
 					break;
 			}
-
-			if ( !bIsAllDataCached )
-			{
-				GShaderCompilingManager->ProcessAsyncResults(true, false, true);
 			}
 		}
 
@@ -1165,7 +1162,7 @@ uint32 UCookOnTheFlyServer::TickCookOnTheSide( const float TimeSlice, uint32 &Co
 			GenerateManifestInfo( PackagesToSave[0], AllTargetPlatformNames );
 		}
 
-
+		bool bFinishedSave = true;
 
 		if ( PackagesToSave.Num() )
 		{
@@ -1205,6 +1202,7 @@ uint32 UCookOnTheFlyServer::TickCookOnTheSide( const float TimeSlice, uint32 &Co
 						}
 
 						// break out of the loop
+						bFinishedSave = false;
 						break;
 					}
 				}
@@ -1291,26 +1289,23 @@ uint32 UCookOnTheFlyServer::TickCookOnTheSide( const float TimeSlice, uint32 &Co
 		// TODO: Daniel: this is reference code needs to be reimplemented on the callee side.
 		//  cooker can't depend on callee being able to garbage collect
 		// collect garbage
-		/*if (NumProcessedSinceLastGC >= GCInterval)
+		if ( CookByTheBookOptions->bLeakTest && bFinishedSave )
 		{
-		UE_LOG(LogCookCommandlet, Display, TEXT("Full GC..."));
+			check( CurrentCookMode == ECookMode::CookByTheBook);
+			UE_LOG(LogCookOnTheFly, Display, TEXT("Full GC..."));
 
-			CollectGarbage( RF_Native );
-		NumProcessedSinceLastGC = 0;
+		CollectGarbage( RF_Native );
 
-		if (bLeakTest)
+		for (FObjectIterator It; It; ++It)
 		{
-			for (FObjectIterator It; It; ++It)
-			{
-		if (!LastGCItems.Contains(FWeakObjectPtr(*It)))
-				{
-		UE_LOG(LogCookCommandlet, Warning, TEXT("\tLeaked %s"), *(It->GetFullName()));
-		LastGCItems.Add(FWeakObjectPtr(*It));
-				}
-			}
+				if (!CookByTheBookOptions->LastGCItems.Contains(FWeakObjectPtr(*It)))
+		{
+					UE_LOG(LogCookOnTheFly, Warning, TEXT("\tLeaked %s"), *(It->GetFullName()));
+					CookByTheBookOptions->LastGCItems.Add(FWeakObjectPtr(*It));
 		}
-		}*/
-		
+		}
+		}
+
 		if ( Timer.IsTimeUp() )
 		{
 			break;
@@ -1367,9 +1362,19 @@ void UCookOnTheFlyServer::MarkPackageDirtyForCooker( UPackage *Package )
 		FPaths::MakeStandardFilename(PackageFilename);*/
 		FString PackageFilename = GetCachedStandardPackageFilename(Package);
 		UE_LOG(LogCookOnTheFly, Display, TEXT("Modification detected to package %s"), *PackageFilename);
-		CookedPackages.RemoveAll( FName(*PackageFilename) );
+		const FName PackageFFileName = FName(*PackageFilename);
 
-		// TODO: go through package dependencies of this package, if a found dependency is not loaded then mark it as need recooking.
+		if ( CurrentCookMode == ECookMode::CookByTheBookFromTheEditor )
+		{
+			TArray<FName> CookedPlatforms;
+			if ( CookedPackages.GetCookedPlatforms(PackageFFileName, CookedPlatforms) )
+			{
+				// if this package was previously cooked and we are doing a cook by the book 
+				// we need to recook this package before finishing cook by the book
+				CookRequests.EnqueueUnique( FFilePlatformRequest(PackageFFileName, CookedPlatforms) );
+			}
+		}
+		CookedPackages.RemoveAll( PackageFFileName );
 	}
 }
 
@@ -1450,7 +1455,6 @@ bool UCookOnTheFlyServer::ShouldCook(const FString& InFileName, const FName &InP
 {
 	bool bDoCook = false;
 
-
 	FString PkgFile;
 	FString PkgFilename;
 	FDateTime DependentTimeStamp = FDateTime::MinValue();
@@ -1712,6 +1716,9 @@ void UCookOnTheFlyServer::Initialize( ECookMode::Type DesiredCookMode, ECookInit
 
 	CleanSandbox(Platforms);
 
+	// generate version info for use by -iterate to figure out what files need recooking
+	SaveIniVersionStrings(Platforms);
+
 	// always generate the asset registry before starting to cook, for either method
 	GenerateAssetRegistry(Platforms);
 
@@ -1751,6 +1758,138 @@ FString UCookOnTheFlyServer::GetOutputDirectoryOverride( const FString &OutputDi
 	return OutputDirectory;
 }
 
+bool UCookOnTheFlyServer::GetCurrentIniVersionStrings( const ITargetPlatform* TargetPlatform, TArray<FString> &IniVersionStrings ) const
+{
+	// there is a list of important ini settings in the Editor config 
+	TArray<FString> IniVersionedParams;
+	GConfig->GetArray( TEXT("CookSettings"), TEXT("VersionedIniParams"), IniVersionedParams, GEditorIni );
+
+	/*const FString EditorIni = FPaths::GameDir() / GEditorIni;
+	const FString SandboxEditorIni = SandboxFile->ConvertToAbsolutePathForExternalAppForWrite(*EditorIni);*/
+
+	// if the old one doesn't contain all the settings in the new one then we fail this check
+	for ( const auto& IniVersioned : IniVersionedParams )
+	{
+
+		TArray<FString> IniVersionedArray;
+		IniVersioned.ParseIntoArray(&IniVersionedArray, TEXT(":"), false);
+
+		if ( IniVersionedArray.Num() != 3 )
+		{
+			UE_LOG(LogCookOnTheFly, Warning, TEXT("Invalid entry in CookSettings, VersionedIniParams %s"), *IniVersioned);
+			return false;
+		}
+
+		const FString& Filename = IniVersionedArray[0];
+		const FString& Section = IniVersionedArray[1];
+		const FString& Key = IniVersionedArray[2];
+
+		const FString& IniFilename = FPaths::GeneratedConfigDir() / TargetPlatform->IniPlatformName() / Filename + TEXT(".ini");
+
+		// get the value of the entry
+		FString Value;
+		if ( !GConfig->GetString(*Section, *Key, Value, IniFilename) )
+		{
+			UE_LOG(LogCookOnTheFly, Warning, TEXT("Unable to find entry in CookSettings, VersionedIniParams %s, assume default is being used"), *IniVersioned);
+			continue;
+		}
+
+		FString CurrentVersionString = FString::Printf(TEXT("%s:%s:%s:%s"), *Filename, *Section, *Key, *Value );
+
+		IniVersionStrings.Emplace( MoveTemp(CurrentVersionString) );
+	}
+
+
+	const FTextureLODSettings& LodSettings = TargetPlatform->GetTextureLODSettings();
+
+	UEnum* TextureGroupEnum = FindObject<UEnum>( NULL, TEXT("Engine.TextureGroup") );
+	UEnum* TextureMipGenSettingsEnum = FindObject<UEnum>( NULL, TEXT("Engine.TextureMipGenSettings") );
+
+	for ( int I = 0; I < TextureGroup::TEXTUREGROUP_MAX; ++I )
+	{
+		const TextureMipGenSettings& MipGenSettings = LodSettings.GetTextureMipGenSettings((TextureGroup)(I));
+
+		FString VersionString = FString::Printf( TEXT("TextureLODGroup:%s:%s"), *TextureGroupEnum->GetEnumName( I ), *TextureMipGenSettingsEnum->GetEnumName((int32)(MipGenSettings)) );
+
+		IniVersionStrings.Emplace( MoveTemp( VersionString ) );
+	}
+	
+
+
+	return true;
+}
+
+bool UCookOnTheFlyServer::GetCookedIniVersionStrings( const ITargetPlatform* TargetPlatform, TArray<FString>& IniVersionStrings ) const
+{
+
+	const FString EditorIni = FPaths::GameDir() / TEXT("CookedIniVersion.txt");
+	const FString SandboxEditorIni = SandboxFile->ConvertToAbsolutePathForExternalAppForWrite(*EditorIni);
+
+
+	const FString PlatformSandboxEditorIni = SandboxEditorIni.Replace(TEXT("[Platform]"), *TargetPlatform->PlatformName());
+
+	TArray<FString> SavedIniVersionedParams;
+	GConfig->GetArray( TEXT("CookedSettings"), TEXT("VersionedIniParams"), IniVersionStrings, PlatformSandboxEditorIni );
+
+	return true;
+}
+
+
+bool UCookOnTheFlyServer::IniSettingsOutOfDate( const TArray<ITargetPlatform*>& TargetPlatforms, TArray<ITargetPlatform*>& OutOfDateTargetPlatforms ) const
+{
+
+
+	for ( const auto& TargetPlatform : TargetPlatforms )
+	{
+		TArray<FString> CurrentIniVersionStrings;
+		if ( GetCurrentIniVersionStrings(TargetPlatform, CurrentIniVersionStrings) == false )
+		{
+			// previous cook seems half baked... bomb out and recook all the things
+			OutOfDateTargetPlatforms.Add( TargetPlatform);
+			continue;
+		}
+
+		TArray<FString> CookedIniVersionStrings;
+		if ( GetCookedIniVersionStrings( TargetPlatform, CookedIniVersionStrings ) )
+		{
+			// can't even get the cooked version strings for this platform add it to the out of date list
+			OutOfDateTargetPlatforms.Add(TargetPlatform);
+			continue;
+		}
+
+		for ( const auto& CurrentVersionString : CurrentIniVersionStrings ) 
+		{
+			if ( CookedIniVersionStrings.Contains(CurrentVersionString) == false )
+			{
+				OutOfDateTargetPlatforms.Add(TargetPlatform);
+				break;
+			}
+		}
+	}
+
+	return true;
+}
+
+void UCookOnTheFlyServer::SaveIniVersionStrings( const TArray<ITargetPlatform*>& TargetPlatforms ) const
+{
+	const FString CookedIni = FPaths::GameDir() / TEXT("CookedIniVersion.txt");
+	const FString SandboxCookedIni = SandboxFile->ConvertToAbsolutePathForExternalAppForWrite(*CookedIni);
+
+
+	for ( const auto& TargetPlatform : TargetPlatforms )
+	{
+		// there is a list of important ini settings in the Editor config 
+		TArray<FString> IniVersionStrings;
+		GetCurrentIniVersionStrings(TargetPlatform, IniVersionStrings);
+
+		const FString PlatformSandboxCookedIni = SandboxCookedIni.Replace(TEXT("[Platform]"), *TargetPlatform->PlatformName());
+
+		GConfig->SetArray(TEXT("CookSettings"), TEXT("VersionedIniParams"), IniVersionStrings, PlatformSandboxCookedIni);
+
+		GConfig->Flush(false, PlatformSandboxCookedIni);
+	}
+}
+
 void UCookOnTheFlyServer::CleanSandbox(const TArray<ITargetPlatform*>& Platforms)
 {
 	double SandboxCleanTime = 0.0;
@@ -1769,8 +1908,18 @@ void UCookOnTheFlyServer::CleanSandbox(const TArray<ITargetPlatform*>& Platforms
 		}
 		else
 		{
-			FPackageDependencyInfoModule& PDInfoModule = FModuleManager::LoadModuleChecked<FPackageDependencyInfoModule>("PackageDependencyInfo");
+			TArray<ITargetPlatform*> OutOfDatePlatforms;
+			IniSettingsOutOfDate(Platforms, OutOfDatePlatforms);
+			// wipe entire directory for any platforms which are out of date
+			for ( const auto& Target : OutOfDatePlatforms )
+			{
+				FString SandboxDirectory = GetOutputDirectory(Target->PlatformName());
+				IFileManager::Get().DeleteDirectory(*SandboxDirectory, false, true);
+			}
 
+			// check each package to see if it's out of date
+			FPackageDependencyInfoModule& PDInfoModule = FModuleManager::LoadModuleChecked<FPackageDependencyInfoModule>("PackageDependencyInfo");
+			
 			// list of directories to skip
 			TArray<FString> DirectoriesToSkip;
 			TArray<FString> DirectoriesToNotRecurse;
@@ -1784,7 +1933,7 @@ void UCookOnTheFlyServer::CleanSandbox(const TArray<ITargetPlatform*>& Platforms
 				// use the timestamp grabbing visitor
 				IPlatformFile& PlatformFile = FPlatformFileManager::Get().GetPlatformFile();
 				FLocalTimestampDirectoryVisitor Visitor(PlatformFile, DirectoriesToSkip, DirectoriesToNotRecurse, false);
-
+			
 				PlatformFile.IterateDirectory(*SandboxDirectory, Visitor);
 
 				for (TMap<FString, FDateTime>::TIterator TimestampIt(Visitor.FileTimes); TimestampIt; ++TimestampIt)
@@ -1891,6 +2040,14 @@ void UCookOnTheFlyServer::GenerateLongPackageNames(TArray<FString>& FilesInPath)
 	Exchange(FilesInPathReverse, FilesInPath);
 }
 
+void UCookOnTheFlyServer::AddFileToCook( TArray<FString>& InOutFilesToCook, const FString &InFilename ) const
+{ 
+	if (!FPackageName::IsScriptPackage(InFilename))
+	{
+		InOutFilesToCook.AddUnique(InFilename);
+	}
+}
+
 void UCookOnTheFlyServer::CollectFilesToCook(TArray<FString>& FilesInPath, const TArray<FString> &CookMaps, const TArray<FString> &InCookDirectories, const TArray<FString> &CookCultures, const TArray<FString> &IniMapSections, bool bCookAll, bool bMapsOnly, bool bNoDev)
 {
 	TArray<FString> MapList;
@@ -1907,7 +2064,7 @@ void UCookOnTheFlyServer::CollectFilesToCook(TArray<FString>& FilesInPath, const
 	/*GEditor->ParseMapSectionIni(*Params, MapList);
 	for (int32 MapIdx = 0; MapIdx < MapList.Num(); MapIdx++)
 	{
-		FilesInPath.AddUnique(MapList[MapIdx]);
+		AddFileToCook( FilesInPath, MapList[MapIdx]);
 	}*/
 
 
@@ -1933,12 +2090,12 @@ void UCookOnTheFlyServer::CollectFilesToCook(TArray<FString>& FilesInPath, const
 			}
 			else
 			{
-				FilesInPath.AddUnique(OutFilename);
+				AddFileToCook( FilesInPath, OutFilename);
 			}
 		}
 		else
 		{
-			FilesInPath.AddUnique(CurrEntry);
+			AddFileToCook( FilesInPath,CurrEntry);
 		}
 	}
 
@@ -1951,7 +2108,7 @@ void UCookOnTheFlyServer::CollectFilesToCook(TArray<FString>& FilesInPath, const
 		{
 			FString StdFile = Files[Index];
 			FPaths::MakeStandardFilename(StdFile);
-			FilesInPath.AddUnique(StdFile);
+			AddFileToCook( FilesInPath,StdFile);
 
 			// this asset may not be in our currently mounted content directories, so try to mount a new one now
 			FString LongPackageName;
@@ -1993,7 +2150,7 @@ void UCookOnTheFlyServer::CollectFilesToCook(TArray<FString>& FilesInPath, const
 
 			for (int32 TokenFileIndex = 0; TokenFileIndex < TokenFiles.Num(); ++TokenFileIndex)
 			{
-				FilesInPath.AddUnique(TokenFiles[TokenFileIndex]);
+				AddFileToCook( FilesInPath, TokenFiles[TokenFileIndex]);
 			}
 		}
 	}
@@ -2013,7 +2170,7 @@ void UCookOnTheFlyServer::CollectFilesToCook(TArray<FString>& FilesInPath, const
 		{
 			if (Obj != FName(NAME_None).ToString())
 			{
-				FilesInPath.AddUnique(Obj);
+				AddFileToCook( FilesInPath, Obj);
 			}
 		}
 		if ( IsCookFlagSet(ECookInitializationFlags::IncludeServerMaps) )
@@ -2022,7 +2179,7 @@ void UCookOnTheFlyServer::CollectFilesToCook(TArray<FString>& FilesInPath, const
 			{
 				if (Obj != FName(NAME_None).ToString())
 				{
-					FilesInPath.AddUnique(Obj);
+					AddFileToCook( FilesInPath, Obj);
 				}
 			}
 		}
@@ -2030,14 +2187,14 @@ void UCookOnTheFlyServer::CollectFilesToCook(TArray<FString>& FilesInPath, const
 		{
 			if (Obj != FName(NAME_None).ToString())
 			{
-				FilesInPath.AddUnique(Obj);
+				AddFileToCook( FilesInPath, Obj);
 			}
 		}
 		if (PlatformEngineIni.GetString(TEXT("/Script/EngineSettings.GameMapsSettings"), TEXT("GlobalDefaultServerGameMode"), Obj))
 		{
 			if (Obj != FName(NAME_None).ToString())
 			{
-				FilesInPath.AddUnique(Obj);
+				AddFileToCook( FilesInPath, Obj);
 			}
 		}
 	}
@@ -2051,7 +2208,7 @@ void UCookOnTheFlyServer::CollectFilesToCook(TArray<FString>& FilesInPath, const
 	{
 		if (InterfaceFile != TEXT("None") && InterfaceFile != TEXT(""))
 		{
-			FilesInPath.AddUnique(InterfaceFile);
+			AddFileToCook( FilesInPath, InterfaceFile);
 		}
 	}
 
@@ -2074,7 +2231,7 @@ void UCookOnTheFlyServer::CollectFilesToCook(TArray<FString>& FilesInPath, const
 				{
 					FString StdFile = Files[Index];
 					FPaths::MakeStandardFilename(StdFile);
-					FilesInPath.AddUnique(StdFile);
+					AddFileToCook( FilesInPath, StdFile);
 				}
 			}
 		}
@@ -2191,7 +2348,10 @@ void UCookOnTheFlyServer::StopAndClearCookedData()
 	CookedPackages.Empty(); // set of files which have been cooked when needing to recook a file the entry will need to be removed from here
 }
 
-void UCookOnTheFlyServer::StartCookByTheBook(const TArray<ITargetPlatform*>& TargetPlatforms, const TArray<FString> &CookMaps, const TArray<FString> &CookDirectories, const TArray<FString> &CookCultures, const TArray<FString> &IniMapSections )
+void UCookOnTheFlyServer::StartCookByTheBook(const TArray<ITargetPlatform*>& TargetPlatforms, 
+											 const TArray<FString> &CookMaps, const TArray<FString> &CookDirectories, 
+											 const TArray<FString> &CookCultures, const TArray<FString> &IniMapSections, 
+											 ECookByTheBookOptions CookOptions )
 {
 	check( IsInGameThread() );
 
@@ -2231,15 +2391,16 @@ void UCookOnTheFlyServer::StartCookByTheBook(const TArray<ITargetPlatform*>& Tar
 	CookByTheBookOptions->CookTime = 0.0f;
 	CookByTheBookOptions->CookStartTime = FPlatformTime::Seconds();
 
-	bool bCookAll = false;
-	bool bMapsOnly = false;
-	bool bNoDev = false;
-	
-	bool bLeakTest = IsCookFlagSet(ECookInitializationFlags::LeakTest); // this won't work from the editor this needs to be standalone
-	check( !bLeakTest || CurrentCookMode == ECookMode::CookByTheBook );
+	bool bCookAll = (CookOptions & ECookByTheBookOptions::CookAll) != ECookByTheBookOptions::None;
+	bool bMapsOnly = (CookOptions & ECookByTheBookOptions::MapsOnly) != ECookByTheBookOptions::None;
+	bool bNoDev = (CookOptions & ECookByTheBookOptions::NoDevContent) != ECookByTheBookOptions::None;
+
+	CookByTheBookOptions->bLeakTest = (CookOptions & ECookByTheBookOptions::LeakTest) != ECookByTheBookOptions::None; // this won't work from the editor this needs to be standalone
+	check( !CookByTheBookOptions->bLeakTest || CurrentCookMode == ECookMode::CookByTheBook );
+
 
 	CookByTheBookOptions->LastGCItems.Empty();
-	if (bLeakTest)
+	if (CookByTheBookOptions->bLeakTest )
 	{
 		for (FObjectIterator It; It; ++It)
 		{
diff --git a/Engine/Source/Runtime/Engine/Private/ShaderCompiler/ShaderCompiler.cpp b/Engine/Source/Runtime/Engine/Private/ShaderCompiler/ShaderCompiler.cpp
--- a/Engine/Source/Runtime/Engine/Private/ShaderCompiler/ShaderCompiler.cpp
+++ b/Engine/Source/Runtime/Engine/Private/ShaderCompiler/ShaderCompiler.cpp
@@ -1894,7 +1894,7 @@ void FShaderCompilingManager::FinishAllCompilation()
 	UE_LOG(LogShaders, Log, TEXT("FinishAllCompilation %.3fs"), (float)(EndTime - StartTime));
 }
 
-void FShaderCompilingManager::ProcessAsyncResults(bool bLimitExecutionTime, bool bBlockOnGlobalShaderCompletion, bool bProcessShaderMapNotForRendering)
+void FShaderCompilingManager::ProcessAsyncResults(bool bLimitExecutionTime, bool bBlockOnGlobalShaderCompletion)
 {
 	if (bAllowAsynchronousShaderCompiling)
 	{
@@ -1929,9 +1929,7 @@ void FShaderCompilingManager::ProcessAsyncResults(bool bLimitExecutionTime, bool
 				{
 					const FShaderMapCompileResults& Results = It.Value();
 
-					if (Results.FinishedJobs.Num() == Results.NumJobsQueued 
-						// Shader maps that don't want to be applied can only be processed by a call to FinishCompilation
-						&& (Results.bApplyCompletedShaderMapForRendering || bProcessShaderMapNotForRendering) )
+					if (Results.FinishedJobs.Num() == Results.NumJobsQueued)
 					{
 						PendingFinalizeShaderMaps.Add(It.Key(), FShaderMapFinalizeResults(Results));
 						ShaderMapsToRemove.Add(It.Key());
diff --git a/Engine/Source/Runtime/Engine/Private/TextureLODSettings.cpp b/Engine/Source/Runtime/Engine/Private/TextureLODSettings.cpp
--- a/Engine/Source/Runtime/Engine/Private/TextureLODSettings.cpp
+++ b/Engine/Source/Runtime/Engine/Private/TextureLODSettings.cpp
@@ -280,6 +280,19 @@ int32 FTextureLODSettings::GetNumStreamedMips( int32 InLODGroup ) const
 	return Retval;
 }
 
+/**
+ * Returns the LODGroup mip gen settings
+ *
+ * @param	InLODGroup		The LOD Group ID 
+ * @return	TextureMipGenSettings for lod group
+ */
+ const TextureMipGenSettings& FTextureLODSettings::GetTextureMipGenSettings( int32 InLODGroup ) const
+{
+	return TextureLODGroups[InLODGroup].MipGenSettings; 
+}
+
+
+
 /**
  * Returns the filter state that should be used for the passed in texture, taking
  * into account other system settings.
diff --git a/Engine/Source/Runtime/Engine/Public/ShaderCompiler.h b/Engine/Source/Runtime/Engine/Public/ShaderCompiler.h
--- a/Engine/Source/Runtime/Engine/Public/ShaderCompiler.h
+++ b/Engine/Source/Runtime/Engine/Public/ShaderCompiler.h
@@ -308,7 +308,7 @@ public:
 	 * @param bBlockOnGlobalShaderCompletion - When enabled, ProcessAsyncResults will block until global shader maps are complete.
 	 *		This must be done before using global shaders for rendering.
 	 */
-	ENGINE_API void ProcessAsyncResults(bool bLimitExecutionTime, bool bBlockOnGlobalShaderCompletion, bool bProcessShaderMapNotForRendering = false);
+	ENGINE_API void ProcessAsyncResults(bool bLimitExecutionTime, bool bBlockOnGlobalShaderCompletion);
 
 	/**
 	 * Returns true if the given shader compile worker is still running.
diff --git a/Engine/Source/Runtime/Engine/Public/TextureLODSettings.h b/Engine/Source/Runtime/Engine/Public/TextureLODSettings.h
--- a/Engine/Source/Runtime/Engine/Public/TextureLODSettings.h
+++ b/Engine/Source/Runtime/Engine/Public/TextureLODSettings.h
@@ -83,6 +83,15 @@ struct FTextureLODSettings
 	 */
 	ENGINE_API ESamplerFilter GetSamplerFilter( const UTexture* Texture ) const;
 
+	/**
+	 * Returns the LODGroup mip gen settings
+	 *
+	 * @param	InLODGroup		The LOD Group ID 
+	 * @return	TextureMipGenSettings for lod group
+	 */
+	ENGINE_API const TextureMipGenSettings& GetTextureMipGenSettings( int32 InLODGroup ) const; 
+
+
 	/**
 	 * Returns the texture group names, sorted like enum.
 	 *
