don't mark BP package dirty on streaming async loads

diff --git a/Engine/Source/Editor/UnrealEd/Private/Kismet2/Kismet2.cpp b/Engine/Source/Editor/UnrealEd/Private/Kismet2/Kismet2.cpp
--- a/Engine/Source/Editor/UnrealEd/Private/Kismet2/Kismet2.cpp
+++ b/Engine/Source/Editor/UnrealEd/Private/Kismet2/Kismet2.cpp
@@ -522,6 +522,9 @@ void FKismetEditorUtilities::RecompileBlueprintBytecode(UBlueprint* BlueprintObj
 	check(BlueprintObj);
 	checkf(BlueprintObj->GeneratedClass, TEXT("Invalid generated class for %s"), *BlueprintObj->GetName());
 
+	UPackage* const BlueprintPackage = Cast<UPackage>(BlueprintObj->GetOutermost());
+	bool const bStartedWithUnsavedChanges = (BlueprintPackage != NULL) ? BlueprintPackage->IsDirty() : true;
+
 	IKismetCompilerInterface& Compiler = FModuleManager::LoadModuleChecked<IKismetCompilerInterface>(KISMET_COMPILER_MODULENAME);
 
 	TGuardValue<bool> GuardTemplateNameFlag(GCompilingBlueprint, true);
@@ -534,6 +537,11 @@ void FKismetEditorUtilities::RecompileBlueprintBytecode(UBlueprint* BlueprintObj
 	Compiler.CompileBlueprint(BlueprintObj, CompileOptions, Results, NULL, ObjLoaded);
 
 	ReinstanceHelper.UpdateBytecodeReferences();
+
+	if (BlueprintPackage != NULL)
+	{
+		BlueprintPackage->SetDirtyFlag(bStartedWithUnsavedChanges);
+	}
 }
 
 /** Tries to make sure that a blueprint is conformed to its native parent, in case any native class flags have changed */
