Cook on the fly null reference access.
Stability fix as cook on the fly doesn't work right now.

diff --git a/Engine/Source/Editor/UnrealEd/Private/CookOnTheFlyServer.cpp b/Engine/Source/Editor/UnrealEd/Private/CookOnTheFlyServer.cpp
--- a/Engine/Source/Editor/UnrealEd/Private/CookOnTheFlyServer.cpp
+++ b/Engine/Source/Editor/UnrealEd/Private/CookOnTheFlyServer.cpp
@@ -35,7 +35,7 @@
 DEFINE_LOG_CATEGORY_STATIC(LogCookOnTheFly, Log, All);
 
 #define DEBUG_COOKONTHEFLY 0
-#define OUTPUT_TIMING 0
+#define OUTPUT_TIMING 1
 
 #if OUTPUT_TIMING
 
@@ -1289,21 +1289,21 @@ uint32 UCookOnTheFlyServer::TickCookOnTheSide( const float TimeSlice, uint32 &Co
 		// TODO: Daniel: this is reference code needs to be reimplemented on the callee side.
 		//  cooker can't depend on callee being able to garbage collect
 		// collect garbage
-		if ( CookByTheBookOptions->bLeakTest && bFinishedSave )
+		if ( CookByTheBookOptions && CookByTheBookOptions->bLeakTest && bFinishedSave )
 		{
 			check( CurrentCookMode == ECookMode::CookByTheBook);
 			UE_LOG(LogCookOnTheFly, Display, TEXT("Full GC..."));
 
-		CollectGarbage( RF_Native );
+			CollectGarbage( RF_Native );
 
-		for (FObjectIterator It; It; ++It)
-		{
+			for (FObjectIterator It; It; ++It)
+			{
 				if (!CookByTheBookOptions->LastGCItems.Contains(FWeakObjectPtr(*It)))
-		{
+				{
 					UE_LOG(LogCookOnTheFly, Warning, TEXT("\tLeaked %s"), *(It->GetFullName()));
 					CookByTheBookOptions->LastGCItems.Add(FWeakObjectPtr(*It));
-		}
-		}
+				}
+			}
 		}
 
 		if ( Timer.IsTimeUp() )
