Fix crash deleting levels from the content browser

diff --git a/Engine/Source/Editor/UnrealEd/Private/PackageTools.cpp b/Engine/Source/Editor/UnrealEd/Private/PackageTools.cpp
--- a/Engine/Source/Editor/UnrealEd/Private/PackageTools.cpp
+++ b/Engine/Source/Editor/UnrealEd/Private/PackageTools.cpp
@@ -323,7 +323,14 @@ namespace PackageTools
 			for ( int32 PackageIndex = 0 ; PackageIndex < PackagesToUnload.Num() ; ++PackageIndex )
 			{
 				PackageBeingUnloaded = PackagesToUnload[PackageIndex];
-				check( !PackageBeingUnloaded->IsDirty() );
+	
+				if( PackageBeingUnloaded->IsDirty() )
+				{
+					// The package was marked dirty as a result of something that happened below (e.g callbacks in CollectGarbage).  
+					// Dirty packages we actually care about unloading were filtered above so if the package becomes dirty here it should still be unloaded
+					PackageBeingUnloaded->SetDirtyFlag(false);
+				}
+
 				GWarn->StatusUpdate( PackageIndex, PackagesToUnload.Num(), FText::Format(NSLOCTEXT("UnrealEd", "Unloadingf", "Unloading {0}..."), FText::FromString(PackageBeingUnloaded->GetName()) ) );
 
 				PackageBeingUnloaded->bHasBeenFullyLoaded = false;
