Adjusted access to the LocalWorkspaceMenuRoot within FTabManager - GetLocalWorkspaceRoot() now returns a const and provides methods to expose adding a category and clearing categories

diff --git a/Engine/Plugins/2D/Paper2D/Source/Paper2DEditor/Private/FlipbookEditor/FlipbookEditor.cpp b/Engine/Plugins/2D/Paper2D/Source/Paper2DEditor/Private/FlipbookEditor/FlipbookEditor.cpp
--- a/Engine/Plugins/2D/Paper2D/Source/Paper2DEditor/Private/FlipbookEditor/FlipbookEditor.cpp
+++ b/Engine/Plugins/2D/Paper2D/Source/Paper2DEditor/Private/FlipbookEditor/FlipbookEditor.cpp
@@ -301,7 +301,7 @@ TSharedRef<SDockTab> FFlipbookEditor::SpawnTab_Details(const FSpawnTabArgs& Args
 
 void FFlipbookEditor::RegisterTabSpawners(const TSharedRef<class FTabManager>& TabManager)
 {
-	WorkspaceMenuCategory = TabManager->GetLocalWorkspaceMenuRoot()->AddGroup(LOCTEXT("WorkspaceMenu_FlipbookEditor", "Flipbook Editor"));
+	WorkspaceMenuCategory = TabManager->AddLocalWorkspaceMenuCategory(LOCTEXT("WorkspaceMenu_FlipbookEditor", "Flipbook Editor"));
 	auto WorkspaceMenuCategoryRef = WorkspaceMenuCategory.ToSharedRef();
 
 	FAssetEditorToolkit::RegisterTabSpawners(TabManager);
diff --git a/Engine/Plugins/2D/Paper2D/Source/Paper2DEditor/Private/SpriteEditor/SpriteEditor.cpp b/Engine/Plugins/2D/Paper2D/Source/Paper2DEditor/Private/SpriteEditor/SpriteEditor.cpp
--- a/Engine/Plugins/2D/Paper2D/Source/Paper2DEditor/Private/SpriteEditor/SpriteEditor.cpp
+++ b/Engine/Plugins/2D/Paper2D/Source/Paper2DEditor/Private/SpriteEditor/SpriteEditor.cpp
@@ -329,7 +329,7 @@ TSharedRef<SDockTab> FSpriteEditor::SpawnTab_SpriteList(const FSpawnTabArgs& Arg
 
 void FSpriteEditor::RegisterTabSpawners(const TSharedRef<class FTabManager>& TabManager)
 {
-	WorkspaceMenuCategory = TabManager->GetLocalWorkspaceMenuRoot()->AddGroup(LOCTEXT("WorkspaceMenu_SpriteEditor", "Sprite Editor"));
+	WorkspaceMenuCategory = TabManager->AddLocalWorkspaceMenuCategory(LOCTEXT("WorkspaceMenu_SpriteEditor", "Sprite Editor"));
 	auto WorkspaceMenuCategoryRef = WorkspaceMenuCategory.ToSharedRef();
 
 	FAssetEditorToolkit::RegisterTabSpawners(TabManager);
diff --git a/Engine/Plugins/2D/Paper2D/Source/Paper2DEditor/Private/TileMapEditing/TileMapEditor.cpp b/Engine/Plugins/2D/Paper2D/Source/Paper2DEditor/Private/TileMapEditing/TileMapEditor.cpp
--- a/Engine/Plugins/2D/Paper2D/Source/Paper2DEditor/Private/TileMapEditing/TileMapEditor.cpp
+++ b/Engine/Plugins/2D/Paper2D/Source/Paper2DEditor/Private/TileMapEditing/TileMapEditor.cpp
@@ -232,7 +232,7 @@ TSharedRef<SDockTab> FTileMapEditor::SpawnTab_Details(const FSpawnTabArgs& Args)
 
 void FTileMapEditor::RegisterTabSpawners(const TSharedRef<class FTabManager>& TabManager)
 {
-	WorkspaceMenuCategory = TabManager->GetLocalWorkspaceMenuRoot()->AddGroup(LOCTEXT("WorkspaceMenu_TileMapEditor", "Tile Map Editor"));
+	WorkspaceMenuCategory = TabManager->AddLocalWorkspaceMenuCategory(LOCTEXT("WorkspaceMenu_TileMapEditor", "Tile Map Editor"));
 	auto WorkspaceMenuCategoryRef = WorkspaceMenuCategory.ToSharedRef();
 
 	FAssetEditorToolkit::RegisterTabSpawners(TabManager);
diff --git a/Engine/Plugins/Media/MediaPlayerEditor/Source/MediaPlayerEditor/Private/MediaPlayerEditorToolkit.cpp b/Engine/Plugins/Media/MediaPlayerEditor/Source/MediaPlayerEditor/Private/MediaPlayerEditorToolkit.cpp
--- a/Engine/Plugins/Media/MediaPlayerEditor/Source/MediaPlayerEditor/Private/MediaPlayerEditorToolkit.cpp
+++ b/Engine/Plugins/Media/MediaPlayerEditor/Source/MediaPlayerEditor/Private/MediaPlayerEditorToolkit.cpp
@@ -107,7 +107,7 @@ FString FMediaPlayerEditorToolkit::GetDocumentationLink() const
 
 void FMediaPlayerEditorToolkit::RegisterTabSpawners( const TSharedRef<class FTabManager>& TabManager )
 {
-	WorkspaceMenuCategory = TabManager->GetLocalWorkspaceMenuRoot()->AddGroup(LOCTEXT("WorkspaceMenu_MediaPlayerEditor", "Media Player Editor"));
+	WorkspaceMenuCategory = TabManager->AddLocalWorkspaceMenuCategory(LOCTEXT("WorkspaceMenu_MediaPlayerEditor", "Media Player Editor"));
 	auto WorkspaceMenuCategoryRef = WorkspaceMenuCategory.ToSharedRef();
 
 	FAssetEditorToolkit::RegisterTabSpawners(TabManager);
diff --git a/Engine/Plugins/Messaging/MessagingDebugger/Source/MessagingDebugger/Private/Widgets/SMessagingDebugger.cpp b/Engine/Plugins/Messaging/MessagingDebugger/Source/MessagingDebugger/Private/Widgets/SMessagingDebugger.cpp
--- a/Engine/Plugins/Messaging/MessagingDebugger/Source/MessagingDebugger/Private/Widgets/SMessagingDebugger.cpp
+++ b/Engine/Plugins/Messaging/MessagingDebugger/Source/MessagingDebugger/Private/Widgets/SMessagingDebugger.cpp
@@ -54,7 +54,7 @@ void SMessagingDebugger::Construct( const FArguments& InArgs, const TSharedRef<S
 
 	// create & initialize tab manager
 	TabManager = FGlobalTabmanager::Get()->NewTabManager(ConstructUnderMajorTab);
-	TSharedRef<FWorkspaceItem> AppMenuGroup = TabManager->GetLocalWorkspaceMenuRoot()->AddGroup(LOCTEXT("MessagingDebuggerGroupName", "Messaging Debugger"));
+	TSharedRef<FWorkspaceItem> AppMenuGroup = TabManager->AddLocalWorkspaceMenuCategory(LOCTEXT("MessagingDebuggerGroupName", "Messaging Debugger"));
 
 	TabManager->RegisterTabSpawner(BreakpointsTabId, FOnSpawnTab::CreateRaw(this, &SMessagingDebugger::HandleTabManagerSpawnTab, BreakpointsTabId))
 		.SetDisplayName(LOCTEXT("BreakpointsTabTitle", "Breakpoints"))
diff --git a/Engine/Source/Developer/DeviceManager/Private/Widgets/SDeviceManager.cpp b/Engine/Source/Developer/DeviceManager/Private/Widgets/SDeviceManager.cpp
--- a/Engine/Source/Developer/DeviceManager/Private/Widgets/SDeviceManager.cpp
+++ b/Engine/Source/Developer/DeviceManager/Private/Widgets/SDeviceManager.cpp
@@ -39,7 +39,7 @@ void SDeviceManager::Construct( const FArguments& InArgs, const ITargetDeviceSer
 
 	// create & initialize tab manager
 	TabManager = FGlobalTabmanager::Get()->NewTabManager(ConstructUnderMajorTab);
-	TSharedRef<FWorkspaceItem> AppMenuGroup = TabManager->GetLocalWorkspaceMenuRoot()->AddGroup(LOCTEXT("DeviceManagerMenuGroupName", "Device Manager"));
+	TSharedRef<FWorkspaceItem> AppMenuGroup = TabManager->AddLocalWorkspaceMenuCategory(LOCTEXT("DeviceManagerMenuGroupName", "Device Manager"));
 
 	TabManager->RegisterTabSpawner(DeviceBrowserTabId, FOnSpawnTab::CreateRaw(this, &SDeviceManager::HandleTabManagerSpawnTab, DeviceBrowserTabId))
 		.SetDisplayName(LOCTEXT("DeviceBrowserTabTitle", "Device Browser"))
diff --git a/Engine/Source/Developer/SessionFrontend/Private/Widgets/SSessionFrontend.cpp b/Engine/Source/Developer/SessionFrontend/Private/Widgets/SSessionFrontend.cpp
--- a/Engine/Source/Developer/SessionFrontend/Private/Widgets/SSessionFrontend.cpp
+++ b/Engine/Source/Developer/SessionFrontend/Private/Widgets/SSessionFrontend.cpp
@@ -26,7 +26,7 @@ void SSessionFrontend::Construct( const FArguments& InArgs, const TSharedRef<SDo
 
 	// create & initialize tab manager
 	TabManager = FGlobalTabmanager::Get()->NewTabManager(ConstructUnderMajorTab);
-	TSharedRef<FWorkspaceItem> AppMenuGroup = TabManager->GetLocalWorkspaceMenuRoot()->AddGroup(LOCTEXT("SessionFrontendMenuGroupName", "Session Frontend"));
+	TSharedRef<FWorkspaceItem> AppMenuGroup = TabManager->AddLocalWorkspaceMenuCategory(LOCTEXT("SessionFrontendMenuGroupName", "Session Frontend"));
 	
 	TabManager->RegisterTabSpawner(AutomationTabId, FOnSpawnTab::CreateRaw(this, &SSessionFrontend::HandleTabManagerSpawnTab, AutomationTabId))
 		.SetDisplayName(LOCTEXT("AutomationTabTitle", "Automation"))
diff --git a/Engine/Source/Editor/Cascade/Private/Cascade.cpp b/Engine/Source/Editor/Cascade/Private/Cascade.cpp
--- a/Engine/Source/Editor/Cascade/Private/Cascade.cpp
+++ b/Engine/Source/Editor/Cascade/Private/Cascade.cpp
@@ -50,7 +50,7 @@ DEFINE_LOG_CATEGORY(LogCascade);
 
 void FCascade::RegisterTabSpawners(const TSharedRef<class FTabManager>& TabManager)
 {
-	WorkspaceMenuCategory = TabManager->GetLocalWorkspaceMenuRoot()->AddGroup(NSLOCTEXT("Cascade", "WorkspaceMenu_Cascade", "Cascade"));
+	WorkspaceMenuCategory = TabManager->AddLocalWorkspaceMenuCategory(NSLOCTEXT("Cascade", "WorkspaceMenu_Cascade", "Cascade"));
 	auto WorkspaceMenuCategoryRef = WorkspaceMenuCategory.ToSharedRef();
 
 	FAssetEditorToolkit::RegisterTabSpawners(TabManager);
diff --git a/Engine/Source/Editor/CurveAssetEditor/Private/CurveAssetEditor.cpp b/Engine/Source/Editor/CurveAssetEditor/Private/CurveAssetEditor.cpp
--- a/Engine/Source/Editor/CurveAssetEditor/Private/CurveAssetEditor.cpp
+++ b/Engine/Source/Editor/CurveAssetEditor/Private/CurveAssetEditor.cpp
@@ -15,7 +15,7 @@ const FName FCurveAssetEditor::CurveTabId( TEXT( "CurveAssetEditor_Curve" ) );
 
 void FCurveAssetEditor::RegisterTabSpawners(const TSharedRef<class FTabManager>& TabManager)
 {
-	WorkspaceMenuCategory = TabManager->GetLocalWorkspaceMenuRoot()->AddGroup(LOCTEXT("WorkspaceMenu_CurveAssetEditor", "Curve Asset Editor"));
+	WorkspaceMenuCategory = TabManager->AddLocalWorkspaceMenuCategory(LOCTEXT("WorkspaceMenu_CurveAssetEditor", "Curve Asset Editor"));
 
 	FAssetEditorToolkit::RegisterTabSpawners(TabManager);
 
diff --git a/Engine/Source/Editor/CurveTableEditor/Private/CurveTableEditor.cpp b/Engine/Source/Editor/CurveTableEditor/Private/CurveTableEditor.cpp
--- a/Engine/Source/Editor/CurveTableEditor/Private/CurveTableEditor.cpp
+++ b/Engine/Source/Editor/CurveTableEditor/Private/CurveTableEditor.cpp
@@ -13,7 +13,7 @@ const FName FCurveTableEditor::CurveTableTabId( TEXT( "CurveTableEditor_CurveTab
 
 void FCurveTableEditor::RegisterTabSpawners(const TSharedRef<class FTabManager>& TabManager)
 {
-	WorkspaceMenuCategory = TabManager->GetLocalWorkspaceMenuRoot()->AddGroup(LOCTEXT("WorkspaceMenu_UserDefinedEnumEditor", "User-Defined Enum Editor"));
+	WorkspaceMenuCategory = TabManager->AddLocalWorkspaceMenuCategory(LOCTEXT("WorkspaceMenu_CurveTableEditor", "Curve Table Editor"));
 
 	TabManager->RegisterTabSpawner( CurveTableTabId, FOnSpawnTab::CreateSP(this, &FCurveTableEditor::SpawnTab_CurveTable) )
 		.SetDisplayName( LOCTEXT("CurveTableTab", "Curve Table") )
diff --git a/Engine/Source/Editor/DataTableEditor/Private/DataTableEditor.cpp b/Engine/Source/Editor/DataTableEditor/Private/DataTableEditor.cpp
--- a/Engine/Source/Editor/DataTableEditor/Private/DataTableEditor.cpp
+++ b/Engine/Source/Editor/DataTableEditor/Private/DataTableEditor.cpp
@@ -14,7 +14,7 @@ const FName FDataTableEditor::DataTableTabId( TEXT( "DataTableEditor_DataTable"
 
 void FDataTableEditor::RegisterTabSpawners(const TSharedRef<class FTabManager>& TabManager)
 {
-	WorkspaceMenuCategory = TabManager->GetLocalWorkspaceMenuRoot()->AddGroup(LOCTEXT("WorkspaceMenu_Data Table Editor", "Data Table Editor"));
+	WorkspaceMenuCategory = TabManager->AddLocalWorkspaceMenuCategory(LOCTEXT("WorkspaceMenu_Data Table Editor", "Data Table Editor"));
 
 	TabManager->RegisterTabSpawner( DataTableTabId, FOnSpawnTab::CreateSP(this, &FDataTableEditor::SpawnTab_DataTable) )
 		.SetDisplayName( LOCTEXT("DataTableTab", "Data Table") )
diff --git a/Engine/Source/Editor/DestructibleMeshEditor/Private/DestructibleMeshEditor.cpp b/Engine/Source/Editor/DestructibleMeshEditor/Private/DestructibleMeshEditor.cpp
--- a/Engine/Source/Editor/DestructibleMeshEditor/Private/DestructibleMeshEditor.cpp
+++ b/Engine/Source/Editor/DestructibleMeshEditor/Private/DestructibleMeshEditor.cpp
@@ -43,7 +43,7 @@ FDestructibleMeshEditor::~FDestructibleMeshEditor()
 
 void FDestructibleMeshEditor::RegisterTabSpawners(const TSharedRef<class FTabManager>& TabManager)
 {
-	WorkspaceMenuCategory = TabManager->GetLocalWorkspaceMenuRoot()->AddGroup(LOCTEXT("WorkspaceMenu_DestructibleMeshEditor", "Destructible Mesh Editor"));
+	WorkspaceMenuCategory = TabManager->AddLocalWorkspaceMenuCategory(LOCTEXT("WorkspaceMenu_DestructibleMeshEditor", "Destructible Mesh Editor"));
 	auto WorkspaceMenuCategoryRef = WorkspaceMenuCategory.ToSharedRef();
 
 	FAssetEditorToolkit::RegisterTabSpawners(TabManager);
diff --git a/Engine/Source/Editor/EnvironmentQueryEditor/Private/EnvironmentQueryEditor.cpp b/Engine/Source/Editor/EnvironmentQueryEditor/Private/EnvironmentQueryEditor.cpp
--- a/Engine/Source/Editor/EnvironmentQueryEditor/Private/EnvironmentQueryEditor.cpp
+++ b/Engine/Source/Editor/EnvironmentQueryEditor/Private/EnvironmentQueryEditor.cpp
@@ -23,7 +23,7 @@ const FName FEnvironmentQueryEditor::EQSPropertiesTabId( TEXT( "EnvironmentQuery
 
 void FEnvironmentQueryEditor::RegisterTabSpawners(const TSharedRef<class FTabManager>& TabManager)
 {
-	WorkspaceMenuCategory = TabManager->GetLocalWorkspaceMenuRoot()->AddGroup(LOCTEXT("WorkspaceMenu_EnvironmentQueryEditor", "Environment Query Editor"));
+	WorkspaceMenuCategory = TabManager->AddLocalWorkspaceMenuCategory(LOCTEXT("WorkspaceMenu_EnvironmentQueryEditor", "Environment Query Editor"));
 	auto WorkspaceMenuCategoryRef = WorkspaceMenuCategory.ToSharedRef();
 
 	FAssetEditorToolkit::RegisterTabSpawners(TabManager);
diff --git a/Engine/Source/Editor/FontEditor/Private/FontEditor.cpp b/Engine/Source/Editor/FontEditor/Private/FontEditor.cpp
--- a/Engine/Source/Editor/FontEditor/Private/FontEditor.cpp
+++ b/Engine/Source/Editor/FontEditor/Private/FontEditor.cpp
@@ -72,7 +72,7 @@ void FFontEditorCommands::RegisterCommands()
 
 void FFontEditor::RegisterTabSpawners(const TSharedRef<class FTabManager>& TabManager)
 {
-	WorkspaceMenuCategory = TabManager->GetLocalWorkspaceMenuRoot()->AddGroup(LOCTEXT("WorkspaceMenu_FontEditor", "Font Editor"));
+	WorkspaceMenuCategory = TabManager->AddLocalWorkspaceMenuCategory(LOCTEXT("WorkspaceMenu_FontEditor", "Font Editor"));
 	auto WorkspaceMenuCategoryRef = WorkspaceMenuCategory.ToSharedRef();
 
 	FAssetEditorToolkit::RegisterTabSpawners(TabManager);
diff --git a/Engine/Source/Editor/Kismet/Private/UserDefinedEnumEditor.cpp b/Engine/Source/Editor/Kismet/Private/UserDefinedEnumEditor.cpp
--- a/Engine/Source/Editor/Kismet/Private/UserDefinedEnumEditor.cpp
+++ b/Engine/Source/Editor/Kismet/Private/UserDefinedEnumEditor.cpp
@@ -14,7 +14,7 @@ const FName FUserDefinedEnumEditor::UserDefinedEnumEditorAppIdentifier( TEXT( "U
 
 void FUserDefinedEnumEditor::RegisterTabSpawners(const TSharedRef<class FTabManager>& TabManager)
 {
-	WorkspaceMenuCategory = TabManager->GetLocalWorkspaceMenuRoot()->AddGroup(LOCTEXT("WorkspaceMenu_UserDefinedEnumEditor", "User-Defined Enum Editor"));
+	WorkspaceMenuCategory = TabManager->AddLocalWorkspaceMenuCategory(LOCTEXT("WorkspaceMenu_UserDefinedEnumEditor", "User-Defined Enum Editor"));
 
 	FAssetEditorToolkit::RegisterTabSpawners(TabManager);
 
diff --git a/Engine/Source/Editor/Kismet/Private/UserDefinedStructureEditor.cpp b/Engine/Source/Editor/Kismet/Private/UserDefinedStructureEditor.cpp
--- a/Engine/Source/Editor/Kismet/Private/UserDefinedStructureEditor.cpp
+++ b/Engine/Source/Editor/Kismet/Private/UserDefinedStructureEditor.cpp
@@ -168,7 +168,7 @@ const FName FUserDefinedStructureEditor::UserDefinedStructureEditorAppIdentifier
 
 void FUserDefinedStructureEditor::RegisterTabSpawners(const TSharedRef<class FTabManager>& TabManager)
 {
-	WorkspaceMenuCategory = TabManager->GetLocalWorkspaceMenuRoot()->AddGroup(LOCTEXT("WorkspaceMenu_UserDefinedStructureEditor", "User-Defined Structure Editor"));
+	WorkspaceMenuCategory = TabManager->AddLocalWorkspaceMenuCategory(LOCTEXT("WorkspaceMenu_UserDefinedStructureEditor", "User-Defined Structure Editor"));
 
 	FAssetEditorToolkit::RegisterTabSpawners(TabManager);
 
diff --git a/Engine/Source/Editor/Kismet/Private/WorkflowCentricApplication.cpp b/Engine/Source/Editor/Kismet/Private/WorkflowCentricApplication.cpp
--- a/Engine/Source/Editor/Kismet/Private/WorkflowCentricApplication.cpp
+++ b/Engine/Source/Editor/Kismet/Private/WorkflowCentricApplication.cpp
@@ -58,10 +58,9 @@ void FWorkflowCentricApplication::SetCurrentMode(FName NewMode)
 			//@TODO: Should do some validation here
 			CurrentAppModePtr = NewModePtr;
 
-			// Establish the workspace menu category
-			auto LocalWorkspaceRoot = TabManager->GetLocalWorkspaceMenuRoot();
-			LocalWorkspaceRoot->ClearItems();
-			LocalWorkspaceRoot->AddItem(CurrentAppModePtr->GetWorkspaceMenuCategory());
+			// Establish the workspace menu category for the new mode
+			TabManager->ClearLocalWorkspaceMenuCategories();
+			TabManager->AddLocalWorkspaceMenuItem( CurrentAppModePtr->GetWorkspaceMenuCategory() );
 
 			// Activate the new layout
 			const TSharedRef<FTabManager::FLayout> NewLayout = CurrentAppModePtr->ActivateMode(TabManager);
diff --git a/Engine/Source/Editor/MaterialEditor/Private/MaterialEditor.cpp b/Engine/Source/Editor/MaterialEditor/Private/MaterialEditor.cpp
--- a/Engine/Source/Editor/MaterialEditor/Private/MaterialEditor.cpp
+++ b/Engine/Source/Editor/MaterialEditor/Private/MaterialEditor.cpp
@@ -171,7 +171,7 @@ void FMatExpressionPreview::NotifyCompilationFinished()
 
 void FMaterialEditor::RegisterTabSpawners(const TSharedRef<class FTabManager>& TabManager)
 {
-	WorkspaceMenuCategory = TabManager->GetLocalWorkspaceMenuRoot()->AddGroup(LOCTEXT("WorkspaceMenu_MaterialEditor", "Material Editor"));
+	WorkspaceMenuCategory = TabManager->AddLocalWorkspaceMenuCategory(LOCTEXT("WorkspaceMenu_MaterialEditor", "Material Editor"));
 	auto WorkspaceMenuCategoryRef = WorkspaceMenuCategory.ToSharedRef();
 
 	FAssetEditorToolkit::RegisterTabSpawners(TabManager);
@@ -223,8 +223,6 @@ void FMaterialEditor::UnregisterTabSpawners(const TSharedRef<class FTabManager>&
 	TabManager->UnregisterTabSpawner( StatsTabId );
 	TabManager->UnregisterTabSpawner( FindTabId );
 	TabManager->UnregisterTabSpawner( HLSLCodeTabId );
-
-	TabManager->GetLocalWorkspaceMenuRoot()->ClearItems();
 }
 
 void FMaterialEditor::InitEditorForMaterial(UMaterial* InMaterial)
diff --git a/Engine/Source/Editor/MaterialEditor/Private/MaterialInstanceEditor.cpp b/Engine/Source/Editor/MaterialEditor/Private/MaterialInstanceEditor.cpp
--- a/Engine/Source/Editor/MaterialEditor/Private/MaterialInstanceEditor.cpp
+++ b/Engine/Source/Editor/MaterialEditor/Private/MaterialInstanceEditor.cpp
@@ -109,7 +109,7 @@ protected:
 
 void FMaterialInstanceEditor::RegisterTabSpawners(const TSharedRef<class FTabManager>& TabManager)
 {
-	WorkspaceMenuCategory = TabManager->GetLocalWorkspaceMenuRoot()->AddGroup(LOCTEXT("WorkspaceMenu_MaterialInstanceEditor", "Material Instance Editor"));
+	WorkspaceMenuCategory = TabManager->AddLocalWorkspaceMenuCategory(LOCTEXT("WorkspaceMenu_MaterialInstanceEditor", "Material Instance Editor"));
 	auto WorkspaceMenuCategoryRef = WorkspaceMenuCategory.ToSharedRef();
 
 	FAssetEditorToolkit::RegisterTabSpawners(TabManager);
diff --git a/Engine/Source/Editor/Matinee/Private/Matinee.cpp b/Engine/Source/Editor/Matinee/Private/Matinee.cpp
--- a/Engine/Source/Editor/Matinee/Private/Matinee.cpp
+++ b/Engine/Source/Editor/Matinee/Private/Matinee.cpp
@@ -118,7 +118,7 @@ static const FName MatineePropertyWindowName("Matinee_PropertyWindow");
 
 void FMatinee::RegisterTabSpawners(const TSharedRef<class FTabManager>& TabManager)
 {
-	WorkspaceMenuCategory = TabManager->GetLocalWorkspaceMenuRoot()->AddGroup(LOCTEXT("WorkspaceMenu_MatineeEditor", "Matinee"));
+	WorkspaceMenuCategory = TabManager->AddLocalWorkspaceMenuCategory(LOCTEXT("WorkspaceMenu_MatineeEditor", "Matinee"));
 	auto WorkspaceMenuCategoryRef = WorkspaceMenuCategory.ToSharedRef();
 
 	FAssetEditorToolkit::RegisterTabSpawners(TabManager);
diff --git a/Engine/Source/Editor/NiagaraEditor/Private/NiagaraEditor.cpp b/Engine/Source/Editor/NiagaraEditor/Private/NiagaraEditor.cpp
--- a/Engine/Source/Editor/NiagaraEditor/Private/NiagaraEditor.cpp
+++ b/Engine/Source/Editor/NiagaraEditor/Private/NiagaraEditor.cpp
@@ -14,7 +14,7 @@ const FName FNiagaraEditor::UpdateGraphTabId( TEXT( "NiagaraEditor_UpdateGraph"
 
 void FNiagaraEditor::RegisterTabSpawners(const TSharedRef<class FTabManager>& TabManager)
 {
-	WorkspaceMenuCategory = TabManager->GetLocalWorkspaceMenuRoot()->AddGroup(LOCTEXT("WorkspaceMenu_NiagaraEditor", "Niagara"));
+	WorkspaceMenuCategory = TabManager->AddLocalWorkspaceMenuCategory(LOCTEXT("WorkspaceMenu_NiagaraEditor", "Niagara"));
 
 	FAssetEditorToolkit::RegisterTabSpawners(TabManager);
 
diff --git a/Engine/Source/Editor/NiagaraEditor/Private/NiagaraEffectEditor.cpp b/Engine/Source/Editor/NiagaraEditor/Private/NiagaraEffectEditor.cpp
--- a/Engine/Source/Editor/NiagaraEditor/Private/NiagaraEffectEditor.cpp
+++ b/Engine/Source/Editor/NiagaraEditor/Private/NiagaraEffectEditor.cpp
@@ -15,7 +15,7 @@ const FName FNiagaraEffectEditor::UpdateTabId(TEXT("NiagaraEditor_Effect"));
 
 void FNiagaraEffectEditor::RegisterTabSpawners(const TSharedRef<class FTabManager>& TabManager)
 {
-	WorkspaceMenuCategory = TabManager->GetLocalWorkspaceMenuRoot()->AddGroup(LOCTEXT("WorkspaceMenu_NiagaraEffectEditor", "Niagara Effect"));
+	WorkspaceMenuCategory = TabManager->AddLocalWorkspaceMenuCategory(LOCTEXT("WorkspaceMenu_NiagaraEffectEditor", "Niagara Effect"));
 
 	FAssetEditorToolkit::RegisterTabSpawners(TabManager);
 
diff --git a/Engine/Source/Editor/PhAT/Private/PhAT.cpp b/Engine/Source/Editor/PhAT/Private/PhAT.cpp
--- a/Engine/Source/Editor/PhAT/Private/PhAT.cpp
+++ b/Engine/Source/Editor/PhAT/Private/PhAT.cpp
@@ -75,7 +75,7 @@ static const FName PhATHierarchyName("PhAT_Hierarchy");
 
 void FPhAT::RegisterTabSpawners(const TSharedRef<class FTabManager>& TabManager)
 {
-	WorkspaceMenuCategory = TabManager->GetLocalWorkspaceMenuRoot()->AddGroup( LOCTEXT( "WorkspaceMenu_MaterialEditor", "Material Editor" ) );
+	WorkspaceMenuCategory = TabManager->AddLocalWorkspaceMenuCategory( LOCTEXT( "WorkspaceMenu_PhAT", "Physics Asset Editor" ) );
 	auto WorkspaceMenuCategoryRef = WorkspaceMenuCategory.ToSharedRef();
 
 	FAssetEditorToolkit::RegisterTabSpawners( TabManager );
diff --git a/Engine/Source/Editor/PropertyEditor/Private/PropertyEditorToolkit.cpp b/Engine/Source/Editor/PropertyEditor/Private/PropertyEditorToolkit.cpp
--- a/Engine/Source/Editor/PropertyEditor/Private/PropertyEditorToolkit.cpp
+++ b/Engine/Source/Editor/PropertyEditor/Private/PropertyEditorToolkit.cpp
@@ -22,7 +22,7 @@ const FName FPropertyEditorToolkit::TreePinAsColumnHeaderId( TEXT( "PropertyEdit
 
 void FPropertyEditorToolkit::RegisterTabSpawners(const TSharedRef<class FTabManager>& TabManager)
 {
-	WorkspaceMenuCategory = TabManager->GetLocalWorkspaceMenuRoot()->AddGroup(LOCTEXT("WorkspaceMenu_PropertyEditorToolkit", "Property Editor"));
+	WorkspaceMenuCategory = TabManager->AddLocalWorkspaceMenuCategory(LOCTEXT("WorkspaceMenu_PropertyEditorToolkit", "Property Editor"));
 
 	TabManager->RegisterTabSpawner( GridTabId, FOnSpawnTab::CreateSP(this, &FPropertyEditorToolkit::SpawnTab_PropertyTable) )
 		.SetDisplayName( LOCTEXT("PropertyTableTab", "Grid") )
diff --git a/Engine/Source/Editor/Sequencer/Private/SequencerAssetEditor.cpp b/Engine/Source/Editor/Sequencer/Private/SequencerAssetEditor.cpp
--- a/Engine/Source/Editor/Sequencer/Private/SequencerAssetEditor.cpp
+++ b/Engine/Source/Editor/Sequencer/Private/SequencerAssetEditor.cpp
@@ -26,7 +26,7 @@ void FSequencerAssetEditor::RegisterTabSpawners(const TSharedRef<class FTabManag
 {
 	if( FSequencer::IsSequencerEnabled() && !IsWorldCentricAssetEditor() )
 	{
-		WorkspaceMenuCategory = TabManager->GetLocalWorkspaceMenuRoot()->AddGroup(LOCTEXT("WorkspaceMenu_SequencerAssetEditor", "Sequencer"));
+		WorkspaceMenuCategory = TabManager->AddLocalWorkspaceMenuCategory(LOCTEXT("WorkspaceMenu_SequencerAssetEditor", "Sequencer"));
 
 		TabManager->RegisterTabSpawner( SequencerMainTabId, FOnSpawnTab::CreateSP(this, &FSequencerAssetEditor::SpawnTab_SequencerMain) )
 			.SetDisplayName( LOCTEXT("SequencerMainTab", "Sequencer") )
diff --git a/Engine/Source/Editor/SoundClassEditor/Private/SoundClassEditor.cpp b/Engine/Source/Editor/SoundClassEditor/Private/SoundClassEditor.cpp
--- a/Engine/Source/Editor/SoundClassEditor/Private/SoundClassEditor.cpp
+++ b/Engine/Source/Editor/SoundClassEditor/Private/SoundClassEditor.cpp
@@ -31,7 +31,7 @@ const FName FSoundClassEditor::PropertiesTabId( TEXT( "SoundClassEditor_Properti
 
 void FSoundClassEditor::RegisterTabSpawners(const TSharedRef<class FTabManager>& TabManager)
 {
-	WorkspaceMenuCategory = TabManager->GetLocalWorkspaceMenuRoot()->AddGroup(LOCTEXT("WorkspaceMenu_SoundClassEditor", "Sound Class Editor"));
+	WorkspaceMenuCategory = TabManager->AddLocalWorkspaceMenuCategory(LOCTEXT("WorkspaceMenu_SoundClassEditor", "Sound Class Editor"));
 	auto WorkspaceMenuCategoryRef = WorkspaceMenuCategory.ToSharedRef();
 
 	FAssetEditorToolkit::RegisterTabSpawners(TabManager);
diff --git a/Engine/Source/Editor/SoundCueEditor/Private/SoundCueEditor.cpp b/Engine/Source/Editor/SoundCueEditor/Private/SoundCueEditor.cpp
--- a/Engine/Source/Editor/SoundCueEditor/Private/SoundCueEditor.cpp
+++ b/Engine/Source/Editor/SoundCueEditor/Private/SoundCueEditor.cpp
@@ -27,7 +27,7 @@ const FName FSoundCueEditor::PaletteTabId( TEXT( "SoundCueEditor_Palette" ) );
 
 void FSoundCueEditor::RegisterTabSpawners(const TSharedRef<class FTabManager>& TabManager)
 {
-	WorkspaceMenuCategory = TabManager->GetLocalWorkspaceMenuRoot()->AddGroup(LOCTEXT("WorkspaceMenu_SoundCueEditor", "Sound Cue Editor"));
+	WorkspaceMenuCategory = TabManager->AddLocalWorkspaceMenuCategory(LOCTEXT("WorkspaceMenu_SoundCueEditor", "Sound Cue Editor"));
 	auto WorkspaceMenuCategoryRef = WorkspaceMenuCategory.ToSharedRef();
 
 	FAssetEditorToolkit::RegisterTabSpawners(TabManager);
diff --git a/Engine/Source/Editor/StaticMeshEditor/Private/StaticMeshEditor.cpp b/Engine/Source/Editor/StaticMeshEditor/Private/StaticMeshEditor.cpp
--- a/Engine/Source/Editor/StaticMeshEditor/Private/StaticMeshEditor.cpp
+++ b/Engine/Source/Editor/StaticMeshEditor/Private/StaticMeshEditor.cpp
@@ -42,7 +42,7 @@ const FName FStaticMeshEditor::CollisionTabId( TEXT( "StaticMeshEditor_Collision
 
 void FStaticMeshEditor::RegisterTabSpawners(const TSharedRef<class FTabManager>& TabManager)
 {
-	WorkspaceMenuCategory = TabManager->GetLocalWorkspaceMenuRoot()->AddGroup(LOCTEXT("WorkspaceMenu_StaticMeshEditor", "Static Mesh Editor"));
+	WorkspaceMenuCategory = TabManager->AddLocalWorkspaceMenuCategory(LOCTEXT("WorkspaceMenu_StaticMeshEditor", "Static Mesh Editor"));
 	auto WorkspaceMenuCategoryRef = WorkspaceMenuCategory.ToSharedRef();
 
 	FAssetEditorToolkit::RegisterTabSpawners(TabManager);
diff --git a/Engine/Source/Editor/TextureEditor/Private/TextureEditorToolkit.cpp b/Engine/Source/Editor/TextureEditor/Private/TextureEditorToolkit.cpp
--- a/Engine/Source/Editor/TextureEditor/Private/TextureEditorToolkit.cpp
+++ b/Engine/Source/Editor/TextureEditor/Private/TextureEditorToolkit.cpp
@@ -42,7 +42,7 @@ FString FTextureEditorToolkit::GetDocumentationLink( ) const
 
 void FTextureEditorToolkit::RegisterTabSpawners( const TSharedRef<class FTabManager>& TabManager )
 {
-	WorkspaceMenuCategory = TabManager->GetLocalWorkspaceMenuRoot()->AddGroup(LOCTEXT("WorkspaceMenu_TextureEditor", "Texture Editor"));
+	WorkspaceMenuCategory = TabManager->AddLocalWorkspaceMenuCategory(LOCTEXT("WorkspaceMenu_TextureEditor", "Texture Editor"));
 	auto WorkspaceMenuCategoryRef = WorkspaceMenuCategory.ToSharedRef();
 
 	FAssetEditorToolkit::RegisterTabSpawners(TabManager);
diff --git a/Engine/Source/Editor/TranslationEditor/Private/TranslationEditor.cpp b/Engine/Source/Editor/TranslationEditor/Private/TranslationEditor.cpp
--- a/Engine/Source/Editor/TranslationEditor/Private/TranslationEditor.cpp
+++ b/Engine/Source/Editor/TranslationEditor/Private/TranslationEditor.cpp
@@ -47,7 +47,7 @@ void FTranslationEditor::Initialize()
 
 void FTranslationEditor::RegisterTabSpawners(const TSharedRef<class FTabManager>& TabManager)
 {
-	WorkspaceMenuCategory = TabManager->GetLocalWorkspaceMenuRoot()->AddGroup(LOCTEXT("WorkspaceMenu_TranslationEditor", "Translation Editor"));
+	WorkspaceMenuCategory = TabManager->AddLocalWorkspaceMenuCategory(LOCTEXT("WorkspaceMenu_TranslationEditor", "Translation Editor"));
 	auto WorkspaceMenuCategoryRef = WorkspaceMenuCategory.ToSharedRef();
 
 	FAssetEditorToolkit::RegisterTabSpawners(TabManager);
diff --git a/Engine/Source/Editor/UnrealEd/Private/Toolkits/AssetEditorToolkit.cpp b/Engine/Source/Editor/UnrealEd/Private/Toolkits/AssetEditorToolkit.cpp
--- a/Engine/Source/Editor/UnrealEd/Private/Toolkits/AssetEditorToolkit.cpp
+++ b/Engine/Source/Editor/UnrealEd/Private/Toolkits/AssetEditorToolkit.cpp
@@ -238,7 +238,7 @@ FAssetEditorToolkit::~FAssetEditorToolkit()
 
 void FAssetEditorToolkit::RegisterTabSpawners(const TSharedRef<class FTabManager>& InTabManager)
 {
-	// Use the first child of the local workspace root if there is one, otherwise use the root itself
+	// Use the first child category of the local workspace root if there is one, otherwise use the root itself
 	const auto& LocalCategories = InTabManager->GetLocalWorkspaceMenuRoot()->GetChildItems();
 	TSharedRef<FWorkspaceItem> ToolbarSpawnerCategory = LocalCategories.Num() > 0 ? LocalCategories[0] : InTabManager->GetLocalWorkspaceMenuRoot();
 
@@ -251,6 +251,7 @@ void FAssetEditorToolkit::RegisterTabSpawners(const TSharedRef<class FTabManager
 void FAssetEditorToolkit::UnregisterTabSpawners(const TSharedRef<class FTabManager>& InTabManager)
 {
 	InTabManager->UnregisterTabSpawner( ToolbarTabId );
+	InTabManager->ClearLocalWorkspaceMenuCategories();
 }
 
 bool FAssetEditorToolkit::IsAssetEditor() const
diff --git a/Engine/Source/Editor/UnrealEd/Private/Toolkits/SimpleAssetEditor.cpp b/Engine/Source/Editor/UnrealEd/Private/Toolkits/SimpleAssetEditor.cpp
--- a/Engine/Source/Editor/UnrealEd/Private/Toolkits/SimpleAssetEditor.cpp
+++ b/Engine/Source/Editor/UnrealEd/Private/Toolkits/SimpleAssetEditor.cpp
@@ -16,7 +16,7 @@ const FName FSimpleAssetEditor::PropertiesTabId( TEXT( "GenericEditor_Properties
 
 void FSimpleAssetEditor::RegisterTabSpawners(const TSharedRef<class FTabManager>& TabManager)
 {
-	WorkspaceMenuCategory = TabManager->GetLocalWorkspaceMenuRoot()->AddGroup(LOCTEXT("WorkspaceMenu_GenericAssetEditor", "Asset Editor"));
+	WorkspaceMenuCategory = TabManager->AddLocalWorkspaceMenuCategory(LOCTEXT("WorkspaceMenu_GenericAssetEditor", "Asset Editor"));
 
 	FAssetEditorToolkit::RegisterTabSpawners(TabManager);
 
diff --git a/Engine/Source/Runtime/Slate/Private/Framework/Docking/TabManager.cpp b/Engine/Source/Runtime/Slate/Private/Framework/Docking/TabManager.cpp
--- a/Engine/Source/Runtime/Slate/Private/Framework/Docking/TabManager.cpp
+++ b/Engine/Source/Runtime/Slate/Private/Framework/Docking/TabManager.cpp
@@ -507,11 +507,26 @@ bool FTabManager::IsTabCloseable(const TSharedRef<const SDockTab>& InTab) const
 	return !(MainNonCloseableTab.Pin() == InTab);
 }
 
-TSharedRef<FWorkspaceItem> FTabManager::GetLocalWorkspaceMenuRoot()
+const TSharedRef<FWorkspaceItem> FTabManager::GetLocalWorkspaceMenuRoot() const
 {
 	return LocalWorkspaceMenuRoot.ToSharedRef();
 }
 
+TSharedRef<FWorkspaceItem> FTabManager::AddLocalWorkspaceMenuCategory( const FText& CategoryTitle )
+{
+	return LocalWorkspaceMenuRoot->AddGroup( CategoryTitle );
+}
+
+void FTabManager::AddLocalWorkspaceMenuItem( const TSharedRef<FWorkspaceItem>& CategoryItem )
+{
+	LocalWorkspaceMenuRoot->AddItem( CategoryItem );
+}
+
+void FTabManager::ClearLocalWorkspaceMenuCategories()
+{
+	LocalWorkspaceMenuRoot->ClearItems();
+}
+
 TSharedPtr<FTabManager::FStack> FTabManager::FLayoutNode::AsStack()
 {
 	return TSharedPtr<FTabManager::FStack>();
diff --git a/Engine/Source/Runtime/Slate/Public/Framework/Docking/TabManager.h b/Engine/Source/Runtime/Slate/Public/Framework/Docking/TabManager.h
--- a/Engine/Source/Runtime/Slate/Public/Framework/Docking/TabManager.h
+++ b/Engine/Source/Runtime/Slate/Public/Framework/Docking/TabManager.h
@@ -616,7 +616,16 @@ class SLATE_API FTabManager : public TSharedFromThis<FTabManager>
 		bool IsTabCloseable(const TSharedRef<const SDockTab>& InTab) const;
 
 		/** @return The local workspace menu root */
-		TSharedRef<FWorkspaceItem> GetLocalWorkspaceMenuRoot();
+		const TSharedRef<FWorkspaceItem> GetLocalWorkspaceMenuRoot() const;
+
+		/** Adds a category to the local workspace menu by name */
+		TSharedRef<FWorkspaceItem> AddLocalWorkspaceMenuCategory( const FText& CategoryTitle );
+
+		/** Adds an existing workspace item to the local workspace menu */
+		void AddLocalWorkspaceMenuItem( const TSharedRef<FWorkspaceItem>& CategoryItem );
+
+		/** Clears all categories in the local workspace menu */
+		void ClearLocalWorkspaceMenuCategories();
 
 	protected:
 		void InvokeTabForMenu( FName TabId );
