Stopped slow process requests from poping up during cook in editor.

diff --git a/Engine/Source/Editor/UnrealEd/Classes/Editor/EditorEngine.h b/Engine/Source/Editor/UnrealEd/Classes/Editor/EditorEngine.h
--- a/Engine/Source/Editor/UnrealEd/Classes/Editor/EditorEngine.h
+++ b/Engine/Source/Editor/UnrealEd/Classes/Editor/EditorEngine.h
@@ -1985,7 +1985,7 @@ public:
 	/** The editor wrapper for UPackage::SavePackage. Auto-adds files to source control when necessary */
 	bool SavePackage( UPackage* InOuter, UObject* Base, EObjectFlags TopLevelFlags, const TCHAR* Filename, 
 		FOutputDevice* Error=GError, ULinkerLoad* Conform=NULL, bool bForceByteSwapping=false, bool bWarnOfLongFilename=true, 
-		uint32 SaveFlags=SAVE_None, const class ITargetPlatform* TargetPlatform = NULL, const FDateTime& FinalTimeStamp = FDateTime::MinValue() );
+		uint32 SaveFlags=SAVE_None, const class ITargetPlatform* TargetPlatform = NULL, const FDateTime& FinalTimeStamp = FDateTime::MinValue(), bool bSlowTask = true );
 
 	/** Invoked before a UWorld is saved to update editor systems */
 	virtual void OnPreSaveWorld(uint32 SaveFlags, UWorld* World);
diff --git a/Engine/Source/Editor/UnrealEd/Private/CookOnTheFlyServer.cpp b/Engine/Source/Editor/UnrealEd/Private/CookOnTheFlyServer.cpp
--- a/Engine/Source/Editor/UnrealEd/Private/CookOnTheFlyServer.cpp
+++ b/Engine/Source/Editor/UnrealEd/Private/CookOnTheFlyServer.cpp
@@ -1662,7 +1662,7 @@ bool UCookOnTheFlyServer::SaveCookedPackage( UPackage* Package, uint32 SaveFlags
 				else
 				{
 					SCOPE_TIMER(GEditorSavePackage);
-					bSavedCorrectly &= GEditor->SavePackage( Package, World, Flags, *PlatFilename, GError, NULL, bSwap, false, SaveFlags, Target, FDateTime::MinValue() );
+					bSavedCorrectly &= GEditor->SavePackage( Package, World, Flags, *PlatFilename, GError, NULL, bSwap, false, SaveFlags, Target, FDateTime::MinValue(), false );
 				}
 
 				
diff --git a/Engine/Source/Editor/UnrealEd/Private/Editor.cpp b/Engine/Source/Editor/UnrealEd/Private/Editor.cpp
--- a/Engine/Source/Editor/UnrealEd/Private/Editor.cpp
+++ b/Engine/Source/Editor/UnrealEd/Private/Editor.cpp
@@ -4329,9 +4329,9 @@ void UEditorEngine::OnSourceControlDialogClosed(bool bEnabled)
 
 bool UEditorEngine::SavePackage( UPackage* InOuter, UObject* InBase, EObjectFlags TopLevelFlags, const TCHAR* Filename, 
 				 FOutputDevice* Error, ULinkerLoad* Conform, bool bForceByteSwapping, bool bWarnOfLongFilename, 
-				 uint32 SaveFlags, const class ITargetPlatform* TargetPlatform, const FDateTime& FinalTimeStamp )
+				 uint32 SaveFlags, const class ITargetPlatform* TargetPlatform, const FDateTime& FinalTimeStamp, bool bSlowTask )
 {
-	FScopedSlowTask SlowTask(100);
+	FScopedSlowTask SlowTask(100, FText(), bSlowTask);
 
 	UObject* Base = InBase;
 	if ( !Base && InOuter && InOuter->PackageFlags & PKG_ContainsMap )
@@ -4364,7 +4364,7 @@ bool UEditorEngine::SavePackage( UPackage* InOuter, UObject* InBase, EObjectFlag
 
 	SlowTask.EnterProgressFrame(70);
 
-	bool bSuccess = UPackage::SavePackage(InOuter, Base, TopLevelFlags, Filename, Error, Conform, bForceByteSwapping, bWarnOfLongFilename, SaveFlags, TargetPlatform, FinalTimeStamp);
+	bool bSuccess = UPackage::SavePackage(InOuter, Base, TopLevelFlags, Filename, Error, Conform, bForceByteSwapping, bWarnOfLongFilename, SaveFlags, TargetPlatform, FinalTimeStamp, bSlowTask);
 
 	SlowTask.EnterProgressFrame(10);
 
diff --git a/Engine/Source/Runtime/CoreUObject/Private/UObject/SavePackage.cpp b/Engine/Source/Runtime/CoreUObject/Private/UObject/SavePackage.cpp
--- a/Engine/Source/Runtime/CoreUObject/Private/UObject/SavePackage.cpp
+++ b/Engine/Source/Runtime/CoreUObject/Private/UObject/SavePackage.cpp
@@ -2339,7 +2339,7 @@ EObjectMark UPackage::GetObjectMarksForTargetPlatform( const class ITargetPlatfo
 
 bool UPackage::SavePackage( UPackage* InOuter, UObject* Base, EObjectFlags TopLevelFlags, const TCHAR* Filename, 
 	FOutputDevice* Error, ULinkerLoad* Conform, bool bForceByteSwapping, bool bWarnOfLongFilename, uint32 SaveFlags, 
-	const class ITargetPlatform* TargetPlatform, const FDateTime&  FinalTimeStamp)
+	const class ITargetPlatform* TargetPlatform, const FDateTime&  FinalTimeStamp, bool bSlowTask)
 {
 	UE_START_LOG_COOK_TIME( Filename );
 	
@@ -2498,7 +2498,7 @@ bool UPackage::SavePackage( UPackage* InOuter, UObject* Base, EObjectFlags TopLe
 		FText StatusMessage = FText::Format( NSLOCTEXT("Core", "SavingFile", "Saving file: {CleanFilename}..."), Args );
 
 		const int32 TotalSaveSteps = 31;
-		FScopedSlowTask SlowTask(TotalSaveSteps, StatusMessage);
+		FScopedSlowTask SlowTask(TotalSaveSteps, StatusMessage, bSlowTask);
 		SlowTask.MakeDialog(SaveFlags & SAVE_FromAutosave ? true : false);
 
 		SlowTask.EnterProgressFrame();
diff --git a/Engine/Source/Runtime/CoreUObject/Public/UObject/CoreObject.h b/Engine/Source/Runtime/CoreUObject/Public/UObject/CoreObject.h
--- a/Engine/Source/Runtime/CoreUObject/Public/UObject/CoreObject.h
+++ b/Engine/Source/Runtime/CoreUObject/Public/UObject/CoreObject.h
@@ -340,7 +340,7 @@ public:
 	 */
 	static bool SavePackage( UPackage* InOuter, UObject* Base, EObjectFlags TopLevelFlags, const TCHAR* Filename, 
 		FOutputDevice* Error=GError, ULinkerLoad* Conform=NULL, bool bForceByteSwapping=false, bool bWarnOfLongFilename=true, 
-		uint32 SaveFlags=SAVE_None, const class ITargetPlatform* TargetPlatform = NULL, const FDateTime& FinalTimeStamp = FDateTime::MinValue() );
+		uint32 SaveFlags=SAVE_None, const class ITargetPlatform* TargetPlatform = NULL, const FDateTime& FinalTimeStamp = FDateTime::MinValue(), bool bSlowTask = true );
 
 	/** Wait for any SAVE_Async file writes to complete **/
 	static void WaitForAsyncFileWrites();
