Missing file for CL 2339083.

diff --git a/Engine/Source/Editor/UnrealEd/Private/Kismet2/Kismet2.cpp b/Engine/Source/Editor/UnrealEd/Private/Kismet2/Kismet2.cpp
--- a/Engine/Source/Editor/UnrealEd/Private/Kismet2/Kismet2.cpp
+++ b/Engine/Source/Editor/UnrealEd/Private/Kismet2/Kismet2.cpp
@@ -479,7 +479,8 @@ void FKismetEditorUtilities::CompileBlueprint(UBlueprint* BlueprintObj, bool bIs
 		BlueprintPackage->SetDirtyFlag(bStartedWithUnsavedChanges);
 
 		UBlueprintEditorSettings* Settings = GetMutableDefault<UBlueprintEditorSettings>();
-		if (Settings->bSaveOnCompile && (BlueprintObj->Status == BS_UpToDate))
+		if ( (Settings->SaveOnCompile == SoC_Always) || 
+			((Settings->SaveOnCompile == SoC_SuccessOnly) && (BlueprintObj->Status == BS_UpToDate)) )
 		{
 			bool const bIsLevelPackage = (UWorld::FindWorldInPackage(BlueprintPackage) != nullptr);
 			// we don't want to save the entire level (especially is this 
