[UEBP-5] eliminate redundant BP skeleton class generation passes during compile-on-load

- defer skeleton class regeneration to the full compile pass for "normal" BPs
- bypass MBASM during COL on the RefreshAllNodes() call
- no longer regenerate skeleton class on bytecode-only pass

diff --git a/Engine/Source/Editor/KismetCompiler/Private/KismetCompilerModule.cpp b/Engine/Source/Editor/KismetCompiler/Private/KismetCompilerModule.cpp
--- a/Engine/Source/Editor/KismetCompiler/Private/KismetCompilerModule.cpp
+++ b/Engine/Source/Editor/KismetCompiler/Private/KismetCompilerModule.cpp
@@ -108,6 +108,7 @@ void FKismet2CompilerModule::CompileBlueprint(class UBlueprint* Blueprint, const
 
 	const bool bIsBrandNewBP = (Blueprint->SkeletonGeneratedClass == NULL) && (Blueprint->GeneratedClass == NULL) && (Blueprint->ParentClass != NULL) && !CompileOptions.bIsDuplicationInstigated;
 
+	if( CompileOptions.CompileType != EKismetCompileType::BytecodeOnly )
 	{
 		BP_SCOPED_COMPILER_EVENT_NAME(TEXT("Compile Skeleton Class"));
 
diff --git a/Engine/Source/Editor/UnrealEd/Private/Kismet2/BlueprintEditorUtils.cpp b/Engine/Source/Editor/UnrealEd/Private/Kismet2/BlueprintEditorUtils.cpp
--- a/Engine/Source/Editor/UnrealEd/Private/Kismet2/BlueprintEditorUtils.cpp
+++ b/Engine/Source/Editor/UnrealEd/Private/Kismet2/BlueprintEditorUtils.cpp
@@ -302,8 +302,8 @@ void FBlueprintEditorUtils::RefreshAllNodes(UBlueprint* Blueprint)
 		Schema->ReconstructNode(*CurrentNode, true);
 	}
 
-	// If all nodes change structure, catch that case and recompile now
-	if( bLastChangesStructure )
+	// If all nodes change structure, catch that case and recompile now (unless we're already compiling on load)
+	if( bLastChangesStructure && !Blueprint->bIsRegeneratingOnLoad )
 	{
 		FBlueprintEditorUtils::MarkBlueprintAsStructurallyModified(Blueprint);
 	}
@@ -1106,8 +1106,6 @@ UClass* FBlueprintEditorUtils::RegenerateBlueprintClass(UBlueprint* Blueprint, U
 		// Make sure all used external classes/functions/structures/macros/etc are loaded and linked
 		FRegenerationHelper::LinkExternalDependencies(Blueprint, ObjLoaded);
 
-		FKismetEditorUtilities::GenerateBlueprintSkeleton(Blueprint);
-
 		static FBoolConfigValueHelper ReplaceBlueprintWithClass(TEXT("EditoronlyBP"), TEXT("bReplaceBlueprintWithClass"));
 		if (ReplaceBlueprintWithClass)
 		{
@@ -1128,16 +1126,22 @@ UClass* FBlueprintEditorUtils::RegenerateBlueprintClass(UBlueprint* Blueprint, U
 			// Refresh all nodes to make sure function signatures are up to date, etc.
 			FBlueprintEditorUtils::RefreshAllNodes(Blueprint);
 
-			// Compile the actual blueprint
+			// Compile the actual blueprint (skeleton class + generated class)
 			FKismetEditorUtilities::CompileBlueprint(Blueprint, true);
 		}
 		else if( bIsMacro )
 		{
-			// Just refresh all nodes in macro blueprints, but don't recompil
+			// Refresh all nodes to make sure function signatures are up to date, etc.
 			FBlueprintEditorUtils::RefreshAllNodes(Blueprint);
+
+			// Generate the skeleton class
+			FKismetEditorUtilities::GenerateBlueprintSkeleton(Blueprint);
 		}
 		else
 		{
+			// Generate the skeleton class
+			FKismetEditorUtilities::GenerateBlueprintSkeleton(Blueprint);
+
 			if (Blueprint->IsGeneratedClassAuthoritative() && (Blueprint->GeneratedClass != NULL))
 			{
 				check(PreviousCDO != NULL);
