BP: blueprint loading crash (introduced in cl# 2339313 ?) - temporary fix

diff --git a/Engine/Source/Editor/UnrealEd/Private/Kismet2/BlueprintEditorUtils.cpp b/Engine/Source/Editor/UnrealEd/Private/Kismet2/BlueprintEditorUtils.cpp
--- a/Engine/Source/Editor/UnrealEd/Private/Kismet2/BlueprintEditorUtils.cpp
+++ b/Engine/Source/Editor/UnrealEd/Private/Kismet2/BlueprintEditorUtils.cpp
@@ -1106,6 +1106,8 @@ UClass* FBlueprintEditorUtils::RegenerateBlueprintClass(UBlueprint* Blueprint, U
 		// Make sure all used external classes/functions/structures/macros/etc are loaded and linked
 		FRegenerationHelper::LinkExternalDependencies(Blueprint, ObjLoaded);
 
+		FKismetEditorUtilities::GenerateBlueprintSkeleton(Blueprint);
+
 		static FBoolConfigValueHelper ReplaceBlueprintWithClass(TEXT("EditoronlyBP"), TEXT("bReplaceBlueprintWithClass"));
 		if (ReplaceBlueprintWithClass)
 		{
