UE-4486  Crash occurs when making a new blueprint after deleting one with the same name

When BP is ForceDeleted its BPGC is moved to transient package.

diff --git a/Engine/Source/Editor/UnrealEd/Private/ObjectTools.cpp b/Engine/Source/Editor/UnrealEd/Private/ObjectTools.cpp
--- a/Engine/Source/Editor/UnrealEd/Private/ObjectTools.cpp
+++ b/Engine/Source/Editor/UnrealEd/Private/ObjectTools.cpp
@@ -1790,12 +1790,6 @@ namespace ObjectTools
 							ObjectsToReplace.AddUnique(BlueprintObject->GeneratedClass);
 						}
 
-						// Clear BP ref so we're not warned about the skeleton class being transient during replacement
-						if (BlueprintObject->SkeletonGeneratedClass)
-						{
-							BlueprintObject->SkeletonGeneratedClass->ClassGeneratedBy = NULL;
-						}
-
 						// Reparent any direct children to the parent class of the blueprint that's about to be deleted
 						if(BlueprintObject->ParentClass != nullptr)
 						{
@@ -1819,6 +1813,8 @@ namespace ObjectTools
 								}
 							}
 						}
+
+						BlueprintObject->RemoveGeneratedClasses();
 					}
 				}
 
