Updating the editor perf test code based on code feedback.

diff --git a/Engine/Source/Editor/UnrealEd/Private/Tests/AutomationEditorCommon.cpp b/Engine/Source/Editor/UnrealEd/Private/Tests/AutomationEditorCommon.cpp
--- a/Engine/Source/Editor/UnrealEd/Private/Tests/AutomationEditorCommon.cpp
+++ b/Engine/Source/Editor/UnrealEd/Private/Tests/AutomationEditorCommon.cpp
@@ -290,10 +290,10 @@ namespace AutomationEditorCommonUtils
 	* @param InNumberToBeWritten is the float number that is expected to be written to the file.
 	* @param Delimiter is the delimiter to be used. TEXT(",")
 	*/
-	void WriteToTextFile(const FString& InTestName, const FString& InItemBeingTested, const FString& InFileName, const float& InNumberToBeWritten, const FString& Delimiter)
+	void WriteToTextFile(const FString& InTestName, const FString& InTestItem, const FString& InFileName, const float& InEntry, const FString& Delimiter)
 	{
 		//Performance file locations and setups.
-		FString FileSaveLocation = FPaths::Combine(*FPaths::AutomationLogDir(), *InTestName, *InItemBeingTested, *InFileName);
+		FString FileSaveLocation = FPaths::Combine(*FPaths::AutomationLogDir(), *InTestName, *InTestItem, *InFileName);
 
 		if (FPaths::FileExists(FileSaveLocation))
 		{
@@ -302,12 +302,12 @@ namespace AutomationEditorCommonUtils
 
 			//Write to the text file the combined contents from the text file with the number to write.
 			FFileHelper::LoadFileToString(TextFileContents, *FileSaveLocation);
-			FString FileSetup = TextFileContents + Delimiter + FString::SanitizeFloat(InNumberToBeWritten);
+			FString FileSetup = TextFileContents + Delimiter + FString::SanitizeFloat(InEntry);
 			FFileHelper::SaveStringToFile(FileSetup, *FileSaveLocation);
 			return;
 		}
 
-		FFileHelper::SaveStringToFile(FString::SanitizeFloat(InNumberToBeWritten), *FileSaveLocation);
+		FFileHelper::SaveStringToFile(FString::SanitizeFloat(InEntry), *FileSaveLocation);
 	}
 
 	/**
@@ -315,7 +315,7 @@ namespace AutomationEditorCommonUtils
 	* @param InFloatArray is the name of the array intended to be used.
 	* @param bisAveragedInstead will return the average of the available numbers instead of the sum.
 	*/
-	float TotalFromFloatArray(const TArray<float>& InFloatArray, bool bisAveragedInstead)
+	float TotalFromFloatArray(const TArray<float>& InFloatArray, bool bIsAveragedInstead)
 	{
 		//Total Value holds the sum of all the numbers available in the array.
 		float TotalValue = 0;
@@ -327,7 +327,7 @@ namespace AutomationEditorCommonUtils
 		}
 
 		//If bAverageInstead equals true then only the average is returned.
-		if (bisAveragedInstead)
+		if (bIsAveragedInstead)
 		{
 			UE_LOG(LogEditorAutomationTests, VeryVerbose, TEXT("Average value of the Array is %f"), (TotalValue / InFloatArray.Num()));
 			return (TotalValue / InFloatArray.Num());
@@ -341,7 +341,7 @@ namespace AutomationEditorCommonUtils
 	* Returns the largest value from an array of float numbers.
 	* @param InFloatArray is the name of the array intended to be used.
 	*/
-	float LargetValueInFloatArray(const TArray<float>& InFloatArray)
+	float LargestValueInFloatArray(const TArray<float>& InFloatArray)
 	{
 		//Total Value holds the sum of all the numbers available in the array.
 		float LargestValue = 0;
@@ -360,35 +360,31 @@ namespace AutomationEditorCommonUtils
 
 	/**
 	* Returns the contents of a text file as an array of FString.
-	* @param InFileLocation is the location of the file.
+	* @param InFileLocation -  is the location of the file.
+	* @param OutArray - The name of the array where the 
 	*/
-	TArray<FString> CreateArrayFromFile(const FString& InFileLocation)
+	void CreateArrayFromFile(const FString& InFileLocation, TArray<FString>& OutArray)
 	{
 		FString RawData;
-		TArray<FString> DataArray;
 
 		if (FPaths::FileExists(*InFileLocation))
 		{
 			UE_LOG(LogEditorAutomationTests, VeryVerbose, TEXT("Loading and parsing the data from '%s' into an array."), *InFileLocation);
 			FFileHelper::LoadFileToString(RawData, *InFileLocation);
-			RawData.ParseIntoArray(&DataArray, TEXT(","), false);
-
-			return DataArray;
+			RawData.ParseIntoArray(&OutArray, TEXT(","), false);
 		}
 
 		UE_LOG(LogEditorAutomationTests, Warning, TEXT("Unable to create an array.  '%s' does not exist."), *InFileLocation);
 		RawData = TEXT("0");
-		DataArray.Add(RawData);
-
-		return DataArray;
+		OutArray.Add(RawData);
 	}
 
 	/**
 	* Returns true if the archive/file can be written to otherwise false..
 	* @param InFilePath - is the location of the file.
-	* @param InArchiveName - is the name of the arhcive to be used.
+	* @param InArchiveName - is the name of the archive to be used.
 	*/
-	bool IsArchiveWriteable(const FString InFilePath, const FArchive* InArchiveName)
+	bool IsArchiveWriteable(const FString& InFilePath, const FArchive* InArchiveName)
 	{
 		if (!InArchiveName)
 		{
diff --git a/Engine/Source/Editor/UnrealEd/Private/Tests/EditorPerformanceTests.cpp b/Engine/Source/Editor/UnrealEd/Private/Tests/EditorPerformanceTests.cpp
--- a/Engine/Source/Editor/UnrealEd/Private/Tests/EditorPerformanceTests.cpp
+++ b/Engine/Source/Editor/UnrealEd/Private/Tests/EditorPerformanceTests.cpp
@@ -57,7 +57,8 @@ void EditorPerfDump(EditorPerfCaptureParameters& EditorPerfStats)
 	FString MapLoadTimeFileLocation = FPaths::Combine(*DataFileLocation, TEXT("RAWMapLoadTime.txt"));
 	if (FPaths::FileExists(*MapLoadTimeFileLocation))
 	{
-		TArray<FString> SavedMapLoadTimes = AutomationEditorCommonUtils::CreateArrayFromFile(MapLoadTimeFileLocation);
+		TArray<FString> SavedMapLoadTimes;
+		AutomationEditorCommonUtils::CreateArrayFromFile(MapLoadTimeFileLocation, SavedMapLoadTimes);
 		EditorPerfStats.MapLoadTime = FCString::Atof(*SavedMapLoadTimes.Last());
 	}
 	
@@ -73,12 +74,12 @@ void EditorPerfDump(EditorPerfCaptureParameters& EditorPerfStats)
 	RAWCSVArchive->Serialize(TCHAR_TO_ANSI(*RAWCSVLine), RAWCSVLine.Len());
 
 	//Dump the stats from each run to the raw csv file and then close it.
-	for (int32 I = 0; I < EditorPerfStats.TimeStamp.Num(); I++)
+	for (int32 i = 0; i < EditorPerfStats.TimeStamp.Num(); i++)
 	{
 		//If the raw file isn't available to write to then we'll fail back this test.
 		if (AutomationEditorCommonUtils::IsArchiveWriteable(RAWCSVFilePath, RAWCSVArchive))
 		{
-			RAWCSVLine = FString::Printf(TEXT("%s,%s,%s,%.3f,%.1f,%.1f,%.0f,%.0f,%.0f,%.0f,%.0f,%.0f%s"), *EditorPerfStats.MapName, *GEngineVersion.ToString(EVersionComponent::Changelist), *EditorPerfStats.TimeStamp[I].ToString(), EditorPerfStats.MapLoadTime, EditorPerfStats.AverageFPS[I], EditorPerfStats.AverageFrameTime[I], EditorPerfStats.UsedPhysical[I], EditorPerfStats.UsedVirtual[I], EditorPerfStats.PeakUsedPhysical[I], EditorPerfStats.PeakUsedVirtual[I], EditorPerfStats.AvailablePhysical[I], EditorPerfStats.AvailableVirtual[I], LINE_TERMINATOR);
+			RAWCSVLine = FString::Printf(TEXT("%s,%s,%s,%.3f,%.1f,%.1f,%.0f,%.0f,%.0f,%.0f,%.0f,%.0f%s"), *EditorPerfStats.MapName, *GEngineVersion.ToString(EVersionComponent::Changelist), *EditorPerfStats.TimeStamp[i].ToString(), EditorPerfStats.MapLoadTime, EditorPerfStats.AverageFPS[i], EditorPerfStats.AverageFrameTime[i], EditorPerfStats.UsedPhysical[i], EditorPerfStats.UsedVirtual[i], EditorPerfStats.PeakUsedPhysical[i], EditorPerfStats.PeakUsedVirtual[i], EditorPerfStats.AvailablePhysical[i], EditorPerfStats.AvailableVirtual[i], LINE_TERMINATOR);
 			RAWCSVArchive->Serialize(TCHAR_TO_ANSI(*RAWCSVLine), RAWCSVLine.Len());
 		}
 	}
@@ -91,8 +92,8 @@ void EditorPerfDump(EditorPerfCaptureParameters& EditorPerfStats)
 	float MemoryAvailPhysAvg = AutomationEditorCommonUtils::TotalFromFloatArray(EditorPerfStats.AvailablePhysical, true);
 	float MemoryAvailVirtualAvg = AutomationEditorCommonUtils::TotalFromFloatArray(EditorPerfStats.AvailableVirtual, true);
 	float MemoryUsedVirtualAvg = AutomationEditorCommonUtils::TotalFromFloatArray(EditorPerfStats.UsedVirtual, true);
-	float MemoryUsedPeak = AutomationEditorCommonUtils::LargetValueInFloatArray(EditorPerfStats.PeakUsedPhysical);
-	float MemoryUsedPeakVirtual = AutomationEditorCommonUtils::LargetValueInFloatArray(EditorPerfStats.PeakUsedVirtual);
+	float MemoryUsedPeak = AutomationEditorCommonUtils::LargestValueInFloatArray(EditorPerfStats.PeakUsedPhysical);
+	float MemoryUsedPeakVirtual = AutomationEditorCommonUtils::LargestValueInFloatArray(EditorPerfStats.PeakUsedVirtual);
 
 	//TestRunDuration is the length of time the test lasted in ticks.
 	FTimespan TestRunDuration = (EditorPerfStats.TimeStamp.Last().GetTicks() - EditorPerfStats.TimeStamp[0].GetTicks()) + ETimespan::TicksPerSecond;
diff --git a/Engine/Source/Editor/UnrealEd/Public/Tests/AutomationEditorCommon.h b/Engine/Source/Editor/UnrealEd/Public/Tests/AutomationEditorCommon.h
--- a/Engine/Source/Editor/UnrealEd/Public/Tests/AutomationEditorCommon.h
+++ b/Engine/Source/Editor/UnrealEd/Public/Tests/AutomationEditorCommon.h
@@ -103,7 +103,7 @@ namespace AutomationEditorCommonUtils
 	* @param InNumberToBeWritten - is the float number that is expected to be written to the file.
 	* @param Delimiter - is the delimiter to be used. TEXT(",")
 	*/
-	void WriteToTextFile(const FString& InTestName, const FString& InItemBeingTested, const FString& InFileName, const float& InNumberToBeWritten, const FString& Delimiter);
+	void WriteToTextFile(const FString& InTestName, const FString& InTestItem, const FString& InFileName, const float& InEntry, const FString& Delimiter);
 
 	/**
 	* Returns the sum of the numbers available in an array of float.
@@ -118,14 +118,15 @@ namespace AutomationEditorCommonUtils
 
 	* @param InFloatArray - is the name of the array intended to be used.
 	*/
-	float LargetValueInFloatArray(const TArray<float>& InFloatArray);
+	float LargestValueInFloatArray(const TArray<float>& InFloatArray);
 
 	/**
 	* Returns the contents of a text file as an array of FString.
 
 	* @param InFileLocation - is the location of the file.
+	* @param OutArray - The name of the array that will store the data.
 	*/
-	TArray<FString> CreateArrayFromFile(const FString& InFileLocation);
+	void CreateArrayFromFile(const FString& InFileLocation, TArray<FString>& OutArray);
 
 	/**
 	* Returns true if the archive/file can be written to otherwise false.
@@ -133,7 +134,7 @@ namespace AutomationEditorCommonUtils
 	* @param InFilePath - is the location of the file.
 	* @param InArchiveName - is the name of the archive to be used.
 	*/
-	bool IsArchiveWriteable(const FString InFilePath, const FArchive* InArchiveName);
+	bool IsArchiveWriteable(const FString& InFilePath, const FArchive* InArchiveName);
 }
 
 
