Switched the time stamp that gets outputted to the raw csv to be the formatted version.

diff --git a/Engine/Source/Editor/UnrealEd/Private/Tests/EditorPerformanceTests.cpp b/Engine/Source/Editor/UnrealEd/Private/Tests/EditorPerformanceTests.cpp
--- a/Engine/Source/Editor/UnrealEd/Private/Tests/EditorPerformanceTests.cpp
+++ b/Engine/Source/Editor/UnrealEd/Private/Tests/EditorPerformanceTests.cpp
@@ -79,7 +79,7 @@ void EditorPerfDump(EditorPerfCaptureParameters& EditorPerfStats)
 		//If the raw file isn't available to write to then we'll fail back this test.
 		if (AutomationEditorCommonUtils::IsArchiveWriteable(RAWCSVFilePath, RAWCSVArchive))
 		{
-			RAWCSVLine = FString::Printf(TEXT("%s,%s,%s,%.3f,%.1f,%.1f,%.0f,%.0f,%.0f,%.0f,%.0f,%.0f%s"), *EditorPerfStats.MapName, *GEngineVersion.ToString(EVersionComponent::Changelist), *EditorPerfStats.TimeStamp[i].ToString(), EditorPerfStats.MapLoadTime, EditorPerfStats.AverageFPS[i], EditorPerfStats.AverageFrameTime[i], EditorPerfStats.UsedPhysical[i], EditorPerfStats.UsedVirtual[i], EditorPerfStats.PeakUsedPhysical[i], EditorPerfStats.PeakUsedVirtual[i], EditorPerfStats.AvailablePhysical[i], EditorPerfStats.AvailableVirtual[i], LINE_TERMINATOR);
+			RAWCSVLine = FString::Printf(TEXT("%s,%s,%s,%.3f,%.1f,%.1f,%.0f,%.0f,%.0f,%.0f,%.0f,%.0f%s"), *EditorPerfStats.MapName, *GEngineVersion.ToString(EVersionComponent::Changelist), *EditorPerfStats.FormattedTimeStamp[i], EditorPerfStats.MapLoadTime, EditorPerfStats.AverageFPS[i], EditorPerfStats.AverageFrameTime[i], EditorPerfStats.UsedPhysical[i], EditorPerfStats.UsedVirtual[i], EditorPerfStats.PeakUsedPhysical[i], EditorPerfStats.PeakUsedVirtual[i], EditorPerfStats.AvailablePhysical[i], EditorPerfStats.AvailableVirtual[i], LINE_TERMINATOR);
 			RAWCSVArchive->Serialize(TCHAR_TO_ANSI(*RAWCSVLine), RAWCSVLine.Len());
 		}
 	}
