back out previous change to RegenerateBlueprintClass(); needs further work at the RefreshAllNodes() level

diff --git a/Engine/Source/Editor/UnrealEd/Private/Kismet2/BlueprintEditorUtils.cpp b/Engine/Source/Editor/UnrealEd/Private/Kismet2/BlueprintEditorUtils.cpp
--- a/Engine/Source/Editor/UnrealEd/Private/Kismet2/BlueprintEditorUtils.cpp
+++ b/Engine/Source/Editor/UnrealEd/Private/Kismet2/BlueprintEditorUtils.cpp
@@ -302,8 +302,8 @@ void FBlueprintEditorUtils::RefreshAllNodes(UBlueprint* Blueprint)
 		Schema->ReconstructNode(*CurrentNode, true);
 	}
 
-	// If all nodes change structure, catch that case and recompile now (unless we're already compiling on load)
-	if( bLastChangesStructure && !Blueprint->bIsRegeneratingOnLoad )
+	// If all nodes change structure, catch that case and recompile now
+	if( bLastChangesStructure )
 	{
 		FBlueprintEditorUtils::MarkBlueprintAsStructurallyModified(Blueprint);
 	}
@@ -1128,22 +1128,16 @@ UClass* FBlueprintEditorUtils::RegenerateBlueprintClass(UBlueprint* Blueprint, U
 			// Refresh all nodes to make sure function signatures are up to date, etc.
 			FBlueprintEditorUtils::RefreshAllNodes(Blueprint);
 
-			// Compile the actual blueprint (skeleton class + generated class)
+			// Compile the actual blueprint
 			FKismetEditorUtilities::CompileBlueprint(Blueprint, true);
 		}
 		else if( bIsMacro )
 		{
-			// Refresh all nodes to make sure function signatures are up to date, etc.
+			// Just refresh all nodes in macro blueprints, but don't recompil
 			FBlueprintEditorUtils::RefreshAllNodes(Blueprint);
-
-			// Generate the skeleton class
-			FKismetEditorUtilities::GenerateBlueprintSkeleton(Blueprint);
 		}
 		else
 		{
-			// Generate the skeleton class
-			FKismetEditorUtilities::GenerateBlueprintSkeleton(Blueprint);
-
 			if (Blueprint->IsGeneratedClassAuthoritative() && (Blueprint->GeneratedClass != NULL))
 			{
 				check(PreviousCDO != NULL);
