Fixed faulty compile errors for non-actor BPs being duplicated by resolving differences in the duplication process for actor BPs and non-actor BPs.

diff --git a/Engine/Source/Editor/UnrealEd/Private/Kismet2/BlueprintEditorUtils.cpp b/Engine/Source/Editor/UnrealEd/Private/Kismet2/BlueprintEditorUtils.cpp
--- a/Engine/Source/Editor/UnrealEd/Private/Kismet2/BlueprintEditorUtils.cpp
+++ b/Engine/Source/Editor/UnrealEd/Private/Kismet2/BlueprintEditorUtils.cpp
@@ -1404,26 +1404,18 @@ void FBlueprintEditorUtils::PostDuplicateBlueprint(UBlueprint* Blueprint)
 			FKismetCompilerOptions CompileOptions;
 			CompileOptions.bIsDuplicationInstigated = true;
 
-			//SCS can change structure of the class
-			const bool bDontCompileClassBeforeScsIsSet = FBlueprintEditorUtils::SupportsConstructionScript(Blueprint) ||  FBlueprintEditorUtils::IsActorBased(Blueprint);
-			if (bDontCompileClassBeforeScsIsSet)
-			{ 
-				FName NewSkelClassName, NewGenClassName;
-				Blueprint->GetBlueprintClassNames(NewGenClassName, NewSkelClassName);
-				UBlueprintGeneratedClass* NewClass = ConstructObject<UBlueprintGeneratedClass>(
-					Blueprint->GetBlueprintClass(), Blueprint->GetOutermost(), NewGenClassName, RF_Public | RF_Transactional);
-
-				Blueprint->GeneratedClass = NewClass;
-				NewClass->ClassGeneratedBy = Blueprint;
-				NewClass->SetSuperStruct(Blueprint->ParentClass);
-
-				// Since we just duplicated the generated class above, we don't need to do a full compile below
-				CompileOptions.CompileType = EKismetCompileType::SkeletonOnly;
-			}
-			else
-			{
-				Compiler.CompileBlueprint(Blueprint, CompileOptions, Results);
-			}
+			FName NewSkelClassName, NewGenClassName;
+			Blueprint->GetBlueprintClassNames(NewGenClassName, NewSkelClassName);
+
+			UClass* NewClass = ConstructObject<UClass>(
+				Blueprint->GetBlueprintClass(), Blueprint->GetOutermost(), NewGenClassName, RF_Public | RF_Transactional);
+
+			Blueprint->GeneratedClass = NewClass;
+			NewClass->ClassGeneratedBy = Blueprint;
+			NewClass->SetSuperStruct(Blueprint->ParentClass);
+
+			// Since we just duplicated the generated class above, we don't need to do a full compile below
+			CompileOptions.CompileType = EKismetCompileType::SkeletonOnly;
 
 			TMap<UObject*, UObject*> OldToNewMap;
 
@@ -1483,10 +1475,7 @@ void FBlueprintEditorUtils::PostDuplicateBlueprint(UBlueprint* Blueprint)
 			Blueprint->ComponentTemplates = NewBPGC->ComponentTemplates;
 			Blueprint->Timelines = NewBPGC->Timelines;
 
-			if (bDontCompileClassBeforeScsIsSet)
-			{ 
-				Compiler.CompileBlueprint(Blueprint, CompileOptions, Results);
-			}
+			Compiler.CompileBlueprint(Blueprint, CompileOptions, Results);
 
 			// Create a new blueprint guid
 			Blueprint->GenerateNewGuid();
diff --git a/Engine/Source/Runtime/GameplayAbilities/Private/GameplayAbilityBlueprint.cpp b/Engine/Source/Runtime/GameplayAbilities/Private/GameplayAbilityBlueprint.cpp
--- a/Engine/Source/Runtime/GameplayAbilities/Private/GameplayAbilityBlueprint.cpp
+++ b/Engine/Source/Runtime/GameplayAbilities/Private/GameplayAbilityBlueprint.cpp
@@ -16,7 +16,7 @@ UGameplayAbilityBlueprint::UGameplayAbilityBlueprint(const FObjectInitializer& O
 
 UClass* UGameplayAbilityBlueprint::GetBlueprintClass() const
 {
-	return UGameplayAbilityBlueprint::StaticClass();
+	return UGameplayAbilityBlueprintGeneratedClass::StaticClass();
 }
 
 
