Improved clarity of the 'pak in apk' functionality in the UI for packaging projects on Android.

diff --git a/Engine/Source/Editor/UnrealEd/Classes/Settings/ProjectPackagingSettings.h b/Engine/Source/Editor/UnrealEd/Classes/Settings/ProjectPackagingSettings.h
--- a/Engine/Source/Editor/UnrealEd/Classes/Settings/ProjectPackagingSettings.h
+++ b/Engine/Source/Editor/UnrealEd/Classes/Settings/ProjectPackagingSettings.h
@@ -66,7 +66,7 @@ public:
 	bool IncludePrerequisites;
 
 	/** If enabled, on Android platforms, .pak files are placed inside the APK. */
-	UPROPERTY(config, EditAnywhere, Category = Packaging)
+	UPROPERTY(config, EditAnywhere, Category = Packaging, meta = (DisplayName = "Package .pak in APK"))
 	bool UseOBB_InAPK;
 
 	/**
