Added option to change the units used by the measuring tool.

diff --git a/Engine/Source/Editor/UnrealEd/Classes/Settings/LevelEditorViewportSettings.h b/Engine/Source/Editor/UnrealEd/Classes/Settings/LevelEditorViewportSettings.h
--- a/Engine/Source/Editor/UnrealEd/Classes/Settings/LevelEditorViewportSettings.h
+++ b/Engine/Source/Editor/UnrealEd/Classes/Settings/LevelEditorViewportSettings.h
@@ -35,6 +35,18 @@ enum EWASDType
 };
 
 
+/**
+ * Units used by measuring tool
+ */
+UENUM()
+enum EMeasuringToolUnits
+{
+	MeasureUnits_Centimeters UMETA(DisplayName="Centimeters"),
+	MeasureUnits_Meters UMETA(DisplayName="Meters"),
+	MeasureUnits_Kilometers UMETA(DisplayName="Kilometers")
+};
+
+
 /**
  * Implements the Level Editor's per-instance view port settings.
  */
@@ -381,6 +393,10 @@ public:
 	UPROPERTY(EditAnywhere, config, Category = LookAndFeel)
 	uint32 bSaveEngineStats : 1;
 
+	/** Specify the units used by the measuring tool */
+	UPROPERTY(EditAnywhere, config, Category = LookAndFeel)
+	TEnumAsByte<EMeasuringToolUnits> MeasuringToolUnits;
+
 private:
 
 	// Per-instance viewport settings.
diff --git a/Engine/Source/Editor/UnrealEd/Private/DragTool_Measure.cpp b/Engine/Source/Editor/UnrealEd/Private/DragTool_Measure.cpp
--- a/Engine/Source/Editor/UnrealEd/Private/DragTool_Measure.cpp
+++ b/Engine/Source/Editor/UnrealEd/Private/DragTool_Measure.cpp
@@ -60,7 +60,8 @@ void FDragTool_Measure::AddDelta(const FVector& InDelta)
 
 void FDragTool_Measure::Render(const FSceneView* View, FCanvas* Canvas)
 {
-	const float Length = FMath::RoundToFloat((PixelEnd - PixelStart).Size() * ViewportClient->GetOrthoUnitsPerPixel(ViewportClient->Viewport));
+	const float OrthoUnitsPerPixel = ViewportClient->GetOrthoUnitsPerPixel(ViewportClient->Viewport);
+	const float Length = FMath::RoundToFloat((PixelEnd - PixelStart).Size() * OrthoUnitsPerPixel);
 
 	if (View != nullptr && Canvas != nullptr && Length >= 1.f)
 	{
@@ -69,9 +70,32 @@ void FDragTool_Measure::Render(const FSceneView* View, FCanvas* Canvas)
 
 		const FVector2D PixelMid = FVector2D(PixelStart + ((PixelEnd - PixelStart) / 2));
 
+		// Calculate number of decimal places to display, based on the current viewport zoom
+		float Divisor = 1.0f;
+		int DecimalPlaces = 0;
+		const float OrderOfMagnitude = FMath::LogX(10.0f, OrthoUnitsPerPixel);
+
+		switch (GetDefault<ULevelEditorViewportSettings>()->MeasuringToolUnits)
+		{
+		case MeasureUnits_Meters:
+			Divisor = 100.0f;
+			// Max one decimal place allowed for meters.
+			DecimalPlaces = FMath::Clamp(FMath::FloorToInt(1.5f - OrderOfMagnitude), 0, 1);
+			break;
+
+		case MeasureUnits_Kilometers:
+			Divisor = 100000.0f;
+			// Max two decimal places allowed for kilometers.
+			DecimalPlaces = FMath::Clamp(FMath::FloorToInt(4.5f - OrderOfMagnitude), 0, 2);
+			break;
+		}
+
 		FNumberFormattingOptions Options;
 		Options.UseGrouping = false;
-		const FText LengthStr = FText::AsNumber( FMath::RoundToInt(Length), &Options );
+		Options.MinimumFractionalDigits = DecimalPlaces;
+		Options.MaximumFractionalDigits = DecimalPlaces;
+
+		const FText LengthStr = FText::AsNumber( Length / Divisor, &Options );
 
 
 		FCanvasTextItem TextItem( FVector2D( FMath::FloorToFloat(PixelMid.X), FMath::FloorToFloat(PixelMid.Y) ), LengthStr, GEngine->GetSmallFont(), FLinearColor::White );
diff --git a/Engine/Source/Editor/UnrealEd/Private/Settings/SettingsClasses.cpp b/Engine/Source/Editor/UnrealEd/Private/Settings/SettingsClasses.cpp
--- a/Engine/Source/Editor/UnrealEd/Private/Settings/SettingsClasses.cpp
+++ b/Engine/Source/Editor/UnrealEd/Private/Settings/SettingsClasses.cpp
@@ -183,8 +183,8 @@ ULevelEditorViewportSettings::ULevelEditorViewportSettings( const FObjectInitial
 {
 	bLevelStreamingVolumePrevis = false;
 	BillboardScale = 1.0f;
-
 	TransformWidgetSizeAdjustment = 0.0f;
+	MeasuringToolUnits = MeasureUnits_Centimeters;
 
 	// Set a default preview mesh
 	PreviewMeshes.Add(FStringAssetReference("/Engine/EditorMeshes/ColorCalibrator/SM_ColorCalibrator.SM_ColorCalibrator"));
