Editor - PkgInfo -export -simple now supports sorting

diff --git a/Engine/Source/Editor/UnrealEd/Private/Commandlets/PackageUtilities.cpp b/Engine/Source/Editor/UnrealEd/Private/Commandlets/PackageUtilities.cpp
--- a/Engine/Source/Editor/UnrealEd/Private/Commandlets/PackageUtilities.cpp
+++ b/Engine/Source/Editor/UnrealEd/Private/Commandlets/PackageUtilities.cpp
@@ -996,58 +996,56 @@ void FPkgInfoReporter_Log::GeneratePackageReport( ULinkerLoad* InLinker/*=NULL*/
 		GWarn->Log ( TEXT("Export Map"));
 		GWarn->Log ( TEXT("=========="));
 
-		if ( (InfoFlags&PKGINFO_Compact) == 0 )
+		TArray<FExportInfo> SortedExportMap;
+		SortedExportMap.Empty(Linker->ExportMap.Num());
+		for( int32 i = 0; i < Linker->ExportMap.Num(); ++i )
 		{
-			TArray<FExportInfo> SortedExportMap;
-			SortedExportMap.Empty(Linker->ExportMap.Num());
-			for( int32 i = 0; i < Linker->ExportMap.Num(); ++i )
-			{
-				new(SortedExportMap) FExportInfo(Linker, i);
-			}
+			new(SortedExportMap) FExportInfo(Linker, i);
+		}
 
-			FString SortingParms;
-			if ( FParse::Value(FCommandLine::Get(), TEXT("SORT="), SortingParms) )
-			{
-				TArray<FString> SortValues;
-				SortingParms.ParseIntoArray(&SortValues, TEXT(","), true);
+		FString SortingParms;
+		if ( FParse::Value(FCommandLine::Get(), TEXT("SORT="), SortingParms) )
+		{
+			TArray<FString> SortValues;
+			SortingParms.ParseIntoArray(&SortValues, TEXT(","), true);
 
-				for ( int32 i = 0; i < EXPORTSORT_MAX; i++ )
+			for ( int32 i = 0; i < EXPORTSORT_MAX; i++ )
+			{
+				if ( i < SortValues.Num() )
 				{
-					if ( i < SortValues.Num() )
+					const FString Value = SortValues[i];
+					if ( Value == TEXT("index") )
 					{
-						const FString Value = SortValues[i];
-						if ( Value == TEXT("index") )
-						{
-							FObjectExport_Sorter::SortPriority[i] = EXPORTSORT_ExportIndex;
-						}
-						else if ( Value == TEXT("size") )
-						{
-							FObjectExport_Sorter::SortPriority[i] = EXPORTSORT_ExportSize;
-						}
-						else if ( Value == TEXT("name") )
-						{
-							FObjectExport_Sorter::SortPriority[i] = EXPORTSORT_ObjectPathname;
-						}
-						else if ( Value == TEXT("outer") )
-						{
-							FObjectExport_Sorter::SortPriority[i] = EXPORTSORT_OuterPathname;
-						}
+						FObjectExport_Sorter::SortPriority[i] = EXPORTSORT_ExportIndex;
 					}
-					else
+					else if ( Value == TEXT("size") )
+					{
+						FObjectExport_Sorter::SortPriority[i] = EXPORTSORT_ExportSize;
+					}
+					else if ( Value == TEXT("name") )
 					{
-						FObjectExport_Sorter::SortPriority[i] = EXPORTSORT_MAX;
+						FObjectExport_Sorter::SortPriority[i] = EXPORTSORT_ObjectPathname;
+					}
+					else if ( Value == TEXT("outer") )
+					{
+						FObjectExport_Sorter::SortPriority[i] = EXPORTSORT_OuterPathname;
 					}
 				}
+				else
+				{
+					FObjectExport_Sorter::SortPriority[i] = EXPORTSORT_MAX;
+				}
 			}
+		}
 
-			SortedExportMap.Sort( FObjectExport_Sorter() );
+		SortedExportMap.Sort( FObjectExport_Sorter() );
 
-			for( int32 SortedIndex = 0; SortedIndex < SortedExportMap.Num(); ++SortedIndex )
+		if ( (InfoFlags&PKGINFO_Compact) == 0 )
+		{
+			for( const auto& ExportInfo : SortedExportMap )
 			{
 				GWarn->Log ( TEXT("\t*************************"));
-				FExportInfo& ExportInfo = SortedExportMap[SortedIndex];
-
-				FObjectExport& Export = ExportInfo.Export;
+				const FObjectExport& Export = ExportInfo.Export;
 
 				UE_LOG(LogPackageUtilities, Warning, TEXT("\tExport %d: '%s'"), ExportInfo.ExportIndex, *Export.ObjectName.ToString() );
 
@@ -1147,12 +1145,12 @@ void FPkgInfoReporter_Log::GeneratePackageReport( ULinkerLoad* InLinker/*=NULL*/
 		}
 		else
 		{
-			for( int32 ExportIndex=0; ExportIndex<Linker->ExportMap.Num(); ExportIndex++ )
+			for( const auto& ExportInfo : SortedExportMap )
 			{
-				const FObjectExport& Export = Linker->ExportMap[ExportIndex];
-				UE_LOG(LogPackageUtilities, Warning, TEXT("  %8i %10i %32s %s"), ExportIndex, Export.SerialSize, 
-					*(Linker->GetExportClassName(ExportIndex).ToString()), 
-					(InfoFlags&PKGINFO_Paths) != 0 ? *Linker->GetExportPathName(ExportIndex) : *Export.ObjectName.ToString());
+				const FObjectExport& Export = ExportInfo.Export;
+				UE_LOG(LogPackageUtilities, Warning, TEXT("  %8i %10i %32s %s"), ExportInfo.ExportIndex, Export.SerialSize, 
+					*(Linker->GetExportClassName(ExportInfo.ExportIndex).ToString()), 
+					(InfoFlags&PKGINFO_Paths) != 0 ? *Linker->GetExportPathName(ExportInfo.ExportIndex) : *Export.ObjectName.ToString());
 			}
 		}
 	}
