Editor Performance Capture tool now grabs data for every second without skipping seconds.

diff --git a/Engine/Source/Editor/UnrealEd/Private/Tests/EditorPerformanceTests.cpp b/Engine/Source/Editor/UnrealEd/Private/Tests/EditorPerformanceTests.cpp
--- a/Engine/Source/Editor/UnrealEd/Private/Tests/EditorPerformanceTests.cpp
+++ b/Engine/Source/Editor/UnrealEd/Private/Tests/EditorPerformanceTests.cpp
@@ -19,6 +19,7 @@ struct EditorPerfCaptureParameters
 
 	//Saved Performance Values
 	float MapLoadTime;
+	int64 Counter;
 	TArray<float> AverageFPS;
 	TArray<float> AverageFrameTime;
 	TArray<float> UsedPhysical;
@@ -34,8 +35,10 @@ struct EditorPerfCaptureParameters
 		: MapName(TEXT("None"))
 		, TestDuration(60)
 		, MapLoadTime(0)
+		, Counter(0)
 	{
 	}
+
 };
 
 //////////////////////////////////////////////////////////////////////////
@@ -135,7 +138,48 @@ void EditorPerfDump(EditorPerfCaptureParameters& EditorPerfStats)
 	UE_LOG(LogEditorAutomationTests, Log, TEXT("Raw performance csv file is located here: %s"), *FPaths::ConvertRelativePathToFull(RAWCSVFilePath));
 }
 
+void CaptureEditorData(EditorPerfCaptureParameters& OutEditorPerfStats)
+{
+	//Capture the current time stamp and format it to YYYY-MM-DD HH:MM:SS.mmm.
+	FDateTime CurrentDateAndTime = FDateTime::Now();
+	
+	//Find the Average FPS
+	//Clamp to avoid huge averages at startup or after hitches
+	const float CurrentFPS = 1.0f / FSlateApplication::Get().GetAverageDeltaTime();
+	const float ClampedFPS = (CurrentFPS < 0.0f || CurrentFPS > 4000.0f) ? 0.0f : CurrentFPS;
+	OutEditorPerfStats.AverageFPS.Add(ClampedFPS);
+
+	//Find the Frame Time in ms.
+	//Clamp to avoid huge averages at startup or after hitches
+	const float AverageMS = FSlateApplication::Get().GetAverageDeltaTime() * 1000.0f;
+	const float ClampedMS = (AverageMS < 0.0f || AverageMS > 4000.0f) ? 0.0f : AverageMS;
+	OutEditorPerfStats.AverageFrameTime.Add(ClampedMS);
+
+	//Query OS for process memory used.
+	FPlatformMemoryStats MemoryStats = FPlatformMemory::GetStats();
+	OutEditorPerfStats.UsedPhysical.Add((float)MemoryStats.UsedPhysical / 1024);
+
+	//Query OS for available physical memory
+	OutEditorPerfStats.AvailablePhysical.Add((float)MemoryStats.AvailablePhysical / 1024);
+
+	//Query OS for available virtual memory
+	OutEditorPerfStats.AvailableVirtual.Add((float)MemoryStats.AvailableVirtual / 1024);
 
+	//Query OS for used virtual memory
+	OutEditorPerfStats.UsedVirtual.Add((float)MemoryStats.UsedVirtual / 1024);
+
+	//Query OS for used Peak Used physical memory
+	OutEditorPerfStats.PeakUsedPhysical.Add((float)MemoryStats.PeakUsedPhysical / 1024);
+
+	//Query OS for used Peak Used virtual memory
+	OutEditorPerfStats.PeakUsedVirtual.Add((float)MemoryStats.PeakUsedVirtual / 1024);
+
+	//Capture the time stamp.
+	FString FormatedTimeStamp = FString::Printf(TEXT("%04i-%02i-%02i %02i:%02i:%02i.%03i"), CurrentDateAndTime.GetYear(), CurrentDateAndTime.GetMonth(), CurrentDateAndTime.GetDay(), CurrentDateAndTime.GetHour(), CurrentDateAndTime.GetMinute(), CurrentDateAndTime.GetSecond(), CurrentDateAndTime.GetMillisecond());
+	OutEditorPerfStats.FormattedTimeStamp.Add(FormatedTimeStamp);
+	OutEditorPerfStats.TimeStamp.Add(CurrentDateAndTime);
+
+}
 
 //////////////////////////////////////////////////////////////////////////
 //Editor Performance Latent Commands
@@ -151,69 +195,20 @@ DEFINE_LATENT_AUTOMATION_COMMAND_ONE_PARAMETER(FEditorPerfCaptureCommand, Editor
 */
 bool FEditorPerfCaptureCommand::Update()
 {
-	//Capture the current time stamp and format it to YYYY-MM-DD HH:MM:SS.mmm.
-	FDateTime CurrentDateAndTime = FDateTime::Now();
-
-	//This is how long it has been since the last run through.
-	FTimespan ElapsedTime = 0;
-	FTimespan TimeBetweenCaptures;
-
-	//We want to only capture data every whole second.
-	if (EditorPerfStats.TimeStamp.Num() > 0)
+	float ElapsedTime = FPlatformTime::Seconds() - StartTime;
+	
+	if ((ElapsedTime <= (EditorPerfStats.TestDuration + 1.0f)))
 	{
-		TimeBetweenCaptures = CurrentDateAndTime.GetTicks() - EditorPerfStats.TimeStamp.Last().GetTicks();
-		if (TimeBetweenCaptures.GetTicks() < ETimespan::TicksPerSecond)
+		if (((int64)ElapsedTime - EditorPerfStats.Counter) >= 1)
 		{
-			return false;
+			EditorPerfStats.Counter = ElapsedTime;
+			CaptureEditorData(EditorPerfStats);
 		}
-
-		ElapsedTime = CurrentDateAndTime.GetTicks() - EditorPerfStats.TimeStamp[0].GetTicks();
-	}
-
-	if (ElapsedTime.GetTotalSeconds() <= EditorPerfStats.TestDuration)
-	{
-		//Find the Average FPS
-		//Clamp to avoid huge averages at startup or after hitches
-		const float CurrentFPS = 1.0f / FSlateApplication::Get().GetAverageDeltaTime();
-		const float ClampedFPS = (CurrentFPS < 0.0f || CurrentFPS > 4000.0f) ? 0.0f : CurrentFPS;
-		EditorPerfStats.AverageFPS.Add(ClampedFPS);
-
-		//Find the Frame Time in ms.
-		//Clamp to avoid huge averages at startup or after hitches
-		const float AverageMS = FSlateApplication::Get().GetAverageDeltaTime() * 1000.0f;
-		const float ClampedMS = (AverageMS < 0.0f || AverageMS > 4000.0f) ? 0.0f : AverageMS;
-		EditorPerfStats.AverageFrameTime.Add(ClampedMS);
-
-		//Query OS for process memory used.
-		FPlatformMemoryStats MemoryStats = FPlatformMemory::GetStats();
-		EditorPerfStats.UsedPhysical.Add((float)MemoryStats.UsedPhysical / 1024);
-
-		//Query OS for available physical memory
-		EditorPerfStats.AvailablePhysical.Add((float)MemoryStats.AvailablePhysical / 1024);
-
-		//Query OS for available virtual memory
-		EditorPerfStats.AvailableVirtual.Add((float)MemoryStats.AvailableVirtual / 1024);
-
-		//Query OS for used virtual memory
-		EditorPerfStats.UsedVirtual.Add((float)MemoryStats.UsedVirtual / 1024);
-
-		//Query OS for used Peak Used physical memory
-		EditorPerfStats.PeakUsedPhysical.Add((float)MemoryStats.PeakUsedPhysical / 1024);
-
-		//Query OS for used Peak Used virtual memory
-		EditorPerfStats.PeakUsedVirtual.Add((float)MemoryStats.PeakUsedVirtual / 1024);
-
-		//Capture the time stamp.
-		FString FormatedTimeStamp = FString::Printf(TEXT("%04i-%02i-%02i %02i:%02i:%02i.%03i"), CurrentDateAndTime.GetYear(), CurrentDateAndTime.GetMonth(), CurrentDateAndTime.GetDay(), CurrentDateAndTime.GetHour(), CurrentDateAndTime.GetMinute(), CurrentDateAndTime.GetSecond(), CurrentDateAndTime.GetMillisecond());
-		EditorPerfStats.FormattedTimeStamp.Add(FormatedTimeStamp);
-		EditorPerfStats.TimeStamp.Add(CurrentDateAndTime);
-
 		return false;
 	}
 
 	//Dump the performance data in a csv file.
 	EditorPerfDump(EditorPerfStats);
-
 	return true;
 }
 
