Added Mip min / max lod settings to the cooker version checking for -iterate.
Added ddc texture / shader / audio version numbers to cooker version checking for -iterate.
Added -iterate as an experimental option for launch on.

diff --git a/Engine/Source/Editor/UnrealEd/Classes/Settings/EditorExperimentalSettings.h b/Engine/Source/Editor/UnrealEd/Classes/Settings/EditorExperimentalSettings.h
--- a/Engine/Source/Editor/UnrealEd/Classes/Settings/EditorExperimentalSettings.h
+++ b/Engine/Source/Editor/UnrealEd/Classes/Settings/EditorExperimentalSettings.h
@@ -77,6 +77,10 @@ public:
 	UPROPERTY(EditAnywhere, config, Category = Cooking, meta = (DisplayName = "Cook On The Side (Run a cook on the fly server in the background of the editor)"))
 	bool bCookOnTheSide;
 
+	/** Enable -iterate for launch on */
+	UPROPERTY(EditAnywhere, config, Category = Cooking, meta = (DisplayName = "Iterative cooking for builds launched form the editor (launch on)"))
+	bool bIterativeCookingForLaunchOn;
+
 	/** Enables Environment Queries editor */
 	UPROPERTY(EditAnywhere, config, Category = AI, meta = (DisplayName = "Environment Querying System"))
 	bool bEQSEditor;
diff --git a/Engine/Source/Editor/UnrealEd/Private/CookOnTheFlyServer.cpp b/Engine/Source/Editor/UnrealEd/Private/CookOnTheFlyServer.cpp
--- a/Engine/Source/Editor/UnrealEd/Private/CookOnTheFlyServer.cpp
+++ b/Engine/Source/Editor/UnrealEd/Private/CookOnTheFlyServer.cpp
@@ -1039,15 +1039,17 @@ uint32 UCookOnTheFlyServer::TickCookOnTheSide( const float TimeSlice, uint32 &Co
 				
 				for ( const auto& Obj : ObjectsInPackage )
 				{
-					if ( Timer.IsTimeUp() && IsRealtimeMode() )
+					Obj->BeginCacheForCookedPlatformData( TargetPlatform );
+					if ( Obj->IsCachedCookedPlatformDataLoaded(TargetPlatform) == false )
 					{
+						UE_LOG(LogCookOnTheFly, Display, TEXT("Object %s isn't cached yet"), *Obj->GetFullName());
 						bIsAllDataCached = false;
 						break;
 					}
 
-					Obj->BeginCacheForCookedPlatformData( TargetPlatform );
-					if ( Obj->IsCachedCookedPlatformDataLoaded(TargetPlatform) == false )
+					if ( Timer.IsTimeUp() && IsRealtimeMode() )
 					{
+						UE_LOG(LogCookOnTheFly, Display, TEXT("Object %s took too long to cache"), *Obj->GetFullName());
 						bIsAllDataCached = false;
 						break;
 					}
@@ -1056,7 +1058,7 @@ uint32 UCookOnTheFlyServer::TickCookOnTheSide( const float TimeSlice, uint32 &Co
 				if ( bIsAllDataCached == false )
 				{
 					break;
-			}
+				}
 			}
 		}
 
@@ -1361,7 +1363,7 @@ void UCookOnTheFlyServer::MarkPackageDirtyForCooker( UPackage *Package )
 		/*FString PackageFilename(GetPackageFilename(Package));
 		FPaths::MakeStandardFilename(PackageFilename);*/
 		FString PackageFilename = GetCachedStandardPackageFilename(Package);
-		UE_LOG(LogCookOnTheFly, Display, TEXT("Modification detected to package %s"), *PackageFilename);
+		// UE_LOG(LogCookOnTheFly, Display, TEXT("Modification detected to package %s"), *PackageFilename);
 		const FName PackageFFileName = FName(*PackageFilename);
 
 		if ( CurrentCookMode == ECookMode::CookByTheBookFromTheEditor )
@@ -1758,14 +1760,31 @@ FString UCookOnTheFlyServer::GetOutputDirectoryOverride( const FString &OutputDi
 	return OutputDirectory;
 }
 
-bool UCookOnTheFlyServer::GetCurrentIniVersionStrings( const ITargetPlatform* TargetPlatform, TArray<FString> &IniVersionStrings ) const
+template<class T>
+void GetVersionFormatNumbersForIniVersionStrings( TArray<FString>& IniVersionStrings, const FString& FormatName, const TArray<T> &FormatArray )
+{
+	for ( const auto& Format : FormatArray )
+	{
+		TArray<FName> SupportedFormats;
+		Format->GetSupportedFormats(SupportedFormats);
+		for ( const auto& SupportedFormat : SupportedFormats )
+		{
+			int32 VersionNumber = Format->GetVersion(SupportedFormat);
+			FString IniVersionString = FString::Printf( TEXT("%s:%s:VersionNumber%d"), *FormatName, *SupportedFormat.ToString(), VersionNumber);
+			IniVersionStrings.Emplace( IniVersionString );
+		}
+	}
+}
+
+bool UCookOnTheFlyServer::GetCurrentIniVersionStrings( const ITargetPlatform* TargetPlatform, TArray<FString>& IniVersionStrings ) const
 {
 	// there is a list of important ini settings in the Editor config 
 	TArray<FString> IniVersionedParams;
 	GConfig->GetArray( TEXT("CookSettings"), TEXT("VersionedIniParams"), IniVersionedParams, GEditorIni );
 
-	/*const FString EditorIni = FPaths::GameDir() / GEditorIni;
-	const FString SandboxEditorIni = SandboxFile->ConvertToAbsolutePathForExternalAppForWrite(*EditorIni);*/
+
+	// used to store temporary platform specific ini files
+	TMap<FString,FConfigFile*> PlatformIniFiles;
 
 	// if the old one doesn't contain all the settings in the new one then we fail this check
 	for ( const auto& IniVersioned : IniVersionedParams )
@@ -1784,11 +1803,18 @@ bool UCookOnTheFlyServer::GetCurrentIniVersionStrings( const ITargetPlatform* Ta
 		const FString& Section = IniVersionedArray[1];
 		const FString& Key = IniVersionedArray[2];
 
-		const FString& IniFilename = FPaths::GeneratedConfigDir() / TargetPlatform->IniPlatformName() / Filename + TEXT(".ini");
+		// const FString& IniFilename = FPaths::GeneratedConfigDir() / TargetPlatform->IniPlatformName() / Filename + TEXT(".ini");
+		FConfigFile *PlatformIniFile = PlatformIniFiles.FindRef( Filename );
+		if ( PlatformIniFile == NULL )
+		{
+			PlatformIniFile = new FConfigFile();
+			FConfigCacheIni::LoadLocalIniFile( *PlatformIniFile, *Filename, true, *TargetPlatform->IniPlatformName() );
+			PlatformIniFiles.Add( Filename, PlatformIniFile );
+		}
 
 		// get the value of the entry
 		FString Value;
-		if ( !GConfig->GetString(*Section, *Key, Value, IniFilename) )
+		if ( !PlatformIniFile->GetString(*Section, *Key, Value) )
 		{
 			UE_LOG(LogCookOnTheFly, Warning, TEXT("Unable to find entry in CookSettings, VersionedIniParams %s, assume default is being used"), *IniVersioned);
 			continue;
@@ -1799,6 +1825,13 @@ bool UCookOnTheFlyServer::GetCurrentIniVersionStrings( const ITargetPlatform* Ta
 		IniVersionStrings.Emplace( MoveTemp(CurrentVersionString) );
 	}
 
+	// clean up our temporary platform ini files
+	for ( const auto& PlatformIniFile : PlatformIniFiles )
+	{
+		delete PlatformIniFile.Value;
+	}
+	PlatformIniFiles.Empty();
+
 
 	const FTextureLODSettings& LodSettings = TargetPlatform->GetTextureLODSettings();
 
@@ -1808,12 +1841,62 @@ bool UCookOnTheFlyServer::GetCurrentIniVersionStrings( const ITargetPlatform* Ta
 	for ( int I = 0; I < TextureGroup::TEXTUREGROUP_MAX; ++I )
 	{
 		const TextureMipGenSettings& MipGenSettings = LodSettings.GetTextureMipGenSettings((TextureGroup)(I));
+		FString MipGenVersionString = FString::Printf( TEXT("TextureLODGroupMipGenSettings:%s:%s"), *TextureGroupEnum->GetEnumName( I ), *TextureMipGenSettingsEnum->GetEnumName((int32)(MipGenSettings)) );
+		IniVersionStrings.Emplace( MoveTemp( MipGenVersionString ) );
 
-		FString VersionString = FString::Printf( TEXT("TextureLODGroup:%s:%s"), *TextureGroupEnum->GetEnumName( I ), *TextureMipGenSettingsEnum->GetEnumName((int32)(MipGenSettings)) );
+		const int32 MinMipCount = LodSettings.GetMinLODMipCount((TextureGroup)(I));
+		FString MinMipVersionString = FString::Printf( TEXT("TextureLODGroupMinMipCount:%s:%d"), *TextureGroupEnum->GetEnumName( I ), MinMipCount);
+		IniVersionStrings.Emplace( MoveTemp( MinMipVersionString ) );
 
-		IniVersionStrings.Emplace( MoveTemp( VersionString ) );
+		const int32 MaxMipCount = LodSettings.GetMaxLODMipCount((TextureGroup)(I));
+		FString MaxMipVersionString = FString::Printf( TEXT("TextureLODGroupMaxMipCount:%s:%d"), *TextureGroupEnum->GetEnumName( I ), MaxMipCount);
+		IniVersionStrings.Emplace( MoveTemp( MaxMipVersionString ) );
 	}
-	
+
+	// save off the ddc version numbers also
+	ITargetPlatformManagerModule* TPM = GetTargetPlatformManager();
+	check(TPM);
+
+	GetVersionFormatNumbersForIniVersionStrings( IniVersionStrings, TEXT("AudioFormat"), TPM->GetAudioFormats() );
+	GetVersionFormatNumbersForIniVersionStrings( IniVersionStrings, TEXT("TextureFormat"), TPM->GetTextureFormats() );
+	GetVersionFormatNumbersForIniVersionStrings( IniVersionStrings, TEXT("ShaderFormat"), TPM->GetShaderFormats() );
+	/*
+	for ( const auto& AudioFormat : TPM->GetAudioFormats() )
+	{
+		TArray<FName> SupportedFormats;
+		AudioFormat.GetSupportedFormats(SupportedFormats);
+		for ( const auto& SupportedFormat : SupportedFormats )
+		{
+			int32 VersionNumber = AudioFormat.GetVersion(SupportedFormat);
+			FString AudioString = FString::Printf( TEXT("AudioFormat:%s:VersionNumber%d"), SupportedFormat.ToString(), VersionNumber);
+			IniVersionStrings.Emplace( AudioString );
+		}
+	}
+
+	for ( const auto& TextureFormat : TPM->GetTextureFormats() )
+	{
+		TArray<FName> SupportedFormats;
+		TextureFormat.GetSupportedFormats(SupportedFormats);
+		for ( const auto& SupportedFormat : SupportedFormats )
+		{
+			int32 VersionNumber = AudioFormat.GetVersion(SupportedFormat);
+			FString AudioString = FString::Printf( TEXT("AudioFormat:%s:VersionNumber%d"), SupportedFormat.ToString(), VersionNumber);
+			IniVersionStrings.Emplace( AudioString );
+		}
+	}
+
+	for ( const auto& AudioFormat : TPM->GetAudioFormats() )
+	{
+		TArray<FName> SupportedFormats;
+		AudioFormat.GetSupportedFormats(SupportedFormats);
+		for ( const auto& SupportedFormat : SupportedFormats )
+		{
+			int32 VersionNumber = AudioFormat.GetVersion(SupportedFormat);
+			FString AudioString = FString::Printf( TEXT("AudioFormat:%s:VersionNumber%d"), SupportedFormat.ToString(), VersionNumber);
+			IniVersionStrings.Emplace( AudioString );
+		}
+	}
+	*/
 
 
 	return true;
diff --git a/Engine/Source/Editor/UnrealEd/Private/UnrealEdEngine.cpp b/Engine/Source/Editor/UnrealEd/Private/UnrealEdEngine.cpp
--- a/Engine/Source/Editor/UnrealEd/Private/UnrealEdEngine.cpp
+++ b/Engine/Source/Editor/UnrealEd/Private/UnrealEdEngine.cpp
@@ -111,11 +111,12 @@ void UUnrealEdEngine::Init(IEngineLoop* InEngineLoop)
 	}
 
 	UEditorExperimentalSettings const* ExperimentalSettings =  GetDefault<UEditorExperimentalSettings>();
-
+	ECookInitializationFlags BaseCookingFlags = ECookInitializationFlags::AutoTick | ECookInitializationFlags::AsyncSave;
+	BaseCookingFlags |= ExperimentalSettings->bIterativeCookingForLaunchOn ? ECookInitializationFlags::Iterative : ECookInitializationFlags::None;
 	if ( ExperimentalSettings->bCookInTheEditor )
 	{
 		CookServer = ConstructObject<UCookOnTheFlyServer>( UCookOnTheFlyServer::StaticClass() );
-		CookServer->Initialize( ECookMode::CookByTheBookFromTheEditor, ECookInitializationFlags::AutoTick | ECookInitializationFlags::AsyncSave );
+		CookServer->Initialize( ECookMode::CookByTheBookFromTheEditor, BaseCookingFlags );
 
 		FCoreUObjectDelegates::OnObjectPropertyChanged.AddUObject(CookServer, &UCookOnTheFlyServer::OnObjectPropertyChanged);
 		FCoreUObjectDelegates::OnObjectModified.AddUObject(CookServer, &UCookOnTheFlyServer::OnObjectModified);
@@ -123,7 +124,7 @@ void UUnrealEdEngine::Init(IEngineLoop* InEngineLoop)
 	else if ( ExperimentalSettings->bCookOnTheSide )
 	{
 		CookServer = ConstructObject<UCookOnTheFlyServer>( UCookOnTheFlyServer::StaticClass() );
-		CookServer->Initialize( ECookMode::CookOnTheFly, ECookInitializationFlags::AutoTick | ECookInitializationFlags::AsyncSave );
+		CookServer->Initialize( ECookMode::CookOnTheFly, BaseCookingFlags );
 		CookServer->StartNetworkFileServer( false );
 
 		FCoreUObjectDelegates::OnObjectPropertyChanged.AddUObject(CookServer, &UCookOnTheFlyServer::OnObjectPropertyChanged);
diff --git a/Engine/Source/Runtime/Engine/Private/TextureLODSettings.cpp b/Engine/Source/Runtime/Engine/Private/TextureLODSettings.cpp
--- a/Engine/Source/Runtime/Engine/Private/TextureLODSettings.cpp
+++ b/Engine/Source/Runtime/Engine/Private/TextureLODSettings.cpp
@@ -262,6 +262,48 @@ int32 FTextureLODSettings::GetTextureLODGroupLODBias( int32 InLODGroup ) const
 	return Retval;
 }
 
+
+
+/**
+ * Returns the LODGroup setting for number of streaming mip-levels.
+ * -1 means that all mip-levels are allowed to stream.
+ *
+ * @param	InLODGroup		The LOD Group ID 
+ * @return	Number of streaming mip-levels for textures in the specified LODGroup
+ */
+int32 FTextureLODSettings::GetMinLODMipCount( int32 InLODGroup ) const
+{
+	int32 Retval = 0;
+
+	const FTextureLODGroup& LODGroup = TextureLODGroups[InLODGroup]; 
+
+	Retval = LODGroup.MinLODMipCount;
+
+	return Retval;
+}
+
+
+
+
+/**
+ * Returns the LODGroup setting for number of streaming mip-levels.
+ * -1 means that all mip-levels are allowed to stream.
+ *
+ * @param	InLODGroup		The LOD Group ID 
+ * @return	Number of streaming mip-levels for textures in the specified LODGroup
+ */
+int32 FTextureLODSettings::GetMaxLODMipCount( int32 InLODGroup ) const
+{
+	int32 Retval = 0;
+
+	const FTextureLODGroup& LODGroup = TextureLODGroups[InLODGroup]; 
+
+	Retval = LODGroup.MaxLODMipCount;
+
+	return Retval;
+}
+
+
 /**
  * Returns the LODGroup setting for number of streaming mip-levels.
  * -1 means that all mip-levels are allowed to stream.
diff --git a/Engine/Source/Runtime/Engine/Public/TextureLODSettings.h b/Engine/Source/Runtime/Engine/Public/TextureLODSettings.h
--- a/Engine/Source/Runtime/Engine/Public/TextureLODSettings.h
+++ b/Engine/Source/Runtime/Engine/Public/TextureLODSettings.h
@@ -74,6 +74,9 @@ struct FTextureLODSettings
 	 */
 	ENGINE_API int32 GetNumStreamedMips( int32 InLODGroup ) const;
 
+	ENGINE_API int32 GetMinLODMipCount( int32 InLODGroup ) const;
+	ENGINE_API int32 GetMaxLODMipCount( int32 InLODGroup ) const;
+
 	/**
 	 * Returns the filter state that should be used for the passed in texture, taking
 	 * into account other system settings.
