#UE4 Fixed a bug that happened when attempting to use save-as to save a world asset over another world asset.

diff --git a/Engine/Source/Editor/UnrealEd/Private/PackageTools.cpp b/Engine/Source/Editor/UnrealEd/Private/PackageTools.cpp
--- a/Engine/Source/Editor/UnrealEd/Private/PackageTools.cpp
+++ b/Engine/Source/Editor/UnrealEd/Private/PackageTools.cpp
@@ -323,13 +323,6 @@ namespace PackageTools
 			for ( int32 PackageIndex = 0 ; PackageIndex < PackagesToUnload.Num() ; ++PackageIndex )
 			{
 				PackageBeingUnloaded = PackagesToUnload[PackageIndex];
-	
-				if( PackageBeingUnloaded->IsDirty() )
-				{
-					// The package was marked dirty as a result of something that happened below (e.g callbacks in CollectGarbage).  
-					// Dirty packages we actually care about unloading were filtered above so if the package becomes dirty here it should still be unloaded
-					PackageBeingUnloaded->SetDirtyFlag(false);
-				}
 
 				GWarn->StatusUpdate( PackageIndex, PackagesToUnload.Num(), FText::Format(NSLOCTEXT("UnrealEd", "Unloadingf", "Unloading {0}..."), FText::FromString(PackageBeingUnloaded->GetName()) ) );
 
@@ -356,6 +349,13 @@ namespace PackageTools
 
 				// Collect garbage.
 				CollectGarbage( GARBAGE_COLLECTION_KEEPFLAGS );
+				
+				if( PackageBeingUnloaded->IsDirty() )
+				{
+					// The package was marked dirty as a result of something that happened above (e.g callbacks in CollectGarbage).  
+					// Dirty packages we actually care about unloading were filtered above so if the package becomes dirty here it should still be unloaded
+					PackageBeingUnloaded->SetDirtyFlag(false);
+				}
 
 				// Cleanup.
 				ObjectsThatHadFlagsCleared.Empty();
