[UE-2889] BP compiler errors can now be double clicked in the message log to focus on the problematic node.

diff --git a/Engine/Source/Developer/MessageLog/Private/UserInterface/SMessageLogListing.cpp b/Engine/Source/Developer/MessageLog/Private/UserInterface/SMessageLogListing.cpp
--- a/Engine/Source/Developer/MessageLog/Private/UserInterface/SMessageLogListing.cpp
+++ b/Engine/Source/Developer/MessageLog/Private/UserInterface/SMessageLogListing.cpp
@@ -183,6 +183,17 @@ void SMessageLogListing::BroadcastMessageTokenClicked( TSharedPtr<FTokenizedMess
 	MessageLogListingViewModel->ExecuteToken(Token);
 }
 
+void SMessageLogListing::BroadcastMessageDoubleClicked(TSharedPtr< class FTokenizedMessage > Message)
+{
+	if (Message->GetMessageTokens().Num() > 0)
+	{
+		TSharedPtr<IMessageToken> MessageLink = Message->GetMessageLink();
+		if (MessageLink.IsValid())
+		{
+			MessageLogListingViewModel->ExecuteToken(MessageLink->AsShared());
+		}
+	}	
+}
 
 const TArray< TSharedRef<FTokenizedMessage> > SMessageLogListing::GetSelectedMessages() const
 {
@@ -240,7 +251,8 @@ TSharedRef<ITableRow> SMessageLogListing::MakeMessageLogListItemWidget( TSharedR
 	return
 		SNew(SMessageLogMessageListRow, OwnerTable)
 		.Message(Message)
-		.OnTokenClicked( this, &SMessageLogListing::BroadcastMessageTokenClicked );
+		.OnTokenClicked( this, &SMessageLogListing::BroadcastMessageTokenClicked )
+		.OnMessageDoubleClicked( this, &SMessageLogListing::BroadcastMessageDoubleClicked );
 }
 
 
diff --git a/Engine/Source/Developer/MessageLog/Private/UserInterface/SMessageLogListing.h b/Engine/Source/Developer/MessageLog/Private/UserInterface/SMessageLogListing.h
--- a/Engine/Source/Developer/MessageLog/Private/UserInterface/SMessageLogListing.h
+++ b/Engine/Source/Developer/MessageLog/Private/UserInterface/SMessageLogListing.h
@@ -27,6 +27,8 @@ public:
 
 	/** Used to execute the 'on clicked token' delegate */
 	void BroadcastMessageTokenClicked( TSharedPtr< class FTokenizedMessage > Message, const TSharedRef<class IMessageToken>& Token );
+	/** Used to execute a message's featured action (a token associated with the entire message) */
+	void BroadcastMessageDoubleClicked(TSharedPtr< class FTokenizedMessage > Message);
 
 	/** Gets a list of the selected messages */
 	const TArray< TSharedRef< class FTokenizedMessage > > GetSelectedMessages() const;
diff --git a/Engine/Source/Developer/MessageLog/Private/UserInterface/SMessageLogMessageListRow.cpp b/Engine/Source/Developer/MessageLog/Private/UserInterface/SMessageLogMessageListRow.cpp
--- a/Engine/Source/Developer/MessageLog/Private/UserInterface/SMessageLogMessageListRow.cpp
+++ b/Engine/Source/Developer/MessageLog/Private/UserInterface/SMessageLogMessageListRow.cpp
@@ -21,6 +21,7 @@ TSharedRef<SWidget> SMessageLogMessageListRow::CreateHyperlink(const TSharedRef<
 void SMessageLogMessageListRow::Construct(const FArguments& InArgs, const TSharedRef< STableViewBase >& InOwnerTableView)
 {
 	this->OnTokenClicked = InArgs._OnTokenClicked;
+	this->OnMessageDoubleClicked = InArgs._OnMessageDoubleClicked;
 
 	Message = InArgs._Message.Get();
 
diff --git a/Engine/Source/Developer/MessageLog/Private/UserInterface/SMessageLogMessageListRow.h b/Engine/Source/Developer/MessageLog/Private/UserInterface/SMessageLogMessageListRow.h
--- a/Engine/Source/Developer/MessageLog/Private/UserInterface/SMessageLogMessageListRow.h
+++ b/Engine/Source/Developer/MessageLog/Private/UserInterface/SMessageLogMessageListRow.h
@@ -18,15 +18,18 @@ class MESSAGELOG_API SMessageLogMessageListRow
 public:
 
 	DECLARE_DELEGATE_TwoParams( FOnTokenClicked, TSharedPtr<FTokenizedMessage>, const TSharedRef<class IMessageToken>& );
+	DECLARE_DELEGATE_OneParam( FOnMessageClicked, TSharedPtr<FTokenizedMessage> );
 
 public:
 
 	SLATE_BEGIN_ARGS(SMessageLogMessageListRow)
 		: _Message()
 		, _OnTokenClicked()
+		, _OnMessageDoubleClicked()
 	{ }
 		SLATE_ATTRIBUTE(TSharedPtr<FTokenizedMessage>, Message)
 		SLATE_EVENT(FOnTokenClicked, OnTokenClicked)
+		SLATE_EVENT(FOnMessageClicked, OnMessageDoubleClicked)
 	SLATE_END_ARGS()
 
 	/**
@@ -41,6 +44,18 @@ public:
 	/** @return Widget for this log listing item*/
 	virtual TSharedRef<SWidget> GenerateWidget();
 
+	virtual FReply OnMouseButtonDoubleClick(const FGeometry& InMyGeometry, const FPointerEvent& InMouseEvent) override
+	{
+		FReply HandledReply = STableRow<TSharedPtr<FTokenizedMessage>>::OnMouseButtonDoubleClick(InMyGeometry, InMouseEvent);
+
+		if (InMouseEvent.GetEffectingButton() == EKeys::LeftMouseButton)
+		{
+			OnMessageDoubleClicked.ExecuteIfBound(Message);
+			HandledReply = FReply::Handled();
+		}
+		return HandledReply;
+	}
+
 protected:
 
 	TSharedRef<SWidget> CreateHyperlink( const TSharedRef<IMessageToken>& InMessageToken, const FText& InToolTip = FText() );
@@ -86,6 +101,9 @@ protected:
 
 	/** Delegate to execute when the token is clicked. */
 	FOnTokenClicked OnTokenClicked;
+
+	/** Delegate to execute when the message is double-clicked. */
+	FOnMessageClicked OnMessageDoubleClicked;
 };
 
 
diff --git a/Engine/Source/Editor/Kismet/Private/BlueprintEditor.cpp b/Engine/Source/Editor/Kismet/Private/BlueprintEditor.cpp
--- a/Engine/Source/Editor/Kismet/Private/BlueprintEditor.cpp
+++ b/Engine/Source/Editor/Kismet/Private/BlueprintEditor.cpp
@@ -231,8 +231,18 @@ namespace BlueprintEditorImpl
 	 * 
 	 * @param  Blueprint	The blueprint to search through.
 	 * @param  Severity		Defines the severity of the error/warning to search for.
+	 * @return The first node found with the specified error.
 	 */
 	static UEdGraphNode* FindNodeWithError(UBlueprint* Blueprint, EMessageSeverity::Type Severity = EMessageSeverity::Error);
+
+	/**
+	 * Searches through an error log, looking for the most severe error'ing node.
+	 * 
+	 * @param  ErrorLog		The error log you want to search through.
+	 * @param  Severity		Defines the severity of the error/warning to search for.
+	 * @return The first node found with the specified error.
+	 */
+	static UEdGraphNode* FindNodeWithError(FCompilerResultsLog const& ErrorLog, EMessageSeverity::Type Severity = EMessageSeverity::Error);
 }
 
 static bool BlueprintEditorImpl::GraphHasUserPlacedNodes(UEdGraph const* InGraph)
@@ -340,6 +350,30 @@ static UEdGraphNode* BlueprintEditorImpl::FindNodeWithError(UBlueprint* Blueprin
 	return ChoiceNode;
 }
 
+static UEdGraphNode* BlueprintEditorImpl::FindNodeWithError(FCompilerResultsLog const& ErrorLog, EMessageSeverity::Type Severity/* = EMessageSeverity::Error*/)
+{
+	UEdGraphNode* ChoiceNode = nullptr;
+	for (TWeakObjectPtr<UEdGraphNode> NodePtr : ErrorLog.AnnotatedNodes)
+	{
+		UEdGraphNode* Node = NodePtr.Get();
+		if ((Node != nullptr) && (Node->ErrorType <= Severity))
+		{
+			if ((ChoiceNode == nullptr) || (Node->ErrorType < ChoiceNode->ErrorType))
+			{
+				ChoiceNode = Node;
+				if (ChoiceNode->ErrorType == 0)
+				{
+					break;
+				}
+			}
+		}
+	}
+
+	return ChoiceNode;
+}
+
+
+
 bool FBlueprintEditor::IsASubGraph( const UEdGraph* GraphPtr )
 {
 	if( GraphPtr && GraphPtr->GetOuter() )
@@ -2618,7 +2652,7 @@ void FBlueprintEditor::Compile()
 		UBlueprintEditorSettings const* BpEditorSettings = GetDefault<UBlueprintEditorSettings>();
 		if ((LogResults.NumErrors > 0) && BpEditorSettings->bJumpToNodeErrors)
 		{
-			if (UEdGraphNode* NodeWithError = BlueprintEditorImpl::FindNodeWithError(BlueprintObj))
+			if (UEdGraphNode* NodeWithError = BlueprintEditorImpl::FindNodeWithError(LogResults))
 			{
 				JumpToNode(NodeWithError, /*bRequestRename =*/false);
 			}
diff --git a/Engine/Source/Editor/UnrealEd/Private/Kismet2/CompilerResultsLog.cpp b/Engine/Source/Editor/UnrealEd/Private/Kismet2/CompilerResultsLog.cpp
--- a/Engine/Source/Editor/UnrealEd/Private/Kismet2/CompilerResultsLog.cpp
+++ b/Engine/Source/Editor/UnrealEd/Private/Kismet2/CompilerResultsLog.cpp
@@ -293,41 +293,48 @@ void FCompilerResultsLog::InternalLogMessage(const EMessageSeverity::Type& Sever
 
 void FCompilerResultsLog::AnnotateNode(class UEdGraphNode* Node, TSharedRef<FTokenizedMessage> LogLine)
 {
-	if ((Node != NULL) && bAnnotateMentionedNodes)
+	if (Node != nullptr)
 	{
-		// Determine if this message is the first or more important than the previous one (only showing one error/warning per node for now)
-		bool bUpdateMessage = true;
-		if (Node->bHasCompilerMessage)
-		{
-			// Already has a message, see if we meet or trump the severity
-			bUpdateMessage = LogLine->GetSeverity() <= Node->ErrorType;
-		}
-		else
-		{
-			Node->ErrorMsg.Empty();
-		}
-		
-		// Update the message
-		if (bUpdateMessage)
+		LogLine->SetMessageLink(FUObjectToken::Create(Node));
+
+		if (bAnnotateMentionedNodes)
 		{
-			Node->ErrorType = (int32)LogLine->GetSeverity();
-			Node->bHasCompilerMessage = true;
-			
-			FText FullMessage = LogLine->ToText();
-			
-			if (Node->ErrorMsg.IsEmpty())
+			// Determine if this message is the first or more important than the previous one (only showing one error/warning per node for now)
+			bool bUpdateMessage = true;
+			if (Node->bHasCompilerMessage)
 			{
-				Node->ErrorMsg = FullMessage.ToString();
+				// Already has a message, see if we meet or trump the severity
+				bUpdateMessage = LogLine->GetSeverity() <= Node->ErrorType;
 			}
 			else
 			{
-				FFormatNamedArguments Args;
-				Args.Add( TEXT("PreviousMessage"), FText::FromString( Node->ErrorMsg ) );
-				Args.Add( TEXT("NewMessage"), FullMessage );
-				Node->ErrorMsg = FText::Format( LOCTEXT("AggregateMessagesFormatter", "{PreviousMessage}\n{NewMessage}"), Args ).ToString();
+				Node->ErrorMsg.Empty();
+			}
+
+			// Update the message
+			if (bUpdateMessage)
+			{
+				Node->ErrorType = (int32)LogLine->GetSeverity();
+				Node->bHasCompilerMessage = true;
+
+				FText FullMessage = LogLine->ToText();
+
+				if (Node->ErrorMsg.IsEmpty())
+				{
+					Node->ErrorMsg = FullMessage.ToString();
+				}
+				else
+				{
+					FFormatNamedArguments Args;
+					Args.Add(TEXT("PreviousMessage"), FText::FromString(Node->ErrorMsg));
+					Args.Add(TEXT("NewMessage"), FullMessage);
+					Node->ErrorMsg = FText::Format(LOCTEXT("AggregateMessagesFormatter", "{PreviousMessage}\n{NewMessage}"), Args).ToString();
+				}
+
+				AnnotatedNodes.Add(Node);
 			}
 		}
-	}
+	}	
 }
 
 TArray< TSharedRef<FTokenizedMessage> > FCompilerResultsLog::ParseCompilerLogDump(const FString& LogDump)
diff --git a/Engine/Source/Editor/UnrealEd/Public/Kismet2/CompilerResultsLog.h b/Engine/Source/Editor/UnrealEd/Public/Kismet2/CompilerResultsLog.h
--- a/Engine/Source/Editor/UnrealEd/Public/Kismet2/CompilerResultsLog.h
+++ b/Engine/Source/Editor/UnrealEd/Public/Kismet2/CompilerResultsLog.h
@@ -92,6 +92,9 @@ public:
 	// Minimum event time (ms) for inclusion into the final summary log
 	int EventDisplayThresholdMs;
 
+	/** Tracks nodes that produced errors/warnings */
+	TSet< TWeakObjectPtr<UEdGraphNode> > AnnotatedNodes;
+
 protected:
 	// Maps from transient object created during compiling to original 'source code' object
 	FBacktrackMap SourceBacktrackMap;
diff --git a/Engine/Source/Runtime/Core/Private/Logging/TokenizedMessage.cpp b/Engine/Source/Runtime/Core/Private/Logging/TokenizedMessage.cpp
--- a/Engine/Source/Runtime/Core/Private/Logging/TokenizedMessage.cpp
+++ b/Engine/Source/Runtime/Core/Private/Logging/TokenizedMessage.cpp
@@ -103,6 +103,12 @@ TSharedRef<FTokenizedMessage> FTokenizedMessage::AddToken( const TSharedRef<IMes
 	return AsShared();
 }
 
+
+void FTokenizedMessage::SetMessageLink(const TSharedRef<IMessageToken>& InToken)
+{
+	MessageLink = InToken;
+}
+
 void FTokenizedMessage::SetSeverity( const EMessageSeverity::Type InSeverity )
 {
 	Severity = InSeverity;
@@ -129,6 +135,11 @@ const TArray<TSharedRef<IMessageToken> >& FTokenizedMessage::GetMessageTokens()
 	return MessageTokens;
 }
 
+TSharedPtr<IMessageToken> FTokenizedMessage::GetMessageLink() const
+{
+	return MessageLink;
+}
+
 FURLToken::FGenerateURL FURLToken::GenerateURL;
 
 void FURLToken::VisitURL(const TSharedRef<IMessageToken>& Token, FString InURL)
diff --git a/Engine/Source/Runtime/Core/Public/Logging/TokenizedMessage.h b/Engine/Source/Runtime/Core/Public/Logging/TokenizedMessage.h
--- a/Engine/Source/Runtime/Core/Public/Logging/TokenizedMessage.h
+++ b/Engine/Source/Runtime/Core/Public/Logging/TokenizedMessage.h
@@ -164,6 +164,22 @@ public:
 	 */
 	CORE_API const TArray< TSharedRef<IMessageToken> >& GetMessageTokens() const;
 
+	/**
+	 * Sets up a token action for the message as a whole (not to be displayed...
+	 * intended to be invoked from a double click).
+	 * 
+	 * @param  InToken	A token for the entire message to link to.
+	 */
+	CORE_API void SetMessageLink(const TSharedRef<IMessageToken>& InToken);
+
+	/**
+	 * Gets the token action associated with this messages as a whole. Should
+	 * link to an associated element.
+	 * 
+	 * @return A token action that was set via SetMessageLink().
+	 */
+	CORE_API TSharedPtr<IMessageToken> GetMessageLink() const;
+
 	/** 
 	 * Helper function for getting a severity as text
 	 *
@@ -192,6 +208,9 @@ protected:
 	/** the array of message tokens this message contains */
 	TArray< TSharedRef<IMessageToken> > MessageTokens;
 
+	/** A token associated with the entire message (doesn't display) */
+	TSharedPtr<IMessageToken> MessageLink;
+
 private:
 
 	/** The severity of this message */
