BP: When a Blueprint is duplicated, generated class of the original bp should not be renamed.

diff --git a/Engine/Source/Editor/UnrealEd/Private/Kismet2/BlueprintEditorUtils.cpp b/Engine/Source/Editor/UnrealEd/Private/Kismet2/BlueprintEditorUtils.cpp
--- a/Engine/Source/Editor/UnrealEd/Private/Kismet2/BlueprintEditorUtils.cpp
+++ b/Engine/Source/Editor/UnrealEd/Private/Kismet2/BlueprintEditorUtils.cpp
@@ -1355,18 +1355,6 @@ void FBlueprintEditorUtils::PropagateParentBlueprintDefaults(UClass* ClassToProp
 	}
 }
 
-void DestroyGeneratedClass(TSubclassOf<UObject> &ClassToDestroy)
-{
-	if (ClassToDestroy != nullptr)
-	{
-		// We just need to rename the object, so the name won't collide.
-		ClassToDestroy->Rename(TEXT("None"), nullptr, REN_DontCreateRedirectors);
-
-		// The rest will do GC after we lose reference to it.
-		ClassToDestroy = nullptr;
-	}
-}
-
 void FBlueprintEditorUtils::PostDuplicateBlueprint(UBlueprint* Blueprint)
 {
 	// Only recompile after duplication if this isn't PIE
@@ -1391,10 +1379,6 @@ void FBlueprintEditorUtils::PostDuplicateBlueprint(UBlueprint* Blueprint)
 			Blueprint->Timelines.Empty();
 
 			// Delete the existing class references, the compile will create new ones
-			DestroyGeneratedClass(Blueprint->GeneratedClass);
-			DestroyGeneratedClass(Blueprint->SkeletonGeneratedClass);
-				
-
 			Blueprint->GeneratedClass = nullptr;
 			Blueprint->SkeletonGeneratedClass = nullptr;
 
