bUseEmissiveForStaticLighting can be used to disable material emissive casting light in Lightmass

diff --git a/Engine/Source/Editor/UnrealEd/Private/Lightmass/Lightmass.cpp b/Engine/Source/Editor/UnrealEd/Private/Lightmass/Lightmass.cpp
--- a/Engine/Source/Editor/UnrealEd/Private/Lightmass/Lightmass.cpp
+++ b/Engine/Source/Editor/UnrealEd/Private/Lightmass/Lightmass.cpp
@@ -1453,8 +1453,7 @@ void FLightmassExporter::WriteMeshInstances( int32 Channel )
 						NewElementData.MaterialId = Material->GetLightingGuid();
 						NewElementData.bUseTwoSidedLighting = Primitive->LightmassSettings.bUseTwoSidedLighting;
 						NewElementData.bShadowIndirectOnly = Primitive->LightmassSettings.bShadowIndirectOnly;
-						// Mesh area lights currently disabled
-						NewElementData.bUseEmissiveForStaticLighting = false;
+						NewElementData.bUseEmissiveForStaticLighting = Primitive->LightmassSettings.bUseEmissiveForStaticLighting;;
 						// Combine primitive and level boost settings so we don't have to send the level settings over to Lightmass  
 						NewElementData.EmissiveLightFalloffExponent = Primitive->LightmassSettings.EmissiveLightFalloffExponent;
 						NewElementData.EmissiveLightExplicitInfluenceRadius = Primitive->LightmassSettings.EmissiveLightExplicitInfluenceRadius;
@@ -1524,7 +1523,7 @@ void FLightmassExporter::WriteLandscapeInstances( int32 Channel )
 			FLightmassPrimitiveSettings& LMSetting = LandscapeComp->GetLandscapeProxy()->LightmassSettings;
 			NewElementData.bUseTwoSidedLighting = LMSetting.bUseTwoSidedLighting;
 			NewElementData.bShadowIndirectOnly = LMSetting.bShadowIndirectOnly;
-			NewElementData.bUseEmissiveForStaticLighting = false;
+			NewElementData.bUseEmissiveForStaticLighting = LMSetting.bUseEmissiveForStaticLighting;
 			// Combine primitive and level boost settings so we don't have to send the level settings over to Lightmass  
 			NewElementData.EmissiveLightFalloffExponent = LMSetting.EmissiveLightFalloffExponent;
 			NewElementData.EmissiveLightExplicitInfluenceRadius = LMSetting.EmissiveLightExplicitInfluenceRadius;
@@ -1628,7 +1627,7 @@ void FLightmassExporter::WriteMappings( int32 Channel )
 			TempData.MaterialId = Material->GetLightingGuid();
 			TempData.bUseTwoSidedLighting = PrimitiveSettings.bUseTwoSidedLighting;
 			TempData.bShadowIndirectOnly = PrimitiveSettings.bShadowIndirectOnly;
-			TempData.bUseEmissiveForStaticLighting = false;
+			TempData.bUseEmissiveForStaticLighting = PrimitiveSettings.bUseEmissiveForStaticLighting;
 			TempData.EmissiveLightFalloffExponent = PrimitiveSettings.EmissiveLightFalloffExponent;
 			TempData.EmissiveLightExplicitInfluenceRadius = PrimitiveSettings.EmissiveLightExplicitInfluenceRadius;
 			TempData.EmissiveBoost = PrimitiveSettings.EmissiveBoost * LevelSettings.EmissiveBoost;
@@ -2270,8 +2269,8 @@ FLightmassProcessor::FLightmassProcessor(const FStaticLightingSystem& InSystem,
 	Messages.Add( TEXT("LightmassError_BuildSelectedNothingSelected"), LOCTEXT("LightmassError_BuildSelectedNothingSelected", "Building selected actors and BSP only, but no actors or BSP selected!") );
 	Messages.Add( TEXT("LightmassError_ObjectWrappedUVs"), LOCTEXT("LightmassError_ObjectWrappedUVs", "Object has wrapping UVs.") );
 	Messages.Add( TEXT("LightmassError_ObjectOverlappedUVs"), LOCTEXT("LightmassError_ObjectOverlappedUVs", "Object has overlapping UVs.") );
-	Messages.Add( TEXT("LightmassError_EmissiveMeshHighPolyCount"), LOCTEXT("LightmassError_EmissiveMeshHighPolyCount", "Object has bUseEmissiveForStaticLighting=true, but a large number of polygons (more than 3000) and will result in a long lighting build.") );
-	Messages.Add( TEXT("LightmassError_EmissiveMeshExtremelyHighPolyCount"), LOCTEXT("LightmassError_EmissiveMeshExtremelyHighPolyCount", "Object did not create emissive lights even though it had bUseEmissiveForStaticLighting=true due to excessive polycount (more than 5000).") );
+	Messages.Add( TEXT("LightmassError_EmissiveMeshHighPolyCount"), LOCTEXT("LightmassError_EmissiveMeshHighPolyCount", "Object has a large number of polygons (more than 3000) and will result in a long lighting build.") );
+	Messages.Add( TEXT("LightmassError_EmissiveMeshExtremelyHighPolyCount"), LOCTEXT("LightmassError_EmissiveMeshExtremelyHighPolyCount", "Object did not create emissive lights due to excessive polycount (more than 5000).") );
 	Messages.Add( TEXT("LightmassError_BadLightMapCoordinateIndex"), LOCTEXT("LightmassError_BadLightMapCoordinateIndex", "StaticMesh has invalid LightMapCoordinateIndex.") );
 	Messages.Add( TEXT("LightmassError_ObjectMultipleDominantLights"), LOCTEXT("LightmassError_ObjectMultipleDominantLights", "Object has multiple dominant lights.") );
 }
diff --git a/Engine/Source/Programs/UnrealLightmass/Private/Lighting/FinalGather.cpp b/Engine/Source/Programs/UnrealLightmass/Private/Lighting/FinalGather.cpp
--- a/Engine/Source/Programs/UnrealLightmass/Private/Lighting/FinalGather.cpp
+++ b/Engine/Source/Programs/UnrealLightmass/Private/Lighting/FinalGather.cpp
@@ -262,9 +262,13 @@ FLinearColor FStaticLightingSystem::CalculateExitantRadiance(
 			}
 		}
 
-		//@todo - cache emissive on surfaces like direct lighting?
-		const FLinearColor EmissiveLighting = HitMesh->EvaluateEmissive(Vertex.TextureCoordinates[0], ElementIndex);
-		AccumulatedRadiance += EmissiveLighting;
+		FLinearColor Emissive = FLinearColor::Black;
+		if (HitMesh->IsEmissive(ElementIndex))
+		{
+			Emissive = HitMesh->EvaluateEmissive(Vertex.TextureCoordinates[0], ElementIndex);
+		}
+
+		AccumulatedRadiance += Emissive;
 	}
 
 	// So we can compare it against FLinearColor::Black easily
diff --git a/Engine/Source/Programs/UnrealLightmass/Private/Lighting/LightingMesh.cpp b/Engine/Source/Programs/UnrealLightmass/Private/Lighting/LightingMesh.cpp
--- a/Engine/Source/Programs/UnrealLightmass/Private/Lighting/LightingMesh.cpp
+++ b/Engine/Source/Programs/UnrealLightmass/Private/Lighting/LightingMesh.cpp
@@ -14,6 +14,9 @@ namespace Lightmass
  */
 TMap<FStaticLightingMesh*, int32> FStaticLightingMesh::MeshToIndexMap;
 
+// Currently disabled due to lack of robustness
+bool bAllowMeshAreaLights = false;
+
 /** Evaluates the mesh's Bidirectional Reflectance Distribution Function. */
 FLinearColor FStaticLightingMesh::EvaluateBRDF(
 	const FStaticLightingVertex& Vertex, 
@@ -113,7 +116,7 @@ void FStaticLightingMesh::Import( FLightmassImporter& Importer )
 		CurrentMaterialElement.MaterialId = TempData.MaterialId;
 		CurrentMaterialElement.bUseTwoSidedLighting = TempData.bUseTwoSidedLighting;
 		CurrentMaterialElement.bShadowIndirectOnly = TempData.bShadowIndirectOnly;
-		CurrentMaterialElement.bUseEmissiveForStaticLighting = false;
+		CurrentMaterialElement.bUseEmissiveForStaticLighting = TempData.bUseEmissiveForStaticLighting;
 		// Validating data here instead of in Unreal since EmissiveLightFalloffExponent is used in so many different object types
 		CurrentMaterialElement.EmissiveLightFalloffExponent = FMath::Max(TempData.EmissiveLightFalloffExponent, 0.0f);
 		CurrentMaterialElement.EmissiveLightExplicitInfluenceRadius = FMath::Max(TempData.EmissiveLightExplicitInfluenceRadius, 0.0f);
@@ -164,7 +167,8 @@ void FStaticLightingMesh::CreateMeshAreaLights(
 	bool bAnyElementsUseEmissiveForLighting = false;
 	for (int32 MaterialIndex = 0; MaterialIndex < MaterialElements.Num(); MaterialIndex++)
 	{
-		if (MaterialElements[MaterialIndex].bUseEmissiveForStaticLighting &&
+		if (bAllowMeshAreaLights && 
+			MaterialElements[MaterialIndex].bUseEmissiveForStaticLighting &&
 			MaterialElements[MaterialIndex].Material->EmissiveSize > 0)
 		{
 			bAnyElementsUseEmissiveForLighting = true;
@@ -240,7 +244,8 @@ void FStaticLightingMesh::CreateMeshAreaLights(
 	for (int32 MaterialIndex = 0; MaterialIndex < MaterialElements.Num(); MaterialIndex++)
 	{
 		const FMaterial& CurrentMaterial = *MaterialElements[MaterialIndex].Material;
-		if (MaterialElements[MaterialIndex].bUseEmissiveForStaticLighting &&
+		if (bAllowMeshAreaLights &&
+			MaterialElements[MaterialIndex].bUseEmissiveForStaticLighting &&
 			CurrentMaterial.EmissiveSize > 0)
 		{
 			// Operate on each layer independently
diff --git a/Engine/Source/Programs/UnrealLightmass/Public/SceneExport.h b/Engine/Source/Programs/UnrealLightmass/Public/SceneExport.h
--- a/Engine/Source/Programs/UnrealLightmass/Public/SceneExport.h
+++ b/Engine/Source/Programs/UnrealLightmass/Public/SceneExport.h
@@ -903,7 +903,7 @@ struct FMaterialElementData
 	FMaterialElementData() :
 		  bUseTwoSidedLighting(false)
 		, bShadowIndirectOnly(false)
-		, bUseEmissiveForStaticLighting(false)
+		, bUseEmissiveForStaticLighting(true)
 	    , EmissiveLightFalloffExponent(2.0f)
 		, EmissiveLightExplicitInfluenceRadius(0.0f)
 		, EmissiveBoost(1.0f)
diff --git a/Engine/Source/Runtime/Engine/Classes/Engine/EngineTypes.h b/Engine/Source/Runtime/Engine/Classes/Engine/EngineTypes.h
--- a/Engine/Source/Runtime/Engine/Classes/Engine/EngineTypes.h
+++ b/Engine/Source/Runtime/Engine/Classes/Engine/EngineTypes.h
@@ -1080,15 +1080,15 @@ struct FLightmassPrimitiveSettings
 	GENERATED_USTRUCT_BODY()
 
 	/** If true, this object will be lit as if it receives light from both sides of its polygons. */
-	UPROPERTY(EditAnywhere, BlueprintReadWrite, Category=Lightmass)
+	UPROPERTY(EditAnywhere, BlueprintReadOnly, Category=Lightmass)
 	uint32 bUseTwoSidedLighting:1;
 
 	/** If true, this object will only shadow indirect lighting.  					*/
-	UPROPERTY(EditAnywhere, BlueprintReadWrite, Category=Lightmass)
+	UPROPERTY(EditAnywhere, BlueprintReadOnly, Category=Lightmass)
 	uint32 bShadowIndirectOnly:1;
 
 	/** If true, allow using the emissive for static lighting.						*/
-	UPROPERTY()
+	UPROPERTY(EditAnywhere, BlueprintReadOnly, Category=Lightmass)
 	uint32 bUseEmissiveForStaticLighting:1;
 
 	/** Direct lighting falloff exponent for mesh area lights created from emissive areas on this primitive. */
@@ -1119,7 +1119,7 @@ struct FLightmassPrimitiveSettings
 	{
 		bUseTwoSidedLighting = false;
 		bShadowIndirectOnly = false;
-		bUseEmissiveForStaticLighting = false;
+		bUseEmissiveForStaticLighting = true;
 		EmissiveLightFalloffExponent = 8.0f;
 		EmissiveLightExplicitInfluenceRadius = 0.0f;
 		EmissiveBoost = 1.0f;
diff --git a/Engine/Source/Runtime/Engine/Private/Polygon.cpp b/Engine/Source/Runtime/Engine/Private/Polygon.cpp
--- a/Engine/Source/Runtime/Engine/Private/Polygon.cpp
+++ b/Engine/Source/Runtime/Engine/Private/Polygon.cpp
@@ -47,7 +47,6 @@ FArchive& operator<<(FArchive& Ar, FLightmassPrimitiveSettings& Settings)
 		float SpecularBoost;
 		Ar << SpecularBoost;
 	}
-	Settings.bUseEmissiveForStaticLighting = false;
 
 	return Ar;
 }
