Removed FEditorDelegates::MapChange broadcast on bsp brush operations which was causing multisecond editor hitches in big maps.
Originaly it was added to fix TTP# 304489.

diff --git a/Engine/Source/Editor/FoliageEdit/Private/FoliageEdMode.cpp b/Engine/Source/Editor/FoliageEdit/Private/FoliageEdMode.cpp
--- a/Engine/Source/Editor/FoliageEdit/Private/FoliageEdMode.cpp
+++ b/Engine/Source/Editor/FoliageEdit/Private/FoliageEdMode.cpp
@@ -77,12 +77,6 @@ void FEdModeFoliage::AddReferencedObjects(FReferenceCollector& Collector)
 	Collector.AddReferencedObject(SphereBrushComponent);
 }
 
-/** FEdMode: Called when the mode is created */
-void FEdModeFoliage::Initialize()
-{
-	FEditorDelegates::MapChange.AddSP(this, &FEdModeFoliage::NotifyMapRebuild);
-}
-
 /** FEdMode: Called when the mode is entered */
 void FEdModeFoliage::Enter()
 {
@@ -217,18 +211,6 @@ void FEdModeFoliage::NotifyToolChanged()
 	}
 }
 
-void FEdModeFoliage::NotifyMapRebuild(uint32 MapChangeFlags) const
-{
-	if (MapChangeEventFlags::MapRebuild & MapChangeFlags)
-	{
-		AInstancedFoliageActor* IFA = AInstancedFoliageActor::GetInstancedFoliageActorForCurrentLevel(GetWorld());
-		if (IFA)
-		{
-			IFA->MapRebuild();
-		}
-	}
-}
-
 bool FEdModeFoliage::DisallowMouseDeltaTracking() const
 {
 	// We never want to use the mouse delta tracker while painting
diff --git a/Engine/Source/Editor/FoliageEdit/Private/FoliageEdMode.h b/Engine/Source/Editor/FoliageEdit/Private/FoliageEdMode.h
--- a/Engine/Source/Editor/FoliageEdit/Private/FoliageEdMode.h
+++ b/Engine/Source/Editor/FoliageEdit/Private/FoliageEdMode.h
@@ -163,9 +163,6 @@ public:
 	/** FGCObject interface */
 	virtual void AddReferencedObjects(FReferenceCollector& Collector) override;
 
-	/** FEdMode: Called when the mode is created */
-	virtual void Initialize() override;
-
 	/** FEdMode: Called when the mode is entered */
 	virtual void Enter() override;
 
@@ -181,12 +178,6 @@ public:
 	/** Called when the user changes the current tool in the UI */
 	void NotifyToolChanged();
 
-	/**
-	* Called when the BSP map is rebuilt so that foliage actors can self-update without tying them to the editor.
-	* @param MapChangeFlags Map change type flags.
-	*/
-	void NotifyMapRebuild(uint32 MapChangeFlags) const;
-
 	/**
 	 * Called when the mouse is moved over the viewport
 	 *
diff --git a/Engine/Source/Editor/UnrealEd/Private/EditorServer.cpp b/Engine/Source/Editor/UnrealEd/Private/EditorServer.cpp
--- a/Engine/Source/Editor/UnrealEd/Private/EditorServer.cpp
+++ b/Engine/Source/Editor/UnrealEd/Private/EditorServer.cpp
@@ -1315,6 +1315,8 @@ void UEditorEngine::RebuildMap(UWorld* InWorld, EMapRebuildType RebuildType)
 		UE_LOG(LogEditorServer, Log, TEXT("Rebuildmap Clear paths rebuilt"));
 	}
 
+	TArray<ULevel*> UpdatedLevels;
+
 	switch (RebuildType)
 	{
 		case EMapRebuildType::MRT_AllVisible:
@@ -1330,6 +1332,7 @@ void UEditorEngine::RebuildMap(UWorld* InWorld, EMapRebuildType RebuildType)
 				csgRebuild( InWorld );
 				InWorld->InvalidateModelGeometry( Level );
 				Level->bGeometryDirtyForLighting = false;
+				UpdatedLevels.AddUnique( Level );
 			}
 
 			// Build CSG for all visible streaming levels
@@ -1345,6 +1348,7 @@ void UEditorEngine::RebuildMap(UWorld* InWorld, EMapRebuildType RebuildType)
 						csgRebuild( InWorld );
 						InWorld->InvalidateModelGeometry( Level );
 						InWorld->GetCurrentLevel()->bGeometryDirtyForLighting = false;
+						UpdatedLevels.AddUnique( Level );
 					}
 				}
 			}
@@ -1366,6 +1370,7 @@ void UEditorEngine::RebuildMap(UWorld* InWorld, EMapRebuildType RebuildType)
 					csgRebuild( InWorld );
 					InWorld->InvalidateModelGeometry( Level );
 					Level->bGeometryDirtyForLighting = false;
+					UpdatedLevels.AddUnique( Level );
 				}
 
 				// Build CSG for each streaming level that is out of date
@@ -1381,6 +1386,7 @@ void UEditorEngine::RebuildMap(UWorld* InWorld, EMapRebuildType RebuildType)
 							csgRebuild( InWorld );
 							InWorld->InvalidateModelGeometry( Level );
 							Level->bGeometryDirtyForLighting = false;
+							UpdatedLevels.AddUnique( Level );
 						}
 					}
 				}
@@ -1396,10 +1402,21 @@ void UEditorEngine::RebuildMap(UWorld* InWorld, EMapRebuildType RebuildType)
 			csgRebuild( InWorld );
 			InWorld->InvalidateModelGeometry( InWorld->GetCurrentLevel() );
 			InWorld->GetCurrentLevel()->bGeometryDirtyForLighting = false;
+			UpdatedLevels.AddUnique( InWorld->GetCurrentLevel() );
 		}
 		break;
 	}
 
+	// See if there is any foliage that also needs to be updated
+	for (ULevel* Level : UpdatedLevels)
+	{
+		AInstancedFoliageActor* IFA = AInstancedFoliageActor::GetInstancedFoliageActorForLevel(Level);
+		if (IFA)
+		{
+			IFA->MapRebuild();
+		}
+	}
+	
 	GWarn->StatusUpdate( -1, -1, NSLOCTEXT("UnrealEd", "CleaningUpE", "Cleaning up...") );
 
 	RedrawLevelEditingViewports();
@@ -1442,6 +1459,13 @@ void UEditorEngine::RebuildLevel(ULevel& Level)
 	Level.UpdateModelComponents();
 
 	RebuildStaticNavigableGeometry(&Level);
+
+	// See if there is any foliage that also needs to be updated
+	AInstancedFoliageActor* IFA = AInstancedFoliageActor::GetInstancedFoliageActorForLevel(&Level);
+	if (IFA)
+	{
+		IFA->MapRebuild();
+	}
 }
 
 void UEditorEngine::RebuildModelFromBrushes(UModel* Model, bool bSelectedBrushesOnly)
@@ -1556,7 +1580,6 @@ void UEditorEngine::RebuildAlteredBSP()
 
 	}
 
-	int32 NumLevelsTouched = 0;
 	// Rebuild the levels
 	for (int32 LevelIdx = 0; LevelIdx < LevelsToRebuild.Num(); ++LevelIdx)
 	{
@@ -1564,17 +1587,9 @@ void UEditorEngine::RebuildAlteredBSP()
 		if ( LevelToRebuild.IsValid() )
 		{
 			RebuildLevel(*LevelToRebuild.Get());
-			NumLevelsTouched++;
 		}
 	}
 
-	// Broadcast MapChange event in case we actually build something 
-	// will mark all loaded levels as dirty during components re-registration
-	if (NumLevelsTouched > 0)
-	{
-		FEditorDelegates::MapChange.Broadcast( MapChangeEventFlags::MapRebuild );
-	}
-	
 	RedrawLevelEditingViewports();
 
 	ABrush::OnRebuildDone();
diff --git a/Engine/Source/Runtime/Engine/Classes/Foliage/InstancedFoliageActor.h b/Engine/Source/Runtime/Engine/Classes/Foliage/InstancedFoliageActor.h
--- a/Engine/Source/Runtime/Engine/Classes/Foliage/InstancedFoliageActor.h
+++ b/Engine/Source/Runtime/Engine/Classes/Foliage/InstancedFoliageActor.h
@@ -42,7 +42,7 @@ public:
 
 #if WITH_EDITOR
 
-	// Called in response to FEditorDelegates::MapChange.Broadcast(MapChangeEventFlags::MapRebuild ) so BSP rebuilds migrate foliage from obsolete to new components.
+	// Called in response to BSP rebuilds to migrate foliage from obsolete to new components.
 	ENGINE_API void MapRebuild();
 
 	// Moves instances based on the specified component to the current streaming level
diff --git a/Engine/Source/Runtime/Engine/Private/InstancedFoliage.cpp b/Engine/Source/Runtime/Engine/Private/InstancedFoliage.cpp
--- a/Engine/Source/Runtime/Engine/Private/InstancedFoliage.cpp
+++ b/Engine/Source/Runtime/Engine/Private/InstancedFoliage.cpp
@@ -1305,7 +1305,7 @@ void AInstancedFoliageActor::MapRebuild()
 	UE_LOG(LogInstancedFoliage, Log, TEXT("Map Rebuilt - Update all BSP painted foliage!"));
 
 	TMap<UFoliageType*, TArray<FFoliageInstance>> NewInstances;
-	TArray<UActorComponent*> RemovedComponents;
+	TArray<UModelComponent*> RemovedModelComponents;
 	UWorld* World = GetWorld();
 	check(World);
 
@@ -1320,11 +1320,11 @@ void AInstancedFoliageActor::MapRebuild()
 		for (auto& ComponentFoliagePair : MeshInfo.ComponentHash)
 		{
 			// BSP components are UModelComponents - they are the only ones we need to change
-			UActorComponent* TargetComponent = ComponentFoliagePair.Key;
-			if (!TargetComponent || TargetComponent->IsA(UModelComponent::StaticClass()))
+			UModelComponent* TargetComponent = Cast<UModelComponent>(ComponentFoliagePair.Key);
+			if (TargetComponent)
 			{
-				// Delete it later
-				RemovedComponents.Add(TargetComponent);
+				// Delete its instances later
+				RemovedModelComponents.Add(TargetComponent);
 
 				FFoliageComponentHashInfo const& FoliageInfo = ComponentFoliagePair.Value;
 
@@ -1349,21 +1349,15 @@ void AInstancedFoliageActor::MapRebuild()
 					if (bHit && Result.Component.IsValid() && Result.Component->IsA(UModelComponent::StaticClass()))
 					{
 						NewInstance.Base = CastChecked<UPrimitiveComponent>(Result.Component.Get());
+						NewInstances.FindOrAdd(Settings).Add(NewInstance);
 					}
-					else
-					{
-						// Untargeted instances remain for manual deletion.
-						NewInstance.Base = nullptr;
-					}
-
-					NewInstances.FindOrAdd(Settings).Add(NewInstance);
 				}
 			}
 		}
 	}
 
 	// Remove all existing & broken instances & component references.
-	for (UActorComponent* Component : RemovedComponents)
+	for (UModelComponent* Component : RemovedModelComponents)
 	{
 		DeleteInstancesForComponent(Component);
 	}
