Fixed return values for ini iterative cooking settings.

diff --git a/Engine/Source/Editor/UnrealEd/Private/CookOnTheFlyServer.cpp b/Engine/Source/Editor/UnrealEd/Private/CookOnTheFlyServer.cpp
--- a/Engine/Source/Editor/UnrealEd/Private/CookOnTheFlyServer.cpp
+++ b/Engine/Source/Editor/UnrealEd/Private/CookOnTheFlyServer.cpp
@@ -1912,9 +1912,7 @@ bool UCookOnTheFlyServer::GetCookedIniVersionStrings( const ITargetPlatform* Tar
 	const FString PlatformSandboxEditorIni = SandboxEditorIni.Replace(TEXT("[Platform]"), *TargetPlatform->PlatformName());
 
 	TArray<FString> SavedIniVersionedParams;
-	GConfig->GetArray( TEXT("CookedSettings"), TEXT("VersionedIniParams"), IniVersionStrings, PlatformSandboxEditorIni );
-
-	return true;
+	return GConfig->GetArray( TEXT("CookSettings"), TEXT("VersionedIniParams"), IniVersionStrings, PlatformSandboxEditorIni ) > 0;
 }
 
 
@@ -1933,7 +1931,7 @@ bool UCookOnTheFlyServer::IniSettingsOutOfDate( const TArray<ITargetPlatform*>&
 		}
 
 		TArray<FString> CookedIniVersionStrings;
-		if ( GetCookedIniVersionStrings( TargetPlatform, CookedIniVersionStrings ) )
+		if ( GetCookedIniVersionStrings( TargetPlatform, CookedIniVersionStrings ) == false )
 		{
 			// can't even get the cooked version strings for this platform add it to the out of date list
 			OutOfDateTargetPlatforms.Add(TargetPlatform);
