Enabled cook in editor by default.

diff --git a/Engine/Source/Editor/UnrealEd/Classes/Settings/EditorExperimentalSettings.h b/Engine/Source/Editor/UnrealEd/Classes/Settings/EditorExperimentalSettings.h
--- a/Engine/Source/Editor/UnrealEd/Classes/Settings/EditorExperimentalSettings.h
+++ b/Engine/Source/Editor/UnrealEd/Classes/Settings/EditorExperimentalSettings.h
@@ -69,9 +69,9 @@ public:
 	UPROPERTY(EditAnywhere, config, Category = Rendering, meta = (DisplayName = "Feature Level Preview"))
 	bool bFeatureLevelPreview;
 
-	/** Enable cook in the editor (overrides cook on the side setting) */
-	UPROPERTY(EditAnywhere, config, Category = Cooking, meta = (DisplayName = "Cook In The Editor (cooks from launchon will be run inside the main editor process instead of starting a new editor)"))
-	bool bCookInTheEditor;
+	/** Disable cook in the editor */
+	UPROPERTY(EditAnywhere, config, Category = Cooking, meta = (DisplayName = "Disable Cook In The Editor feature (cooks from launch on will be run in a separate process if disabled)"))
+	bool bDisableCookInEditor;
 
 	/** Enable cook on the side */
 	UPROPERTY(EditAnywhere, config, Category = Cooking, meta = (DisplayName = "Cook On The Side (Run a cook on the fly server in the background of the editor)"))
diff --git a/Engine/Source/Editor/UnrealEd/Private/UnrealEdEngine.cpp b/Engine/Source/Editor/UnrealEd/Private/UnrealEdEngine.cpp
--- a/Engine/Source/Editor/UnrealEd/Private/UnrealEdEngine.cpp
+++ b/Engine/Source/Editor/UnrealEd/Private/UnrealEdEngine.cpp
@@ -113,7 +113,8 @@ void UUnrealEdEngine::Init(IEngineLoop* InEngineLoop)
 	UEditorExperimentalSettings const* ExperimentalSettings =  GetDefault<UEditorExperimentalSettings>();
 	ECookInitializationFlags BaseCookingFlags = ECookInitializationFlags::AutoTick | ECookInitializationFlags::AsyncSave;
 	BaseCookingFlags |= ExperimentalSettings->bIterativeCookingForLaunchOn ? ECookInitializationFlags::Iterative : ECookInitializationFlags::None;
-	if ( ExperimentalSettings->bCookInTheEditor )
+
+	if ( !ExperimentalSettings->bDisableCookInEditor)
 	{
 		CookServer = ConstructObject<UCookOnTheFlyServer>( UCookOnTheFlyServer::StaticClass() );
 		CookServer->Initialize( ECookMode::CookByTheBookFromTheEditor, BaseCookingFlags );
