Launch game can be aborted if the user cancels when prompted to save unsaved packages.

diff --git a/Engine/Source/Editor/UnrealEd/Classes/Editor/EditorEngine.h b/Engine/Source/Editor/UnrealEd/Classes/Editor/EditorEngine.h
--- a/Engine/Source/Editor/UnrealEd/Classes/Editor/EditorEngine.h
+++ b/Engine/Source/Editor/UnrealEd/Classes/Editor/EditorEngine.h
@@ -1540,6 +1540,11 @@ public:
 	 */
 	void RequestPlaySession( const FString& DeviceId, const FString& DeviceName );
 
+	/**
+	 * Cancel request to start a play session
+	 */
+	void CancelRequestPlaySession();
+
 	/**
 	 * Makes a request to start a play from a Slate editor session
 	 */
diff --git a/Engine/Source/Editor/UnrealEd/Private/PlayLevel.cpp b/Engine/Source/Editor/UnrealEd/Private/PlayLevel.cpp
--- a/Engine/Source/Editor/UnrealEd/Private/PlayLevel.cpp
+++ b/Engine/Source/Editor/UnrealEd/Private/PlayLevel.cpp
@@ -642,6 +642,14 @@ void UEditorEngine::RequestPlaySession(const FString& DeviceId, const FString& D
 	bIsPlayWorldQueued = true;
 }
 
+void UEditorEngine::CancelRequestPlaySession()
+{
+	bIsPlayWorldQueued = false;
+	bPlayOnLocalPcSession = false;
+	bPlayUsingLauncher = false;
+	bPlayUsingMobilePreview = false;
+}
+
 void UEditorEngine::PlaySessionPaused()
 {
 	FEditorDelegates::PausePIE.Broadcast(bIsSimulatingInEditor);
@@ -1389,6 +1397,12 @@ void UEditorEngine::PlayUsingLauncher()
 		else
 		{
 			SaveWorldForPlay(MapNames);
+
+			if (MapNames.Num() == 0)
+			{
+				GEditor->CancelRequestPlaySession();
+				return;
+		}
 		}
 	
 		FString InitialMapName;
