Saving a new level now attempts to find a unique default name.

diff --git a/Engine/Source/Editor/UnrealEd/Private/FileHelpers.cpp b/Engine/Source/Editor/UnrealEd/Private/FileHelpers.cpp
--- a/Engine/Source/Editor/UnrealEd/Private/FileHelpers.cpp
+++ b/Engine/Source/Editor/UnrealEd/Private/FileHelpers.cpp
@@ -601,13 +601,14 @@ static bool OpenLevelSaveAsDialog(const FString& InDefaultPath, const FString& I
 	FString NewNameSuggestion = InNewNameSuggestion;
 	if (NewNameSuggestion.IsEmpty())
 	{
-		NewNameSuggestion = TEXT("NewMap");
-	}
+		const FString DefaultName = TEXT("NewMap");
+		FString PackageName = DefaultPath / DefaultName;
+		FString Name;
+		FAssetToolsModule& AssetToolsModule = FModuleManager::LoadModuleChecked<FAssetToolsModule>("AssetTools");
+		AssetToolsModule.Get().CreateUniqueAssetName(PackageName, TEXT(""), PackageName, Name);
 
-	FString PackageName = DefaultPath / NewNameSuggestion;
-	FString Name;
-	FAssetToolsModule& AssetToolsModule = FModuleManager::LoadModuleChecked<FAssetToolsModule>("AssetTools");
-	AssetToolsModule.Get().CreateUniqueAssetName(PackageName, TEXT(""), PackageName, Name);
+		NewNameSuggestion = FPaths::GetCleanFilename(PackageName);
+	}
 
 	FSaveAssetDialogConfig SaveAssetDialogConfig;
 	SaveAssetDialogConfig.DialogTitleOverride = LOCTEXT("SaveLevelDialogTitle", "Save Level As");
