Unity build fix

diff --git a/Engine/Source/Editor/UnrealEd/Private/EditorServer.cpp b/Engine/Source/Editor/UnrealEd/Private/EditorServer.cpp
--- a/Engine/Source/Editor/UnrealEd/Private/EditorServer.cpp
+++ b/Engine/Source/Editor/UnrealEd/Private/EditorServer.cpp
@@ -28,6 +28,7 @@
 #include "Matinee/InterpTrackAnimControl.h"
 #include "Matinee/InterpData.h"
 #include "Animation/SkeletalMeshActor.h"
+#include "Foliage/InstancedFoliageActor.h"
 
 #include "Editor/UnrealEd/Public/Kismet2/KismetEditorUtilities.h"
 #include "Editor/PropertyEditor/Public/PropertyEditorModule.h"
