Actually fixing the date time stamp to be UTC.

diff --git a/Engine/Source/Editor/UnrealEd/Private/Tests/EditorPerformanceTests.cpp b/Engine/Source/Editor/UnrealEd/Private/Tests/EditorPerformanceTests.cpp
--- a/Engine/Source/Editor/UnrealEd/Private/Tests/EditorPerformanceTests.cpp
+++ b/Engine/Source/Editor/UnrealEd/Private/Tests/EditorPerformanceTests.cpp
@@ -141,7 +141,7 @@ void EditorPerfDump(EditorPerfCaptureParameters& EditorPerfStats)
 void CaptureEditorData(EditorPerfCaptureParameters& OutEditorPerfStats)
 {
 	//Capture the current time stamp and format it to YYYY-MM-DD HH:MM:SS.mmm.
-	FDateTime CurrentDateAndTime = FDateTime::Now();
+	FDateTime CurrentDateAndTime = FDateTime::UtcNow();
 	
 	//Find the Average FPS
 	//Clamp to avoid huge averages at startup or after hitches
