Fixing a misspelling
UEAUTO-55/TTP346636

diff --git a/Engine/Source/Editor/UnrealEd/Private/Tests/EditorBuildPromotionTests.cpp b/Engine/Source/Editor/UnrealEd/Private/Tests/EditorBuildPromotionTests.cpp
--- a/Engine/Source/Editor/UnrealEd/Private/Tests/EditorBuildPromotionTests.cpp
+++ b/Engine/Source/Editor/UnrealEd/Private/Tests/EditorBuildPromotionTests.cpp
@@ -4253,12 +4253,12 @@ bool FSettingsCheckForPIECommand::Update()
 	if (GEditor->PlayWorld != NULL)
 	{
 		//Success
-		UE_LOG(LogEditorBuildPromotionTests, Display, TEXT("PlayInEditor keyboard shotcut success"));
+		UE_LOG(LogEditorBuildPromotionTests, Display, TEXT("PlayInEditor keyboard shortcut success"));
 		EditorBuildPromotionTestUtils::EndPIE();
 	}
 	else
 	{
-		UE_LOG(LogEditorBuildPromotionTests, Error, TEXT("PlayInEditor keyboard shotcut failed"));
+		UE_LOG(LogEditorBuildPromotionTests, Error, TEXT("PlayInEditor keyboard shortcut failed"));
 	}
 	return true;
 }
