Fix assert in WarmCookedPackages

diff --git a/Engine/Source/Editor/UnrealEd/Private/CookOnTheFlyServer.cpp b/Engine/Source/Editor/UnrealEd/Private/CookOnTheFlyServer.cpp
--- a/Engine/Source/Editor/UnrealEd/Private/CookOnTheFlyServer.cpp
+++ b/Engine/Source/Editor/UnrealEd/Private/CookOnTheFlyServer.cpp
@@ -2694,7 +2694,10 @@ void UCookOnTheFlyServer::WarmCookedPackages(const FString& AssetRegistryPath, c
 			UE_LOG(LogCookOnTheFly, Verbose, TEXT("Read package %s from %s"), *GetCachedStandardPackageFilename(NewAssetData->ObjectPath), *AssetRegistryPath);
 
 			FileRequest.SetFilename(GetCachedStandardPackageFilename(NewAssetData->ObjectPath));
-			CookedPackages.Add(FileRequest);
+			if (FileRequest.IsValid())
+			{
+				CookedPackages.Add(FileRequest);
+			}
 		}
 
 		delete[] PreallocatedAssetDataBuffer;
