Fixed regression in landscape normal map rendering during material exporting

diff --git a/Engine/Source/Editor/UnrealEd/Private/EditorExporters.cpp b/Engine/Source/Editor/UnrealEd/Private/EditorExporters.cpp
--- a/Engine/Source/Editor/UnrealEd/Private/EditorExporters.cpp
+++ b/Engine/Source/Editor/UnrealEd/Private/EditorExporters.cpp
@@ -2833,12 +2833,14 @@ namespace MaterialExportUtils
 				ViewFamily.EngineShowFlags.DisableAdvancedFeatures();
 				ViewFamily.EngineShowFlags.MotionBlur = 0;
 				ViewFamily.EngineShowFlags.Lighting = 0;
-				ViewFamily.EngineShowFlags.PostProcessing = 1;
-				ViewFamily.EngineShowFlags.VisualizeBuffer = 1;
 				ViewFamily.EngineShowFlags.LightFunctions = 0;
 				ViewFamily.EngineShowFlags.DynamicShadows = 0;
 				ViewFamily.EngineShowFlags.Atmosphere = 0;
-
+				// To enable WorldNormal visualization mode
+				ViewFamily.EngineShowFlags.PostProcessing = 1;
+				ViewFamily.EngineShowFlags.PostProcessMaterial = 1;
+				ViewFamily.EngineShowFlags.VisualizeBuffer = 1; 
+				
 				FSceneViewInitOptions ViewInitOptions;
 				ViewInitOptions.SetViewRectangle(FIntRect(0, 0, OutFlattenMaterial.NormalSize.X, OutFlattenMaterial.NormalSize.Y));
 				ViewInitOptions.ViewFamily = &ViewFamily;
