Invalidate hit proxy when rotating camera in editor viewports using gestures

diff --git a/Engine/Source/Editor/UnrealEd/Private/EditorViewportClient.cpp b/Engine/Source/Editor/UnrealEd/Private/EditorViewportClient.cpp
--- a/Engine/Source/Editor/UnrealEd/Private/EditorViewportClient.cpp
+++ b/Engine/Source/Editor/UnrealEd/Private/EditorViewportClient.cpp
@@ -2930,7 +2930,7 @@ void FEditorViewportClient::UpdateGestureDelta()
 	{
 		MoveViewportCamera( CurrentGestureDragDelta, CurrentGestureRotDelta, false );
 		
-		Invalidate( true, false );
+		Invalidate( true, true );
 		
 		CurrentGestureDragDelta = FVector::ZeroVector;
 		CurrentGestureRotDelta = FRotator::ZeroRotator;
