Added an array defined by an enum test to the property test object

diff --git a/Engine/Source/Editor/UnrealEd/Classes/Editor/PropertyEditorTestObject.h b/Engine/Source/Editor/UnrealEd/Classes/Editor/PropertyEditorTestObject.h
--- a/Engine/Source/Editor/UnrealEd/Classes/Editor/PropertyEditorTestObject.h
+++ b/Engine/Source/Editor/UnrealEd/Classes/Editor/PropertyEditorTestObject.h
@@ -22,6 +22,19 @@ enum PropertEditorTestEnum
 	PropertyEditorTest_MAX,
 };
 
+UENUM()
+enum ArrayLabelEnum
+{
+	ArrayIndex0,
+	ArrayIndex1,
+	ArrayIndex2,
+	ArrayIndex3,
+	ArrayIndex4,
+	ArrayIndex5,
+	ArrayIndex_MAX,
+};
+
+
 USTRUCT()
 struct FPropertyEditorTestSubStruct
 {
@@ -185,6 +198,9 @@ class UPropertyEditorTestObject : public UObject
 	UPROPERTY(EditAnywhere, Category=ArraysOfProperties)
 	int32 StaticArrayOfInts[5];
 
+	UPROPERTY(EditAnywhere, Category=ArraysOfProperties)
+	int32 StaticArrayOfIntsWithEnumLabels[ArrayIndex_MAX];
+
 	/** This is a float property tooltip that is overridden */
 	UPROPERTY(EditAnywhere, Category=AdvancedProperties, meta=(ClampMin = "0.0", ClampMax = "100.0", UIMin = "0.0", UIMax = "50.0", ToolTip = "This is a custom tooltip that should be shown"))
 	float FloatPropertyWithClampedRange;
