Asset/FeaturePack factory

diff --git a/Engine/Source/Editor/UnrealEd/Classes/Factories/PackFactory.h b/Engine/Source/Editor/UnrealEd/Classes/Factories/PackFactory.h
new file mode 100644
--- /dev/null
+++ b/Engine/Source/Editor/UnrealEd/Classes/Factories/PackFactory.h
@@ -0,0 +1,17 @@
+// Copyright 1998-2014 Epic Games, Inc. All Rights Reserved.
+
+#pragma once
+#include "PackFactory.generated.h"
+
+/**
+ * Factory for unpacking upack files
+ */
+UCLASS(hidecategories=Object)
+class UPackFactory : public UFactory
+{
+	GENERATED_UCLASS_BODY()
+
+	// Begin UFactory Interface
+	virtual UObject* FactoryCreateBinary( UClass* Class, UObject* InParent, FName Name, EObjectFlags Flags, UObject* Context, const TCHAR* Type, const uint8*& Buffer, const uint8* BufferEnd, FFeedbackContext* Warn ) override;
+	// End UFactory Interface
+};
\ No newline at end of file
diff --git a/Engine/Source/Editor/UnrealEd/Private/Factories/PackFactory.cpp b/Engine/Source/Editor/UnrealEd/Private/Factories/PackFactory.cpp
new file mode 100644
--- /dev/null
+++ b/Engine/Source/Editor/UnrealEd/Private/Factories/PackFactory.cpp
@@ -0,0 +1,524 @@
+// Copyright 1998-2014 Epic Games, Inc. All Rights Reserved.
+
+/*=============================================================================
+	PackFactory.cpp: Factory for importing asset and feature packs
+=============================================================================*/
+
+#include "UnrealEd.h"
+#include "Runtime/PakFile/Public/IPlatformFilePak.h"
+#include "ISourceControlModule.h"
+#include "SourceControlHelpers.h"
+#include "HotReloadInterface.h"
+#include "AES.h"
+#include "ModuleManager.h"
+#include "GameProjectGenerationModule.h"
+
+DEFINE_LOG_CATEGORY_STATIC(LogPackFactory, Log, All);
+
+UPackFactory::UPackFactory(const class FObjectInitializer& PCIP)
+	: Super(PCIP)
+{
+	// Since this factory can output multiple and any number of class it doesn't really have a 
+	// SupportedClass per say, but one must be defined, so we just reference ourself
+	SupportedClass = UPackFactory::StaticClass();
+
+	Formats.Add(TEXT("upack;Asset Pack"));
+	Formats.Add(TEXT("upack;Feature Pack"));
+
+	bEditorImport = true;
+}
+
+namespace PackFactoryHelper
+{
+	// Utility function to copy a single pak entry out of the Source archive and in to the Destination archive using Buffer as temporary space
+	bool BufferedCopyFile(FArchive& DestAr, FArchive& Source, const FPakEntry& Entry, TArray<uint8>& Buffer)
+	{	
+		// Align down
+		const int64 BufferSize = Buffer.Num() & ~(FAES::AESBlockSize-1);
+		int64 RemainingSizeToCopy = Entry.Size;
+		while (RemainingSizeToCopy > 0)
+		{
+			const int64 SizeToCopy = FMath::Min(BufferSize, RemainingSizeToCopy);
+			// If file is encrypted so we need to account for padding
+			int64 SizeToRead = Entry.bEncrypted ? Align(SizeToCopy,FAES::AESBlockSize) : SizeToCopy;
+
+			Source.Serialize(Buffer.GetData(),SizeToRead);
+			if (Entry.bEncrypted)
+			{
+				FAES::DecryptData(Buffer.GetData(),SizeToRead);
+			}
+			DestAr.Serialize(Buffer.GetData(), SizeToCopy);
+			RemainingSizeToCopy -= SizeToRead;
+		}
+		return true;
+	}
+
+	// Utility function to uncompress and copy a single pak entry out of the Source archive and in to the Destination archive using PersistentBuffer as temporary space
+	bool UncompressCopyFile(FArchive& DestAr, FArchive& Source, const FPakEntry& Entry, TArray<uint8>& PersistentBuffer)
+	{
+		if (Entry.UncompressedSize == 0)
+		{
+			return false;
+		}
+
+		int64 WorkingSize = Entry.CompressionBlockSize;
+		int32 MaxCompressionBlockSize = FCompression::CompressMemoryBound((ECompressionFlags)Entry.CompressionMethod, WorkingSize);
+		WorkingSize += MaxCompressionBlockSize;
+		if (PersistentBuffer.Num() < WorkingSize)
+		{
+			PersistentBuffer.SetNumUninitialized(WorkingSize);
+		}
+
+		uint8* UncompressedBuffer = PersistentBuffer.GetData() + MaxCompressionBlockSize;
+
+		for (uint32 BlockIndex=0, BlockIndexNum=Entry.CompressionBlocks.Num(); BlockIndex < BlockIndexNum; ++BlockIndex)
+		{
+			uint32 CompressedBlockSize = Entry.CompressionBlocks[BlockIndex].CompressedEnd - Entry.CompressionBlocks[BlockIndex].CompressedStart;
+			uint32 UncompressedBlockSize = (uint32)FMath::Min<int64>(Entry.UncompressedSize - Entry.CompressionBlockSize*BlockIndex, Entry.CompressionBlockSize);
+			Source.Seek(Entry.CompressionBlocks[BlockIndex].CompressedStart);
+			uint32 SizeToRead = Entry.bEncrypted ? Align(CompressedBlockSize, FAES::AESBlockSize) : CompressedBlockSize;
+			Source.Serialize(PersistentBuffer.GetData(), SizeToRead);
+
+			if (Entry.bEncrypted)
+			{
+				FAES::DecryptData(PersistentBuffer.GetData(), SizeToRead);
+			}
+
+			if(!FCompression::UncompressMemory((ECompressionFlags)Entry.CompressionMethod,UncompressedBuffer,UncompressedBlockSize,PersistentBuffer.GetData(),CompressedBlockSize))
+			{
+				return false;
+			}
+			DestAr.Serialize(UncompressedBuffer,UncompressedBlockSize);
+		}
+
+		return true;
+	}
+
+	// Utility function to extract a pak entry out of the memory reader containing the pak file and place in the destination archive.
+	// Uses Buffer or PersistentCompressionBuffer depending on whether the entry is compressed or not.
+	void ExtractFile(const FPakEntry& Entry, FBufferReader& PakReader, TArray<uint8>& Buffer, TArray<uint8>& PersistentCompressionBuffer, FArchive& DestAr)
+	{
+		if (Entry.CompressionMethod == COMPRESS_None)
+		{
+			PackFactoryHelper::BufferedCopyFile(DestAr, PakReader, Entry, Buffer);
+		}
+		else
+		{
+			PackFactoryHelper::UncompressCopyFile(DestAr, PakReader, Entry, PersistentCompressionBuffer);
+		}
+	}
+
+	// Utility function to extract a pak entry out of the memory reader containing the pak file and place in a string.
+	// Uses Buffer or PersistentCompressionBuffer depending on whether the entry is compressed or not.
+	void ExtractFileToString(const FPakEntry& Entry, FBufferReader& PakReader, TArray<uint8>& Buffer, TArray<uint8>& PersistentCompressionBuffer, FString& FileContents)
+	{
+		TArray<uint8> Contents;
+		FMemoryWriter MemWriter(Contents);
+
+		ExtractFile(Entry, PakReader, Buffer, PersistentCompressionBuffer, MemWriter);
+
+		// Add a line feed at the end because the FString archive read will consume the last byte
+		Contents.Add('\n');
+
+		// Insert the length of the string to the front of the memory chunk so we can use FString archive read
+		const int32 StringLength = Contents.Num();
+		Contents.InsertUninitialized(0,sizeof(int32));
+		*(reinterpret_cast<int32*>(Contents.GetData())) = StringLength;
+
+		FMemoryReader MemReader(Contents);
+		MemReader << FileContents;
+	}
+
+	// Takes a string that represents the contents of a config file and sets up the supported config properties based on it
+	// Currently we support Action and Axis Mappings and a GameName (for setting up redirects)
+	void ProcessPackConfig(const FString& ConfigString, TMap<FString, FString>& ConfigParameters)
+	{
+		FConfigFile PackConfig;
+		PackConfig.ProcessInputFileContents(ConfigString);
+
+		// Input Settings
+		static UArrayProperty* ActionMappingsProp = FindFieldChecked<UArrayProperty>(UInputSettings::StaticClass(), "ActionMappings");
+		static UArrayProperty* AxisMappingsProp = FindFieldChecked<UArrayProperty>(UInputSettings::StaticClass(), "AxisMappings");
+
+		UInputSettings* InputSettingsCDO = GetMutableDefault<UInputSettings>();
+		bool bCheckedOut = false;
+
+		TArray<FInputActionKeyMapping> ActionMappingsToAdd;
+		TArray<FInputAxisKeyMapping> AxisMappingsToAdd;
+
+		FConfigSection* InputSettingsSection = PackConfig.Find("InputSettings");
+		if (InputSettingsSection)
+		{
+			for (auto SettingPair : *InputSettingsSection)
+			{
+				struct FMatchMappingByName
+				{
+					FMatchMappingByName(const FName InName)
+						: Name(InName)
+					{
+					}
+
+					bool operator() (const FInputActionKeyMapping& ActionMapping)
+					{
+						return ActionMapping.ActionName == Name;
+					}
+
+					bool operator() (const FInputAxisKeyMapping& AxisMapping)
+					{
+						return AxisMapping.AxisName == Name;
+					}
+
+					FName Name;
+				};
+
+				if (SettingPair.Key.ToString().Contains("ActionMappings"))
+				{
+					FInputActionKeyMapping ActionKeyMapping;
+					ActionMappingsProp->Inner->ImportText(*SettingPair.Value, &ActionKeyMapping, PPF_None, nullptr);
+
+					if (InputSettingsCDO->ActionMappings.FindByPredicate(FMatchMappingByName(ActionKeyMapping.ActionName)) == nullptr)
+					{
+						ActionMappingsToAdd.Add(ActionKeyMapping);
+					}
+				}
+				else if (SettingPair.Key.ToString().Contains("AxisMappings"))
+				{
+					FInputAxisKeyMapping AxisKeyMapping;
+					AxisMappingsProp->Inner->ImportText(*SettingPair.Value, &AxisKeyMapping, PPF_None, nullptr);
+
+					if (InputSettingsCDO->AxisMappings.FindByPredicate(FMatchMappingByName(AxisKeyMapping.AxisName)) == nullptr)
+					{
+						AxisMappingsToAdd.Add(AxisKeyMapping);
+					}
+				}
+
+				if (ActionMappingsToAdd.Num() > 0 || AxisMappingsToAdd.Num() > 0)
+				{
+					if (ISourceControlModule::Get().IsEnabled())
+					{
+						FText ErrorMessage;
+
+						const FString InputSettingsFilename = FPaths::ConvertRelativePathToFull(InputSettingsCDO->GetDefaultConfigFilename());
+						if (!SourceControlHelpers::CheckoutOrMarkForAdd(InputSettingsFilename, FText::FromString(InputSettingsFilename), NULL, ErrorMessage))
+						{
+							UE_LOG(LogPackFactory, Error, TEXT("%s"), *ErrorMessage.ToString());
+						}
+					}
+
+					for (const FInputActionKeyMapping& ActionKeyMapping : ActionMappingsToAdd)
+					{
+						InputSettingsCDO->AddActionMapping(ActionKeyMapping);
+					}
+					for (const FInputAxisKeyMapping& AxisKeyMapping : AxisMappingsToAdd)
+					{
+						InputSettingsCDO->AddAxisMapping(AxisKeyMapping);
+					}
+
+					InputSettingsCDO->SaveKeyMappings();
+					InputSettingsCDO->UpdateDefaultConfigFile();
+				}
+			}
+		}
+
+		FConfigSection* RedirectsSection = PackConfig.Find("Redirects");
+		if (RedirectsSection)
+		{	
+			if (FString* GameName = RedirectsSection->Find("GameName"))
+			{
+				ConfigParameters.Add("GameName", *GameName);
+			}
+		}
+	}
+}
+
+UObject* UPackFactory::FactoryCreateBinary
+(
+	UClass*				Class,
+	UObject*			InParent,
+	FName				Name,
+	EObjectFlags		Flags,
+	UObject*			Context,
+	const TCHAR*		FileType,
+	const uint8*&		Buffer,
+	const uint8*		BufferEnd,
+	FFeedbackContext*	Warn
+)
+{
+	FBufferReader PakReader((void*)Buffer, BufferEnd-Buffer, false);
+	FPakFile PakFile(&PakReader);
+
+	UObject* ReturnAsset = nullptr;
+
+	if (PakFile.IsValid())
+	{
+		static FString ContentFolder(TEXT("/Content/"));
+		static FString SourceFolder(TEXT("/Source/"));
+		FString ContentDestinationRoot = FPaths::GameContentDir();
+
+		const int32 ChopIndex = PakFile.GetMountPoint().Find(ContentFolder);
+		if (ChopIndex != INDEX_NONE)
+		{
+			ContentDestinationRoot /= PakFile.GetMountPoint().RightChop(ChopIndex + ContentFolder.Len());
+		}
+
+		TArray<uint8> Buffer;
+		TArray<uint8> PersistentCompressionBuffer;
+		Buffer.AddUninitialized(8 * 1024 * 1024); // 8MB buffer for extracting
+		int32 ErrorCount = 0;
+		int32 FileCount = 0;
+
+		bool bContainsSource = false;
+		FModuleContextInfo SourceModuleInfo;
+		TMap<FString, FString> ConfigParameters;
+
+		TArray<FString> WrittenFiles;
+
+		// Process the config files and identify if we have source files
+		for (FPakFile::FFileIterator It(PakFile); It; ++It, ++FileCount)
+		{
+			if (It.Filename().StartsWith(TEXT("Config/")) || It.Filename().Contains(TEXT("/Config/")))
+			{
+				const FPakEntry& Entry = It.Info();
+				PakReader.Seek(Entry.Offset);
+				FPakEntry EntryInfo;
+				EntryInfo.Serialize(PakReader, PakFile.GetInfo().Version);
+
+				if (EntryInfo == Entry)
+				{
+					FString ConfigString;
+					PackFactoryHelper::ExtractFileToString(Entry, PakReader, Buffer, PersistentCompressionBuffer, ConfigString);
+					PackFactoryHelper::ProcessPackConfig(ConfigString, ConfigParameters);
+				}
+				else
+				{
+					UE_LOG(LogPackFactory, Error, TEXT("Serialized hash mismatch for \"%s\"."), *It.Filename());
+					ErrorCount++;
+				}
+			}
+			else if (!bContainsSource && (It.Filename().StartsWith(TEXT("Source/")) || It.Filename().Contains(TEXT("/Source/"))))
+			{
+				bContainsSource = true;
+			}
+		}
+
+		// If we have source files, set up the project files if necessary and the game name redirects for blueprints saved with class
+		// references to the module name from the source template
+		if (bContainsSource)
+		{
+			FGameProjectGenerationModule& GameProjectModule = FModuleManager::LoadModuleChecked<FGameProjectGenerationModule>(TEXT("GameProjectGeneration"));
+
+			if (!GameProjectModule.Get().ProjectHasCodeFiles())
+			{
+				TArray<FString> StartupModuleNames;
+				TArray<FString> CreatedFiles;
+				FText OutFailReason;
+				if ( GameProjectModule.Get().GenerateBasicSourceCode(CreatedFiles, OutFailReason) )
+				{
+					WrittenFiles.Append(CreatedFiles);
+				}
+				else
+				{
+					UE_LOG(LogPackFactory, Error, TEXT("Unable to create basic source code: '%s'"), *OutFailReason.ToString());
+				}
+			}
+
+			for (const FModuleContextInfo& ModuleInfo : GameProjectModule.Get().GetCurrentProjectModules())
+			{
+				// Pick the module to insert the code in.  For now always pick the first Runtime module
+				if (ModuleInfo.ModuleType == EHostType::Runtime)
+				{
+					SourceModuleInfo = ModuleInfo;
+
+					// Setup the game name redirect
+					if (FString* GameName = ConfigParameters.Find("GameName"))
+					{
+
+						const FString EngineIniFilename = FPaths::ConvertRelativePathToFull(GetDefault<UEngine>()->GetDefaultConfigFilename());
+
+						if (ISourceControlModule::Get().IsEnabled())
+						{
+							FText ErrorMessage;
+
+							if (!SourceControlHelpers::CheckoutOrMarkForAdd(EngineIniFilename, FText::FromString(EngineIniFilename), NULL, ErrorMessage))
+							{
+								UE_LOG(LogPackFactory, Error, TEXT("%s"), *ErrorMessage.ToString());
+							}
+						}
+
+						const FString RedirectsSection(TEXT("/Script/Engine.Engine"));
+						const FString LongOldGameName = FString::Printf(TEXT("/Script/%s"), **GameName);
+						const FString LongNewGameName = FString::Printf(TEXT("/Script/%s"), *ModuleInfo.ModuleName);
+						
+						FConfigCacheIni Config;
+						FConfigFile& NewFile = Config.Add(EngineIniFilename, FConfigFile());
+						FConfigCacheIni::LoadLocalIniFile(NewFile, TEXT("DefaultEngine"), false);
+						FConfigSection* PackageRedirects = Config.GetSectionPrivate(*RedirectsSection, true, false, EngineIniFilename);
+
+						PackageRedirects->Add(TEXT("+ActiveGameNameRedirects"), FString::Printf(TEXT("(OldGameName=\"%s\",NewGameName=\"%s\")"), *LongOldGameName, *LongNewGameName));
+						PackageRedirects->Add(TEXT("+ActiveGameNameRedirects"), FString::Printf(TEXT("(OldGameName=\"%s\",NewGameName=\"%s\")"), **GameName, *LongNewGameName));
+
+						NewFile.UpdateSections(*EngineIniFilename, *RedirectsSection);
+
+						FString FinalIniFileName;
+						GConfig->LoadGlobalIniFile(FinalIniFileName, *RedirectsSection, NULL, NULL, true);
+
+						ULinkerLoad::AddGameNameRedirect(*LongOldGameName, *LongNewGameName);
+						ULinkerLoad::AddGameNameRedirect(**GameName, *LongNewGameName);
+					}
+					break;
+				}
+			}
+		}
+
+		// Process everything else and copy out to disk
+		for (FPakFile::FFileIterator It(PakFile); It; ++It, ++FileCount)
+		{
+			// config files already handled
+			if (It.Filename().StartsWith(TEXT("Config/")) || It.Filename().Contains(TEXT("/Config/")))
+			{
+				continue;
+			}
+
+			const FPakEntry& Entry = It.Info();
+			PakReader.Seek(Entry.Offset);
+			FPakEntry EntryInfo;
+			EntryInfo.Serialize(PakReader, PakFile.GetInfo().Version);
+
+			if (EntryInfo == Entry)
+			{
+				if (It.Filename().StartsWith(TEXT("Source/")) || It.Filename().Contains(TEXT("/Source/")))
+				{
+					FString DestFilename = It.Filename();
+					if (DestFilename.StartsWith(TEXT("Source/")))
+					{
+						DestFilename = DestFilename.RightChop(7);
+					}
+					else 
+					{
+						const int32 SourceIndex = DestFilename.Find(SourceFolder);
+						if (SourceIndex != INDEX_NONE)
+						{
+							DestFilename = DestFilename.RightChop(SourceIndex + 8);
+						}
+					}
+
+					DestFilename = SourceModuleInfo.ModuleSourcePath / DestFilename;
+					UE_LOG(LogPackFactory, Log, TEXT("%s (%ld) -> %s"), *It.Filename(), Entry.Size, *DestFilename);
+
+					FString SourceContents;
+					PackFactoryHelper::ExtractFileToString(Entry, PakReader, Buffer, PersistentCompressionBuffer, SourceContents);
+
+					FGameProjectGenerationModule& GameProjectModule = FModuleManager::LoadModuleChecked<FGameProjectGenerationModule>(TEXT("GameProjectGeneration"));
+					SourceContents = SourceContents.Replace(TEXT("%MODULE_INCLUDE_PATH%"), *GameProjectModule.Get().DetermineModuleIncludePath(SourceModuleInfo, DestFilename), ESearchCase::CaseSensitive);
+
+					if (FFileHelper::SaveStringToFile(SourceContents, *DestFilename))
+					{
+						WrittenFiles.Add(*DestFilename);
+					}
+					else
+					{
+						UE_LOG(LogPackFactory, Error, TEXT("Unable to write file \"%s\"."), *DestFilename);
+						++ErrorCount;
+					}
+				}
+				else
+				{
+					FString DestFilename = It.Filename();
+					if (DestFilename.StartsWith(TEXT("Content/")))
+					{
+						DestFilename = DestFilename.RightChop(8);
+					}
+					else
+					{
+						const int32 ContentIndex = DestFilename.Find(ContentFolder);
+						if (ContentIndex != INDEX_NONE)
+						{
+							DestFilename = DestFilename.RightChop(ContentIndex + 9);
+						}
+					}
+					DestFilename = ContentDestinationRoot / DestFilename;
+					UE_LOG(LogPackFactory, Log, TEXT("%s (%ld) -> %s"), *It.Filename(), Entry.Size, *DestFilename);
+
+					TAutoPtr<FArchive> FileHandle(IFileManager::Get().CreateFileWriter(*DestFilename));
+
+					if (FileHandle.IsValid())
+					{
+						PackFactoryHelper::ExtractFile(Entry, PakReader, Buffer, PersistentCompressionBuffer, *FileHandle);
+						WrittenFiles.Add(*DestFilename);
+					}
+					else
+					{
+						UE_LOG(LogPackFactory, Error, TEXT("Unable to create file \"%s\"."), *DestFilename);
+						++ErrorCount;
+					}
+				}
+			}
+			else
+			{
+				UE_LOG(LogPackFactory, Error, TEXT("Serialized hash mismatch for \"%s\"."), *It.Filename());
+				ErrorCount++;
+			}
+		}
+
+		UE_LOG(LogPackFactory, Log, TEXT("Finished extracting %d files (including %d errors)."), FileCount, ErrorCount);
+
+		if (WrittenFiles.Num() > 0)
+		{
+			// If we wrote out source files, kick off the hot reload process
+			// TODO: Ideally we would block on finishing the factory until the hot reload completed
+			if (bContainsSource)
+			{
+				IHotReloadInterface& HotReloadSupport = FModuleManager::LoadModuleChecked<IHotReloadInterface>("HotReload");
+				if( !HotReloadSupport.IsCurrentlyCompiling() )
+				{
+					HotReloadSupport.DoHotReloadFromEditor();
+				}
+			}
+
+			// Find an asset to return as the file to be selected in the content browser
+			// TODO: Should we load all objects? Should the config file be able to specify which asset is selected?
+			for (const FString& Filename : WrittenFiles)
+			{
+				static const FString AssetExtension(TEXT(".uasset"));
+				if (Filename.EndsWith(AssetExtension))
+				{
+					FString GameFileName = Filename;
+					if (FPaths::MakePathRelativeTo(GameFileName, *FPaths::GameContentDir()))
+					{
+						int32 SlashIndex = INDEX_NONE;
+						GameFileName = FString(TEXT("/Game/")) / GameFileName.LeftChop(AssetExtension.Len());
+						if (GameFileName.FindLastChar(TEXT('/'), SlashIndex))
+						{
+							const FString AssetName = GameFileName.RightChop(SlashIndex + 1);
+							ReturnAsset = LoadObject<UObject>(nullptr, *(GameFileName + TEXT(".") + AssetName));
+							if (ReturnAsset) 
+							{
+								break;
+							}
+						}
+					}
+				}
+			}
+
+			// If source control is enabled mark all the added files for checkout/add
+			if (ISourceControlModule::Get().IsEnabled() && GetDefault<UEditorLoadingSavingSettings>()->bSCCAutoAddNewFiles)
+			{
+				for (const FString& Filename : WrittenFiles)
+				{
+					FText ErrorMessage;
+					if (!SourceControlHelpers::CheckoutOrMarkForAdd(Filename, FText::FromString(Filename), NULL, ErrorMessage))
+					{
+						UE_LOG(LogPackFactory, Error, TEXT("%s"), *ErrorMessage.ToString());
+					}
+				}
+			}
+		}
+	}
+	else
+	{
+		UE_LOG(LogPackFactory, Warning, TEXT("Invalid pak file."));
+	}
+
+
+	return ReturnAsset;
+}
\ No newline at end of file
diff --git a/Engine/Source/Editor/UnrealEd/UnrealEd.Build.cs b/Engine/Source/Editor/UnrealEd/UnrealEd.Build.cs
--- a/Engine/Source/Editor/UnrealEd/UnrealEd.Build.cs
+++ b/Engine/Source/Editor/UnrealEd/UnrealEd.Build.cs
@@ -95,6 +95,7 @@ public class UnrealEd : ModuleRules
 				"MaterialEditor",
 				"MessageLog",
 				"NetworkFileSystem",
+				"PakFile",
 				"PropertyEditor",
 				"Projects",
 				"RawMesh",
