Fixed editor startup cooker spam.

diff --git a/Engine/Source/Editor/UnrealEd/Private/CookOnTheFlyServer.cpp b/Engine/Source/Editor/UnrealEd/Private/CookOnTheFlyServer.cpp
--- a/Engine/Source/Editor/UnrealEd/Private/CookOnTheFlyServer.cpp
+++ b/Engine/Source/Editor/UnrealEd/Private/CookOnTheFlyServer.cpp
@@ -35,7 +35,7 @@
 DEFINE_LOG_CATEGORY_STATIC(LogCookOnTheFly, Log, All);
 
 #define DEBUG_COOKONTHEFLY 0
-#define OUTPUT_TIMING 1
+#define OUTPUT_TIMING 0
 
 #if OUTPUT_TIMING
 
@@ -1795,7 +1795,7 @@ bool UCookOnTheFlyServer::GetCurrentIniVersionStrings( const ITargetPlatform* Ta
 
 		if ( IniVersionedArray.Num() != 3 )
 		{
-			UE_LOG(LogCookOnTheFly, Warning, TEXT("Invalid entry in CookSettings, VersionedIniParams %s"), *IniVersioned);
+			// UE_LOG(LogCookOnTheFly, Warning, TEXT("Invalid entry in CookSettings, VersionedIniParams %s"), *IniVersioned);
 			return false;
 		}
 
@@ -1816,7 +1816,7 @@ bool UCookOnTheFlyServer::GetCurrentIniVersionStrings( const ITargetPlatform* Ta
 		FString Value;
 		if ( !PlatformIniFile->GetString(*Section, *Key, Value) )
 		{
-			UE_LOG(LogCookOnTheFly, Warning, TEXT("Unable to find entry in CookSettings, VersionedIniParams %s, assume default is being used"), *IniVersioned);
+			// UE_LOG(LogCookOnTheFly, Warning, TEXT("Unable to find entry in CookSettings, VersionedIniParams %s, assume default is being used"), *IniVersioned);
 			continue;
 		}
 
