Don't throttle editor during amortized lighting export

diff --git a/Engine/Source/Editor/UnrealEd/Classes/Editor/EditorEngine.h b/Engine/Source/Editor/UnrealEd/Classes/Editor/EditorEngine.h
--- a/Engine/Source/Editor/UnrealEd/Classes/Editor/EditorEngine.h
+++ b/Engine/Source/Editor/UnrealEd/Classes/Editor/EditorEngine.h
@@ -1431,6 +1431,8 @@ public:
 	/** Checks to see if the asynchronous lighting build is running or not */
 	bool IsLightingBuildCurrentlyRunning() const;
 
+	bool IsLightingBuildCurrentlyExporting() const;
+
 	/** Checks if asynchronous lighting is building, if so, it throws a warning notification and returns true */
 	bool WarnIfLightingBuildIsCurrentlyRunning();
 
diff --git a/Engine/Source/Editor/UnrealEd/Private/Editor.cpp b/Engine/Source/Editor/UnrealEd/Private/Editor.cpp
--- a/Engine/Source/Editor/UnrealEd/Private/Editor.cpp
+++ b/Engine/Source/Editor/UnrealEd/Private/Editor.cpp
@@ -6041,6 +6041,12 @@ bool UEditorEngine::ShouldThrottleCPUUsage() const
 		}
 	}
 
+	// Don't throttle during amortized export, greatly increases export time
+	if (IsLightingBuildCurrentlyExporting())
+	{
+		return false;
+	}
+
 	return bShouldThrottle;
 }
 
diff --git a/Engine/Source/Editor/UnrealEd/Private/StaticLightingSystem/StaticLightingPrivate.h b/Engine/Source/Editor/UnrealEd/Private/StaticLightingSystem/StaticLightingPrivate.h
--- a/Engine/Source/Editor/UnrealEd/Private/StaticLightingSystem/StaticLightingPrivate.h
+++ b/Engine/Source/Editor/UnrealEd/Private/StaticLightingSystem/StaticLightingPrivate.h
@@ -204,6 +204,8 @@ public:
 	/** True if a lighting build is currently in the works, or is finished but not accepted yet */
 	bool IsLightingBuildCurrentlyRunning() const;
 
+	bool IsLightingBuildCurrentlyExporting() const;
+
 private:
 	FStaticLightingManager()
 		: StaticLightingSystem(NULL) {}
@@ -259,6 +261,8 @@ public:
 	/** True if the current stage of building is asynchronous (lightmass in flight) */
 	bool IsAsyncBuilding() const;
 
+	bool IsAmortizedExporting() const;
+
 private:
 	/**
 	 * Generates mappings/meshes for all BSP in the given level
diff --git a/Engine/Source/Editor/UnrealEd/Private/StaticLightingSystem/StaticLightingSystem.cpp b/Engine/Source/Editor/UnrealEd/Private/StaticLightingSystem/StaticLightingSystem.cpp
--- a/Engine/Source/Editor/UnrealEd/Private/StaticLightingSystem/StaticLightingSystem.cpp
+++ b/Engine/Source/Editor/UnrealEd/Private/StaticLightingSystem/StaticLightingSystem.cpp
@@ -288,6 +288,11 @@ bool FStaticLightingManager::IsLightingBuildCurrentlyRunning() const
 	return StaticLightingSystem != NULL;
 }
 
+bool FStaticLightingManager::IsLightingBuildCurrentlyExporting() const
+{
+	return StaticLightingSystem != NULL && StaticLightingSystem->IsAmortizedExporting();
+}
+
 
 
 
@@ -2105,6 +2110,11 @@ bool FStaticLightingSystem::IsAsyncBuilding() const
 	return CurrentBuildStage == FStaticLightingSystem::AsynchronousBuilding;
 }
 
+bool FStaticLightingSystem::IsAmortizedExporting() const
+{
+	return CurrentBuildStage == FStaticLightingSystem::AmortizedExport;
+}
+
 #endif
 
 
@@ -2144,6 +2154,11 @@ bool UEditorEngine::IsLightingBuildCurrentlyRunning() const
 	return FStaticLightingManager::Get()->IsLightingBuildCurrentlyRunning();
 }
 
+bool UEditorEngine::IsLightingBuildCurrentlyExporting() const
+{
+	return FStaticLightingManager::Get()->IsLightingBuildCurrentlyExporting(); 
+}
+
 bool UEditorEngine::WarnIfLightingBuildIsCurrentlyRunning()
 {
 	bool bFailure = IsLightingBuildCurrentlyRunning();
