Updated MeshMerging UI to use new lightmap generation method, also uses channel 0 as lightmap source UV

diff --git a/Engine/Source/Developer/MeshUtilities/Private/MeshUtilities.cpp b/Engine/Source/Developer/MeshUtilities/Private/MeshUtilities.cpp
--- a/Engine/Source/Developer/MeshUtilities/Private/MeshUtilities.cpp
+++ b/Engine/Source/Developer/MeshUtilities/Private/MeshUtilities.cpp
@@ -3457,13 +3457,9 @@ void FMeshUtilities::MergeActors(
 	struct FRawMeshExt
 	{
 		FRawMeshExt() 
-			: LightMapCoordinateIndex(0)
-			, LightMapRes(32)
 		{}
 		
 		FRawMesh	Mesh;
-		int32		LightMapCoordinateIndex;
-		int32		LightMapRes;
 		FString		AssetPackageName;
 		FVector		Pivot;	
 	};
@@ -3486,12 +3482,7 @@ void FMeshUtilities::MergeActors(
 		{
 			MaterialMap.Add(MeshId, MeshMaterialMap);
 
-			// Store mesh lightmap info
-			FIntPoint ActorLightMapRes; 
-			MeshComponent->GetLightMapResolution(ActorLightMapRes.X, ActorLightMapRes.Y);
-			SourceMeshes[MeshId].LightMapRes = ActorLightMapRes.X;
-			SourceMeshes[MeshId].LightMapCoordinateIndex = MeshComponent->StaticMesh->LightMapCoordinateIndex;
-			
+
 			// Store component location
 			SourceMeshes[MeshId].Pivot = MeshComponent->ComponentToWorld.GetLocation();
 			
@@ -3538,10 +3529,6 @@ void FMeshUtilities::MergeActors(
 
 	FRawMeshExt MergedMesh;
 
-	// Attempt to pack lightmaps
-	float MergedLightMapScale = 1.f;
-	bool bCreateLightMapChannel = false;
-
 	// Use first mesh for naming and pivot
 	MergedMesh.AssetPackageName = SourceMeshes[0].AssetPackageName;
 	MergedMesh.Pivot = InSettings.bPivotPointAtZero ? FVector::ZeroVector : SourceMeshes[0].Pivot;
@@ -3593,12 +3580,6 @@ void FMeshUtilities::MergeActors(
 		// Merge all other UV channels 
 		for (int32 ChannelIdx = 0; ChannelIdx < MAX_MESH_TEXTURE_COORDS; ++ChannelIdx)
 		{
-			// Skip Lightmap channel if any
-			if (bCreateLightMapChannel && ChannelIdx == MergedMesh.LightMapCoordinateIndex)
-			{
-				continue;
-			}
-			
 			// Whether this channel has data
 			if (bOcuppiedUVChannels[ChannelIdx])
 			{
@@ -3660,8 +3641,9 @@ void FMeshUtilities::MergeActors(
 		SrcModel->BuildSettings.bRecomputeTangents = false;
 		SrcModel->BuildSettings.bRemoveDegenerates = false;
 		SrcModel->BuildSettings.bUseFullPrecisionUVs = false;
-		SrcModel->BuildSettings.bGenerateLightmapUVs = InSettings.bGenerateAtlasedLightMapUV;
-		SrcModel->BuildSettings.SrcLightmapIndex = InSettings.TargetLightMapUVChannel;
+		SrcModel->BuildSettings.bGenerateLightmapUVs = InSettings.bGenerateLightMapUV;
+		SrcModel->BuildSettings.MinLightmapResolution = InSettings.TargetLightMapResolution;
+		SrcModel->BuildSettings.SrcLightmapIndex = 0;
 		SrcModel->BuildSettings.DstLightmapIndex = InSettings.TargetLightMapUVChannel;
 
 		SrcModel->RawMeshBulkData->SaveRawMesh(MergedMesh.Mesh);
diff --git a/Engine/Source/Developer/MeshUtilities/Public/MeshUtilities.h b/Engine/Source/Developer/MeshUtilities/Public/MeshUtilities.h
--- a/Engine/Source/Developer/MeshUtilities/Public/MeshUtilities.h
+++ b/Engine/Source/Developer/MeshUtilities/Public/MeshUtilities.h
@@ -73,14 +73,14 @@ public:
  */
 struct FMeshMergingSettings
 {
-	/** Whether to generate atlased lightmap UVs for a merged mesh*/
-	bool bGenerateAtlasedLightMapUV;
+	/** Whether to generate lightmap UVs for a merged mesh*/
+	bool bGenerateLightMapUV;
 	
-	/** Target UV channel in a merged mesh for an atlased lightmap */
+	/** Target UV channel in a merged mesh for a lightmap */
 	int32 TargetLightMapUVChannel;
 
-	/** Upper bounds for an atlased lightmap resolution */
-	int32 MaxAltlasedLightMapResolution;
+	/** Target lightmap resolution */
+	int32 TargetLightMapResolution;
 		
 	/** Whether we should import vertex colors into merged mesh */
 	bool bImportVertexColors;
@@ -90,9 +90,9 @@ struct FMeshMergingSettings
 		
 	/** Default settings. */
 	FMeshMergingSettings()
-		: bGenerateAtlasedLightMapUV(true)
+		: bGenerateLightMapUV(false)
 		, TargetLightMapUVChannel(1)
-		, MaxAltlasedLightMapResolution(1024)
+		, TargetLightMapResolution(256)
 		, bImportVertexColors(false)
 		, bPivotPointAtZero(false)
 	{
diff --git a/Engine/Source/Editor/UnrealEd/Private/Dialogs/SMeshMergingDialog.cpp b/Engine/Source/Editor/UnrealEd/Private/Dialogs/SMeshMergingDialog.cpp
--- a/Engine/Source/Editor/UnrealEd/Private/Dialogs/SMeshMergingDialog.cpp
+++ b/Engine/Source/Editor/UnrealEd/Private/Dialogs/SMeshMergingDialog.cpp
@@ -71,11 +71,11 @@ void SMeshMergingDialog::Construct(const FArguments& InArgs)
 				[
 					SNew(SCheckBox)
 					.Type(ESlateCheckBoxType::CheckBox)
-					.IsChecked(this, &SMeshMergingDialog::GetAtlasLightmapUV)
-					.OnCheckStateChanged(this, &SMeshMergingDialog::SetAtlasLightmapUV)
+					.IsChecked(this, &SMeshMergingDialog::GetGenerateLightmapUV)
+					.OnCheckStateChanged(this, &SMeshMergingDialog::SetGenerateLightmapUV)
 					.Content()
 					[
-						SNew(STextBlock).Text(LOCTEXT("AtlasLightmapUVLabel", "Generate Atlased Lightmap UVs"))
+						SNew(STextBlock).Text(LOCTEXT("AtlasLightmapUVLabel", "Generate Lightmap UVs"))
 					]
 				]
 					
@@ -111,7 +111,7 @@ void SMeshMergingDialog::Construct(const FArguments& InArgs)
 					.VAlign(VAlign_Center)
 					[
 						SNew(STextBlock)
-						.Text(LOCTEXT("MaxLightMapResolutionLabel", "Max Resolution:"))
+						.Text(LOCTEXT("TargetLightMapResolutionLabel", "Target Resolution:"))
 					]
 												
 					+SHorizontalBox::Slot()
@@ -122,8 +122,8 @@ void SMeshMergingDialog::Construct(const FArguments& InArgs)
 						SNew(STextComboBox)
 						.IsEnabled( this, &SMeshMergingDialog::IsLightmapChannelEnabled )
 						.OptionsSource(&LightMapResolutionOptions)
-						.InitiallySelectedItem(LightMapResolutionOptions[FMath::FloorLog2(MergingSettings.MaxAltlasedLightMapResolution)])
-						.OnSelectionChanged(this, &SMeshMergingDialog::SetMaxLightMapResolution)
+						.InitiallySelectedItem(LightMapResolutionOptions[FMath::FloorLog2(MergingSettings.TargetLightMapResolution)])
+						.OnSelectionChanged(this, &SMeshMergingDialog::SetTargetLightMapResolution)
 					]
 				]
 			]
@@ -267,19 +267,19 @@ FReply SMeshMergingDialog::OnMergeClicked()
 	return FReply::Handled();
 }
 
-ESlateCheckBoxState::Type SMeshMergingDialog::GetAtlasLightmapUV() const
+ESlateCheckBoxState::Type SMeshMergingDialog::GetGenerateLightmapUV() const
 {
-	return (MergingSettings.bGenerateAtlasedLightMapUV ? ESlateCheckBoxState::Checked : ESlateCheckBoxState::Unchecked);
+	return (MergingSettings.bGenerateLightMapUV ? ESlateCheckBoxState::Checked : ESlateCheckBoxState::Unchecked);
 }
 
-void SMeshMergingDialog::SetAtlasLightmapUV(ESlateCheckBoxState::Type NewValue)
+void SMeshMergingDialog::SetGenerateLightmapUV(ESlateCheckBoxState::Type NewValue)
 {
-	MergingSettings.bGenerateAtlasedLightMapUV = (ESlateCheckBoxState::Checked == NewValue);
+	MergingSettings.bGenerateLightMapUV = (ESlateCheckBoxState::Checked == NewValue);
 }
 
 bool SMeshMergingDialog::IsLightmapChannelEnabled() const
 {
-	return MergingSettings.bGenerateAtlasedLightMapUV;
+	return MergingSettings.bGenerateLightMapUV;
 }
 
 void SMeshMergingDialog::SetTargetLightMapChannel(TSharedPtr<FString> NewSelection, ESelectInfo::Type SelectInfo)
@@ -287,9 +287,9 @@ void SMeshMergingDialog::SetTargetLightMapChannel(TSharedPtr<FString> NewSelecti
 	TTypeFromString<int32>::FromString(MergingSettings.TargetLightMapUVChannel, **NewSelection);
 }
 
-void SMeshMergingDialog::SetMaxLightMapResolution(TSharedPtr<FString> NewSelection, ESelectInfo::Type SelectInfo)
+void SMeshMergingDialog::SetTargetLightMapResolution(TSharedPtr<FString> NewSelection, ESelectInfo::Type SelectInfo)
 {
-	TTypeFromString<int32>::FromString(MergingSettings.MaxAltlasedLightMapResolution, **NewSelection);
+	TTypeFromString<int32>::FromString(MergingSettings.TargetLightMapResolution, **NewSelection);
 }
 
 ESlateCheckBoxState::Type SMeshMergingDialog::GetImportVertexColors() const
diff --git a/Engine/Source/Editor/UnrealEd/Public/Dialogs/SMeshMergingDialog.h b/Engine/Source/Editor/UnrealEd/Public/Dialogs/SMeshMergingDialog.h
--- a/Engine/Source/Editor/UnrealEd/Public/Dialogs/SMeshMergingDialog.h
+++ b/Engine/Source/Editor/UnrealEd/Public/Dialogs/SMeshMergingDialog.h
@@ -34,13 +34,13 @@ private:
 	FReply OnMergeClicked();
 
 	/**  */
-	ESlateCheckBoxState::Type GetAtlasLightmapUV() const;
-	void SetAtlasLightmapUV(ESlateCheckBoxState::Type NewValue);
+	ESlateCheckBoxState::Type GetGenerateLightmapUV() const;
+	void SetGenerateLightmapUV(ESlateCheckBoxState::Type NewValue);
 
 	/** Target lightmap channel */
 	bool IsLightmapChannelEnabled() const;
 	void SetTargetLightMapChannel(TSharedPtr<FString> NewSelection, ESelectInfo::Type SelectInfo);
-	void SetMaxLightMapResolution(TSharedPtr<FString> NewSelection, ESelectInfo::Type SelectInfo);
+	void SetTargetLightMapResolution(TSharedPtr<FString> NewSelection, ESelectInfo::Type SelectInfo);
 		
 	/**  */
 	ESlateCheckBoxState::Type GetImportVertexColors() const;
