Implemented new editor-only methods in the UPrimitiveComponent interface to determine whether a selection box or frustum intersect with the component.
#change Removed old monolithic methods in FLevelEditorViewportClient and DragTool_FrustumSelect.
#change Improved selection criteria for ULandscapeComponent, UArrowComponent, and fixed a bug in selection criterion for UBillboardComponent.

diff --git a/Engine/Source/Editor/UnrealEd/Private/DragTool_BoxSelect.cpp b/Engine/Source/Editor/UnrealEd/Private/DragTool_BoxSelect.cpp
--- a/Engine/Source/Editor/UnrealEd/Private/DragTool_BoxSelect.cpp
+++ b/Engine/Source/Editor/UnrealEd/Private/DragTool_BoxSelect.cpp
@@ -283,7 +283,6 @@ bool FDragTool_ActorBoxSelect::IntersectsBox( AActor& InActor, const FBox& InBox
 	// Check for special cases (like certain show flags that might hide an actor)
 	bool bActorIsHiddenByShowFlags = false;
 
-
 	// Check to see that volume actors are visible in the viewport
 	if( InActor.IsA(AVolume::StaticClass()) && (!LevelViewportClient->EngineShowFlags.Volumes || !LevelViewportClient->IsVolumeVisibleInViewport(InActor) ) )
 	{
@@ -300,12 +299,12 @@ bool FDragTool_ActorBoxSelect::IntersectsBox( AActor& InActor, const FBox& InBox
 		TArray<UPrimitiveComponent*> PrimitiveComponents;
 		InActor.GetComponents(PrimitiveComponents);
 
-		for( int32 ComponentIndex = 0 ; ComponentIndex < PrimitiveComponents.Num() ; ++ComponentIndex )
+		for (const UPrimitiveComponent* PrimitiveComponent : PrimitiveComponents)
 		{
-			UPrimitiveComponent* PrimitiveComponent = PrimitiveComponents[ComponentIndex];
-			if ( PrimitiveComponent->IsRegistered() && PrimitiveComponent->IsVisibleInEditor() )
+			check(PrimitiveComponent != nullptr);
+			if (PrimitiveComponent->IsRegistered() && PrimitiveComponent->IsVisibleInEditor())
 			{
-				if ( LevelViewportClient->ComponentIsTouchingSelectionBox( &InActor, PrimitiveComponent, InBox, bGeometryMode, bUseStrictSelection ) )
+				if (PrimitiveComponent->ComponentIsTouchingSelectionBox(InBox, LevelViewportClient->EngineShowFlags, bGeometryMode, bUseStrictSelection))
 				{
 					bActorHitByBox = true;
 					break;
diff --git a/Engine/Source/Editor/UnrealEd/Private/DragTool_FrustumSelect.cpp b/Engine/Source/Editor/UnrealEd/Private/DragTool_FrustumSelect.cpp
--- a/Engine/Source/Editor/UnrealEd/Private/DragTool_FrustumSelect.cpp
+++ b/Engine/Source/Editor/UnrealEd/Private/DragTool_FrustumSelect.cpp
@@ -23,7 +23,7 @@
 void FDragTool_ActorFrustumSelect::AddDelta( const FVector& InDelta )
 {
 	FIntPoint MousePos;
-	ViewportClient->Viewport->GetMousePos(MousePos);
+	LevelViewportClient->Viewport->GetMousePos(MousePos);
 
 	EndWk = FVector(MousePos);
 	End = EndWk;
@@ -36,8 +36,6 @@ void FDragTool_ActorFrustumSelect::StartDrag(FEditorViewportClient* InViewportCl
 {
 	FDragTool::StartDrag(InViewportClient, InStart, InStartScreen);
 
-	ViewportClient = InViewportClient;
-
 	const bool bUseHoverFeedback = GEditor != NULL && GetDefault<ULevelEditorViewportSettings>()->bEnableViewportHoverFeedback;
 
 	// Remove any active hover objects
@@ -59,8 +57,8 @@ void FDragTool_ActorFrustumSelect::EndDrag()
 	FEditorModeTools& EdModeTools = GLevelEditorModeTools();
 	const bool bGeometryMode = EdModeTools.IsModeActive( FBuiltinEditorModes::EM_Geometry );
 
-	FSceneViewFamilyContext ViewFamily(FSceneViewFamily::ConstructionValues(ViewportClient->Viewport, ViewportClient->GetScene(), ViewportClient->EngineShowFlags ));
-	FSceneView* SceneView = ViewportClient->CalcSceneView(&ViewFamily);
+	FSceneViewFamilyContext ViewFamily(FSceneViewFamily::ConstructionValues(LevelViewportClient->Viewport, LevelViewportClient->GetScene(), LevelViewportClient->EngineShowFlags ));
+	FSceneView* SceneView = LevelViewportClient->CalcSceneView(&ViewFamily);
 
 	// Generate a frustum out of the dragged box
 	FConvexVolume Frustum;
@@ -93,8 +91,8 @@ void FDragTool_ActorFrustumSelect::EndDrag()
 		}
 		
 
-		const int32 ViewportSizeX = ViewportClient->Viewport->GetSizeXY().X;
-		const int32 ViewportSizeY = ViewportClient->Viewport->GetSizeXY().Y;
+		const int32 ViewportSizeX = LevelViewportClient->Viewport->GetSizeXY().X;
+		const int32 ViewportSizeY = LevelViewportClient->Viewport->GetSizeXY().Y;
 
 		if( Start.X > End.X )
 		{
@@ -110,7 +108,7 @@ void FDragTool_ActorFrustumSelect::EndDrag()
 		if (bTransparentBoxSelection)
 		{
 			// Get a list of frustum-culled actors
-			for(FActorIterator It(ViewportClient->GetWorld()); It; ++It)
+			for(FActorIterator It(LevelViewportClient->GetWorld()); It; ++It)
 			{
 				AActor* Actor = *It;
 				if (IntersectsFrustum( *Actor, Frustum, bStrictDragSelection))
@@ -126,7 +124,7 @@ void FDragTool_ActorFrustumSelect::EndDrag()
 			// Extend the endpoint of the rect to get the actual line
 			FIntRect BoxRect( FIntPoint( FMath::Max( 0.0f, Start.X ), FMath::Max( 0.0f, Start.Y ) ), FIntPoint( FMath::Min(ViewportSizeX, FMath::TruncToInt(End.X+1)), FMath::Min( ViewportSizeY, FMath::TruncToInt(End.Y+1) ) ) );
 
-			const TArray<FColor>& RawHitProxyData = ViewportClient->Viewport->GetRawHitProxyData(BoxRect);
+			const TArray<FColor>& RawHitProxyData = LevelViewportClient->Viewport->GetRawHitProxyData(BoxRect);
 
 			TSet<AActor*> HitActors;
 			TSet<UModel*> HitModels;
@@ -230,140 +228,19 @@ void FDragTool_ActorFrustumSelect::Render(const FSceneView* View, FCanvas* Canva
 	Canvas->DrawItem( BoxItem );
 }
 
-/** 
- * Returns true if the mesh on the component has vertices which intersect the frustum
- *
- * @param InComponent			The static mesh or skeletal mesh component to check
- * @param InFrustum				The frustum to check against.
- * @param bUseStrictSelection	true if all the vertices must be entirely within the frustum
- */
-bool FDragTool_ActorFrustumSelect::IntersectsVertices( UPrimitiveComponent& InComponent, const FConvexVolume& InFrustum, bool bUseStrictSelection ) const
-{
-	bool bAlreadyProcessed = false;
-	bool bResult = false;
-	UStaticMeshComponent* StaticMeshComponent = Cast<UStaticMeshComponent>(&InComponent);
-	USkeletalMeshComponent* SkeletalMeshComponent = Cast<USkeletalMeshComponent>(&InComponent);
-
-	if( StaticMeshComponent && StaticMeshComponent->StaticMesh && ViewportClient->EngineShowFlags.StaticMeshes )
-	{
-		check( StaticMeshComponent->StaticMesh->RenderData );
-		if( StaticMeshComponent->StaticMesh->RenderData->LODResources.Num() > 0 )
-		{
-			const FStaticMeshLODResources& LODModel = StaticMeshComponent->StaticMesh->RenderData->LODResources[0];
-			uint32 NumVertices = LODModel.VertexBuffer.GetNumVertices();
-			for (uint32 VertexIndex = 0; VertexIndex < NumVertices; ++VertexIndex)
-			{
-				const FVector& LocalPosition = LODModel.PositionVertexBuffer.VertexPosition(VertexIndex);
-				const FVector& WorldPosition = StaticMeshComponent->ComponentToWorld.TransformPosition(LocalPosition);
-				bool bLocationIntersected = InFrustum.IntersectBox(WorldPosition, FVector::ZeroVector);
-				if (bLocationIntersected && !bUseStrictSelection)
-				{
-					bResult = true;
-					bAlreadyProcessed = true;
-					break;
-				}
-				else if (!bLocationIntersected && bUseStrictSelection)
-				{
-					bResult = false;
-					bAlreadyProcessed = true;
-					break;
-				}
-			}
-		}
-
-		if (!bAlreadyProcessed && bUseStrictSelection)
-		{
-			bResult = true;
-			bAlreadyProcessed = true;
-		}
-	}
-	else if( SkeletalMeshComponent && SkeletalMeshComponent->MeshObject && ViewportClient->EngineShowFlags.SkeletalMeshes )
-	{
-		FSkeletalMeshResource* SkelMeshResource = SkeletalMeshComponent->GetSkeletalMeshResource();
-		check(SkelMeshResource);
-		check( SkelMeshResource->LODModels.Num() > 0 );
-
-		const FStaticLODModel& LODModel = SkelMeshResource->LODModels[0];
-		for( int32 ChunkIndex = 0 ; ChunkIndex < LODModel.Chunks.Num() && !bAlreadyProcessed; ++ChunkIndex )
-		{
-			const FSkelMeshChunk& Chunk = LODModel.Chunks[ChunkIndex];
-			for( int32 VertexIndex = 0; VertexIndex < Chunk.RigidVertices.Num(); ++VertexIndex )
-			{
-				const FVector Location = SkeletalMeshComponent->ComponentToWorld.TransformPosition( Chunk.RigidVertices[VertexIndex].Position );
-				const bool bLocationIntersected = InFrustum.IntersectBox( Location, FVector::ZeroVector );
-
-				// If the selection box doesn't have to encompass the entire component and a skeletal mesh vertex has intersected with
-				// the selection box, this component is being touched by the selection box
-				if( bLocationIntersected && !bUseStrictSelection )
-				{
-					bResult = true;
-					bAlreadyProcessed = true;
-					break;
-				}
-
-				// If the selection box has to encompass the entire component and a skeletal mesh vertex didn't intersect with the selection
-				// box, this component does not qualify
-				else if ( !bLocationIntersected && bUseStrictSelection )
-				{
-					bResult = false;
-					bAlreadyProcessed = true;
-					break;
-				}
-			}
-
-			for( int32 VertexIndex = 0 ; VertexIndex < Chunk.SoftVertices.Num() && !bAlreadyProcessed ; ++VertexIndex )
-			{
-				const FVector Location = SkeletalMeshComponent->ComponentToWorld.TransformPosition( Chunk.SoftVertices[VertexIndex].Position );
-				const bool bLocationIntersected = InFrustum.IntersectBox( Location, FVector::ZeroVector );
-
-				// If the selection box doesn't have to encompass the entire component and a skeletal mesh vertex has intersected with
-				// the selection box, this component is being touched by the selection box
-				if( bLocationIntersected && !bUseStrictSelection )
-				{
-					bResult = true;
-					bAlreadyProcessed = true;
-					break;
-				}
-
-				// If the selection box has to encompass the entire component and a skeletal mesh vertex didn't intersect with the selection
-				// box, this component does not qualify
-				else if ( !bLocationIntersected && bUseStrictSelection  )
-				{
-					bResult = false;
-					bAlreadyProcessed = true;
-					break;
-				}
-			}
-
-			// If the selection box has to encompass all of the component and none of the component's verts failed the intersection test, this component
-			// is consider touching
-			if ( !bAlreadyProcessed && bUseStrictSelection )
-			{
-				bResult = true;
-				bAlreadyProcessed = true;
-			}
-		}
-	}
-
-	return bResult;
-}
-
 bool FDragTool_ActorFrustumSelect::IntersectsFrustum( AActor& InActor, const FConvexVolume& InFrustum, bool bUseStrictSelection ) const
 {	
 	bool bActorHitByBox = false;
-	bool bActorFullyContained = false;
+
+	FEditorModeTools& EdModeTools = GLevelEditorModeTools();
+	const bool bGeometryMode = EdModeTools.IsModeActive(FBuiltinEditorModes::EM_Geometry);
 
 	// Check for special cases (like certain show flags that might hide an actor)
 	bool bActorIsHiddenByShowFlags = false;
 
-	FLevelEditorViewportClient* LevelViewportClient = NULL;
-	if ( ViewportClient->IsLevelEditorClient()  )
-	{
-		LevelViewportClient = (FLevelEditorViewportClient*)ViewportClient;
-	}
-
 	// Check to see that volume actors are visible in the viewport
-	if( LevelViewportClient && InActor.IsA(AVolume::StaticClass()) && (!LevelViewportClient->EngineShowFlags.Volumes || !LevelViewportClient->IsVolumeVisibleInViewport(InActor) ) )
+	check(LevelViewportClient != nullptr);
+	if( InActor.IsA(AVolume::StaticClass()) && (!LevelViewportClient->EngineShowFlags.Volumes || !LevelViewportClient->IsVolumeVisibleInViewport(InActor) ) )
 	{
 		bActorIsHiddenByShowFlags = true;
 	}
@@ -375,104 +252,15 @@ bool FDragTool_ActorFrustumSelect::IntersectsFrustum( AActor& InActor, const FCo
 		TArray<UPrimitiveComponent*> Components;
 		InActor.GetComponents(Components);
 
-		for( int32 ComponentIndex = 0 ; ComponentIndex < Components.Num() ; ++ComponentIndex )
+		for (const UPrimitiveComponent* PrimitiveComponent : Components)
 		{
-			UPrimitiveComponent* PrimitiveComponent = Components[ComponentIndex];
+			check(PrimitiveComponent != nullptr);
 			if ( PrimitiveComponent->IsRegistered() && PrimitiveComponent->IsVisibleInEditor() )
 			{
-				FVector Origin = PrimitiveComponent->Bounds.Origin;
-				FVector Extent;
-				if( PrimitiveComponent->IsA(UDrawSphereComponent::StaticClass()) )
-				{
-					// Do not select by DrawLight Radii or Cones
-					continue;
-				}
-				else if( PrimitiveComponent->IsA(UNavLinkRenderingComponent::StaticClass() ) )
-				{
-					continue;
-				}
-				if(PrimitiveComponent->IsA(UBillboardComponent::StaticClass()) && Cast<UBillboardComponent>(PrimitiveComponent)->Sprite != NULL)
-				{
-					// Use the size of the sprite itself rather than its box extent
-					UBillboardComponent* SpriteComponent = Cast<UBillboardComponent>(PrimitiveComponent);
-					float Scale = SpriteComponent->ComponentToWorld.GetMaximumAxisScale();
-					Extent = FVector(Scale * SpriteComponent->Sprite->GetSizeX()/2.0f, Scale * SpriteComponent->Sprite->GetSizeY()/2.0f, 0.0f);
-				}
-				else
-				{
-					Extent = PrimitiveComponent->Bounds.BoxExtent;
-				}
-
-
-				UBrushComponent* BrushComponent = Cast<UBrushComponent>( PrimitiveComponent );
-				if( BrushComponent && BrushComponent->Brush && BrushComponent->Brush->Polys )
-				{
-					bool bAlreadyProcessed = false;
-
-					// Need to check each vert of the brush separately
-					for( int32 PolyIndex = 0 ; PolyIndex < BrushComponent->Brush->Polys->Element.Num() && !bAlreadyProcessed ; ++PolyIndex )
-					{
-						const FPoly& Poly = BrushComponent->Brush->Polys->Element[ PolyIndex ];
-
-						for( int32 VertexIndex = 0 ; VertexIndex < Poly.Vertices.Num() ; ++VertexIndex )
-						{
-							const FVector Location = BrushComponent->ComponentToWorld.TransformPosition( Poly.Vertices[VertexIndex] );
-							const bool bIntersect = InFrustum.IntersectBox( Location, FVector::ZeroVector );
-
-							if( bIntersect && !bUseStrictSelection )
-							{
-								// If we intersected a vertex and we dont require the box to encompass the entire component
-								// then the actor should be selected and we can stop checking
-								bAlreadyProcessed = true;
-								bActorHitByBox = true;
-								break;
-							}
-							else if( !bIntersect && bUseStrictSelection )
-							{
-								// If we didnt intersect a vertex but we require the box to encompass the entire component
-								// then this test failed and we can stop checking
-								bAlreadyProcessed = true;
-								bActorHitByBox = false;
-								break;
-							}
-						}
-					}
-
-					if( bUseStrictSelection && !bAlreadyProcessed )
-					{
-						// If we are here then every vert was intersected so we should select the actor
-						bActorHitByBox = true;
-						bAlreadyProcessed = true;
-					}
-
-					if( bAlreadyProcessed )
-					{
-						// if this component has been processed, dont bother checking other components
-						break;
-					}
-				}
-				else
+				if (PrimitiveComponent->ComponentIsTouchingSelectionFrustum(InFrustum, LevelViewportClient->EngineShowFlags, bGeometryMode, bUseStrictSelection))
 				{
-					if ( InFrustum.IntersectBox( Origin, Extent, bActorFullyContained) 
-						&& (!bUseStrictSelection || bActorFullyContained))
-					{
-						if( PrimitiveComponent->IsA( UStaticMeshComponent::StaticClass() ) || PrimitiveComponent->IsA( USkeletalMeshComponent::StaticClass() ) )
-						{
-							// Check each vertex on the component's mesh for intersection
-							bActorHitByBox = IntersectsVertices( *PrimitiveComponent, InFrustum, bUseStrictSelection );
-						}
-						else
-						{
-							// Use the bounding box check for all other components
-							bActorHitByBox = true;
-						}
-							
-						if( bActorHitByBox )
-						{
-							// no need to continue checking other components
-							break;
-						}
-					}
+					bActorHitByBox = true;
+					break;
 				}
 			}
 		}
@@ -513,7 +301,7 @@ void FDragTool_ActorFrustumSelect::CalculateFrustum( FSceneView* View, FConvexVo
 {
 	if( bUseBoxFrustum )
 	{
-		FVector CamPoint = ViewportClient->GetViewLocation();
+		FVector CamPoint = LevelViewportClient->GetViewLocation();
 		FVector BoxPoint1, BoxPoint2, BoxPoint3, BoxPoint4;
 		FVector WorldDir1, WorldDir2, WorldDir3, WorldDir4;
 		// Deproject the four corners of the selection box
diff --git a/Engine/Source/Editor/UnrealEd/Private/LevelEditorViewport.cpp b/Engine/Source/Editor/UnrealEd/Private/LevelEditorViewport.cpp
--- a/Engine/Source/Editor/UnrealEd/Private/LevelEditorViewport.cpp
+++ b/Engine/Source/Editor/UnrealEd/Private/LevelEditorViewport.cpp
@@ -2155,7 +2155,7 @@ TSharedPtr<FDragTool> FLevelEditorViewportClient::MakeDragTool( EDragTool::Type
 		DragTool = MakeShareable( new FDragTool_ActorBoxSelect(this) );
 		break;
 	case EDragTool::FrustumSelect:
-		DragTool = MakeShareable( new FDragTool_ActorFrustumSelect() );
+		DragTool = MakeShareable( new FDragTool_ActorFrustumSelect(this) );
 		break;	
 	case EDragTool::Measure:
 		DragTool = MakeShareable( new FDragTool_Measure(this) );
@@ -2749,276 +2749,6 @@ bool FLevelEditorViewportClient::InputAxis(FViewport* Viewport, int32 Controller
 
 
 
-
-/**
- * Determines if InComponent is inside of InSelBBox.  This check differs depending on the type of component.
- * If InComponent is NULL, false is returned.
- *
- * @param	InActor							Used only when testing billboard components.
- * @param	InComponent						The component to query.  If NULL, false is returned.
- * @param	InSelBox						The selection box.
- * @param	bConsiderOnlyBSP				If true, consider only BSP.
- * @param	bMustEncompassEntireComponent	If true, the entire component must be encompassed by the selection box in order to return true.
- */
-bool FLevelEditorViewportClient::ComponentIsTouchingSelectionBox(AActor* InActor, UPrimitiveComponent* InComponent, const FBox& InSelBBox, bool bConsiderOnlyBSP, bool bMustEncompassEntireComponent)
-{
-	bool bResult = false;
-	bool bAlreadyProcessed = false;
-
-	if( (EngineShowFlags.Volumes || EngineShowFlags.BSP) && InComponent->IsA( UBrushComponent::StaticClass() ) )
-	{
-		if ( InActor->IsA(AVolume::StaticClass()) )
-		{
-			// Don't select if the brush is a volume and the volume show flag is unset.
-			if ( !EngineShowFlags.Volumes )
-			{
-				bResult = false;
-				bAlreadyProcessed = true;
-			}
-		}
-		else
-		{
-			// Don't select if the brush is regular bsp and the bsp show flag is unset.
-			if ( !EngineShowFlags.BSP )
-			{
-				bResult = false;
-				bAlreadyProcessed = true;
-			}
-		}
-
-		const UBrushComponent* BrushComponent = static_cast<UBrushComponent*>( InComponent );
-
-		// Check the brush component
-		if( BrushComponent->Brush && BrushComponent->Brush->Polys && !bAlreadyProcessed )
-		{
-			TArray<FVector> Vertices;
-
-			for (int32 PolyIndex = 0; PolyIndex < BrushComponent->Brush->Polys->Element.Num() && !bAlreadyProcessed; ++PolyIndex)
-			{
-				const FPoly& Poly = BrushComponent->Brush->Polys->Element[ PolyIndex ];
-
-				if (!bMustEncompassEntireComponent)
-				{
-					Vertices.Empty(Poly.Vertices.Num());
-					for (const auto& Vertex : Poly.Vertices)
-					{
-						Vertices.Add(InComponent->ComponentToWorld.TransformPosition(Vertex));
-					}
-
-					FSeparatingAxisPointCheck PointCheck(Vertices, InSelBBox.GetCenter(), InSelBBox.GetExtent(), false);
-					if (PointCheck.bHit)
-					{
-						bResult = true;
-						bAlreadyProcessed = true;
-						break;
-					}
-				}
-				else
-				{
-					for (int32 VertexIndex = 0; VertexIndex < Poly.Vertices.Num(); ++VertexIndex)
-					{
-						const FVector Location = InComponent->ComponentToWorld.TransformPosition(Poly.Vertices[VertexIndex]);
-						const bool bLocationIntersected = FMath::PointBoxIntersection(Location, InSelBBox);
-
-						// If the selection box has to encompass the entire component and a poly vertex didn't intersect with the selection
-						// box, this component does not qualify
-						if (!bLocationIntersected)
-						{
-							bResult = false;
-							bAlreadyProcessed = true;
-							break;
-						}
-					}
-				}
-			}
-
-			// If the selection box has to encompass all of the component and none of the component's verts failed the intersection test, this component
-			// is consider touching
-			if ( bMustEncompassEntireComponent && !bAlreadyProcessed )
-			{
-				bResult = true;
-				bAlreadyProcessed = true;
-			}
-		}
-	}
-	else if( !bConsiderOnlyBSP && EngineShowFlags.StaticMeshes && InComponent->IsA( UStaticMeshComponent::StaticClass() ) )
-	{
-		const UStaticMeshComponent* StaticMeshComponent = static_cast<UStaticMeshComponent*>( InComponent );
-
-		if( StaticMeshComponent->StaticMesh
-			&& StaticMeshComponent->StaticMesh->HasValidRenderData() )
-		{
-			// Check if we are even inside it's bounding box, if we are not, there is no way we colliding via the more advanced checks we will do.
-			if(InSelBBox.Intersect(InComponent->Bounds.GetBox()))
-			{
-				TArray<FVector> Vertex;
-
-				FStaticMeshLODResources& LODModel = StaticMeshComponent->StaticMesh->RenderData->LODResources[0];
-				FIndexArrayView Indices = LODModel.IndexBuffer.GetArrayView();
-
-				int32 NumSections = LODModel.Sections.Num();
-				for (int32 SectionIndex = 0; !bAlreadyProcessed && SectionIndex < NumSections; ++SectionIndex)
-				{
-					FStaticMeshSection& Section = LODModel.Sections[SectionIndex];
-
-					// Iterate over each triangle.
-					for( int32 TriangleIndex = 0 ; TriangleIndex < (int32)Section.NumTriangles ; TriangleIndex++ )
-					{
-						Vertex.Empty(3);
-
-						int32 FirstIndex = TriangleIndex * 3 + Section.FirstIndex;
-						for (int32 i = 0; i < 3; i++)
-						{
-							int32 VertexIndex = Indices[FirstIndex+i];
-							FVector LocalPosition = LODModel.PositionVertexBuffer.VertexPosition(VertexIndex);
-							Vertex.Emplace(StaticMeshComponent->ComponentToWorld.TransformPosition(LocalPosition));
-						}
-
-						// Check if the triangle is colliding with the bounding box.
-						FSeparatingAxisPointCheck ThePointCheck(Vertex, InSelBBox.GetCenter(), InSelBBox.GetExtent(), false);
-						if( !bMustEncompassEntireComponent && ThePointCheck.bHit )
-						{
-							bResult = true;
-							bAlreadyProcessed = true;
-							break;
-						}
-						else if( bMustEncompassEntireComponent && !ThePointCheck.bHit )
-						{
-							bResult = false;
-							bAlreadyProcessed = true;
-							break;
-						}
-					}
-				}
-
-				// If the selection box has to encompass all of the component and none of the component's verts failed the intersection test, this component
-				// is consider touching
-				if ( bMustEncompassEntireComponent && !bAlreadyProcessed )
-				{
-					bResult = true;
-					bAlreadyProcessed = true;
-				}
-			}
-		}
-	}
-	else if (!bConsiderOnlyBSP && EngineShowFlags.BillboardSprites && InComponent->IsA(UBillboardComponent::StaticClass()))
-	{
-		const UBillboardComponent* SpriteComponent = CastChecked<const UBillboardComponent>(InComponent);
-		if (SpriteComponent->Sprite != nullptr)
-		{
-			float Scale = SpriteComponent->ComponentToWorld.GetMaximumAxisScale();
-
-			// Construct a box representing the sprite
-			FBox SpriteBox(
-				InActor->GetActorLocation() - Scale * FMath::Max(SpriteComponent->Sprite->GetSizeX(), SpriteComponent->Sprite->GetSizeY()) * FVector(1, 1, 1),
-				InActor->GetActorLocation() + Scale * FMath::Max(SpriteComponent->Sprite->GetSizeX(), SpriteComponent->Sprite->GetSizeY()) * FVector(1, 1, 1));
-
-			// If the selection box doesn't have to encompass the entire component and it intersects with the box constructed for the sprite, then it is valid.
-			// Additionally, if the selection box does have to encompass the entire component and both the min and max vectors of the sprite box are inside the selection box,
-			// then it is valid.
-			if ((!bMustEncompassEntireComponent && InSelBBox.Intersect(SpriteBox))
-				|| (bMustEncompassEntireComponent && InSelBBox.IsInside(SpriteBox.Min) && InSelBBox.IsInside(SpriteBox.Max)))
-			{
-				bResult = true;
-				bAlreadyProcessed = true;
-			}
-		}
-	}
-	else if( !bConsiderOnlyBSP && EngineShowFlags.SkeletalMeshes && InComponent->IsA( USkeletalMeshComponent::StaticClass() ) )
-	{
-		const USkeletalMeshComponent* SkeletalMeshComponent = static_cast<USkeletalMeshComponent*>( InComponent );
-		if( SkeletalMeshComponent->MeshObject )
-		{
-			FSkeletalMeshResource* SkelMeshResource = SkeletalMeshComponent->GetSkeletalMeshResource();
-			check(SkelMeshResource);
-			check(SkelMeshResource->LODModels.Num() > 0);
-
-			// Transform hard and soft verts into world space. Note that this assumes skeletal mesh is in reference pose...
-			const FStaticLODModel& LODModel = SkelMeshResource->LODModels[0];
-			for( int32 ChunkIndex = 0 ; ChunkIndex < LODModel.Chunks.Num() && !bAlreadyProcessed; ++ChunkIndex )
-			{
-				const FSkelMeshChunk& Chunk = LODModel.Chunks[ChunkIndex];
-				for( int32 i = 0; i < Chunk.RigidVertices.Num(); ++i )
-				{
-					const FVector Location = SkeletalMeshComponent->ComponentToWorld.TransformPosition( Chunk.RigidVertices[i].Position );
-					const bool bLocationIntersected = FMath::PointBoxIntersection( Location, InSelBBox );
-
-					// If the selection box doesn't have to encompass the entire component and a skeletal mesh vertex has intersected with
-					// the selection box, this component is being touched by the selection box
-					if( !bMustEncompassEntireComponent && bLocationIntersected )
-					{
-						bResult = true;
-						bAlreadyProcessed = true;
-						break;
-					}
-
-					// If the selection box has to encompass the entire component and a skeletal mesh vertex didn't intersect with the selection
-					// box, this component does not qualify
-					else if ( bMustEncompassEntireComponent && !bLocationIntersected )
-					{
-						bResult = false;
-						bAlreadyProcessed = true;
-						break;
-					}
-				}
-
-				for( int32 i = 0 ; i < Chunk.SoftVertices.Num() ; ++i )
-				{
-					const FVector Location = SkeletalMeshComponent->ComponentToWorld.TransformPosition( Chunk.SoftVertices[i].Position );
-					const bool bLocationIntersected = FMath::PointBoxIntersection( Location, InSelBBox );
-
-					// If the selection box doesn't have to encompass the entire component and a skeletal mesh vertex has intersected with
-					// the selection box, this component is being touched by the selection box
-					if( !bMustEncompassEntireComponent && bLocationIntersected )
-					{
-						bResult = true;
-						bAlreadyProcessed = true;
-						break;
-					}
-
-					// If the selection box has to encompass the entire component and a skeletal mesh vertex didn't intersect with the selection
-					// box, this component does not qualify
-					else if ( bMustEncompassEntireComponent && !bLocationIntersected )
-					{
-						bResult = false;
-						bAlreadyProcessed = true;
-						break;
-					}
-				}
-			}
-
-			// If the selection box has to encompass all of the component and none of the component's verts failed the intersection test, this component
-			// is consider touching
-			if ( bMustEncompassEntireComponent && !bAlreadyProcessed )
-			{
-				bResult = true;
-				bAlreadyProcessed = true;
-			}
-		}
-	}
-	else if ( !bConsiderOnlyBSP )
-	{
-		bool bSelectByBoundingBox = true;
-
-		if( bSelectByBoundingBox )
-		{
-			const FBox& ComponentBounds = InComponent->Bounds.GetBox();
-
-			// Check the component bounds versus the selection box
-			// If the selection box must encompass the entire component, then both the min and max vector of the bounds must be inside in the selection
-			// box to be valid. If the selection box only has to touch the component, then it is sufficient to check if it intersects with the bounds.
-			if (   ( !bMustEncompassEntireComponent && InSelBBox.Intersect( ComponentBounds ) ) 
-				|| ( bMustEncompassEntireComponent && InSelBBox.IsInside( ComponentBounds.Min ) && InSelBBox.IsInside( ComponentBounds.Max ) ) )
-			{
-				bResult = true;
-				bAlreadyProcessed = true;
-			}
-		}
-	}
-
-	return bResult;
-}
-
 static uint32 GetVolumeActorVisibilityId( const AActor& InActor )
 {
 	UClass* Class = InActor.GetClass();
diff --git a/Engine/Source/Editor/UnrealEd/Public/DragTool_FrustumSelect.h b/Engine/Source/Editor/UnrealEd/Public/DragTool_FrustumSelect.h
--- a/Engine/Source/Editor/UnrealEd/Public/DragTool_FrustumSelect.h
+++ b/Engine/Source/Editor/UnrealEd/Public/DragTool_FrustumSelect.h
@@ -12,8 +12,9 @@
 class UNREALED_API FDragTool_ActorFrustumSelect : public FDragTool
 {
 public:
-	FDragTool_ActorFrustumSelect()
+	explicit FDragTool_ActorFrustumSelect(FLevelEditorViewportClient* InLevelViewportClient)
 		: FDragTool()
+		, LevelViewportClient( InLevelViewportClient )
 	{}
 
 	/* Updates the drag tool's end location with the specified delta.  The end location is
@@ -65,15 +66,6 @@ private:
 	 */
 	bool IntersectsFrustum( const UModel& InModel, int32 NodeIndex, const FConvexVolume& InFrustum, bool bUseStrictSelection ) const;
 	
-	/** 
-	 * Returns true if the mesh on the component has vertices which intersect the frustum
-	 *
-	 * @param InComponent			The static mesh or skeletal mesh component to check
-	 * @param InFrustum				The frustum to check against.
-	 * @param bUseStrictSelection	true if all the vertices must be entirely within the frustum
-	 */
-	bool IntersectsVertices( UPrimitiveComponent& InComponent, const FConvexVolume& InFrustum, bool bUseStrictSelection ) const;
-	
 	/** Adds a hover effect to the passed in actor */
 	void AddHoverEffect( AActor& InActor );
 	
@@ -92,5 +84,5 @@ private:
 	/** List of BSP models to check for selection */
 	TArray<UModel*> ModelsToCheck;
 
-	FEditorViewportClient* ViewportClient;
+	FLevelEditorViewportClient* LevelViewportClient;
 };
diff --git a/Engine/Source/Editor/UnrealEd/Public/LevelEditorViewport.h b/Engine/Source/Editor/UnrealEd/Public/LevelEditorViewport.h
--- a/Engine/Source/Editor/UnrealEd/Public/LevelEditorViewport.h
+++ b/Engine/Source/Editor/UnrealEd/Public/LevelEditorViewport.h
@@ -196,18 +196,6 @@ public:
 	/** Determines if the new MoveCanvas movement should be used */
 	bool ShouldUseMoveCanvasMovement (void);
 
-	/**
-	 * Determines if InComponent is inside of InSelBBox.  This check differs depending on the type of component.
-	 * If InComponent is NULL, false is returned.
-	 *
-	 * @param	InActor							Used only when testing billboard components.
-	 * @param	InComponent						The component to query.  If NULL, false is returned.
-	 * @param	InSelBox						The selection box.
-	 * @param	bConsiderOnlyBSP				If true, consider only BSP.
-	 * @param	bMustEncompassEntireComponent	If true, the entire component must be encompassed by the selection box in order to return true.
-	 */
-	bool ComponentIsTouchingSelectionBox( AActor* InActor, UPrimitiveComponent* InComponent, const FBox& InSelBBox, bool bConsiderOnlyBSP, bool bMustEncompassEntireComponent );
-
 	/** 
 	 * Returns true if the passed in volume is visible in the viewport (due to volume actor visibility flags)
 	 *
diff --git a/Engine/Source/Runtime/Engine/Classes/AI/Navigation/NavLinkRenderingComponent.h b/Engine/Source/Runtime/Engine/Classes/AI/Navigation/NavLinkRenderingComponent.h
--- a/Engine/Source/Runtime/Engine/Classes/AI/Navigation/NavLinkRenderingComponent.h
+++ b/Engine/Source/Runtime/Engine/Classes/AI/Navigation/NavLinkRenderingComponent.h
@@ -13,6 +13,10 @@ class ENGINE_API UNavLinkRenderingComponent : public UPrimitiveComponent
 
 	/** Should recreate proxy one very update */
 	virtual bool ShouldRecreateProxyOnUpdateTransform() const override { return true; }
+#if WITH_EDITOR
+	virtual bool ComponentIsTouchingSelectionBox(const FBox& InSelBBox, const FEngineShowFlags& ShowFlags, const bool bConsiderOnlyBSP, const bool bMustEncompassEntireComponent) const override;
+	virtual bool ComponentIsTouchingSelectionFrustum(const FConvexVolume& InFrustum, const FEngineShowFlags& ShowFlags, const bool bConsiderOnlyBSP, const bool bMustEncompassEntireComponent) const override;
+#endif
 	// End UPrimitiveComponent Interface
 
 	// Begin USceneComponent Interface
diff --git a/Engine/Source/Runtime/Engine/Classes/Components/ArrowComponent.h b/Engine/Source/Runtime/Engine/Classes/Components/ArrowComponent.h
--- a/Engine/Source/Runtime/Engine/Classes/Components/ArrowComponent.h
+++ b/Engine/Source/Runtime/Engine/Classes/Components/ArrowComponent.h
@@ -58,6 +58,10 @@ class UArrowComponent : public UPrimitiveComponent
 
 	// Begin UPrimitiveComponent interface.
 	virtual FPrimitiveSceneProxy* CreateSceneProxy() override;
+#if WITH_EDITOR
+	virtual bool ComponentIsTouchingSelectionBox(const FBox& InSelBBox, const FEngineShowFlags& ShowFlags, const bool bConsiderOnlyBSP, const bool bMustEncompassEntireComponent) const override;
+	virtual bool ComponentIsTouchingSelectionFrustum(const FConvexVolume& InFrustum, const FEngineShowFlags& ShowFlags, const bool bConsiderOnlyBSP, const bool bMustEncompassEntireComponent) const override;
+#endif
 	// End UPrimitiveComponent interface.
 
 	// Begin USceneComponent interface.
diff --git a/Engine/Source/Runtime/Engine/Classes/Components/BillboardComponent.h b/Engine/Source/Runtime/Engine/Classes/Components/BillboardComponent.h
--- a/Engine/Source/Runtime/Engine/Classes/Components/BillboardComponent.h
+++ b/Engine/Source/Runtime/Engine/Classes/Components/BillboardComponent.h
@@ -4,6 +4,9 @@
 #pragma once
 #include "BillboardComponent.generated.h"
 
+struct FEngineShowFlags;
+struct FConvexVolume;
+
 /** 
  * A 2d texture that will be rendered always facing the camera.
  */
@@ -62,6 +65,10 @@ class UBillboardComponent : public UPrimitiveComponent
 	// Begin UPrimitiveComponent Interface
 	virtual FPrimitiveSceneProxy* CreateSceneProxy() override;
 	virtual FBoxSphereBounds CalcBounds(const FTransform& LocalToWorld) const override;
+#if WITH_EDITOR
+	virtual bool ComponentIsTouchingSelectionBox(const FBox& InSelBBox, const FEngineShowFlags& ShowFlags, const bool bConsiderOnlyBSP, const bool bMustEncompassEntireComponent) const override;
+	virtual bool ComponentIsTouchingSelectionFrustum(const FConvexVolume& InFrustum, const FEngineShowFlags& ShowFlags, const bool bConsiderOnlyBSP, const bool bMustEncompassEntireComponent) const override;
+#endif
 	// End UPrimitiveComponent Interface
 
 #if WITH_EDITORONLY_DATA
diff --git a/Engine/Source/Runtime/Engine/Classes/Components/BrushComponent.h b/Engine/Source/Runtime/Engine/Classes/Components/BrushComponent.h
--- a/Engine/Source/Runtime/Engine/Classes/Components/BrushComponent.h
+++ b/Engine/Source/Runtime/Engine/Classes/Components/BrushComponent.h
@@ -4,6 +4,9 @@
 #pragma once
 #include "BrushComponent.generated.h"
 
+struct FEngineShowFlags;
+struct FConvexVolume;
+
 UCLASS(editinlinenew, MinimalAPI, hidecategories=(Physics, Lighting, LOD, Rendering, TextureStreaming, Transform, Activation, "Components|Activation"), showcategories=(Mobility, "Rendering|Material"))
 class UBrushComponent : public UPrimitiveComponent
 {
@@ -43,6 +46,10 @@ public:
 	virtual class UBodySetup* GetBodySetup() override { return BrushBodySetup; };
 	virtual void GetUsedMaterials( TArray<UMaterialInterface*>& OutMaterials ) const override;
 	virtual uint8 GetStaticDepthPriorityGroup() const override;
+#if WITH_EDITOR
+	virtual bool ComponentIsTouchingSelectionBox(const FBox& InSelBBox, const FEngineShowFlags& ShowFlags, const bool bConsiderOnlyBSP, const bool bMustEncompassEntireComponent) const override;
+	virtual bool ComponentIsTouchingSelectionFrustum(const FConvexVolume& InFrustum, const FEngineShowFlags& ShowFlags, const bool bConsiderOnlyBSP, const bool bMustEncompassEntireComponent) const override;
+#endif
 	// End UPrimitiveComponent interface.
 
 	/** Create the AggGeom collection-of-convex-primitives from the Brush UModel data. */
diff --git a/Engine/Source/Runtime/Engine/Classes/Components/DrawSphereComponent.h b/Engine/Source/Runtime/Engine/Classes/Components/DrawSphereComponent.h
--- a/Engine/Source/Runtime/Engine/Classes/Components/DrawSphereComponent.h
+++ b/Engine/Source/Runtime/Engine/Classes/Components/DrawSphereComponent.h
@@ -11,6 +11,12 @@ UCLASS(collapsecategories, hidecategories=Object, editinlinenew, MinimalAPI)
 class UDrawSphereComponent : public USphereComponent
 {
 	GENERATED_UCLASS_BODY()
+
+public:
+#if WITH_EDITOR
+	virtual bool ComponentIsTouchingSelectionBox(const FBox& InSelBBox, const FEngineShowFlags& ShowFlags, const bool bConsiderOnlyBSP, const bool bMustEncompassEntireComponent) const override;
+	virtual bool ComponentIsTouchingSelectionFrustum(const FConvexVolume& InFrustum, const FEngineShowFlags& ShowFlags, const bool bConsiderOnlyBSP, const bool bMustEncompassEntireComponent) const override;
+#endif
 };
 
 
diff --git a/Engine/Source/Runtime/Engine/Classes/Components/PrimitiveComponent.h b/Engine/Source/Runtime/Engine/Classes/Components/PrimitiveComponent.h
--- a/Engine/Source/Runtime/Engine/Classes/Components/PrimitiveComponent.h
+++ b/Engine/Source/Runtime/Engine/Classes/Components/PrimitiveComponent.h
@@ -14,6 +14,8 @@
 class FPrimitiveSceneProxy;
 class AController; 
 class UTexture;
+struct FEngineShowFlags;
+struct FConvexVolume;
 
 /** Information about a streaming texture that a primitive uses for rendering. */
 struct FStreamingTexturePrimitiveInfo
@@ -1130,6 +1132,34 @@ public:
 	*/
 	virtual void GetWeldedBodies(TArray<FBodyInstance*> & OutWeldedBodies, TArray<FName> & OutLabels);
 
+#if WITH_EDITOR
+	/**
+	 * Determines whether the supplied bounding box intersects with the component.
+	 * Used by the editor in orthographic viewports.
+	 *
+	 * @param	InSelBBox						Bounding box to test against
+	 * @param	ShowFlags						Engine ShowFlags for the viewport
+	 * @param	bConsiderOnlyBSP				If only BSP geometry should be tested
+	 * @param	bMustEncompassEntireComponent	Whether the component bounding box must lay wholly within the supplied bounding box
+	 *
+	 * @return	true if the supplied bounding box is determined to intersect the component (partially or wholly)
+	 */
+	virtual bool ComponentIsTouchingSelectionBox(const FBox& InSelBBox, const FEngineShowFlags& ShowFlags, const bool bConsiderOnlyBSP, const bool bMustEncompassEntireComponent) const;
+
+	/**
+	 * Determines whether the supplied frustum intersects with the component.
+	 * Used by the editor in perspective viewports.
+	 *
+	 * @param	InFrustum						Frustum to test against
+	 * @param	ShowFlags						Engine ShowFlags for the viewport
+	 * @param	bConsiderOnlyBSP				If only BSP geometry should be tested
+	 * @param	bMustEncompassEntireComponent	Whether the component bounding box must lay wholly within the supplied bounding box
+	 *
+	 * @return	true if the supplied bounding box is determined to intersect the component (partially or wholly)
+	 */
+	virtual bool ComponentIsTouchingSelectionFrustum(const FConvexVolume& InFrustum, const FEngineShowFlags& ShowFlags, const bool bConsiderOnlyBSP, const bool bMustEncompassEntireComponent) const;
+#endif
+
 protected:
 
 	/** Give the static mesh component recreate render state context access to Create/DestroyRenderState_Concurrent(). */
diff --git a/Engine/Source/Runtime/Engine/Classes/Components/SkeletalMeshComponent.h b/Engine/Source/Runtime/Engine/Classes/Components/SkeletalMeshComponent.h
--- a/Engine/Source/Runtime/Engine/Classes/Components/SkeletalMeshComponent.h
+++ b/Engine/Source/Runtime/Engine/Classes/Components/SkeletalMeshComponent.h
@@ -9,6 +9,8 @@
 #include "SkeletalMeshComponent.generated.h"
 
 class UAnimInstance;
+struct FEngineShowFlags;
+struct FConvexVolume;
 
 struct FAnimationEvaluationContext
 {
@@ -831,6 +833,10 @@ public:
 	virtual void SetAllPhysicsLinearVelocity(FVector NewVel,bool bAddToCurrent = false) override;
 	virtual float GetMass() const override;
 	virtual float CalculateMass(FName BoneName = NAME_None) override;
+#if WITH_EDITOR
+	virtual bool ComponentIsTouchingSelectionBox(const FBox& InSelBBox, const FEngineShowFlags& ShowFlags, const bool bConsiderOnlyBSP, const bool bMustEncompassEntireComponent) const override;
+	virtual bool ComponentIsTouchingSelectionFrustum(const FConvexVolume& InFrustum, const FEngineShowFlags& ShowFlags, const bool bConsiderOnlyBSP, const bool bMustEncompassEntireComponent) const override;
+#endif
 protected:
 	virtual FTransform GetComponentTransformFromBodyInstance(FBodyInstance* UseBI) override;
 	// End UPrimitiveComponent interface.
diff --git a/Engine/Source/Runtime/Engine/Classes/Components/StaticMeshComponent.h b/Engine/Source/Runtime/Engine/Classes/Components/StaticMeshComponent.h
--- a/Engine/Source/Runtime/Engine/Classes/Components/StaticMeshComponent.h
+++ b/Engine/Source/Runtime/Engine/Classes/Components/StaticMeshComponent.h
@@ -12,6 +12,8 @@ class FColorVertexBuffer;
 class UStaticMesh;
 class FStaticMeshStaticLightingMesh;
 class ULightComponent;
+struct FEngineShowFlags;
+struct FConvexVolume;
 
 /** Cached vertex information at the time the mesh was painted. */
 USTRUCT()
@@ -288,6 +290,10 @@ public:
 	virtual UMaterialInterface* GetMaterial(int32 MaterialIndex) const override;
 
 	virtual bool DoCustomNavigableGeometryExport(struct FNavigableGeometryExport* GeomExport) const;
+#if WITH_EDITOR
+	virtual bool ComponentIsTouchingSelectionBox(const FBox& InSelBBox, const FEngineShowFlags& ShowFlags, const bool bConsiderOnlyBSP, const bool bMustEncompassEntireComponent) const override;
+	virtual bool ComponentIsTouchingSelectionFrustum(const FConvexVolume& InSelBBox, const FEngineShowFlags& ShowFlags, const bool bConsiderOnlyBSP, const bool bMustEncompassEntireComponent) const override;
+#endif
 	// End UPrimitiveComponent interface.
 
 	/**
diff --git a/Engine/Source/Runtime/Engine/Private/AI/Navigation/NavLinkRenderingComponent.cpp b/Engine/Source/Runtime/Engine/Private/AI/Navigation/NavLinkRenderingComponent.cpp
--- a/Engine/Source/Runtime/Engine/Private/AI/Navigation/NavLinkRenderingComponent.cpp
+++ b/Engine/Source/Runtime/Engine/Private/AI/Navigation/NavLinkRenderingComponent.cpp
@@ -3,6 +3,10 @@
 #include "EnginePrivate.h"
 #include "AI/NavLinkRenderingProxy.h"
 #include "AI/Navigation/NavLinkRenderingComponent.h"
+#if WITH_EDITOR
+#include "ShowFlags.h"
+#include "ConvexVolume.h"
+#endif
 
 //----------------------------------------------------------------------//
 // UNavLinkRenderingComponent
@@ -71,6 +75,20 @@ FPrimitiveSceneProxy* UNavLinkRenderingComponent::CreateSceneProxy()
 	return new FNavLinkRenderingProxy(this);
 }
 
+#if WITH_EDITOR
+bool UNavLinkRenderingComponent::ComponentIsTouchingSelectionBox(const FBox& InSelBBox, const FEngineShowFlags& ShowFlags, const bool bConsiderOnlyBSP, const bool bMustEncompassEntireComponent) const
+{
+	// NavLink rendering components not treated as 'selectable' in editor
+	return false;
+}
+
+bool UNavLinkRenderingComponent::ComponentIsTouchingSelectionFrustum(const FConvexVolume& InFrustum, const FEngineShowFlags& ShowFlags, const bool bConsiderOnlyBSP, const bool bMustEncompassEntireComponent) const
+{
+	// NavLink rendering components not treated as 'selectable' in editor
+	return false;
+}
+#endif
+
 //----------------------------------------------------------------------//
 // FNavLinkRenderingProxy
 //----------------------------------------------------------------------//
diff --git a/Engine/Source/Runtime/Engine/Private/ArrowComponent.cpp b/Engine/Source/Runtime/Engine/Private/ArrowComponent.cpp
--- a/Engine/Source/Runtime/Engine/Private/ArrowComponent.cpp
+++ b/Engine/Source/Runtime/Engine/Private/ArrowComponent.cpp
@@ -384,6 +384,20 @@ FPrimitiveSceneProxy* UArrowComponent::CreateSceneProxy()
 	return new FArrowSceneProxy(this);
 }
 
+#if WITH_EDITOR
+bool UArrowComponent::ComponentIsTouchingSelectionBox(const FBox& InSelBBox, const FEngineShowFlags& ShowFlags, const bool bConsiderOnlyBSP, const bool bMustEncompassEntireComponent) const
+{
+	// Arrow components not treated as 'selectable' in editor
+	return false;
+}
+
+bool UArrowComponent::ComponentIsTouchingSelectionFrustum(const FConvexVolume& InFrustum, const FEngineShowFlags& ShowFlags, const bool bConsiderOnlyBSP, const bool bMustEncompassEntireComponent) const
+{
+	// Arrow components not treated as 'selectable' in editor
+	return false;
+}
+#endif
+
 
 FBoxSphereBounds UArrowComponent::CalcBounds(const FTransform& LocalToWorld) const
 {
diff --git a/Engine/Source/Runtime/Engine/Private/BillboardComponent.cpp b/Engine/Source/Runtime/Engine/Private/BillboardComponent.cpp
--- a/Engine/Source/Runtime/Engine/Private/BillboardComponent.cpp
+++ b/Engine/Source/Runtime/Engine/Private/BillboardComponent.cpp
@@ -3,6 +3,10 @@
 #include "EnginePrivate.h"
 #include "Engine/Light.h"
 #include "LevelUtils.h"
+#if WITH_EDITOR
+#include "ShowFlags.h"
+#include "ConvexVolume.h"
+#endif
 
 namespace BillboardConstants
 {
@@ -354,6 +358,54 @@ FBoxSphereBounds UBillboardComponent::CalcBounds(const FTransform& LocalToWorld)
 	return FBoxSphereBounds(LocalToWorld.GetLocation(),FVector(NewScale,NewScale,NewScale),FMath::Sqrt(3.0f * FMath::Square(NewScale)));
 }
 
+#if WITH_EDITOR
+bool UBillboardComponent::ComponentIsTouchingSelectionBox(const FBox& InSelBBox, const FEngineShowFlags& ShowFlags, const bool bConsiderOnlyBSP, const bool bMustEncompassEntireComponent) const
+{
+	AActor* Actor = GetOwner();
+
+	if (!bConsiderOnlyBSP && ShowFlags.BillboardSprites && Sprite != nullptr && Actor != nullptr)
+	{
+		const float Scale = ComponentToWorld.GetMaximumAxisScale();
+
+		// Construct a box representing the sprite
+		const FBox SpriteBox(
+			Actor->GetActorLocation() - Scale * FMath::Max(Sprite->GetSizeX(), Sprite->GetSizeY()) * FVector(0.5f, 0.5f, 0.5f),
+			Actor->GetActorLocation() + Scale * FMath::Max(Sprite->GetSizeX(), Sprite->GetSizeY()) * FVector(0.5f, 0.5f, 0.5f));
+
+		// If the selection box doesn't have to encompass the entire component and it intersects with the box constructed for the sprite, then it is valid.
+		// Additionally, if the selection box does have to encompass the entire component and both the min and max vectors of the sprite box are inside the selection box,
+		// then it is valid.
+		if ((!bMustEncompassEntireComponent && InSelBBox.Intersect(SpriteBox))
+			|| (bMustEncompassEntireComponent && InSelBBox.IsInside(SpriteBox)))
+		{
+			return true;
+		}
+	}
+
+	return false;
+}
+
+bool UBillboardComponent::ComponentIsTouchingSelectionFrustum(const FConvexVolume& InFrustum, const FEngineShowFlags& ShowFlags, const bool bConsiderOnlyBSP, const bool bMustEncompassEntireComponent) const
+{
+	AActor* Actor = GetOwner();
+
+	if (!bConsiderOnlyBSP && ShowFlags.BillboardSprites && Sprite != nullptr && Actor != nullptr)
+	{
+		const float Scale = ComponentToWorld.GetMaximumAxisScale();
+		const float MaxExtent = FMath::Max(Sprite->GetSizeX(), Sprite->GetSizeY());
+		const FVector Extent = Scale * MaxExtent * FVector(0.5f, 0.5f, 0.0f);
+
+		bool bIsFullyContained;
+		if (InFrustum.IntersectBox(Actor->GetActorLocation(), Extent, bIsFullyContained))
+		{
+			return !bMustEncompassEntireComponent || bIsFullyContained;
+		}
+	}
+
+	return false;
+}
+#endif
+
 void UBillboardComponent::SetSprite(UTexture2D* NewSprite)
 {
 	Sprite = NewSprite;
diff --git a/Engine/Source/Runtime/Engine/Private/BrushComponent.cpp b/Engine/Source/Runtime/Engine/Private/BrushComponent.cpp
--- a/Engine/Source/Runtime/Engine/Private/BrushComponent.cpp
+++ b/Engine/Source/Runtime/Engine/Private/BrushComponent.cpp
@@ -7,7 +7,11 @@
 #include "EnginePrivate.h"
 #include "Model.h"
 #include "LevelUtils.h"
-
+#if WITH_EDITOR
+#include "Collision.h"
+#include "ShowFlags.h"
+#include "ConvexVolume.h"
+#endif
 #include "DebuggingDefines.h"
 #include "ActorEditorUtils.h"
 
@@ -684,6 +688,116 @@ uint8 UBrushComponent::GetStaticDepthPriorityGroup() const
 	}
 }
 
+#if WITH_EDITOR
+static bool IsComponentTypeShown(AActor* Actor, const FEngineShowFlags& ShowFlags)
+{
+	if (Actor != nullptr)
+	{
+		return (Actor->IsA(AVolume::StaticClass())) ? ShowFlags.Volumes : ShowFlags.BSP;
+	}
+
+	return false;
+}
+
+bool UBrushComponent::ComponentIsTouchingSelectionBox(const FBox& InSelBBox, const FEngineShowFlags& ShowFlags, const bool bConsiderOnlyBSP, const bool bMustEncompassEntireComponent) const
+{
+	if (!IsComponentTypeShown(GetOwner(), ShowFlags))
+	{
+		return false;
+	}
+
+	if (Brush != nullptr && Brush->Polys != nullptr)
+	{
+		TArray<FVector> Vertices;
+
+		for (const auto& Poly : Brush->Polys->Element)
+		{
+			if (!bMustEncompassEntireComponent)
+			{
+				// Just an intersection will do...
+				Vertices.Empty(Poly.Vertices.Num());
+				for (const auto& Vertex : Poly.Vertices)
+				{
+					Vertices.Add(ComponentToWorld.TransformPosition(Vertex));
+				}
+
+				FSeparatingAxisPointCheck PointCheck(Vertices, InSelBBox.GetCenter(), InSelBBox.GetExtent(), false);
+				if (PointCheck.bHit)
+				{
+					// If any poly intersected with the bounding box, this component is considered to be touching
+					return true;
+				}
+			}
+			else
+			{
+				// The component must be entirely within the bounding box...
+				for (const auto& Vertex : Poly.Vertices)
+				{
+					const FVector Location = ComponentToWorld.TransformPosition(Vertex);
+					const bool bLocationIntersected = FMath::PointBoxIntersection(Location, InSelBBox);
+
+					// If the selection box has to encompass the entire component and a poly vertex didn't intersect with the selection
+					// box, this component does not qualify
+					if (!bLocationIntersected)
+					{
+						return false;
+					}
+				}
+			}
+		}
+
+		// If the selection box has to encompass all of the component and none of the component's verts failed the intersection test, this component
+		// is considered touching
+		return true;
+	}
+
+	return false;
+}
+
+
+bool UBrushComponent::ComponentIsTouchingSelectionFrustum(const FConvexVolume& InFrustum, const FEngineShowFlags& ShowFlags, const bool bConsiderOnlyBSP, const bool bMustEncompassEntireComponent) const
+{
+	if (!IsComponentTypeShown(GetOwner(), ShowFlags))
+	{
+		return false;
+	}
+
+	if (Brush != nullptr && Brush->Polys != nullptr)
+	{
+		TArray<FVector> Vertices;
+
+		for (const auto& Poly : Brush->Polys->Element)
+		{
+			for (const auto& Vertex : Poly.Vertices)
+			{
+				const FVector Location = ComponentToWorld.TransformPosition(Vertex);
+				const bool bIntersect = InFrustum.IntersectSphere(Location, 0.0f);
+
+				if (bIntersect && !bMustEncompassEntireComponent)
+				{
+					// If we intersected a vertex and we don't require the box to encompass the entire component
+					// then the actor should be selected and we can stop checking
+					return true;
+				}
+				else if (!bIntersect && bMustEncompassEntireComponent)
+				{
+					// If we didn't intersect a vertex but we require the box to encompass the entire component
+					// then this test failed and we can stop checking
+					return false;
+				}
+			}
+		}
+
+		// If the selection box has to encompass all of the component and none of the component's verts failed the intersection test, this component
+		// is considered touching
+		return true;
+	}
+
+	return false;
+}
+#endif
+
+
 void UBrushComponent::BuildSimpleBrushCollision()
 {
 	AActor* Owner = GetOwner();
diff --git a/Engine/Source/Runtime/Engine/Private/DrawSphereComponent.cpp b/Engine/Source/Runtime/Engine/Private/DrawSphereComponent.cpp
--- a/Engine/Source/Runtime/Engine/Private/DrawSphereComponent.cpp
+++ b/Engine/Source/Runtime/Engine/Private/DrawSphereComponent.cpp
@@ -2,6 +2,10 @@
 
 #include "EnginePrivate.h"
 #include "Components/DrawSphereComponent.h"
+#if WITH_EDITOR
+#include "ShowFlags.h"
+#include "ConvexVolume.h"
+#endif
 
 UDrawSphereComponent::UDrawSphereComponent(const FObjectInitializer& ObjectInitializer)
 	: Super(ObjectInitializer)
@@ -13,3 +17,17 @@ UDrawSphereComponent::UDrawSphereComponent(const FObjectInitializer& ObjectIniti
 	bUseEditorCompositing = true;
 	bGenerateOverlapEvents = false;
 }
+
+#if WITH_EDITOR
+bool UDrawSphereComponent::ComponentIsTouchingSelectionBox(const FBox& InSelBBox, const FEngineShowFlags& ShowFlags, const bool bConsiderOnlyBSP, const bool bMustEncompassEntireComponent) const
+{
+	// Draw sphere components not treated as 'selectable' in editor
+	return false;
+}
+
+bool UDrawSphereComponent::ComponentIsTouchingSelectionFrustum(const FConvexVolume& InFrustum, const FEngineShowFlags& ShowFlags, const bool bConsiderOnlyBSP, const bool bMustEncompassEntireComponent) const
+{
+	// Draw sphere components not treated as 'selectable' in editor
+	return false;
+}
+#endif
diff --git a/Engine/Source/Runtime/Engine/Private/PrimitiveComponent.cpp b/Engine/Source/Runtime/Engine/Private/PrimitiveComponent.cpp
--- a/Engine/Source/Runtime/Engine/Private/PrimitiveComponent.cpp
+++ b/Engine/Source/Runtime/Engine/Private/PrimitiveComponent.cpp
@@ -8,6 +8,11 @@
 #include "GameFramework/PhysicsVolume.h"
 #include "PhysicsPublic.h"
 #include "LevelUtils.h"
+#if WITH_EDITOR
+#include "ShowFlags.h"
+#include "Collision.h"
+#include "ConvexVolume.h"
+#endif
 #if WITH_PHYSX
 #include "PhysicsEngine/PhysXSupport.h"
 #include "Collision/PhysXCollision.h"
@@ -1851,6 +1856,43 @@ bool UPrimitiveComponent::GetOverlapsWithActor(const AActor* Actor, TArray<FOver
 	return InitialCount != OutOverlaps.Num();
 }
 
+#if WITH_EDITOR
+bool UPrimitiveComponent::ComponentIsTouchingSelectionBox(const FBox& InSelBBox, const FEngineShowFlags& ShowFlags, const bool bConsiderOnlyBSP, const bool bMustEncompassEntireComponent) const
+{
+	if (!bConsiderOnlyBSP)
+	{
+		const FBox& ComponentBounds = Bounds.GetBox();
+
+		// Check the component bounds versus the selection box
+		// If the selection box must encompass the entire component, then both the min and max vector of the bounds must be inside in the selection
+		// box to be valid. If the selection box only has to touch the component, then it is sufficient to check if it intersects with the bounds.
+		if ((!bMustEncompassEntireComponent && InSelBBox.Intersect(ComponentBounds))
+			|| (bMustEncompassEntireComponent && InSelBBox.IsInside(ComponentBounds)))
+		{
+			return true;
+		}
+	}
+
+	return false;
+}
+
+bool UPrimitiveComponent::ComponentIsTouchingSelectionFrustum(const FConvexVolume& InFrustum, const FEngineShowFlags& ShowFlags, const bool bConsiderOnlyBSP, const bool bMustEncompassEntireComponent) const
+{
+	if (!bConsiderOnlyBSP)
+	{
+		bool bIsFullyContained;
+		if (InFrustum.IntersectBox(Bounds.Origin, Bounds.BoxExtent, bIsFullyContained))
+		{
+			return !bMustEncompassEntireComponent || bIsFullyContained;
+		}
+	}
+
+	return false;
+}
+#endif
+
+
+
 /** Used to determine if it is ok to call a notification on this object */
 static bool IsActorValidToNotify(AActor* Actor)
 {
diff --git a/Engine/Source/Runtime/Engine/Private/SkeletalMeshComponent.cpp b/Engine/Source/Runtime/Engine/Private/SkeletalMeshComponent.cpp
--- a/Engine/Source/Runtime/Engine/Private/SkeletalMeshComponent.cpp
+++ b/Engine/Source/Runtime/Engine/Private/SkeletalMeshComponent.cpp
@@ -18,6 +18,11 @@
 #include "Animation/VertexAnim/VertexAnimation.h"
 #include "Particles/ParticleSystemComponent.h"
 #include "Animation/AnimSingleNodeInstance.h"
+#if WITH_EDITOR
+#include "ShowFlags.h"
+#include "Collision.h"
+#include "ConvexVolume.h"
+#endif
 
 TAutoConsoleVariable<int32> CVarUseParallelAnimationEvaluation(TEXT("a.ParallelAnimEvaluation"), 0, TEXT("If 1, animation evaluation will be run across the task graph system. If 0, evaluation will run purely on the game thread"));
 
@@ -1720,6 +1725,129 @@ float USkeletalMeshComponent::CalculateMass(FName BoneName)
 
 #if WITH_EDITOR
 
+bool USkeletalMeshComponent::ComponentIsTouchingSelectionBox(const FBox& InSelBBox, const FEngineShowFlags& ShowFlags, const bool bConsiderOnlyBSP, const bool bMustEncompassEntireComponent) const
+{
+	if (!bConsiderOnlyBSP && ShowFlags.SkeletalMeshes && MeshObject != nullptr)
+	{
+		FSkeletalMeshResource* SkelMeshResource = GetSkeletalMeshResource();
+		check(SkelMeshResource);
+		check(SkelMeshResource->LODModels.Num() > 0);
+
+		// Transform hard and soft verts into world space. Note that this assumes skeletal mesh is in reference pose...
+		const FStaticLODModel& LODModel = SkelMeshResource->LODModels[0];
+		for (const auto& Chunk : LODModel.Chunks)
+		{
+			for (const auto& Vertex : Chunk.RigidVertices)
+			{
+				const FVector Location = ComponentToWorld.TransformPosition(Vertex.Position);
+				const bool bLocationIntersected = FMath::PointBoxIntersection(Location, InSelBBox);
+
+				// If the selection box doesn't have to encompass the entire component and a skeletal mesh vertex has intersected with
+				// the selection box, this component is being touched by the selection box
+				if (!bMustEncompassEntireComponent && bLocationIntersected)
+				{
+					return true;
+				}
+
+				// If the selection box has to encompass the entire component and a skeletal mesh vertex didn't intersect with the selection
+				// box, this component does not qualify
+				else if (bMustEncompassEntireComponent && !bLocationIntersected)
+				{
+					return false;
+				}
+			}
+
+			for (const auto& Vertex : Chunk.SoftVertices)
+			{
+				const FVector Location = ComponentToWorld.TransformPosition(Vertex.Position);
+				const bool bLocationIntersected = FMath::PointBoxIntersection(Location, InSelBBox);
+
+				// If the selection box doesn't have to encompass the entire component and a skeletal mesh vertex has intersected with
+				// the selection box, this component is being touched by the selection box
+				if (!bMustEncompassEntireComponent && bLocationIntersected)
+				{
+					return true;
+				}
+
+				// If the selection box has to encompass the entire component and a skeletal mesh vertex didn't intersect with the selection
+				// box, this component does not qualify
+				else if (bMustEncompassEntireComponent && !bLocationIntersected)
+				{
+					return false;
+				}
+			}
+		}
+
+		// If the selection box has to encompass all of the component and none of the component's verts failed the intersection test, this component
+		// is consider touching
+		return true;
+	}
+
+	return false;
+}
+
+bool USkeletalMeshComponent::ComponentIsTouchingSelectionFrustum(const FConvexVolume& InFrustum, const FEngineShowFlags& ShowFlags, const bool bConsiderOnlyBSP, const bool bMustEncompassEntireComponent) const
+{
+	if (!bConsiderOnlyBSP && ShowFlags.SkeletalMeshes && MeshObject != nullptr)
+	{
+		FSkeletalMeshResource* SkelMeshResource = GetSkeletalMeshResource();
+		check(SkelMeshResource);
+		check(SkelMeshResource->LODModels.Num() > 0);
+
+		// Transform hard and soft verts into world space. Note that this assumes skeletal mesh is in reference pose...
+		const FStaticLODModel& LODModel = SkelMeshResource->LODModels[0];
+		for (const auto& Chunk : LODModel.Chunks)
+		{
+			for (const auto& Vertex : Chunk.RigidVertices)
+			{
+				const FVector Location = ComponentToWorld.TransformPosition(Vertex.Position);
+				const bool bLocationIntersected = InFrustum.IntersectSphere(Location, 0.0f);
+
+				// If the selection box doesn't have to encompass the entire component and a skeletal mesh vertex has intersected with
+				// the selection box, this component is being touched by the selection box
+				if (!bMustEncompassEntireComponent && bLocationIntersected)
+				{
+					return true;
+				}
+
+				// If the selection box has to encompass the entire component and a skeletal mesh vertex didn't intersect with the selection
+				// box, this component does not qualify
+				else if (bMustEncompassEntireComponent && !bLocationIntersected)
+				{
+					return false;
+				}
+			}
+
+			for (const auto& Vertex : Chunk.SoftVertices)
+			{
+				const FVector Location = ComponentToWorld.TransformPosition(Vertex.Position);
+				const bool bLocationIntersected = InFrustum.IntersectSphere(Location, 0.0f);
+
+				// If the selection box doesn't have to encompass the entire component and a skeletal mesh vertex has intersected with
+				// the selection box, this component is being touched by the selection box
+				if (!bMustEncompassEntireComponent && bLocationIntersected)
+				{
+					return true;
+				}
+
+				// If the selection box has to encompass the entire component and a skeletal mesh vertex didn't intersect with the selection
+				// box, this component does not qualify
+				else if (bMustEncompassEntireComponent && !bLocationIntersected)
+				{
+					return false;
+				}
+			}
+		}
+
+		// If the selection box has to encompass all of the component and none of the component's verts failed the intersection test, this component
+		// is consider touching
+		return true;
+	}
+
+	return false;
+}
+
+
 void USkeletalMeshComponent::UpdateCollisionProfile()
 {
 	Super::UpdateCollisionProfile();
diff --git a/Engine/Source/Runtime/Engine/Private/StaticMeshComponent.cpp b/Engine/Source/Runtime/Engine/Private/StaticMeshComponent.cpp
--- a/Engine/Source/Runtime/Engine/Private/StaticMeshComponent.cpp
+++ b/Engine/Source/Runtime/Engine/Private/StaticMeshComponent.cpp
@@ -7,6 +7,11 @@
 #include "MessageLog.h"
 #include "UObjectToken.h"
 #include "MapErrors.h"
+#if WITH_EDITOR
+#include "ShowFlags.h"
+#include "Collision.h"
+#include "ConvexVolume.h"
+#endif
 #include "ComponentInstanceDataCache.h"
 #include "LightMap.h"
 #include "ShadowMap.h"
@@ -1615,6 +1620,95 @@ bool UStaticMeshComponent::DoCustomNavigableGeometryExport(struct FNavigableGeom
 	return true;
 }
 
+#if WITH_EDITOR
+bool UStaticMeshComponent::ComponentIsTouchingSelectionBox(const FBox& InSelBBox, const FEngineShowFlags& ShowFlags, const bool bConsiderOnlyBSP, const bool bMustEncompassEntireComponent) const
+{
+	if (!bConsiderOnlyBSP && ShowFlags.StaticMeshes && StaticMesh != nullptr && StaticMesh->HasValidRenderData())
+	{
+		// Check if we are even inside it's bounding box, if we are not, there is no way we colliding via the more advanced checks we will do.
+		if (Super::ComponentIsTouchingSelectionBox(InSelBBox, ShowFlags, bConsiderOnlyBSP, false))
+		{
+			TArray<FVector> Vertex;
+
+			FStaticMeshLODResources& LODModel = StaticMesh->RenderData->LODResources[0];
+			FIndexArrayView Indices = LODModel.IndexBuffer.GetArrayView();
+
+			for (const auto& Section : LODModel.Sections)
+			{
+				// Iterate over each triangle.
+				for (int32 TriangleIndex = 0; TriangleIndex < (int32)Section.NumTriangles; TriangleIndex++)
+				{
+					Vertex.Empty(3);
+
+					int32 FirstIndex = TriangleIndex * 3 + Section.FirstIndex;
+					for (int32 i = 0; i < 3; i++)
+					{
+						int32 VertexIndex = Indices[FirstIndex + i];
+						FVector LocalPosition = LODModel.PositionVertexBuffer.VertexPosition(VertexIndex);
+						Vertex.Emplace(ComponentToWorld.TransformPosition(LocalPosition));
+					}
+
+					// Check if the triangle is colliding with the bounding box.
+					FSeparatingAxisPointCheck ThePointCheck(Vertex, InSelBBox.GetCenter(), InSelBBox.GetExtent(), false);
+					if (!bMustEncompassEntireComponent && ThePointCheck.bHit)
+					{
+						// Needn't encompass entire component: any intersection, we consider as touching
+						return true;
+					}
+					else if (bMustEncompassEntireComponent && !ThePointCheck.bHit)
+					{
+						// Must encompass entire component: any non intersection, we consider as not touching
+						return false;
+					}
+				}
+			}
+
+			// If the selection box has to encompass all of the component and none of the component's verts failed the intersection test, this component
+			// is consider touching
+			return true;
+		}
+	}
+
+	return false;
+}
+
+
+bool UStaticMeshComponent::ComponentIsTouchingSelectionFrustum(const FConvexVolume& InFrustum, const FEngineShowFlags& ShowFlags, const bool bConsiderOnlyBSP, const bool bMustEncompassEntireComponent) const
+{
+	if (!bConsiderOnlyBSP && ShowFlags.StaticMeshes && StaticMesh != nullptr && StaticMesh->HasValidRenderData())
+	{
+		// Check if we are even inside it's bounding box, if we are not, there is no way we colliding via the more advanced checks we will do.
+		if (Super::ComponentIsTouchingSelectionFrustum(InFrustum, ShowFlags, bConsiderOnlyBSP, false))
+		{
+			TArray<FVector> Vertex;
+
+			FStaticMeshLODResources& LODModel = StaticMesh->RenderData->LODResources[0];
+
+			uint32 NumVertices = LODModel.VertexBuffer.GetNumVertices();
+			for (uint32 VertexIndex = 0; VertexIndex < NumVertices; ++VertexIndex)
+			{
+				const FVector& LocalPosition = LODModel.PositionVertexBuffer.VertexPosition(VertexIndex);
+				const FVector WorldPosition = ComponentToWorld.TransformPosition(LocalPosition);
+				bool bLocationIntersected = InFrustum.IntersectSphere(WorldPosition, 0.0f);
+				if (bLocationIntersected && !bMustEncompassEntireComponent)
+				{
+					return true;
+				}
+				else if (!bLocationIntersected && bMustEncompassEntireComponent)
+				{
+					return false;
+				}
+			}
+
+			return true;
+		}
+	}
+
+	return false;
+}
+#endif
+
+
 //////////////////////////////////////////////////////////////////////////
 // StaticMeshComponentLODInfo
 
diff --git a/Engine/Source/Runtime/Landscape/Classes/LandscapeComponent.h b/Engine/Source/Runtime/Landscape/Classes/LandscapeComponent.h
--- a/Engine/Source/Runtime/Landscape/Classes/LandscapeComponent.h
+++ b/Engine/Source/Runtime/Landscape/Classes/LandscapeComponent.h
@@ -11,6 +11,8 @@
 class ULandscapeLayerInfoObject;
 class ULandscapeInfo;
 class ALandscapeProxy;
+struct FEngineShowFlags;
+struct FConvexVolume;
 
 class FLandscapeComponentDerivedData
 {
@@ -277,6 +279,8 @@ public:
 	virtual int32 GetNumMaterials() const override;
 	virtual class UMaterialInterface* GetMaterial(int32 ElementIndex) const override;
 	virtual void SetMaterial(int32 ElementIndex, class UMaterialInterface* Material) override;
+	virtual bool ComponentIsTouchingSelectionBox(const FBox& InSelBBox, const FEngineShowFlags& ShowFlags, const bool bConsiderOnlyBSP, const bool bMustEncompassEntireComponent) const override;
+	virtual bool ComponentIsTouchingSelectionFrustum(const FConvexVolume& InFrustum, const FEngineShowFlags& ShowFlags, const bool bConsiderOnlyBSP, const bool bMustEncompassEntireComponent) const override;
 #endif
 	// End UPrimitiveComponent interface.
 
diff --git a/Engine/Source/Runtime/Landscape/Classes/LandscapeHeightfieldCollisionComponent.h b/Engine/Source/Runtime/Landscape/Classes/LandscapeHeightfieldCollisionComponent.h
--- a/Engine/Source/Runtime/Landscape/Classes/LandscapeHeightfieldCollisionComponent.h
+++ b/Engine/Source/Runtime/Landscape/Classes/LandscapeHeightfieldCollisionComponent.h
@@ -142,6 +142,10 @@ class ULandscapeHeightfieldCollisionComponent : public UPrimitiveComponent
 
 	// Begin UPrimitiveComponent interface
 	virtual bool DoCustomNavigableGeometryExport(struct FNavigableGeometryExport* GeomExport) const override;
+#if WITH_EDITOR
+	virtual bool ComponentIsTouchingSelectionBox(const FBox& InSelBBox, const FEngineShowFlags& ShowFlags, const bool bConsiderOnlyBSP, const bool bMustEncompassEntireComponent) const override;
+	virtual bool ComponentIsTouchingSelectionFrustum(const FConvexVolume& InFrustum, const FEngineShowFlags& ShowFlags, const bool bConsiderOnlyBSP, const bool bMustEncompassEntireComponent) const override;
+#endif
 	//End UPrimitiveComponent interface
 
 	// Begin UObject Interface.
diff --git a/Engine/Source/Runtime/Landscape/Private/LandscapeCollision.cpp b/Engine/Source/Runtime/Landscape/Private/LandscapeCollision.cpp
--- a/Engine/Source/Runtime/Landscape/Private/LandscapeCollision.cpp
+++ b/Engine/Source/Runtime/Landscape/Private/LandscapeCollision.cpp
@@ -1165,6 +1165,26 @@ void ULandscapeHeightfieldCollisionComponent::PostEditUndo()
 
 	UNavigationSystem::UpdateNavOctree(this);
 }
+
+bool ULandscapeHeightfieldCollisionComponent::ComponentIsTouchingSelectionBox(const FBox& InSelBBox, const FEngineShowFlags& ShowFlags, const bool bConsiderOnlyBSP, const bool bMustEncompassEntireComponent) const
+{
+	if (ShowFlags.Landscape)
+	{
+		return Super::ComponentIsTouchingSelectionBox(InSelBBox, ShowFlags, bConsiderOnlyBSP, bMustEncompassEntireComponent);
+	}
+
+	return false;
+}
+
+bool ULandscapeHeightfieldCollisionComponent::ComponentIsTouchingSelectionFrustum(const FConvexVolume& InFrustum, const FEngineShowFlags& ShowFlags, const bool bConsiderOnlyBSP, const bool bMustEncompassEntireComponent) const
+{
+	if (ShowFlags.Landscape)
+	{
+		return Super::ComponentIsTouchingSelectionFrustum(InFrustum, ShowFlags, bConsiderOnlyBSP, bMustEncompassEntireComponent);
+	}
+
+	return false;
+}
 #endif
 
 bool ULandscapeHeightfieldCollisionComponent::DoCustomNavigableGeometryExport(FNavigableGeometryExport* GeomExport) const
diff --git a/Engine/Source/Runtime/Landscape/Private/LandscapeEdit.cpp b/Engine/Source/Runtime/Landscape/Private/LandscapeEdit.cpp
--- a/Engine/Source/Runtime/Landscape/Private/LandscapeEdit.cpp
+++ b/Engine/Source/Runtime/Landscape/Private/LandscapeEdit.cpp
@@ -37,6 +37,8 @@ LandscapeEdit.cpp: Landscape editing
 #include "EngineUtils.h"
 #include "Engine/Engine.h"
 #include "EngineGlobals.h"
+#include "ShowFlags.h"
+#include "ConvexVolume.h"
 #endif
 #include "ComponentReregisterContext.h"
 
@@ -312,6 +314,26 @@ void ULandscapeComponent::SetMaterial(int32 ElementIndex, class UMaterialInterfa
 	}
 }
 
+bool ULandscapeComponent::ComponentIsTouchingSelectionBox(const FBox& InSelBBox, const FEngineShowFlags& ShowFlags, const bool bConsiderOnlyBSP, const bool bMustEncompassEntireComponent) const
+{
+	if (ShowFlags.Landscape)
+	{
+		return Super::ComponentIsTouchingSelectionBox(InSelBBox, ShowFlags, bConsiderOnlyBSP, bMustEncompassEntireComponent);
+	}
+
+	return false;
+}
+
+bool ULandscapeComponent::ComponentIsTouchingSelectionFrustum(const FConvexVolume& InFrustum, const FEngineShowFlags& ShowFlags, const bool bConsiderOnlyBSP, const bool bMustEncompassEntireComponent) const
+{
+	if (ShowFlags.Landscape)
+	{
+		return Super::ComponentIsTouchingSelectionFrustum(InFrustum, ShowFlags, bConsiderOnlyBSP, bMustEncompassEntireComponent);
+	}
+
+	return false;
+}
+
 void ULandscapeComponent::PostEditUndo()
 {
 	Super::PostEditUndo();
