Fix to prevent a crash if the selected GameInstance class cannot be found when running PIE.

diff --git a/Engine/Source/Editor/UnrealEd/Private/PlayLevel.cpp b/Engine/Source/Editor/UnrealEd/Private/PlayLevel.cpp
--- a/Engine/Source/Editor/UnrealEd/Private/PlayLevel.cpp
+++ b/Engine/Source/Editor/UnrealEd/Private/PlayLevel.cpp
@@ -2408,6 +2408,11 @@ UGameInstance* UEditorEngine::CreatePIEGameInstance(int32 PIEInstance, bool bInS
 	FStringClassReference GameInstanceClassName = GetDefault<UGameMapsSettings>()->GameInstanceClass;
 	UClass* GameInstanceClass = (GameInstanceClassName.IsValid() ? LoadObject<UClass>(NULL, *GameInstanceClassName.ToString()) : UGameInstance::StaticClass());
 
+	// If the GameInstance class from the settings cannot be found, fall back to the base class
+	if(GameInstanceClass == nullptr)
+	{
+		GameInstanceClass = UGameInstance::StaticClass();
+	}
 	UGameInstance* GameInstance = ConstructObject<UGameInstance>(GameInstanceClass, this);
 
 	// We need to temporarily add the GameInstance to the root because the InitPIE call can do garbage collection wiping out the GameInstance
