Editor - Adding a check to ensure the slate application is initialized before hooking the delegate for handling asset open requests.  Fixes cooking.

diff --git a/Engine/Source/Editor/UnrealEd/Private/Editor.cpp b/Engine/Source/Editor/UnrealEd/Private/Editor.cpp
--- a/Engine/Source/Editor/UnrealEd/Private/Editor.cpp
+++ b/Engine/Source/Editor/UnrealEd/Private/Editor.cpp
@@ -447,8 +447,11 @@ void UEditorEngine::InitEditor(IEngineLoop* InEngineLoop)
 	// Load any modules that might be required by commandlets
 	FModuleManager::Get().LoadModule(TEXT("OnlineBlueprintSupport"));
 
-	// Setup a delegate to handle requests for opening assets
-	FSlateApplication::Get().SetWidgetReflectorAssetAccessDelegate(FAccessAsset::CreateUObject(this, &UEditorEngine::HandleOpenAsset));
+	if ( FSlateApplication::IsInitialized() )
+	{
+		// Setup a delegate to handle requests for opening assets
+		FSlateApplication::Get().SetWidgetReflectorAssetAccessDelegate(FAccessAsset::CreateUObject(this, &UEditorEngine::HandleOpenAsset));
+	}
 }
 
 bool UEditorEngine::HandleOpenAsset(UObject* Asset)
