Nonunity fix

diff --git a/Engine/Source/Editor/UnrealEd/Private/Factories/PackFactory.cpp b/Engine/Source/Editor/UnrealEd/Private/Factories/PackFactory.cpp
--- a/Engine/Source/Editor/UnrealEd/Private/Factories/PackFactory.cpp
+++ b/Engine/Source/Editor/UnrealEd/Private/Factories/PackFactory.cpp
@@ -12,6 +12,7 @@
 #include "AES.h"
 #include "ModuleManager.h"
 #include "GameProjectGenerationModule.h"
+#include "Factories/PackFactory.h"
 
 DEFINE_LOG_CATEGORY_STATIC(LogPackFactory, Log, All);
 
