Fixed merge levels operation.  Fixed bug in UEditorEngine::CopySelectedActorsToClipboard when performing a slow copy due to the copied actors not being correctly pasted into the buffer level, which was affecting both 'merge levels' and copying of multiple actors from different levels.

diff --git a/Engine/Source/Editor/UnrealEd/Private/EditorSelectUtils.cpp b/Engine/Source/Editor/UnrealEd/Private/EditorSelectUtils.cpp
--- a/Engine/Source/Editor/UnrealEd/Private/EditorSelectUtils.cpp
+++ b/Engine/Source/Editor/UnrealEd/Private/EditorSelectUtils.cpp
@@ -443,10 +443,10 @@ bool UUnrealEdEngine::CanSelectActor(AActor* Actor, bool bInSelected, bool bSele
 		return false;
 	}
 
-	bool bSelectionAllowed = false;
-
+	// Allow active modes to determine whether the selection is allowed. If there are no active modes, allow selection anyway.
 	TArray<FEdMode*> ActiveModes;
 	GLevelEditorModeTools().GetActiveModes( ActiveModes );
+	bool bSelectionAllowed = (ActiveModes.Num() == 0);
 	for( int32 ModeIndex = 0; ModeIndex < ActiveModes.Num(); ++ModeIndex )
 	{
 		bSelectionAllowed |= ActiveModes[ModeIndex]->IsSelectionAllowed( Actor, bInSelected );
diff --git a/Engine/Source/Editor/UnrealEd/Private/EditorServer.cpp b/Engine/Source/Editor/UnrealEd/Private/EditorServer.cpp
--- a/Engine/Source/Editor/UnrealEd/Private/EditorServer.cpp
+++ b/Engine/Source/Editor/UnrealEd/Private/EditorServer.cpp
@@ -2753,6 +2753,7 @@ ULevel*  UEditorEngine::CreateTransLevelMoveBuffer( UWorld* InWorld )
 	BufferLevel->AddToRoot();
 	BufferLevel->OwningWorld = InWorld;
 	BufferLevel->Model = new( BufferLevel ) UModel( FObjectInitializer(), NULL, true );
+	BufferLevel->bIsVisible = true;
 		
 	BufferLevel->SetFlags( RF_Transactional );
 	BufferLevel->Model->SetFlags( RF_Transactional );
diff --git a/Engine/Source/Editor/WorldBrowser/Private/LevelCollectionModel.cpp b/Engine/Source/Editor/WorldBrowser/Private/LevelCollectionModel.cpp
--- a/Engine/Source/Editor/WorldBrowser/Private/LevelCollectionModel.cpp
+++ b/Engine/Source/Editor/WorldBrowser/Private/LevelCollectionModel.cpp
@@ -565,7 +565,9 @@ void FLevelCollectionModel::UnloadLevels(const FLevelModelList& InLevelList)
 
 	
 	// Remove each level!
-	for(auto It = InLevelList.CreateConstIterator(); It; ++It)
+	// Take a copy of the list rather than using a reference to the selected levels list, as this will be modified in the loop below
+	const FLevelModelList LevelListCopy = InLevelList;
+	for (auto It = LevelListCopy.CreateConstIterator(); It; ++It)
 	{
 		TSharedPtr<FLevelModel> LevelModel = (*It);
 		ULevel* Level = LevelModel->GetLevelObject();
