Fixed category on blueprint merge tool setting (sorry AI!)

diff --git a/Engine/Source/Editor/UnrealEd/Classes/Settings/EditorExperimentalSettings.h b/Engine/Source/Editor/UnrealEd/Classes/Settings/EditorExperimentalSettings.h
--- a/Engine/Source/Editor/UnrealEd/Classes/Settings/EditorExperimentalSettings.h
+++ b/Engine/Source/Editor/UnrealEd/Classes/Settings/EditorExperimentalSettings.h
@@ -86,7 +86,7 @@ public:
 	bool bEQSEditor;
 
 	/** Enables the Blueprint merge tool */
-	UPROPERTY(EditAnywhere, config, Category = AI, meta = (DisplayName = "Enable Blueprint Merge Tool"))
+	UPROPERTY(EditAnywhere, config, Category = Blueprints, meta = (DisplayName = "Enable Blueprint Merge Tool"))
 	bool bEnableBlueprintMergeTool;
 
 	/**
