Stale Kismet commandlets were removed.

diff --git a/Engine/Source/Editor/UnrealEd/Classes/Commandlets/KismetUpdateCommandlet.h b/Engine/Source/Editor/UnrealEd/Classes/Commandlets/KismetUpdateCommandlet.h
deleted file mode 100644
diff --git a/Engine/Source/Editor/UnrealEd/Classes/Commandlets/RecompileBlueprintsCommandlet.h b/Engine/Source/Editor/UnrealEd/Classes/Commandlets/RecompileBlueprintsCommandlet.h
deleted file mode 100644
diff --git a/Engine/Source/Editor/UnrealEd/Private/Commandlets/ContentCommandlets.cpp b/Engine/Source/Editor/UnrealEd/Private/Commandlets/ContentCommandlets.cpp
--- a/Engine/Source/Editor/UnrealEd/Private/Commandlets/ContentCommandlets.cpp
+++ b/Engine/Source/Editor/UnrealEd/Private/Commandlets/ContentCommandlets.cpp
@@ -710,98 +710,6 @@ void UResavePackagesCommandlet::VerboseMessage(const FString& Message)
 	}
 }
 
-/*-----------------------------------------------------------------------------
-	URecompileBlueprintsCommandlet commandlet.
------------------------------------------------------------------------------*/
-
-URecompileBlueprintsCommandlet::URecompileBlueprintsCommandlet(const FObjectInitializer& ObjectInitializer)
-	: Super(ObjectInitializer)
-{
-}
-
-bool URecompileBlueprintsCommandlet::ShouldSkipPackage(const FString& Filename)
-{
-	//@todo:  Only do this for packages with blueprints in them
-	return false;
-}
-
-bool URecompileBlueprintsCommandlet::PerformPreloadOperations( ULinkerLoad* PackageLinker, bool& bSavePackage )
-{
-	bool RetVal = Super::PerformPreloadOperations(PackageLinker, bSavePackage);
-
-	// Force the metadata to be loaded, which will fully load the package
-	UPackage* Package = PackageLinker->LinkerRoot;
-	if( Package )
-	{
-		Package->GetMetaData();
-	}
-
-	return RetVal;
-}
-
-void URecompileBlueprintsCommandlet::PerformAdditionalOperations( class UObject* Object, bool& bSavePackage )
-{
-	if( UBlueprint* TargetBlueprint = Cast<UBlueprint>(Object) )
-	{
-		// Obtain the filename so we can regenerate the ULinkerLoad after the compile.
-		FString Filename;
-		FPackageName::DoesPackageExist(Object->GetOutermost()->GetName(), NULL, &Filename);
-
-		FBlueprintEditorUtils::RefreshAllNodes(TargetBlueprint);
-		FKismetEditorUtilities::CompileBlueprint(TargetBlueprint, true, false);
-
-		// Clean up the metadata, now that we've recompiled and potentially moved stuff into the transient package
-		UPackage* Package = CastChecked<UPackage>(Object->GetOutermost());
-		UMetaData* MetaData = Package->GetMetaData();
-		check(MetaData);
-		MetaData->RemoveMetaDataOutsidePackage();
-
-		if ( !Filename.IsEmpty() )
-		{
-			// Regenerate the LinkerLoad if it was reset during the compile. This will add the ULinkerLoad to GObjLoaders.
-			// This will allow thumbnails to be properly preserved in UPackage::SavePackage (ResetLoadersForSave)
-			BeginLoad();
-			ULinkerLoad* Linker = GetPackageLinker(NULL,*Filename,LOAD_NoVerify,NULL,NULL);
-			EndLoad();
-		}
-	}
-}
-
-int32 URecompileBlueprintsCommandlet::Main(const FString& Params)
-{
-	// Verify that all blueprints are set to NOT compile on load
-	for( TObjectIterator<UBlueprint> BlueprintIt; BlueprintIt; ++BlueprintIt )
-	{
-		UBlueprint* Blueprint = *BlueprintIt;
-		if( Blueprint->bRecompileOnLoad )
-		{
-			UE_LOG(LogContentCommandlet, Error, TEXT("Blueprint %s is set to compile on load, which is unsafe for this commandlet.  Please disable in the engine's INI file."), *Blueprint->GetClass()->GetName());
-			return -1;
-		}
-	}
-
-	return Super::Main(Params);
-}
-
-int32 URecompileBlueprintsCommandlet::InitializeResaveParameters( const TArray<FString>& Tokens, const TArray<FString>& Switches, TArray<FString>& MapPathNames )
-{
-	int32 RetVal = Super::InitializeResaveParameters(Tokens, Switches, MapPathNames);
-	
-	ResaveClasses.Add(FName(TEXT("Blueprint")));
-	ResaveClasses.Add(FName(TEXT("LevelScriptBlueprint")));
-	ResaveClasses.Add(FName(TEXT("AnimBlueprint")));
-	ResaveClasses.Add(FName(TEXT("EditorUtilityBlueprint")));
-	bOnlySaveDirtyPackages = false;
-
-	// Checking in automatically is too dangerous in the case that something goes wrong!
-	bAutoCheckIn = false;
-
-	// Too slow to GC every time
-	GarbageCollectionFrequency = 50;
-
-	return RetVal;
-}
-
 /*-----------------------------------------------------------------------------
 	UWrangleContent.
 -----------------------------------------------------------------------------*/
diff --git a/Engine/Source/Editor/UnrealEd/Private/Commandlets/KismetUpdateCommandlet.cpp b/Engine/Source/Editor/UnrealEd/Private/Commandlets/KismetUpdateCommandlet.cpp
deleted file mode 100644
