file missed in submission for UE-3347
#ios

diff --git a/Engine/Source/Editor/UnrealEd/UnrealEd.Build.cs b/Engine/Source/Editor/UnrealEd/UnrealEd.Build.cs
--- a/Engine/Source/Editor/UnrealEd/UnrealEd.Build.cs
+++ b/Engine/Source/Editor/UnrealEd/UnrealEd.Build.cs
@@ -188,6 +188,8 @@ public class UnrealEd : ModuleRules
 				"ReferenceViewer",
 				"EditorLiveStreaming",
 				"HotReload",
+                "AndroidPlatformEditor",
+                "IOSPlatformEditor"
 			}
 		);
 
