UE-4908 GitHub 570 : Fix reference component for diffing properties values while cloning components
https://github.com/EpicGames/UnrealEngine/pull/570

diff --git a/Engine/Source/Editor/UnrealEd/Private/Kismet2/Kismet2.cpp b/Engine/Source/Editor/UnrealEd/Private/Kismet2/Kismet2.cpp
--- a/Engine/Source/Editor/UnrealEd/Private/Kismet2/Kismet2.cpp
+++ b/Engine/Source/Editor/UnrealEd/Private/Kismet2/Kismet2.cpp
@@ -910,7 +910,7 @@ protected:
 	{
 		check(Source->GetClass() == Destination->GetClass());
 		UClass* CommonBaseClass = Source->GetClass();
-		UActorComponent* DefaultComponent = CastChecked<UActorComponent>(Source->GetArchetype());
+		UActorComponent* DefaultComponent = CastChecked<UActorComponent>(Destination->GetArchetype());
 
 		//@TODO: Copy of code from inside CopyActorComponentProperties
 
