UEBP-6 PERFORMANCE (REINSTANCING) - Investigate/optimize reinstancing necessity.

Changing default values doesn't cause reinstancing while recompilation.

bChangeDefaultValueWithoutReinstancing switch in ini.

diff --git a/Engine/Config/BaseEngine.ini b/Engine/Config/BaseEngine.ini
--- a/Engine/Config/BaseEngine.ini
+++ b/Engine/Config/BaseEngine.ini
@@ -1726,6 +1726,7 @@ PrintStringDuration=2.0
 bUseLocalGraphVariables=false
 bPersistentUberGraphFrame=true
 bReinstanceOnlyWhenNecessary=true
+bChangeDefaultValueWithoutReinstancing=false
 bDisplaysLayout=false
 bNativeCodeGenerationTool=false
 
diff --git a/Engine/Source/Editor/KismetCompiler/Private/KismetCompiler.cpp b/Engine/Source/Editor/KismetCompiler/Private/KismetCompiler.cpp
--- a/Engine/Source/Editor/KismetCompiler/Private/KismetCompiler.cpp
+++ b/Engine/Source/Editor/KismetCompiler/Private/KismetCompiler.cpp
@@ -3605,7 +3605,8 @@ void FKismetCompilerContext::Compile()
 		// TODO What do we do if validation fails?
 	}
 
-	if (bIsFullCompile)
+	static const FBoolConfigValueHelper ChangeDefaultValueWithoutReinstancing(TEXT("Kismet"), TEXT("bChangeDefaultValueWithoutReinstancing"), GEngineIni);
+	if (bIsFullCompile && !ChangeDefaultValueWithoutReinstancing)
 	{
 		BP_SCOPED_COMPILER_EVENT_STAT(EKismetCompilerStats_ChecksumCDO);
 		UObject* NewCDO = NewClass->GetDefaultObject(false);
diff --git a/Engine/Source/Editor/UnrealEd/Private/Kismet2/KismetReinstanceUtilities.cpp b/Engine/Source/Editor/UnrealEd/Private/Kismet2/KismetReinstanceUtilities.cpp
--- a/Engine/Source/Editor/UnrealEd/Private/Kismet2/KismetReinstanceUtilities.cpp
+++ b/Engine/Source/Editor/UnrealEd/Private/Kismet2/KismetReinstanceUtilities.cpp
@@ -181,8 +181,9 @@ void FBlueprintCompileReinstancer::ReinstanceObjects(bool bAlwaysReinstance)
 			auto BPClassA = Cast<const UBlueprintGeneratedClass>(DuplicatedClass);
 			auto BPClassB = Cast<const UBlueprintGeneratedClass>(ClassToReinstance);
 			auto BP = Cast<const UBlueprint>(ClassToReinstance->ClassGeneratedBy);
+			static const FBoolConfigValueHelper ChangeDefaultValueWithoutReinstancing(TEXT("Kismet"), TEXT("bChangeDefaultValueWithoutReinstancing"), GEngineIni);
 			const bool bTheSameDefaultValues = BP && ClassToReinstanceDefaultValuesCRC && (BP->CrcPreviousCompiledCDO == ClassToReinstanceDefaultValuesCRC);
-			const bool bTheSame = bTheSameDefaultValues && BPClassA && BPClassB && FStructUtils::TheSameLayout(BPClassA, BPClassB, true);
+			const bool bTheSame = (ChangeDefaultValueWithoutReinstancing || bTheSameDefaultValues) && BPClassA && BPClassB && FStructUtils::TheSameLayout(BPClassA, BPClassB, true);
 			if (bTheSame)
 			{
 				UE_LOG(LogBlueprint, Log, TEXT("BlueprintCompileReinstancer: class '%s' is replaced without reinstancing (the optimized way)."), *GetPathNameSafe(ClassToReinstance));
diff --git a/Engine/Source/Runtime/Engine/Classes/Engine/Blueprint.h b/Engine/Source/Runtime/Engine/Classes/Engine/Blueprint.h
--- a/Engine/Source/Runtime/Engine/Classes/Engine/Blueprint.h
+++ b/Engine/Source/Runtime/Engine/Classes/Engine/Blueprint.h
@@ -436,7 +436,7 @@ public:
 	UPROPERTY(transient)
 	uint32 bBeingCompiled:1;
 
-	/** CRC for CDO calculated right after the latest compilation */
+	/** CRC for CDO calculated right after the latest compilation used by Reinstancer to check if default values were changed */
 	UPROPERTY(transient, duplicatetransient)
 	uint32 CrcPreviousCompiledCDO;
 
