Fixed textures trying to be imported as font atlas pages

UE-5113 - Importing texture stuck at determining asset types

bEditorImport had been set to true rather than false since I'd removed a base class and missed that when merging the constructors together.

diff --git a/Engine/Source/Editor/UnrealEd/Private/Factories/TTFontImport.cpp b/Engine/Source/Editor/UnrealEd/Private/Factories/TTFontImport.cpp
--- a/Engine/Source/Editor/UnrealEd/Private/Factories/TTFontImport.cpp
+++ b/Engine/Source/Editor/UnrealEd/Private/Factories/TTFontImport.cpp
@@ -40,6 +40,7 @@ UTrueTypeFontFactory::UTrueTypeFontFactory(const FObjectInitializer& ObjectIniti
 {
 	SupportedClass = UFont::StaticClass();
 	bCreateNew = true;
+	bEditorImport = false;
 	bEditAfterNew = true;
 	AutoPriority = -1;
 	LODGroup = TEXTUREGROUP_UI;
