Fixed a comment spelling errror.
Started working on launch on test.

diff --git a/Engine/Source/Editor/UnrealEd/Private/Tests/EditorAutomationTests.cpp b/Engine/Source/Editor/UnrealEd/Private/Tests/EditorAutomationTests.cpp
--- a/Engine/Source/Editor/UnrealEd/Private/Tests/EditorAutomationTests.cpp
+++ b/Engine/Source/Editor/UnrealEd/Private/Tests/EditorAutomationTests.cpp
@@ -1496,7 +1496,7 @@ bool FLightPlacement::RunTest(const FString& Parameters)
 	ADD_LATENT_AUTOMATION_COMMAND(FWaitLatentCommand(0.1f));
 
 	//Take a screenshot of the final scene.
-	//It is expected to show two moveable point lights (one red, one white) and a static mesh.
+	//It is expected to show two movable point lights (one red, one white) and a static mesh.
 	TakeLatentAutomationScreenshot(PointLightPlacementWindowParameters, BaseFileName, FString::Printf(TEXT("FinalPointLight")), FString::Printf(TEXT("04_Final")), true);
 
 	return true;
@@ -2065,3 +2065,40 @@ bool FStaticMeshUVsTest::RunTest(const FString& Parameters)
 	ADD_LATENT_AUTOMATION_COMMAND(FPerformUVTestCommand(TestHelper));
 	return true;
 }
+
+
+/**
+* Launches a map onto a specified device after making a change to it.
+*/
+IMPLEMENT_COMPLEX_AUTOMATION_TEST(FLaunchOnTestWithMapIteration, "Performance.Launch Test With Map Iteration", (EAutomationTestFlags::ATF_Editor | EAutomationTestFlags::ATF_RequiresUser))
+
+void FLaunchOnTestWithMapIteration::GetTests(TArray<FString>& OutBeautifiedNames, TArray <FString>& OutTestCommands) const
+{
+	//Get Map names from preferences (engine.ini)
+
+	//Check to see if the map has a set launch on device.
+		//If then don't add it as a test.
+}
+
+bool FLaunchOnTestWithMapIteration::RunTest(const FString& Parameters)
+{
+	//Get the map Name
+
+	//Load Map
+	
+	//Make an adjustment to the map
+		
+	//Rebuild Map
+
+	//Check out Map and then save it
+
+	//Set the device to be launched onto
+
+	//Launch onto device and get launch on times and cook times
+
+	//Verify the launch worked.
+
+	//Revert Map
+
+	return true;
+}
\ No newline at end of file
