Fixes the session frontends automation tab long load times for bigger projects.
--Refactored each individual open asset test so that it is now a simple test instead of complex.

diff --git a/Engine/Source/Editor/UnrealEd/Private/Tests/AutomationEditorCommon.cpp b/Engine/Source/Editor/UnrealEd/Private/Tests/AutomationEditorCommon.cpp
--- a/Engine/Source/Editor/UnrealEd/Private/Tests/AutomationEditorCommon.cpp
+++ b/Engine/Source/Editor/UnrealEd/Private/Tests/AutomationEditorCommon.cpp
@@ -591,7 +591,6 @@ void FEditorAutomationTestUtilities::CollectTestsByClass(UClass * Class, TArray<
 void FEditorAutomationTestUtilities::CollectGameContentTestsByClass(UClass * Class, bool bRecursiveClass, TArray<FString>& OutBeautifiedNames, TArray <FString>& OutTestCommands)
 {
 	//Setting the Asset Registry
-
 	FAssetRegistryModule& AssetRegistryModule = FModuleManager::LoadModuleChecked<FAssetRegistryModule>(TEXT("AssetRegistry"));
 
 	//Variable setups
@@ -676,60 +675,86 @@ void FEditorAutomationTestUtilities::CollectMiscGameContentTestsByClass(TArray<F
 	ExcludeClassesList.Add(UWorld::StaticClass()->GetFName());
 
 	//Generating the list of assets.
-	//This list isn't expected to have anything that may be obtained with another function.
 	//This list is being filtered by the game folder and class type.  The results are placed into the ObjectList variable.
 	AssetFilter.PackagePaths.Add("/Game");
 	AssetFilter.bRecursivePaths = true;
 	AssetRegistryModule.Get().GetAssets(AssetFilter, ObjectList);
-	
-	//if (ObjectList.Num() < 25000)
-	//{
 
-		//Loop through the list of assets, make their path full and a string, then add them to the test.
-		for (auto ObjIter = ObjectList.CreateConstIterator(); ObjIter; ++ObjIter)
+	//Loop through the list of assets, make their path full and a string, then add them to the test.
+	for (auto ObjIter = ObjectList.CreateConstIterator(); ObjIter; ++ObjIter)
+	{
+		const FAssetData& Asset = *ObjIter;
+		//First we check if the class is valid.  If not then we move onto the next object.
+		if (Asset.GetClass() != NULL)
 		{
-			const FAssetData& Asset = *ObjIter;
-			//First we check if the class is valid.  If not then we move onto the next object.
-			if (Asset.GetClass() != NULL)
-			{
-				//Variable that holds the class FName for the current asset iteration.
-				FName AssetClassFName = Asset.GetClass()->GetFName();
+			//Variable that holds the class FName for the current asset iteration.
+			FName AssetClassFName = Asset.GetClass()->GetFName();
+
+			//Counter used to keep track for the following for loop.
+			float ExcludedClassesCounter = 1;
 
-				//Counter used to keep track for the following for loop.
-				float ExcludedClassesCounter = 1;
+			for (auto ExcludeIter = ExcludeClassesList.CreateConstIterator(); ExcludeIter; ++ExcludeIter)
+			{
+				FName ExludedName = *ExcludeIter;
 
-				for (auto ExcludeIter = ExcludeClassesList.CreateConstIterator(); ExcludeIter; ++ExcludeIter)
+				//If the classes are the same then we don't want this asset. So we move onto the next one instead.
+				if (AssetClassFName == ExludedName)
 				{
-					FName ExludedName = *ExcludeIter;
+					break;
+				}
 
-					//If the classes are the same then we don't want this asset. So we move onto the next one instead.
-					if (AssetClassFName == ExludedName)
-					{
-						break;
-					}
+				//We run out of class names in our Excluded list then we want the current ObjectList asset.
+				if ((ExcludedClassesCounter + 1) > ExcludeClassesList.Num())
+				{
+					FString Filename = Asset.ObjectPath.ToString();
+					//convert to full paths
+					Filename = FPackageName::LongPackageNameToFilename(Filename);
 
-					//We run out of class names in our Excluded list then we want the current ObjectList asset.
-					if ((ExcludedClassesCounter + 1) > ExcludeClassesList.Num())
+					if (FAutomationTestFramework::GetInstance().ShouldTestContent(Filename))
 					{
-						FString Filename = Asset.ObjectPath.ToString();
-						//convert to full paths
-						Filename = FPackageName::LongPackageNameToFilename(Filename);
-
-						if (FAutomationTestFramework::GetInstance().ShouldTestContent(Filename))
-						{
-							FString BeautifiedFilename = Asset.AssetName.ToString();
-							OutBeautifiedNames.Add(BeautifiedFilename);
-							OutTestCommands.Add(Asset.ObjectPath.ToString());
-						}
-
-						break;
+						FString BeautifiedFilename = Asset.AssetName.ToString();
+						OutBeautifiedNames.Add(BeautifiedFilename);
+						OutTestCommands.Add(Asset.ObjectPath.ToString());
 					}
 
-					//increment the counter.
-					ExcludedClassesCounter++;
+					break;
 				}
+				ExcludedClassesCounter++;
 			}
 		}
-	//}
+	}
 }
 
+/**
+* Generates a list of assets from the GAME by a specific type.
+*/
+void FEditorAutomationTestUtilities::CollectGameContentByClass(const UClass * Class, bool bRecursiveClass, TArray<FString>& OutAssetList)
+{
+	//Setting the Asset Registry
+	FAssetRegistryModule& AssetRegistryModule = FModuleManager::LoadModuleChecked<FAssetRegistryModule>(TEXT("AssetRegistry"));
+
+	//Variable setups
+	TArray<FAssetData> ObjectList;
+	FARFilter AssetFilter;
+
+	//Generating the list of assets.
+	//This list is being filtered by the game folder and class type.  The results are placed into the ObjectList variable.
+	AssetFilter.ClassNames.Add(Class->GetFName());
+	AssetFilter.PackagePaths.Add("/Game");
+	AssetFilter.bRecursiveClasses = bRecursiveClass;
+	AssetFilter.bRecursivePaths = true;
+	AssetRegistryModule.Get().GetAssets(AssetFilter, ObjectList);
+
+	//Loop through the list of assets, make their path full and a string, then add them to the test.
+	for (auto ObjIter = ObjectList.CreateConstIterator(); ObjIter; ++ObjIter)
+	{
+		const FAssetData& Asset = *ObjIter;
+		FString Filename = Asset.ObjectPath.ToString();
+		//convert to full paths
+		Filename = FPackageName::LongPackageNameToFilename(Filename);
+		if (FAutomationTestFramework::GetInstance().ShouldTestContent(Filename))
+		{
+			OutAssetList.Add(Asset.ObjectPath.ToString());
+		}
+	}
+}
\ No newline at end of file
diff --git a/Engine/Source/Editor/UnrealEd/Private/Tests/EditorOpenAssetAutomationTests.cpp b/Engine/Source/Editor/UnrealEd/Private/Tests/EditorOpenAssetAutomationTests.cpp
--- a/Engine/Source/Editor/UnrealEd/Private/Tests/EditorOpenAssetAutomationTests.cpp
+++ b/Engine/Source/Editor/UnrealEd/Private/Tests/EditorOpenAssetAutomationTests.cpp
@@ -7,6 +7,7 @@
 #include "AssertionMacros.h"
 #include "AutomationCommon.h"
 #include "AssetEditorManager.h"
+#include "ShaderCompiler.h"
 
 //Includes needed for opening certain assets
 #include "Materials/MaterialFunction.h"
@@ -14,7 +15,7 @@
 #include "BehaviorTree/BehaviorTree.h"
 #include "Engine/DestructibleMesh.h"
 
-
+#define ASSET_OPEN_INTERVAL 1.5f
 
 bool FOpenActualAssetEditors(const FString& Parameters)
 {
@@ -35,6 +36,7 @@ bool FOpenActualAssetEditors(const FString& Parameters)
 
 	return true;
 }
+
 //////////////////////////////////////////////////////////////////////////
 /**
 * Test to open the sub editor windows for a specified list of assets.
@@ -138,657 +140,664 @@ bool FOpenAssetEditors::RunTest(const FString& Parameters)
 }
 
 
+struct OpenAssetParameters
+{
+	TArray<FString> AssetPathList;
+	int32 AssetListLength;
+	int64 TimeSpent;
+
+	OpenAssetParameters()
+		: AssetListLength(0)
+		, TimeSpent(0)
+	{
+	}
+};
+
 //////////////////////////////////////////////////////////////////////////
 /**
-* This test opens each AimOffsetBlendSpace into its sub-editor.
+* Latent commands for Open Assets Type tests
 */
-IMPLEMENT_COMPLEX_AUTOMATION_TEST(FOpenAimOffsetBlendSpaceAssetEditors, "Tools.Skeletal Mesh.Animation.BlendSpaces.Open Aim Offset BlendSpace Assets", (EAutomationTestFlags::ATF_Editor | EAutomationTestFlags::ATF_RequiresUser));
+DEFINE_LATENT_AUTOMATION_COMMAND_ONE_PARAMETER(FOpenAllofAssetTypeCommand, OpenAssetParameters, AssetList);
 
-void FOpenAimOffsetBlendSpaceAssetEditors::GetTests(TArray<FString>& OutBeautifiedNames, TArray<FString>& OutTestCommands) const
+bool FOpenAllofAssetTypeCommand::Update()
 {
-	//This grabs each AimOffsetBlendSpace asset in the Game/Content directory
-	FEditorAutomationTestUtilities::CollectGameContentTestsByClass(UAimOffsetBlendSpace::StaticClass(), true, OutBeautifiedNames, OutTestCommands);
+	if ((FPlatformTime::Seconds() - AssetList.TimeSpent) > ASSET_OPEN_INTERVAL)
+	{
+		FAssetEditorManager::Get().CloseAllAssetEditors();
+
+		//Get all assets currently being tracked with open editors and make sure they are not opened.
+		if (FAssetEditorManager::Get().GetAllEditedAssets().Num() > 0)
+		{
+			UE_LOG(LogEditorAutomationTests, Warning, TEXT("Not all of the editors were closed."));
+		}
+
+		CollectGarbage(GARBAGE_COLLECTION_KEEPFLAGS);
+
+		UE_LOG(LogEditorAutomationTests, Log, TEXT("Remaining assets to open: %i"), AssetList.AssetPathList.Num() - AssetList.AssetListLength);
+
+		if (AssetList.AssetListLength < AssetList.AssetPathList.Num())
+		{
+			UObject* Object = StaticLoadObject(UObject::StaticClass(), NULL, *AssetList.AssetPathList[AssetList.AssetListLength]);
+			if (Object)
+			{
+				FAssetEditorManager::Get().OpenEditorForAsset(Object);
+				//This checks to see if the asset sub editor is loaded.
+				if (FAssetEditorManager::Get().FindEditorForAsset(Object, true) != NULL)
+				{
+					UE_LOG(LogEditorAutomationTests, Log, TEXT("Verified asset editor for: %s."), *AssetList.AssetPathList[AssetList.AssetListLength]);
+					UE_LOG(LogEditorAutomationTests, Display, TEXT("The editor successffully loaded for: %s."), *AssetList.AssetPathList[AssetList.AssetListLength]);
+					if (GShaderCompilingManager->IsCompiling())
+					{
+						UE_LOG(LogEditorAutomationTests, Log, TEXT("Waiting for %i shaders to finish."), GShaderCompilingManager->GetNumRemainingJobs());
+						GShaderCompilingManager->FinishAllCompilation();
+						UE_LOG(LogEditorAutomationTests, Log, TEXT("Done waiting for shaders to finish."));
+					}
+				}
+			}
+			else
+			{
+				UE_LOG(LogEditorAutomationTests, Error, TEXT("Failed to find object: %s."), *AssetList.AssetPathList[AssetList.AssetListLength]);
+			}
+			AssetList.AssetListLength++;
+			AssetList.TimeSpent = FPlatformTime::Seconds();
+			return false;
+		}
+		else
+		{
+			return true;
+		}
+	}
+		return false;
 }
 
-bool FOpenAimOffsetBlendSpaceAssetEditors::RunTest(const FString& Parameters)
+
+void FOpenAllAssetTests(OpenAssetParameters& AssetList)
 {
-	bool bDidTheTestPass = FOpenActualAssetEditors(Parameters);
-	CollectGarbage(GARBAGE_COLLECTION_KEEPFLAGS);
-	return bDidTheTestPass;
+	ADD_LATENT_AUTOMATION_COMMAND(FWaitLatentCommand(0.5f));
+	ADD_LATENT_AUTOMATION_COMMAND(FOpenAllofAssetTypeCommand(AssetList));
+	ADD_LATENT_AUTOMATION_COMMAND(FWaitLatentCommand(0.5f));
+	ADD_LATENT_AUTOMATION_COMMAND(FCloseAllAssetEditorsCommand());
 }
 
 //////////////////////////////////////////////////////////////////////////
 /**
-* This test opens each AimOffsetBlendSpace1D into its sub-editor.
+* This test opens each AIM OFFSET BLENDSPACE into its sub-editor.
 */
-IMPLEMENT_COMPLEX_AUTOMATION_TEST(FOpenAimOffsetBlendSpace1DAssetEditors, "Tools.Skeletal Mesh.Animation.BlendSpaces.Open Aim Offset BlendSpace1D Assets", (EAutomationTestFlags::ATF_Editor | EAutomationTestFlags::ATF_RequiresUser));
+IMPLEMENT_SIMPLE_AUTOMATION_TEST(FOpenAimOffsetBlendspaceAssetEditorsSLOW, "Tools.Skeletal Mesh.Open All Aim Offset Blendspace Assets (SLOW)", (EAutomationTestFlags::ATF_Editor | EAutomationTestFlags::ATF_RequiresUser));
 
-void FOpenAimOffsetBlendSpace1DAssetEditors::GetTests(TArray<FString>& OutBeautifiedNames, TArray<FString>& OutTestCommands) const
+bool FOpenAimOffsetBlendspaceAssetEditorsSLOW::RunTest(const FString& Parameters)
 {
-	//This grabs each AimOffsetBlendSpace1D asset in the Game/Content directory
-	FEditorAutomationTestUtilities::CollectGameContentTestsByClass(UAimOffsetBlendSpace1D::StaticClass(), true, OutBeautifiedNames, OutTestCommands);
-}
+	OpenAssetParameters AssetList;
+	AssetList.TimeSpent = FPlatformTime::Seconds();
 
-bool FOpenAimOffsetBlendSpace1DAssetEditors::RunTest(const FString& Parameters)
-{
-	bool bDidTheTestPass = FOpenActualAssetEditors(Parameters);
-	CollectGarbage(GARBAGE_COLLECTION_KEEPFLAGS);
-	return bDidTheTestPass;
+	FEditorAutomationTestUtilities::CollectGameContentByClass(UAimOffsetBlendSpace::StaticClass(), true, AssetList.AssetPathList);
+	FOpenAllAssetTests(AssetList);
+
+	return true;
 }
 
-//////////////////////////////////////////////////////////////////////////
 /**
-* This test opens each AnimMontage into its sub-editor.
+* This test opens each AIM OFFSET BLENDSPACE 1D into its sub-editor.
 */
-IMPLEMENT_COMPLEX_AUTOMATION_TEST(FOpenAnimMontageAssetEditors, "Tools.Skeletal Mesh.Animation.Open AnimMontage Assets", (EAutomationTestFlags::ATF_Editor | EAutomationTestFlags::ATF_RequiresUser));
+IMPLEMENT_SIMPLE_AUTOMATION_TEST(FOpenAimOffSetBlendspace1DAssetEditorsSLOW, "Tools.Skeletal Mesh.Open All Aim Offset Blendspace 1D Assets (SLOW)", (EAutomationTestFlags::ATF_Editor | EAutomationTestFlags::ATF_RequiresUser));
 
-void FOpenAnimMontageAssetEditors::GetTests(TArray<FString>& OutBeautifiedNames, TArray<FString>& OutTestCommands) const
+bool FOpenAimOffSetBlendspace1DAssetEditorsSLOW::RunTest(const FString& Parameters)
 {
-	//This grabs each AnimMontage asset in the Game/Content directory
-	FEditorAutomationTestUtilities::CollectGameContentTestsByClass(UAnimMontage::StaticClass(), true, OutBeautifiedNames, OutTestCommands);
-}
+	OpenAssetParameters AssetList;
+	AssetList.TimeSpent = FPlatformTime::Seconds();
 
-bool FOpenAnimMontageAssetEditors::RunTest(const FString& Parameters)
-{
-	bool bDidTheTestPass = FOpenActualAssetEditors(Parameters);
-	CollectGarbage(GARBAGE_COLLECTION_KEEPFLAGS);
-	return bDidTheTestPass;
+	FEditorAutomationTestUtilities::CollectGameContentByClass(UAimOffsetBlendSpace1D::StaticClass(), true, AssetList.AssetPathList);
+
+	FOpenAllAssetTests(AssetList);
+
+	return true;
 }
 
-//////////////////////////////////////////////////////////////////////////
 /**
-* This test opens each AnimSequence into its sub-editor.
+* This test opens each ANIMMONTAGE into its sub-editor.
 */
-IMPLEMENT_COMPLEX_AUTOMATION_TEST(FOpenAnimSequenceAssetEditors, "Tools.Skeletal Mesh.Animation.Open AnimSequence Assets", (EAutomationTestFlags::ATF_Editor | EAutomationTestFlags::ATF_RequiresUser));
+IMPLEMENT_SIMPLE_AUTOMATION_TEST(FOpenAnimMontageAssetEditorsSLOW, "Tools.Skeletal Mesh.Open All Anim Montage Assets (SLOW)", (EAutomationTestFlags::ATF_Editor | EAutomationTestFlags::ATF_RequiresUser));
 
-void FOpenAnimSequenceAssetEditors::GetTests(TArray<FString>& OutBeautifiedNames, TArray<FString>& OutTestCommands) const
+bool FOpenAnimMontageAssetEditorsSLOW::RunTest(const FString& Parameters)
 {
-	//This grabs each AnimSequence asset in the Game/Content directory
-	FEditorAutomationTestUtilities::CollectGameContentTestsByClass(UAnimSequence::StaticClass(), true, OutBeautifiedNames, OutTestCommands);
-}
+	OpenAssetParameters AssetList;
+	AssetList.TimeSpent = FPlatformTime::Seconds();
 
-bool FOpenAnimSequenceAssetEditors::RunTest(const FString& Parameters)
-{
-	bool bDidTheTestPass = FOpenActualAssetEditors(Parameters);
-	CollectGarbage(GARBAGE_COLLECTION_KEEPFLAGS);
-	return bDidTheTestPass;
+	FEditorAutomationTestUtilities::CollectGameContentByClass(UAnimMontage::StaticClass(), true, AssetList.AssetPathList);
+
+	FOpenAllAssetTests(AssetList);
+
+	return true;
 }
 
-//////////////////////////////////////////////////////////////////////////
 /**
-* This test opens each BEHAVIOR TREE into its sub-editor.
+* This test opens each ANIMSEQUENCE into its sub-editor.
 */
-IMPLEMENT_COMPLEX_AUTOMATION_TEST(FOpenBehaviorTreeAssetEditors, "Tools.AI.Open Behavior Tree Assets", (EAutomationTestFlags::ATF_Editor | EAutomationTestFlags::ATF_RequiresUser));
+IMPLEMENT_SIMPLE_AUTOMATION_TEST(FOpenAnimSequenceAssetEditorsSLOW, "Tools.Skeletal Mesh.Open All Anim Sequence Assets (SLOW)", (EAutomationTestFlags::ATF_Editor | EAutomationTestFlags::ATF_RequiresUser));
 
-void FOpenBehaviorTreeAssetEditors::GetTests(TArray<FString>& OutBeautifiedNames, TArray<FString>& OutTestCommands) const
+bool FOpenAnimSequenceAssetEditorsSLOW::RunTest(const FString& Parameters)
 {
-	//This grabs each Behavior Tree asset in the Game/Content directory
-	FEditorAutomationTestUtilities::CollectGameContentTestsByClass(UBehaviorTree::StaticClass(), true, OutBeautifiedNames, OutTestCommands);
-}
+	OpenAssetParameters AssetList;
+	AssetList.TimeSpent = FPlatformTime::Seconds();
 
-bool FOpenBehaviorTreeAssetEditors::RunTest(const FString& Parameters)
-{
-	bool bDidTheTestPass = FOpenActualAssetEditors(Parameters);
-	CollectGarbage(GARBAGE_COLLECTION_KEEPFLAGS);
-	return bDidTheTestPass;
+	FEditorAutomationTestUtilities::CollectGameContentByClass(UAnimSequence::StaticClass(), true, AssetList.AssetPathList);
+
+	FOpenAllAssetTests(AssetList);
+
+	return true;
 }
 
-//////////////////////////////////////////////////////////////////////////
 /**
-* This test opens each BlendSpace into its sub-editor.
+* This test opens each BEHAVIORTREE into its sub-editor.
 */
-IMPLEMENT_COMPLEX_AUTOMATION_TEST(FOpenBlendSpaceAssetEditors, "Tools.Skeletal Mesh.Animation.BlendSpaces.Open BlendSpace Assets", (EAutomationTestFlags::ATF_Editor | EAutomationTestFlags::ATF_RequiresUser));
+IMPLEMENT_SIMPLE_AUTOMATION_TEST(FOpenBehaviorTreeAssetEditorsSLOW, "Tools.AI.Open All Behaviour Tree Assets (SLOW)", (EAutomationTestFlags::ATF_Editor | EAutomationTestFlags::ATF_RequiresUser));
 
-void FOpenBlendSpaceAssetEditors::GetTests(TArray<FString>& OutBeautifiedNames, TArray<FString>& OutTestCommands) const
+bool FOpenBehaviorTreeAssetEditorsSLOW::RunTest(const FString& Parameters)
 {
-	//This grabs each BlendSpace asset in the Game/Content directory
-	FEditorAutomationTestUtilities::CollectGameContentTestsByClass(UBlendSpace::StaticClass(), false, OutBeautifiedNames, OutTestCommands);
-}
+	OpenAssetParameters AssetList;
+	AssetList.TimeSpent = FPlatformTime::Seconds();
 
-bool FOpenBlendSpaceAssetEditors::RunTest(const FString& Parameters)
-{
-	bool bDidTheTestPass = FOpenActualAssetEditors(Parameters);
-	CollectGarbage(GARBAGE_COLLECTION_KEEPFLAGS);
-	return bDidTheTestPass;
+	FEditorAutomationTestUtilities::CollectGameContentByClass(UBehaviorTree::StaticClass(), true, AssetList.AssetPathList);
+
+	FOpenAllAssetTests(AssetList);
+
+	return true;
 }
 
-//////////////////////////////////////////////////////////////////////////
 /**
-* This test opens each BlendSpace1D into its sub-editor.
+* This test opens each BLENDSPACE into its sub-editor.
 */
-IMPLEMENT_COMPLEX_AUTOMATION_TEST(FOpenBlendSpace1DAssetEditors, "Tools.Skeletal Mesh.Animation.BlendSpaces.Open BlendSpace1D Assets", (EAutomationTestFlags::ATF_Editor | EAutomationTestFlags::ATF_RequiresUser));
+IMPLEMENT_SIMPLE_AUTOMATION_TEST(FOpenBlendSpaceAssetEditorsSLOW, "Tools.Skeletal Mesh.Open All BlendSpace Assets (SLOW)", (EAutomationTestFlags::ATF_Editor | EAutomationTestFlags::ATF_RequiresUser));
 
-void FOpenBlendSpace1DAssetEditors::GetTests(TArray<FString>& OutBeautifiedNames, TArray<FString>& OutTestCommands) const
+bool FOpenBlendSpaceAssetEditorsSLOW::RunTest(const FString& Parameters)
 {
-	//This grabs each BlendSpace1D asset in the Game/Content directory
-	FEditorAutomationTestUtilities::CollectGameContentTestsByClass(UBlendSpace1D::StaticClass(), false, OutBeautifiedNames, OutTestCommands);
-}
+	OpenAssetParameters AssetList;
+	AssetList.TimeSpent = FPlatformTime::Seconds();
 
-bool FOpenBlendSpace1DAssetEditors::RunTest(const FString& Parameters)
-{
-	bool bDidTheTestPass = FOpenActualAssetEditors(Parameters);
-	CollectGarbage(GARBAGE_COLLECTION_KEEPFLAGS);
-	return bDidTheTestPass;
-}
+	FEditorAutomationTestUtilities::CollectGameContentByClass(UBlendSpace::StaticClass(), true, AssetList.AssetPathList);
 
+	FOpenAllAssetTests(AssetList);
+
+	return true;
+}
 
-//////////////////////////////////////////////////////////////////////////
 /**
-* This test opens each BLUEPRINT into its sub-editor.
+* This test opens each BLENDSPACE 1D into its sub-editor.
 */
-IMPLEMENT_COMPLEX_AUTOMATION_TEST(FOpenBlueprintAssetEditors, "Tools.Blueprint.Open Blueprint Assets", (EAutomationTestFlags::ATF_Editor | EAutomationTestFlags::ATF_RequiresUser));
+IMPLEMENT_SIMPLE_AUTOMATION_TEST(FOpenBlendSpace1DAssetEditorsSLOW, "Tools.Skeletal Mesh.Open All BlendSpace 1D Assets (SLOW)", (EAutomationTestFlags::ATF_Editor | EAutomationTestFlags::ATF_RequiresUser));
 
-void FOpenBlueprintAssetEditors::GetTests(TArray<FString>& OutBeautifiedNames, TArray<FString>& OutTestCommands) const
+bool FOpenBlendSpace1DAssetEditorsSLOW::RunTest(const FString& Parameters)
 {
-	//This grabs each BlendSpace asset in the Game/Content directory
-	FEditorAutomationTestUtilities::CollectGameContentTestsByClass(UBlueprint::StaticClass(), true, OutBeautifiedNames, OutTestCommands);
-}
+	OpenAssetParameters AssetList;
+	AssetList.TimeSpent = FPlatformTime::Seconds();
 
-bool FOpenBlueprintAssetEditors::RunTest(const FString& Parameters)
-{
-	bool bDidTheTestPass = FOpenActualAssetEditors(Parameters);
-	CollectGarbage(GARBAGE_COLLECTION_KEEPFLAGS);
-	return bDidTheTestPass;
+	FEditorAutomationTestUtilities::CollectGameContentByClass(UBlendSpace::StaticClass(), true, AssetList.AssetPathList);
+
+	FOpenAllAssetTests(AssetList);
+
+	return true;
 }
 
-//////////////////////////////////////////////////////////////////////////
 /**
-* This test opens each UserDefinedEnum into its sub-editor.
+* This test opens each BLUEPRINT into its sub-editor.
 */
-IMPLEMENT_COMPLEX_AUTOMATION_TEST(FOpenUserDefinedEnumAssetEditors, "Tools.Blueprint.Open User Defined Enum Assets", (EAutomationTestFlags::ATF_Editor | EAutomationTestFlags::ATF_RequiresUser));
+IMPLEMENT_SIMPLE_AUTOMATION_TEST(FOpenBlueprintAssetEditorsSLOW, "Tools.Blueprint.Open All Blueprint Assets (SLOW)", (EAutomationTestFlags::ATF_Editor | EAutomationTestFlags::ATF_RequiresUser));
 
-void FOpenUserDefinedEnumAssetEditors::GetTests(TArray<FString>& OutBeautifiedNames, TArray<FString>& OutTestCommands) const
+bool FOpenBlueprintAssetEditorsSLOW::RunTest(const FString& Parameters)
 {
-	//This grabs each UserDefinedEnum asset in the Game/Content directory
-	FEditorAutomationTestUtilities::CollectGameContentTestsByClass(UUserDefinedEnum::StaticClass(), true, OutBeautifiedNames, OutTestCommands);
-}
+	OpenAssetParameters AssetList;
+	AssetList.TimeSpent = FPlatformTime::Seconds();
 
-bool FOpenUserDefinedEnumAssetEditors::RunTest(const FString& Parameters)
-{
-	bool bDidTheTestPass = FOpenActualAssetEditors(Parameters);
-	CollectGarbage(GARBAGE_COLLECTION_KEEPFLAGS);
-	return bDidTheTestPass;
+	FEditorAutomationTestUtilities::CollectGameContentByClass(UBlueprint::StaticClass(), true, AssetList.AssetPathList);
+
+	FOpenAllAssetTests(AssetList);
+
+	return true;
 }
 
-//////////////////////////////////////////////////////////////////////////
 /**
 * This test opens each DESTRUCTIBLE MESH into its sub-editor.
 */
-IMPLEMENT_COMPLEX_AUTOMATION_TEST(FOpenDestructibleMeshAssetEditors, "Tools.Skeletal Mesh.Open Destructible Mesh Assets", (EAutomationTestFlags::ATF_Editor | EAutomationTestFlags::ATF_RequiresUser));
+IMPLEMENT_SIMPLE_AUTOMATION_TEST(FOpenDestructibleMeshAssetEditorsSLOW, "Tools.Skeletal Mesh.Open All Destructible Mesh Assets (SLOW)", (EAutomationTestFlags::ATF_Editor | EAutomationTestFlags::ATF_RequiresUser));
 
-void FOpenDestructibleMeshAssetEditors::GetTests(TArray<FString>& OutBeautifiedNames, TArray<FString>& OutTestCommands) const
+bool FOpenDestructibleMeshAssetEditorsSLOW::RunTest(const FString& Parameters)
 {
-	//This grabs each DestructibleMesh asset in the Game/Content directory
-	FEditorAutomationTestUtilities::CollectGameContentTestsByClass(UDestructibleMesh::StaticClass(), false, OutBeautifiedNames, OutTestCommands);
-}
+	OpenAssetParameters AssetList;
+	AssetList.TimeSpent = FPlatformTime::Seconds();
 
-bool FOpenDestructibleMeshAssetEditors::RunTest(const FString& Parameters)
-{
-	bool bDidTheTestPass = FOpenActualAssetEditors(Parameters);
-	CollectGarbage(GARBAGE_COLLECTION_KEEPFLAGS);
-	return bDidTheTestPass;
+	FEditorAutomationTestUtilities::CollectGameContentByClass(UDestructibleMesh::StaticClass(), false, AssetList.AssetPathList);
+
+	FOpenAllAssetTests(AssetList);
+
+	return true;
 }
 
-//////////////////////////////////////////////////////////////////////////
 /**
-* This test opens each DialogueVoice into its sub-editor.
+* This test opens each DIALOGUE VOICE into its sub-editor.
 */
-IMPLEMENT_COMPLEX_AUTOMATION_TEST(FOpenDialogueVoiceAssetEditors, "Tools.Sounds.Open Dialogue Voice Assets", (EAutomationTestFlags::ATF_Editor | EAutomationTestFlags::ATF_RequiresUser));
+IMPLEMENT_SIMPLE_AUTOMATION_TEST(FOpenDialogueVoiceAssetEditorsSLOW, "Tools.Sounds.Open All Dialogue Voice Assets (SLOW)", (EAutomationTestFlags::ATF_Editor | EAutomationTestFlags::ATF_RequiresUser));
 
-void FOpenDialogueVoiceAssetEditors::GetTests(TArray<FString>& OutBeautifiedNames, TArray<FString>& OutTestCommands) const
+bool FOpenDialogueVoiceAssetEditorsSLOW::RunTest(const FString& Parameters)
 {
-	//This grabs each DialogueVoice asset in the Game/Content directory
-	FEditorAutomationTestUtilities::CollectGameContentTestsByClass(UDialogueVoice::StaticClass(), true, OutBeautifiedNames, OutTestCommands);
-}
+	OpenAssetParameters AssetList;
+	AssetList.TimeSpent = FPlatformTime::Seconds();
 
-bool FOpenDialogueVoiceAssetEditors::RunTest(const FString& Parameters)
-{
-	bool bDidTheTestPass = FOpenActualAssetEditors(Parameters);
-	CollectGarbage(GARBAGE_COLLECTION_KEEPFLAGS);
-	return bDidTheTestPass;
+	FEditorAutomationTestUtilities::CollectGameContentByClass(UDialogueVoice::StaticClass(), true, AssetList.AssetPathList);
+
+	FOpenAllAssetTests(AssetList);
+
+	return true;
 }
 
-//////////////////////////////////////////////////////////////////////////
 /**
-* This test opens each DialogueWave into its sub-editor.
+* This test opens each DIALOGUE WAVE into its sub-editor.
 */
-IMPLEMENT_COMPLEX_AUTOMATION_TEST(FOpenDialogueWaveAssetEditors, "Tools.Sounds.Open Dialogue Wave Assets", (EAutomationTestFlags::ATF_Editor | EAutomationTestFlags::ATF_RequiresUser));
+IMPLEMENT_SIMPLE_AUTOMATION_TEST(FOpenDialogueWaveAssetEditorsSLOW, "Tools.Sounds.Open All Dialogue Wave Assets (SLOW)", (EAutomationTestFlags::ATF_Editor | EAutomationTestFlags::ATF_RequiresUser));
 
-void FOpenDialogueWaveAssetEditors::GetTests(TArray<FString>& OutBeautifiedNames, TArray<FString>& OutTestCommands) const
+bool FOpenDialogueWaveAssetEditorsSLOW::RunTest(const FString& Parameters)
 {
-	//This grabs each DialogueWave asset in the Game/Content directory
-	FEditorAutomationTestUtilities::CollectGameContentTestsByClass(UDialogueWave::StaticClass(), true, OutBeautifiedNames, OutTestCommands);
-}
+	OpenAssetParameters AssetList;
+	AssetList.TimeSpent = FPlatformTime::Seconds();
 
-bool FOpenDialogueWaveAssetEditors::RunTest(const FString& Parameters)
-{
-	bool bDidTheTestPass = FOpenActualAssetEditors(Parameters);
-	CollectGarbage(GARBAGE_COLLECTION_KEEPFLAGS);
-	return bDidTheTestPass;
+	FEditorAutomationTestUtilities::CollectGameContentByClass(UDialogueWave::StaticClass(), true, AssetList.AssetPathList);
+
+	FOpenAllAssetTests(AssetList);
+
+	return true;
 }
 
-//////////////////////////////////////////////////////////////////////////
 /**
-* This test opens each Material into its sub-editor.
+* This test opens each MATERIAL into its sub-editor.
 */
-IMPLEMENT_COMPLEX_AUTOMATION_TEST(FOpenMaterialAssetEditors, "Tools.Materials.Open Material Assets", (EAutomationTestFlags::ATF_Editor | EAutomationTestFlags::ATF_RequiresUser));
+IMPLEMENT_SIMPLE_AUTOMATION_TEST(FOpenMaterialAssetEditorsSLOW, "Tools.Materials.Open All Material Assets (SLOW)", (EAutomationTestFlags::ATF_Editor | EAutomationTestFlags::ATF_RequiresUser));
 
-void FOpenMaterialAssetEditors::GetTests(TArray<FString>& OutBeautifiedNames, TArray<FString>& OutTestCommands) const
+bool FOpenMaterialAssetEditorsSLOW::RunTest(const FString& Parameters)
 {
-	//This grabs each Material asset in the Game/Content directory
-	FEditorAutomationTestUtilities::CollectGameContentTestsByClass(UMaterial::StaticClass(), true, OutBeautifiedNames, OutTestCommands);
-}
+	OpenAssetParameters AssetList;
+	AssetList.TimeSpent = FPlatformTime::Seconds();
 
-bool FOpenMaterialAssetEditors::RunTest(const FString& Parameters)
-{
-	bool bDidTheTestPass = FOpenActualAssetEditors(Parameters);
-	CollectGarbage(GARBAGE_COLLECTION_KEEPFLAGS);
-	return bDidTheTestPass;
-}
+	FEditorAutomationTestUtilities::CollectGameContentByClass(UMaterial::StaticClass(), true, AssetList.AssetPathList);
 
-////////////////////////////////////////////////////////////////////////////
-///**
-//* This test opens each MaterialFunction into its sub-editor.
-//*/
-IMPLEMENT_COMPLEX_AUTOMATION_TEST(FOpenMaterialFunctionAssetEditors, "Tools.Materials.Open Material Function Assets", (EAutomationTestFlags::ATF_Editor | EAutomationTestFlags::ATF_RequiresUser));
+	FOpenAllAssetTests(AssetList);
 
-void FOpenMaterialFunctionAssetEditors::GetTests(TArray<FString>& OutBeautifiedNames, TArray<FString>& OutTestCommands) const
-{
-	//This grabs each MaterialFunction asset in the Game/Content directory
-	FEditorAutomationTestUtilities::CollectGameContentTestsByClass(UMaterialFunction::StaticClass(), true, OutBeautifiedNames, OutTestCommands);
+	return true;
 }
 
-bool FOpenMaterialFunctionAssetEditors::RunTest(const FString& Parameters)
+/**
+* This test opens each MATERIAL FUNCTION into its sub-editor.
+*/
+IMPLEMENT_SIMPLE_AUTOMATION_TEST(FOpenMaterialFunctionAssetEditorsSLOW, "Tools.Materials.Open All Material Function Assets (SLOW)", (EAutomationTestFlags::ATF_Editor | EAutomationTestFlags::ATF_RequiresUser));
+
+bool FOpenMaterialFunctionAssetEditorsSLOW::RunTest(const FString& Parameters)
 {
-	bool bDidTheTestPass = FOpenActualAssetEditors(Parameters);
-	CollectGarbage(GARBAGE_COLLECTION_KEEPFLAGS);
-	return bDidTheTestPass;
+	OpenAssetParameters AssetList;
+	AssetList.TimeSpent = FPlatformTime::Seconds();
+
+	FEditorAutomationTestUtilities::CollectGameContentByClass(UMaterialFunction::StaticClass(), true, AssetList.AssetPathList);
+
+	FOpenAllAssetTests(AssetList);
+
+	return true;
 }
 
-//////////////////////////////////////////////////////////////////////////
 /**
-* This test opens each MaterialInstanceConstant into its sub-editor.
+* This test opens each MATERIAL INSTANCE CONSTANT into its sub-editor.
 */
-IMPLEMENT_COMPLEX_AUTOMATION_TEST(FOpenMaterialInstanceConstantAssetEditors, "Tools.Materials.Open Material Instance Constant Assets", (EAutomationTestFlags::ATF_Editor | EAutomationTestFlags::ATF_RequiresUser));
+IMPLEMENT_SIMPLE_AUTOMATION_TEST(FOpenMICAssetEditorsSLOW, "Tools.Materials.Open All Material Instance Constant Assets (SLOW)", (EAutomationTestFlags::ATF_Editor | EAutomationTestFlags::ATF_RequiresUser));
 
-void FOpenMaterialInstanceConstantAssetEditors::GetTests(TArray<FString>& OutBeautifiedNames, TArray<FString>& OutTestCommands) const
+bool FOpenMICAssetEditorsSLOW::RunTest(const FString& Parameters)
 {
-	//This grabs each MaterialInstanceConstant asset in the Game/Content directory
-	FEditorAutomationTestUtilities::CollectGameContentTestsByClass(UMaterialInstanceConstant::StaticClass(), true, OutBeautifiedNames, OutTestCommands);
-}
+	OpenAssetParameters AssetList;
+	AssetList.TimeSpent = FPlatformTime::Seconds();
 
-bool FOpenMaterialInstanceConstantAssetEditors::RunTest(const FString& Parameters)
-{
-	bool bDidTheTestPass = FOpenActualAssetEditors(Parameters);
-	CollectGarbage(GARBAGE_COLLECTION_KEEPFLAGS);
-	return bDidTheTestPass;
+	FEditorAutomationTestUtilities::CollectGameContentByClass(UMaterialInstanceConstant::StaticClass(), true, AssetList.AssetPathList);
+
+	FOpenAllAssetTests(AssetList);
+
+	return true;
 }
 
-//////////////////////////////////////////////////////////////////////////
 /**
-* This test opens each MaterialParameterCollection into its sub-editor.
+* This test opens each MATERIAL PARAMETER COLLECTION into its sub-editor.
 */
-IMPLEMENT_COMPLEX_AUTOMATION_TEST(FOpenMaterialParameterCollectionAssetEditors, "Tools.Materials.Open Material Parameter Collection Assets", (EAutomationTestFlags::ATF_Editor | EAutomationTestFlags::ATF_RequiresUser));
+IMPLEMENT_SIMPLE_AUTOMATION_TEST(FOpenMaterialParameterCollectionAssetEditorsSLOW, "Tools.Materials.Open All Material Parameter Collection Assets (SLOW)", (EAutomationTestFlags::ATF_Editor | EAutomationTestFlags::ATF_RequiresUser));
 
-void FOpenMaterialParameterCollectionAssetEditors::GetTests(TArray<FString>& OutBeautifiedNames, TArray<FString>& OutTestCommands) const
+bool FOpenMaterialParameterCollectionAssetEditorsSLOW::RunTest(const FString& Parameters)
 {
-	//This grabs each MaterialParameterCollection asset in the Game/Content directory
-	FEditorAutomationTestUtilities::CollectGameContentTestsByClass(UMaterialParameterCollection::StaticClass(), true, OutBeautifiedNames, OutTestCommands);
-}
+	OpenAssetParameters AssetList;
+	AssetList.TimeSpent = FPlatformTime::Seconds();
 
-bool FOpenMaterialParameterCollectionAssetEditors::RunTest(const FString& Parameters)
-{
-	bool bDidTheTestPass = FOpenActualAssetEditors(Parameters);
-	CollectGarbage(GARBAGE_COLLECTION_KEEPFLAGS);
-	return bDidTheTestPass;
+	FEditorAutomationTestUtilities::CollectGameContentByClass(UMaterialParameterCollection::StaticClass(), true, AssetList.AssetPathList);
+
+	FOpenAllAssetTests(AssetList);
+
+	return true;
 }
 
-//////////////////////////////////////////////////////////////////////////
 /**
 * This test opens each PARTICLE SYSTEM into its sub-editor.
 */
-IMPLEMENT_COMPLEX_AUTOMATION_TEST(FOpenParticleSystemAssetEditors, "Tools.Particle System.Open Particle System Assets", (EAutomationTestFlags::ATF_Editor | EAutomationTestFlags::ATF_RequiresUser));
+IMPLEMENT_SIMPLE_AUTOMATION_TEST(FOpenParticleSystemAssetEditorsSLOW, "Tools.Particle System.Open All Particle System Assets (SLOW)", (EAutomationTestFlags::ATF_Editor | EAutomationTestFlags::ATF_RequiresUser));
 
-void FOpenParticleSystemAssetEditors::GetTests(TArray<FString>& OutBeautifiedNames, TArray<FString>& OutTestCommands) const
+bool FOpenParticleSystemAssetEditorsSLOW::RunTest(const FString& Parameters)
 {
-	//This grabs each ParticleSystem asset in the Game/Content directory
-	FEditorAutomationTestUtilities::CollectGameContentTestsByClass(UParticleSystem::StaticClass(), true, OutBeautifiedNames, OutTestCommands);
-}
+	OpenAssetParameters AssetList;
+	AssetList.TimeSpent = FPlatformTime::Seconds();
 
-bool FOpenParticleSystemAssetEditors::RunTest(const FString& Parameters)
-{
-	bool bDidTheTestPass = FOpenActualAssetEditors(Parameters);
-	CollectGarbage(GARBAGE_COLLECTION_KEEPFLAGS);
-	return bDidTheTestPass;
+	FEditorAutomationTestUtilities::CollectGameContentByClass(UParticleSystem::StaticClass(), true, AssetList.AssetPathList);
+
+	FOpenAllAssetTests(AssetList);
+
+	return true;
 }
 
-//////////////////////////////////////////////////////////////////////////
 /**
-* This test opens each PhysicalMaterial into its sub-editor.
+* This test opens each PHYSICAL MATERIAL into its sub-editor.
 */
-IMPLEMENT_COMPLEX_AUTOMATION_TEST(FOpenPhysicalMaterialAssetEditors, "Tools.Physics.Open Physical Material Assets", (EAutomationTestFlags::ATF_Editor | EAutomationTestFlags::ATF_RequiresUser));
+IMPLEMENT_SIMPLE_AUTOMATION_TEST(FOpenPhysicalMaterialAssetEditorsSLOW, "Tools.Physics.Open All Physical Material Assets (SLOW)", (EAutomationTestFlags::ATF_Editor | EAutomationTestFlags::ATF_RequiresUser));
 
-void FOpenPhysicalMaterialAssetEditors::GetTests(TArray<FString>& OutBeautifiedNames, TArray<FString>& OutTestCommands) const
+bool FOpenPhysicalMaterialAssetEditorsSLOW::RunTest(const FString& Parameters)
 {
-	//This grabs each PhysicalMaterial asset in the Game/Content directory
-	FEditorAutomationTestUtilities::CollectGameContentTestsByClass(UPhysicalMaterial::StaticClass(), true, OutBeautifiedNames, OutTestCommands);
-}
+	OpenAssetParameters AssetList;
+	AssetList.TimeSpent = FPlatformTime::Seconds();
 
-bool FOpenPhysicalMaterialAssetEditors::RunTest(const FString& Parameters)
-{
-	bool bDidTheTestPass = FOpenActualAssetEditors(Parameters);
-	CollectGarbage(GARBAGE_COLLECTION_KEEPFLAGS);
-	return bDidTheTestPass;
-}
+	FEditorAutomationTestUtilities::CollectGameContentByClass(UPhysicalMaterial::StaticClass(), true, AssetList.AssetPathList);
 
+	FOpenAllAssetTests(AssetList);
+
+	return true;
+}
 
-//////////////////////////////////////////////////////////////////////////
 /**
-* This test opens each PhysicsAsset into its sub-editor.
+* This test opens each PHYSICS ASSET into its sub-editor.
 */
-IMPLEMENT_COMPLEX_AUTOMATION_TEST(FOpenPhysicsAssetAssetEditors, "Tools.Physics.Open Physics Assets", (EAutomationTestFlags::ATF_Editor | EAutomationTestFlags::ATF_RequiresUser));
+IMPLEMENT_SIMPLE_AUTOMATION_TEST(FOpenPhysicAssetEditorsSLOW, "Tools.Physics.Open All Physic Assets (SLOW)", (EAutomationTestFlags::ATF_Editor | EAutomationTestFlags::ATF_RequiresUser));
 
-void FOpenPhysicsAssetAssetEditors::GetTests(TArray<FString>& OutBeautifiedNames, TArray<FString>& OutTestCommands) const
+bool FOpenPhysicAssetEditorsSLOW::RunTest(const FString& Parameters)
 {
-	//This grabs each PhysicsAsset asset in the Game/Content directory
-	FEditorAutomationTestUtilities::CollectGameContentTestsByClass(UPhysicsAsset::StaticClass(), true, OutBeautifiedNames, OutTestCommands);
-}
+	OpenAssetParameters AssetList;
+	AssetList.TimeSpent = FPlatformTime::Seconds();
 
-bool FOpenPhysicsAssetAssetEditors::RunTest(const FString& Parameters)
-{
-	bool bDidTheTestPass = FOpenActualAssetEditors(Parameters);
-	CollectGarbage(GARBAGE_COLLECTION_KEEPFLAGS);
-	return bDidTheTestPass;
-}
+	FEditorAutomationTestUtilities::CollectGameContentByClass(UPhysicsAsset::StaticClass(), true, AssetList.AssetPathList);
 
+	FOpenAllAssetTests(AssetList);
+
+	return true;
+}
 
-//////////////////////////////////////////////////////////////////////////
 /**
-* This test opens each ReverbEffect into its sub-editor.
+* This test opens each REVERB EFFECT into its sub-editor.
 */
-IMPLEMENT_COMPLEX_AUTOMATION_TEST(FOpenReverbEffectAssetEditors, "Tools.Sounds.Open Reverb Effect Assets", (EAutomationTestFlags::ATF_Editor | EAutomationTestFlags::ATF_RequiresUser));
+IMPLEMENT_SIMPLE_AUTOMATION_TEST(FOpenReverbEffectAssetEditorsSLOW, "Tools.Sounds.Open All Reverb Effect Assets (SLOW)", (EAutomationTestFlags::ATF_Editor | EAutomationTestFlags::ATF_RequiresUser));
 
-void FOpenReverbEffectAssetEditors::GetTests(TArray<FString>& OutBeautifiedNames, TArray<FString>& OutTestCommands) const
+bool FOpenReverbEffectAssetEditorsSLOW::RunTest(const FString& Parameters)
 {
-	//This grabs each ReverbEffect asset in the Game/Content directory
-	FEditorAutomationTestUtilities::CollectGameContentTestsByClass(UReverbEffect::StaticClass(), true, OutBeautifiedNames, OutTestCommands);
-}
+	OpenAssetParameters AssetList;
+	AssetList.TimeSpent = FPlatformTime::Seconds();
 
-bool FOpenReverbEffectAssetEditors::RunTest(const FString& Parameters)
-{
-	bool bDidTheTestPass = FOpenActualAssetEditors(Parameters);
-	CollectGarbage(GARBAGE_COLLECTION_KEEPFLAGS);
-	return bDidTheTestPass;
-}
+	FEditorAutomationTestUtilities::CollectGameContentByClass(UReverbEffect::StaticClass(), true, AssetList.AssetPathList);
 
+	FOpenAllAssetTests(AssetList);
 
+	return true;
+}
 
-//////////////////////////////////////////////////////////////////////////
 /**
-* This test opens each SlateWidgetStyleAsset into its sub-editor.
+* This test opens each SLATE WIDGET STYLE ASSET into its sub-editor.
 */
-IMPLEMENT_COMPLEX_AUTOMATION_TEST(FOpenSlateWidgetStyleAssetAssetEditors, "Tools.Slate.Open Slate Widget Style Assets", (EAutomationTestFlags::ATF_Editor | EAutomationTestFlags::ATF_RequiresUser));
+IMPLEMENT_SIMPLE_AUTOMATION_TEST(FOpenSlateWidgetStyleAssetEditorsSLOW, "Tools.Slate.Open All Slate Widget Style Assets (SLOW)", (EAutomationTestFlags::ATF_Editor | EAutomationTestFlags::ATF_RequiresUser));
 
-void FOpenSlateWidgetStyleAssetAssetEditors::GetTests(TArray<FString>& OutBeautifiedNames, TArray<FString>& OutTestCommands) const
+bool FOpenSlateWidgetStyleAssetEditorsSLOW::RunTest(const FString& Parameters)
 {
-	//This grabs each SlateWidgetStyleAsset asset in the Game/Content directory
-	FEditorAutomationTestUtilities::CollectGameContentTestsByClass(USlateWidgetStyleAsset::StaticClass(), true, OutBeautifiedNames, OutTestCommands);
-}
+	OpenAssetParameters AssetList;
+	AssetList.TimeSpent = FPlatformTime::Seconds();
 
-bool FOpenSlateWidgetStyleAssetAssetEditors::RunTest(const FString& Parameters)
-{
-	bool bDidTheTestPass = FOpenActualAssetEditors(Parameters);
-	CollectGarbage(GARBAGE_COLLECTION_KEEPFLAGS);
-	return bDidTheTestPass;
+	FEditorAutomationTestUtilities::CollectGameContentByClass(USlateWidgetStyleAsset::StaticClass(), true, AssetList.AssetPathList);
+
+	FOpenAllAssetTests(AssetList);
+
+	return true;
 }
 
-//////////////////////////////////////////////////////////////////////////
 /**
-* This test opens each Slate Brush Asset into its sub-editor.
+* This test opens each SLATE BRUSH ASSET into its sub-editor.
 */
-IMPLEMENT_COMPLEX_AUTOMATION_TEST(FOpenSlateBrushAssetAssetEditors, "Tools.Slate.Open Slate Brush Assets", (EAutomationTestFlags::ATF_Editor | EAutomationTestFlags::ATF_RequiresUser));
+IMPLEMENT_SIMPLE_AUTOMATION_TEST(FOpenSlateBrushAssetEditorsSLOW, "Tools.Slate.Open All Slate Brush Assets (SLOW)", (EAutomationTestFlags::ATF_Editor | EAutomationTestFlags::ATF_RequiresUser));
 
-void FOpenSlateBrushAssetAssetEditors::GetTests(TArray<FString>& OutBeautifiedNames, TArray<FString>& OutTestCommands) const
+bool FOpenSlateBrushAssetEditorsSLOW::RunTest(const FString& Parameters)
 {
-	//This grabs each SlateBrushAsset asset in the Game/Content directory
-	FEditorAutomationTestUtilities::CollectGameContentTestsByClass(USlateBrushAsset::StaticClass(), true, OutBeautifiedNames, OutTestCommands);
-}
+	OpenAssetParameters AssetList;
+	AssetList.TimeSpent = FPlatformTime::Seconds();
 
-bool FOpenSlateBrushAssetAssetEditors::RunTest(const FString& Parameters)
-{
-	bool bDidTheTestPass = FOpenActualAssetEditors(Parameters);
-	CollectGarbage(GARBAGE_COLLECTION_KEEPFLAGS);
-	return bDidTheTestPass;
+	FEditorAutomationTestUtilities::CollectGameContentByClass(USlateBrushAsset::StaticClass(), true, AssetList.AssetPathList);
+
+	FOpenAllAssetTests(AssetList);
+
+	return true;
 }
-//////////////////////////////////////////////////////////////////////////
+
 /**
-* This test opens each SoundAttenuation into its sub-editor.
+* This test opens each SOUNDATTENUATION into its sub-editor.
 */
-IMPLEMENT_COMPLEX_AUTOMATION_TEST(FOpenSoundAttenuationAssetEditors, "Tools.Sounds.Open Sound Attenuation Assets", (EAutomationTestFlags::ATF_Editor | EAutomationTestFlags::ATF_RequiresUser));
+IMPLEMENT_SIMPLE_AUTOMATION_TEST(FOpenSoundAttenuationAssetEditorsSLOW, "Tools.Sounds.Open All Sound Attenuation Assets (SLOW)", (EAutomationTestFlags::ATF_Editor | EAutomationTestFlags::ATF_RequiresUser));
 
-void FOpenSoundAttenuationAssetEditors::GetTests(TArray<FString>& OutBeautifiedNames, TArray<FString>& OutTestCommands) const
+bool FOpenSoundAttenuationAssetEditorsSLOW::RunTest(const FString& Parameters)
 {
-	//This grabs each SoundAttenuation asset in the Game/Content directory
-	FEditorAutomationTestUtilities::CollectGameContentTestsByClass(USoundAttenuation::StaticClass(), true, OutBeautifiedNames, OutTestCommands);
-}
+	OpenAssetParameters AssetList;
+	AssetList.TimeSpent = FPlatformTime::Seconds();
 
-bool FOpenSoundAttenuationAssetEditors::RunTest(const FString& Parameters)
-{
-	bool bDidTheTestPass = FOpenActualAssetEditors(Parameters);
-	CollectGarbage(GARBAGE_COLLECTION_KEEPFLAGS);
-	return bDidTheTestPass;
+	FEditorAutomationTestUtilities::CollectGameContentByClass(USoundAttenuation::StaticClass(), true, AssetList.AssetPathList);
+
+	FOpenAllAssetTests(AssetList);
+
+	return true;
 }
 
-//////////////////////////////////////////////////////////////////////////
 /**
-* This test opens each Static Mesh into their sub-editor.
+* This test opens each STATIC MESH into its sub-editor.
 */
-IMPLEMENT_COMPLEX_AUTOMATION_TEST(FOpenStaticMeshAssetEditors, "Tools.Static Mesh.Open Static Mesh Assets", (EAutomationTestFlags::ATF_Editor | EAutomationTestFlags::ATF_RequiresUser));
+IMPLEMENT_SIMPLE_AUTOMATION_TEST(FOpenStaticMeshAssetEditorsSLOW, "Tools.Static Mesh.Open All Static Mesh Assets (SLOW)", (EAutomationTestFlags::ATF_Editor | EAutomationTestFlags::ATF_RequiresUser));
 
-void FOpenStaticMeshAssetEditors::GetTests(TArray<FString>& OutBeautifiedNames, TArray<FString>& OutTestCommands) const
+bool FOpenStaticMeshAssetEditorsSLOW::RunTest(const FString& Parameters)
 {
-	//This grabs each Static Mesh in the Game/Content
-	FEditorAutomationTestUtilities::CollectGameContentTestsByClass(UStaticMesh::StaticClass(), true, OutBeautifiedNames, OutTestCommands);
-}
+	OpenAssetParameters AssetList;
+	AssetList.TimeSpent = FPlatformTime::Seconds();
 
-bool FOpenStaticMeshAssetEditors::RunTest(const FString& Parameters)
-{
-	bool bDidTheTestPass = FOpenActualAssetEditors(Parameters);
-	CollectGarbage(GARBAGE_COLLECTION_KEEPFLAGS);
-	return bDidTheTestPass;
+	FEditorAutomationTestUtilities::CollectGameContentByClass(UStaticMesh::StaticClass(), true, AssetList.AssetPathList);
+
+	FOpenAllAssetTests(AssetList);
+
+	return true;
 }
 
-//////////////////////////////////////////////////////////////////////////
 /**
-* This test opens each Skeletal Mesh into their sub-editor.
+* This test opens each SKELETAL MESH into its sub-editor.
 */
-IMPLEMENT_COMPLEX_AUTOMATION_TEST(FOpenSkeletalMeshAssetEditors, "Tools.Skeletal Mesh.Open Skeletal Mesh Assets", (EAutomationTestFlags::ATF_Editor | EAutomationTestFlags::ATF_RequiresUser));
+IMPLEMENT_SIMPLE_AUTOMATION_TEST(FOpenSkeletalMeshAssetEditorsSLOW, "Tools.Skeletal Mesh.Open All Skeletal Mesh Assets (SLOW)", (EAutomationTestFlags::ATF_Editor | EAutomationTestFlags::ATF_RequiresUser));
 
-void FOpenSkeletalMeshAssetEditors::GetTests(TArray<FString>& OutBeautifiedNames, TArray<FString>& OutTestCommands) const
+bool FOpenSkeletalMeshAssetEditorsSLOW::RunTest(const FString& Parameters)
 {
-	//This grabs each SKeletal Mesh in the Game/Content
-	FEditorAutomationTestUtilities::CollectGameContentTestsByClass(USkeletalMesh::StaticClass(), false, OutBeautifiedNames, OutTestCommands);
-}
+	OpenAssetParameters AssetList;
+	AssetList.TimeSpent = FPlatformTime::Seconds();
 
-bool FOpenSkeletalMeshAssetEditors::RunTest(const FString& Parameters)
-{
-	bool bDidTheTestPass = FOpenActualAssetEditors(Parameters);
-	CollectGarbage(GARBAGE_COLLECTION_KEEPFLAGS);
-	return bDidTheTestPass;
+	FEditorAutomationTestUtilities::CollectGameContentByClass(USkeletalMesh::StaticClass(), false, AssetList.AssetPathList);
+
+	FOpenAllAssetTests(AssetList);
+
+	return true;
 }
 
-//////////////////////////////////////////////////////////////////////////
 /**
-* This test opens each Skeleton into its sub-editor.
+* This test opens each SKELETON into its sub-editor.
 */
-IMPLEMENT_COMPLEX_AUTOMATION_TEST(FOpenSkeletonAssetEditors, "Tools.Skeletal Mesh.Open Skeleton Assets", (EAutomationTestFlags::ATF_Editor | EAutomationTestFlags::ATF_RequiresUser));
+IMPLEMENT_SIMPLE_AUTOMATION_TEST(FOpenSkeletonAssetEditorsSLOW, "Tools.Skeletal Mesh.Open All Skeleton Assets (SLOW)", (EAutomationTestFlags::ATF_Editor | EAutomationTestFlags::ATF_RequiresUser));
 
-void FOpenSkeletonAssetEditors::GetTests(TArray<FString>& OutBeautifiedNames, TArray<FString>& OutTestCommands) const
+bool FOpenSkeletonAssetEditorsSLOW::RunTest(const FString& Parameters)
 {
-	//This grabs each Skeleton asset in the Game/Content directory
-	FEditorAutomationTestUtilities::CollectGameContentTestsByClass(USkeleton::StaticClass(), true, OutBeautifiedNames, OutTestCommands);
-}
+	OpenAssetParameters AssetList;
+	AssetList.TimeSpent = FPlatformTime::Seconds();
 
-bool FOpenSkeletonAssetEditors::RunTest(const FString& Parameters)
-{
-	bool bDidTheTestPass = FOpenActualAssetEditors(Parameters);
-	CollectGarbage(GARBAGE_COLLECTION_KEEPFLAGS);
-	return bDidTheTestPass;
+	FEditorAutomationTestUtilities::CollectGameContentByClass(USkeleton::StaticClass(), true, AssetList.AssetPathList);
+
+	FOpenAllAssetTests(AssetList);
+
+	return true;
 }
 
-//////////////////////////////////////////////////////////////////////////
 /**
-* This test opens each SoundClass into its sub-editor.
+* This test opens each SOUND CLASS into its sub-editor.
 */
-IMPLEMENT_COMPLEX_AUTOMATION_TEST(FOpenSoundClassAssetEditors, "Tools.Sounds.Open SoundClass Assets", (EAutomationTestFlags::ATF_Editor | EAutomationTestFlags::ATF_RequiresUser));
+IMPLEMENT_SIMPLE_AUTOMATION_TEST(FOpenSoundClassAssetEditorsSLOW, "Tools.Sounds.Open All Sound Class Assets (SLOW)", (EAutomationTestFlags::ATF_Editor | EAutomationTestFlags::ATF_RequiresUser));
 
-void FOpenSoundClassAssetEditors::GetTests(TArray<FString>& OutBeautifiedNames, TArray<FString>& OutTestCommands) const
+bool FOpenSoundClassAssetEditorsSLOW::RunTest(const FString& Parameters)
 {
-	//This grabs each SoundClass asset in the Game/Content directory
-	FEditorAutomationTestUtilities::CollectGameContentTestsByClass(USoundClass::StaticClass(), true, OutBeautifiedNames, OutTestCommands);
-}
+	OpenAssetParameters AssetList;
+	AssetList.TimeSpent = FPlatformTime::Seconds();
 
-bool FOpenSoundClassAssetEditors::RunTest(const FString& Parameters)
-{
-	bool bDidTheTestPass = FOpenActualAssetEditors(Parameters);
-	CollectGarbage(GARBAGE_COLLECTION_KEEPFLAGS);
-	return bDidTheTestPass;
-}
+	FEditorAutomationTestUtilities::CollectGameContentByClass(USoundClass::StaticClass(), true, AssetList.AssetPathList);
 
+	FOpenAllAssetTests(AssetList);
+
+	return true;
+}
 
-//////////////////////////////////////////////////////////////////////////
 /**
-* This test opens each SoundCue into its sub-editor.
+* This test opens each SOUND CUE into its sub-editor.
 */
-IMPLEMENT_COMPLEX_AUTOMATION_TEST(FOpenSoundCueAssetEditors, "Tools.Sounds.Open Sound Cue Assets", (EAutomationTestFlags::ATF_Editor | EAutomationTestFlags::ATF_RequiresUser));
+IMPLEMENT_SIMPLE_AUTOMATION_TEST(FOpenSoundCueAssetEditorsSLOW, "Tools.Sounds.Open All Sound Cue Assets (SLOW)", (EAutomationTestFlags::ATF_Editor | EAutomationTestFlags::ATF_RequiresUser));
 
-void FOpenSoundCueAssetEditors::GetTests(TArray<FString>& OutBeautifiedNames, TArray<FString>& OutTestCommands) const
+bool FOpenSoundCueAssetEditorsSLOW::RunTest(const FString& Parameters)
 {
-	//This grabs each SoundCue asset in the Game/Content directory
-	FEditorAutomationTestUtilities::CollectGameContentTestsByClass(USoundCue::StaticClass(), true, OutBeautifiedNames, OutTestCommands);
-}
+	OpenAssetParameters AssetList;
+	AssetList.TimeSpent = FPlatformTime::Seconds();
 
-bool FOpenSoundCueAssetEditors::RunTest(const FString& Parameters)
-{
-	bool bDidTheTestPass = FOpenActualAssetEditors(Parameters);
-	CollectGarbage(GARBAGE_COLLECTION_KEEPFLAGS);
-	return bDidTheTestPass;
-}
+	FEditorAutomationTestUtilities::CollectGameContentByClass(USoundCue::StaticClass(), true, AssetList.AssetPathList);
 
+	FOpenAllAssetTests(AssetList);
+
+	return true;
+}
 
-//////////////////////////////////////////////////////////////////////////
 /**
-* This test opens each SoundMix into its sub-editor.
+* This test opens each SOUND MIX into its sub-editor.
 */
-IMPLEMENT_COMPLEX_AUTOMATION_TEST(FOpenSoundMixAssetEditors, "Tools.Sounds.Open Sound Mix Assets", (EAutomationTestFlags::ATF_Editor | EAutomationTestFlags::ATF_RequiresUser));
+IMPLEMENT_SIMPLE_AUTOMATION_TEST(FOpenSoundMixAssetEditorsSLOW, "Tools.Sounds.Open All Sound Mix Assets (SLOW)", (EAutomationTestFlags::ATF_Editor | EAutomationTestFlags::ATF_RequiresUser));
 
-void FOpenSoundMixAssetEditors::GetTests(TArray<FString>& OutBeautifiedNames, TArray<FString>& OutTestCommands) const
+bool FOpenSoundMixAssetEditorsSLOW::RunTest(const FString& Parameters)
 {
-	//This grabs each SoundMix asset in the Game/Content directory
-	FEditorAutomationTestUtilities::CollectGameContentTestsByClass(USoundMix::StaticClass(), true, OutBeautifiedNames, OutTestCommands);
-}
+	OpenAssetParameters AssetList;
+	AssetList.TimeSpent = FPlatformTime::Seconds();
 
-bool FOpenSoundMixAssetEditors::RunTest(const FString& Parameters)
-{
-	bool bDidTheTestPass = FOpenActualAssetEditors(Parameters);
-	CollectGarbage(GARBAGE_COLLECTION_KEEPFLAGS);
-	return bDidTheTestPass;
-}
+	FEditorAutomationTestUtilities::CollectGameContentByClass(USoundMix::StaticClass(), true, AssetList.AssetPathList);
 
+	FOpenAllAssetTests(AssetList);
+
+	return true;
+}
 
-//////////////////////////////////////////////////////////////////////////
 /**
-* This test opens each SoundWave into its sub-editor.
+* This test opens each SOUND WAVE into its sub-editor.
 */
-IMPLEMENT_COMPLEX_AUTOMATION_TEST(FOpenSoundWaveAssetEditors, "Tools.Sounds.Open Sound Wave Assets", (EAutomationTestFlags::ATF_Editor | EAutomationTestFlags::ATF_RequiresUser));
+IMPLEMENT_SIMPLE_AUTOMATION_TEST(FOpenSoundWaveAssetEditorsSLOW, "Tools.Sounds.Open All Sound Wave Assets (SLOW)", (EAutomationTestFlags::ATF_Editor | EAutomationTestFlags::ATF_RequiresUser));
 
-void FOpenSoundWaveAssetEditors::GetTests(TArray<FString>& OutBeautifiedNames, TArray<FString>& OutTestCommands) const
+bool FOpenSoundWaveAssetEditorsSLOW::RunTest(const FString& Parameters)
 {
-	//This grabs each SoundWave asset in the Game/Content directory
-	FEditorAutomationTestUtilities::CollectGameContentTestsByClass(USoundWave::StaticClass(), true, OutBeautifiedNames, OutTestCommands);
-}
+	OpenAssetParameters AssetList;
+	AssetList.TimeSpent = FPlatformTime::Seconds();
 
-bool FOpenSoundWaveAssetEditors::RunTest(const FString& Parameters)
-{
-	bool bDidTheTestPass = FOpenActualAssetEditors(Parameters);
-	CollectGarbage(GARBAGE_COLLECTION_KEEPFLAGS);
-	return bDidTheTestPass;
+	FEditorAutomationTestUtilities::CollectGameContentByClass(USoundWave::StaticClass(), true, AssetList.AssetPathList);
+
+	FOpenAllAssetTests(AssetList);
+
+	return true;
 }
-//////////////////////////////////////////////////////////////////////////
+
 /**
-* This test opens each SubsurfaceProfile into its sub-editor.
+* This test opens each SUBSURFACEPROFILE into its sub-editor.
 */
-IMPLEMENT_COMPLEX_AUTOMATION_TEST(FOpenSubsurfaceProfileAssetEditors, "Tools.Materials.Open Subsurface Profile Assets", (EAutomationTestFlags::ATF_Editor | EAutomationTestFlags::ATF_RequiresUser));
+IMPLEMENT_SIMPLE_AUTOMATION_TEST(FOpenSubsurfaceProfileAssetEditorsSLOW, "Tools.Materials.Open All Subsurface Profile Assets (SLOW)", (EAutomationTestFlags::ATF_Editor | EAutomationTestFlags::ATF_RequiresUser));
 
-void FOpenSubsurfaceProfileAssetEditors::GetTests(TArray<FString>& OutBeautifiedNames, TArray<FString>& OutTestCommands) const
+bool FOpenSubsurfaceProfileAssetEditorsSLOW::RunTest(const FString& Parameters)
 {
-	//This grabs each SubsurfaceProfile asset in the Game/Content directory
-	FEditorAutomationTestUtilities::CollectGameContentTestsByClass(USubsurfaceProfile::StaticClass(), true, OutBeautifiedNames, OutTestCommands);
-}
+	OpenAssetParameters AssetList;
+	AssetList.TimeSpent = FPlatformTime::Seconds();
 
-bool FOpenSubsurfaceProfileAssetEditors::RunTest(const FString& Parameters)
-{
-	bool bDidTheTestPass = FOpenActualAssetEditors(Parameters);
-	CollectGarbage(GARBAGE_COLLECTION_KEEPFLAGS);
-	return bDidTheTestPass;
-}
+	FEditorAutomationTestUtilities::CollectGameContentByClass(USubsurfaceProfile::StaticClass(), true, AssetList.AssetPathList);
 
+	FOpenAllAssetTests(AssetList);
 
+	return true;
+}
 
-//////////////////////////////////////////////////////////////////////////
 /**
-* This test opens each TEXTURE asset into its sub-editor.
+* This test opens each TEXTURE into its sub-editor.
 */
-IMPLEMENT_COMPLEX_AUTOMATION_TEST(FOpenTextureAssetEditors, "Tools.Textures.Open Texture Assets", (EAutomationTestFlags::ATF_Editor | EAutomationTestFlags::ATF_RequiresUser));
+IMPLEMENT_SIMPLE_AUTOMATION_TEST(FOpenTextureAssetEditorsSLOW, "Tools.Textures.Open All Texture Assets (SLOW)", (EAutomationTestFlags::ATF_Editor | EAutomationTestFlags::ATF_RequiresUser));
 
-void FOpenTextureAssetEditors::GetTests(TArray<FString>& OutBeautifiedNames, TArray<FString>& OutTestCommands) const
+bool FOpenTextureAssetEditorsSLOW::RunTest(const FString& Parameters)
 {
-	//This grabs each texture file in the Game/Content
-	FEditorAutomationTestUtilities::CollectGameContentTestsByClass(UTexture::StaticClass(), true, OutBeautifiedNames, OutTestCommands);
-}
+	OpenAssetParameters AssetList;
+	AssetList.TimeSpent = FPlatformTime::Seconds();
 
-bool FOpenTextureAssetEditors::RunTest(const FString& Parameters)
-{
-	bool bDidTheTestPass = FOpenActualAssetEditors(Parameters);
-	CollectGarbage(GARBAGE_COLLECTION_KEEPFLAGS);
-	return bDidTheTestPass;
-}
+	FEditorAutomationTestUtilities::CollectGameContentByClass(UTexture::StaticClass(), true, AssetList.AssetPathList);
 
+	FOpenAllAssetTests(AssetList);
+
+	return true;
+}
 
-//////////////////////////////////////////////////////////////////////////
 /**
-* This test opens each Font into its sub-editor.
+* This test opens each USER DEFINED ENUM into its sub-editor.
 */
-IMPLEMENT_COMPLEX_AUTOMATION_TEST(FOpenFontAssetEditors, "Tools.Textures.Open Font Assets", (EAutomationTestFlags::ATF_Editor | EAutomationTestFlags::ATF_RequiresUser));
+IMPLEMENT_SIMPLE_AUTOMATION_TEST(FOpenUserDefinedEnumAssetEditorsSLOW, "Tools.Blueprint.Open All User Defined Enum Assets (SLOW)", (EAutomationTestFlags::ATF_Editor | EAutomationTestFlags::ATF_RequiresUser));
 
-void FOpenFontAssetEditors::GetTests(TArray<FString>& OutBeautifiedNames, TArray<FString>& OutTestCommands) const
+bool FOpenUserDefinedEnumAssetEditorsSLOW::RunTest(const FString& Parameters)
 {
-	//This grabs each Font asset in the Game/Content directory
-	FEditorAutomationTestUtilities::CollectGameContentTestsByClass(UFont::StaticClass(), true, OutBeautifiedNames, OutTestCommands);
-}
+	OpenAssetParameters AssetList;
+	AssetList.TimeSpent = FPlatformTime::Seconds();
 
-bool FOpenFontAssetEditors::RunTest(const FString& Parameters)
-{
-	bool bDidTheTestPass = FOpenActualAssetEditors(Parameters);
-	CollectGarbage(GARBAGE_COLLECTION_KEEPFLAGS);
-	return bDidTheTestPass;
+	FEditorAutomationTestUtilities::CollectGameContentByClass(UUserDefinedEnum::StaticClass(), true, AssetList.AssetPathList);
+
+	FOpenAllAssetTests(AssetList);
+
+	return true;
 }
 
-//////////////////////////////////////////////////////////////////////////
 /**
-* This test opens each misc asset into its sub-editor.
+* This test opens each FONT into its sub-editor.
 */
-IMPLEMENT_COMPLEX_AUTOMATION_TEST(FOpenMiscAssetEditors, "Tools.Misc.Open Misc Assets", (EAutomationTestFlags::ATF_Editor | EAutomationTestFlags::ATF_RequiresUser));
+IMPLEMENT_SIMPLE_AUTOMATION_TEST(FOpenFontAssetEditorsSLOW, "Tools.Textures.Open All Font Assets (SLOW)", (EAutomationTestFlags::ATF_Editor | EAutomationTestFlags::ATF_RequiresUser));
 
-void FOpenMiscAssetEditors::GetTests(TArray<FString>& OutBeautifiedNames, TArray<FString>& OutTestCommands) const
+bool FOpenFontAssetEditorsSLOW::RunTest(const FString& Parameters)
 {
-	//This grabs each TextureRenderTarget2D asset in the Game/Content directory
-	FEditorAutomationTestUtilities::CollectMiscGameContentTestsByClass(OutBeautifiedNames, OutTestCommands);
+	OpenAssetParameters AssetList;
+	AssetList.TimeSpent = FPlatformTime::Seconds();
+
+	FEditorAutomationTestUtilities::CollectGameContentByClass(UFont::StaticClass(), true, AssetList.AssetPathList);
+
+	FOpenAllAssetTests(AssetList);
+
+	return true;
 }
 
-bool FOpenMiscAssetEditors::RunTest(const FString& Parameters)
-{
-	bool bDidTheTestPass = FOpenActualAssetEditors(Parameters);
-	CollectGarbage(GARBAGE_COLLECTION_KEEPFLAGS);
-	return bDidTheTestPass;
-}
\ No newline at end of file
+////////////////////////////////////////////////////////////////////////////
+///**
+//* This test opens each misc asset into its sub-editor.
+//*/
+//IMPLEMENT_COMPLEX_AUTOMATION_TEST(FOpenMiscAssetEditors, "Tools.Misc.Open Misc Assets", (EAutomationTestFlags::ATF_Editor | EAutomationTestFlags::ATF_RequiresUser));
+//
+//void FOpenMiscAssetEditors::GetTests(TArray<FString>& OutBeautifiedNames, TArray<FString>& OutTestCommands) const
+//{
+//	//This grabs each TextureRenderTarget2D asset in the Game/Content directory
+//	FEditorAutomationTestUtilities::CollectMiscGameContentTestsByClass(OutBeautifiedNames, OutTestCommands);
+//}
+//
+//bool FOpenMiscAssetEditors::RunTest(const FString& Parameters)
+//{
+//	bool bDidTheTestPass = FOpenActualAssetEditors(Parameters);
+//	CollectGarbage(GARBAGE_COLLECTION_KEEPFLAGS);
+//	return bDidTheTestPass;
+//}
+//
+//
+//
