Fixing a compiling error.

diff --git a/Engine/Source/Editor/UnrealEd/Public/Tests/AutomationEditorCommon.h b/Engine/Source/Editor/UnrealEd/Public/Tests/AutomationEditorCommon.h
--- a/Engine/Source/Editor/UnrealEd/Public/Tests/AutomationEditorCommon.h
+++ b/Engine/Source/Editor/UnrealEd/Public/Tests/AutomationEditorCommon.h
@@ -238,7 +238,10 @@ public:
 	*/
 	static void CollectMiscGameContentTestsByClass(TArray<FString>& OutBeautifiedNames, TArray <FString>& OutTestCommands);
 	
-
+	/**
+	* Generates a list of assets from the GAME by a specific type.
+	*/
+	static void CollectGameContentByClass(const UClass * Class, bool bRecursiveClass, TArray<FString>& OutAssetList);
 };
 
 
