Editor: Fix typo in description of ToggleFPS command

diff --git a/Engine/Source/Editor/UnrealEd/Private/EditorViewportCommands.cpp b/Engine/Source/Editor/UnrealEd/Private/EditorViewportCommands.cpp
--- a/Engine/Source/Editor/UnrealEd/Private/EditorViewportCommands.cpp
+++ b/Engine/Source/Editor/UnrealEd/Private/EditorViewportCommands.cpp
@@ -29,7 +29,7 @@ void FEditorViewportCommands::RegisterCommands()
 
 	UI_COMMAND( ToggleRealTime, "Realtime", "Toggles real time rendering in this viewport", EUserInterfaceActionType::ToggleButton, FInputGesture( EModifierKey::Control, EKeys::R ) );
 	UI_COMMAND( ToggleStats, "Show Stats", "Toggles the ability to show stats in this viewport (enables realtime)", EUserInterfaceActionType::ToggleButton, FInputGesture( EModifierKey::Shift, EKeys::L ) );
-	UI_COMMAND( ToggleFPS, "Show FPS", "Toggles showing frames per section in this viewport (enables realtime)", EUserInterfaceActionType::ToggleButton, FInputGesture( EModifierKey::Control|EModifierKey::Shift, EKeys::H ) );
+	UI_COMMAND( ToggleFPS, "Show FPS", "Toggles showing frames per second in this viewport (enables realtime)", EUserInterfaceActionType::ToggleButton, FInputGesture( EModifierKey::Control|EModifierKey::Shift, EKeys::H ) );
 
 	UI_COMMAND( ScreenCapture, "Screen Capture", "Take a screenshot of the active viewport.", EUserInterfaceActionType::Button, FInputGesture(EKeys::F9) );
 	UI_COMMAND( ScreenCaptureForProjectThumbnail, "Update Project Thumbnail", "Take a screenshot of the active viewport for use as the project thumbnail.", EUserInterfaceActionType::Button, FInputGesture() );
