Editor: Push stat display/toggle functionality up from SLevelViewport/FLevelEditorViewportClient to SEditorViewport/FEditorViewportClient, so it can be used in asset editors as well

diff --git a/Engine/Source/Editor/LevelEditor/Private/SLevelViewport.cpp b/Engine/Source/Editor/LevelEditor/Private/SLevelViewport.cpp
--- a/Engine/Source/Editor/LevelEditor/Private/SLevelViewport.cpp
+++ b/Engine/Source/Editor/LevelEditor/Private/SLevelViewport.cpp
@@ -1704,15 +1704,6 @@ bool SLevelViewport::IsBufferVisualizationModeSelected( FName InName ) const
 	return LevelViewportClient->IsViewModeEnabled( VMI_VisualizeBuffer ) && LevelViewportClient->CurrentBufferVisualizationMode == InName;	
 }
 
-void SLevelViewport::ToggleShowFlag( uint32 EngineShowFlagIndex )
-{
-	bool bOldState = LevelViewportClient->EngineShowFlags.GetSingleFlag(EngineShowFlagIndex);
-	LevelViewportClient->EngineShowFlags.SetSingleFlag(EngineShowFlagIndex, !bOldState);
-
-	// Invalidate clients which aren't real-time so we see the changes
-	LevelViewportClient->Invalidate();
-}
-
 void SLevelViewport::OnToggleAllVolumeActors( bool bVisible )
 {
 	// Reinitialize the volume actor visibility flags to the new state.  All volumes should be visible if "Show All" was selected and hidden if it was not selected.
@@ -1834,16 +1825,6 @@ void SLevelViewport::OnToggleAllStatCommands( bool bVisible )
 	}
 }
 
-void SLevelViewport::ToggleStatCommand(FString CommandName)
-{
-	GEngine->ExecEngineStat(GetWorld(), LevelViewportClient.Get(), *CommandName);
-}
-
-bool SLevelViewport::IsShowFlagEnabled( uint32 EngineShowFlagIndex ) const
-{
-	return LevelViewportClient->EngineShowFlags.GetSingleFlag(EngineShowFlagIndex);
-}
-
 void SLevelViewport::OnUseDefaultShowFlags(bool bUseSavedDefaults)
 {
 	// cache off the current viewmode as it gets trashed when applying FEngineShowFlags()
diff --git a/Engine/Source/Editor/LevelEditor/Public/SLevelViewport.h b/Engine/Source/Editor/LevelEditor/Public/SLevelViewport.h
--- a/Engine/Source/Editor/LevelEditor/Public/SLevelViewport.h
+++ b/Engine/Source/Editor/LevelEditor/Public/SLevelViewport.h
@@ -97,6 +97,7 @@ public:
 	/** SEditorViewport Interface */
 	virtual void OnFocusViewportToSelection() override;
 	virtual EVisibility GetTransformToolbarVisibility() const override;
+	virtual UWorld* GetWorld() const override;
 
 	/**
 	 * Called when the maximize command is executed                   
@@ -210,9 +211,6 @@ public:
 	 * @return a new custom column for a scene outliner that indicates whether each actor is locked to this viewport
 	 */
 	TSharedRef< class ISceneOutlinerColumn > CreateActorLockSceneOutlinerColumn( const TWeakPtr< class ISceneOutliner >& SceneOutliner ) const;
-	
-	/** Gets the world this viewport is for */
-	UWorld* GetWorld() const;
 
 	/** Called when Preview Selected Cameras preference is changed.*/
 	void OnPreviewSelectedCamerasChange();
@@ -316,22 +314,6 @@ private:
 	/** Starts previewing any selected camera actors using live "PIP" sub-views */
 	void PreviewSelectedCameraActors();
 
-	/**
-	 * Called when stat rendering should be enabled in the viewport
-	 * NOTE: This enables realtime as well because we cant render stats without the viewport being realtime
-	 */
-	void OnToggleStats();
-
-	/**
-	 * @return true if stats are enabled in the viewport
-	 */
-	bool AreStatsEnabled() const;
-
-	/**
-	 * @return true if FPS is visible in the viewport                   
-	 */
-	bool IsShowingFPS() const;
-
 	/**
 	 * Called to create a cameraActor in the currently selected perspective viewport
 	 */
@@ -346,21 +328,6 @@ private:
 	 * Called to check currently selected editor viewport is a perspective one
 	 */
 	bool IsPerspectiveViewport() const;
-	
-	/**
-	 * Toggles a show flag in this viewport
-	 *
-	 * @param EngineShowFlagIndex	the ID to toggle
-	 */
-	void ToggleShowFlag( uint32 EngineShowFlagIndex );
-
-	/**
-	 * Checks if a show flag is enabled in this viewport
-	 *
-	 * @param EngineShowFlagIndex	the ID to check
-	 * @return true if the show flag is enabled, false otherwise
-	 */
-	bool IsShowFlagEnabled( uint32 EngineShowFlagIndex ) const;
 
 	/**
 	 * Toggles all volume classes visibility
@@ -418,13 +385,6 @@ private:
 	 */
 	void OnToggleAllStatCommands(bool bVisible);
 
-	/**
-	 * Toggles Stat command visibility in this viewport
-	 *
-	 * @param CommandName				Name of the command
-	 */
-	virtual void ToggleStatCommand(FString CommandName) override;
-
 	/**
 	 * Called when show flags for this viewport should be reset to default, or the saved settings
 	 */
diff --git a/Engine/Source/Editor/UnrealEd/Private/EditorViewportClient.cpp b/Engine/Source/Editor/UnrealEd/Private/EditorViewportClient.cpp
--- a/Engine/Source/Editor/UnrealEd/Private/EditorViewportClient.cpp
+++ b/Engine/Source/Editor/UnrealEd/Private/EditorViewportClient.cpp
@@ -282,6 +282,11 @@ FEditorViewportClient::FEditorViewportClient(FEditorModeTools& InModeTools, FPre
 	SetViewMode(IsPerspective() ? PerspViewModeIndex : OrthoViewModeIndex);
 
 	ModeTools->OnEditorModeChanged().AddRaw(this, &FEditorViewportClient::OnEditorModeChanged);
+
+	FCoreDelegates::StatCheckEnabled.AddRaw(this, &FLevelEditorViewportClient::HandleViewportStatCheckEnabled);
+	FCoreDelegates::StatEnabled.AddRaw(this, &FLevelEditorViewportClient::HandleViewportStatEnabled);
+	FCoreDelegates::StatDisabled.AddRaw(this, &FLevelEditorViewportClient::HandleViewportStatDisabled);
+	FCoreDelegates::StatDisableAll.AddRaw(this, &FLevelEditorViewportClient::HandleViewportStatDisableAll);
 }
 
 FEditorViewportClient::~FEditorViewportClient()
@@ -1416,6 +1421,78 @@ void FEditorViewportClient::DrawSafeFrameQuad( FCanvas &Canvas, FVector2D V1, FV
 	TriItem.Draw( &Canvas );
 }
 
+int32 FEditorViewportClient::SetStatEnabled(const TCHAR* InName, const bool bEnable, const bool bAll)
+{
+	if (bEnable)
+	{
+		check(!bAll);	// Not possible to enable all
+		EnabledStats.AddUnique(InName);
+	}
+	else
+	{
+		if (bAll)
+		{
+			EnabledStats.Empty();
+		}
+		else
+		{
+			EnabledStats.Remove(InName);
+		}
+	}
+	return EnabledStats.Num();
+}
+
+
+void FEditorViewportClient::HandleViewportStatCheckEnabled(const TCHAR* InName, bool& bOutCurrentEnabled, bool& bOutOthersEnabled)
+{
+	// Check to see which viewports have this enabled (current, non-current)
+	const bool bEnabled = IsStatEnabled(InName);
+	if (GStatProcessingViewportClient == this)
+	{
+		// Only if realtime and stats are also enabled should we show the stat as visible
+		bOutCurrentEnabled = IsRealtime() && ShouldShowStats() && bEnabled;
+	}
+	else
+	{
+		bOutOthersEnabled |= bEnabled;
+	}
+}
+
+void FEditorViewportClient::HandleViewportStatEnabled(const TCHAR* InName)
+{
+	// Just enable this on the active viewport
+	if (GStatProcessingViewportClient == this)
+	{
+		SetShowStats(true);
+		SetRealtime(true);
+		SetStatEnabled(InName, true);
+	}
+}
+
+void FEditorViewportClient::HandleViewportStatDisabled(const TCHAR* InName)
+{
+	// Just disable this on the active viewport
+	if (GStatProcessingViewportClient == this)
+	{
+		if (SetStatEnabled(InName, false) == 0)
+		{
+			SetShowStats(false);
+			// Note: we can't disable realtime as we don't know the setting it was previously
+		}
+	}
+}
+
+void FEditorViewportClient::HandleViewportStatDisableAll(const bool bInAnyViewport)
+{
+	// Disable all on either all or the current viewport (depending on the flag)
+	if (bInAnyViewport || GStatProcessingViewportClient == this)
+	{
+		SetShowStats(false);
+		// Note: we can't disable realtime as we don't know the setting it was previously
+		SetStatEnabled(NULL, false, true);
+	}
+}
+
 void FEditorViewportClient::UpdateMouseDelta()
 {
 	FVector DragDelta = FVector::ZeroVector;
@@ -4136,6 +4213,33 @@ void FEditorViewportClient::SetGameView(bool bGameViewEnable)
 	Invalidate();
 }
 
+FStatUnitData* FEditorViewportClient::GetStatUnitData() const
+{
+	return &StatUnitData;
+}
+
+FStatHitchesData* FEditorViewportClient::GetStatHitchesData() const
+{
+	return &StatHitchesData;
+}
+
+const TArray<FString>* FEditorViewportClient::GetEnabledStats() const
+{
+	return &EnabledStats;
+}
+
+void FEditorViewportClient::SetEnabledStats(const TArray<FString>& InEnabledStats)
+{
+	EnabledStats = InEnabledStats;
+}
+
+bool FEditorViewportClient::IsStatEnabled(const TCHAR* InName) const
+{
+	return EnabledStats.Contains(InName);
+}
+
+////////////////
+
 bool FEditorViewportStats::bInitialized(false);
 bool FEditorViewportStats::bUsingCalledThisFrame(false);
 FEditorViewportStats::Category FEditorViewportStats::LastUsing(FEditorViewportStats::CAT_MAX);
diff --git a/Engine/Source/Editor/UnrealEd/Private/LevelEditorViewport.cpp b/Engine/Source/Editor/UnrealEd/Private/LevelEditorViewport.cpp
--- a/Engine/Source/Editor/UnrealEd/Private/LevelEditorViewport.cpp
+++ b/Engine/Source/Editor/UnrealEd/Private/LevelEditorViewport.cpp
@@ -1430,13 +1430,6 @@ FLevelEditorViewportClient::FLevelEditorViewportClient()
 
 	// Sign up for notifications about users changing settings.
 	GetMutableDefault<ULevelEditorViewportSettings>()->OnSettingChanged().AddRaw(this, &FLevelEditorViewportClient::HandleViewportSettingChanged);
-
-	StatUnitData = new FStatUnitData();
-	StatHitchesData = new FStatHitchesData();
-	FCoreDelegates::StatCheckEnabled.AddRaw(this, &FLevelEditorViewportClient::HandleViewportStatCheckEnabled);
-	FCoreDelegates::StatEnabled.AddRaw(this, &FLevelEditorViewportClient::HandleViewportStatEnabled);
-	FCoreDelegates::StatDisabled.AddRaw(this, &FLevelEditorViewportClient::HandleViewportStatDisabled);
-	FCoreDelegates::StatDisableAll.AddRaw(this, &FLevelEditorViewportClient::HandleViewportStatDisableAll);
 }
 
 //
@@ -1470,16 +1463,6 @@ FLevelEditorViewportClient::~FLevelEditorViewportClient()
 	FCoreDelegates::StatEnabled.RemoveAll(this);
 	FCoreDelegates::StatDisabled.RemoveAll(this);
 	FCoreDelegates::StatDisableAll.RemoveAll(this);
-	if (StatHitchesData)
-	{
-		delete StatHitchesData;
-		StatHitchesData = NULL;
-	}
-	if (StatUnitData)
-	{
-		delete StatUnitData;
-		StatUnitData = NULL;
-	}
 
 	GEditor->LevelViewportClients.Remove(this);
 
@@ -4140,56 +4123,6 @@ void FLevelEditorViewportClient::SetIsSimulateInEditorViewport( bool bInIsSimula
 	}
 }
 
-void FLevelEditorViewportClient::HandleViewportStatCheckEnabled(const TCHAR* InName, bool& bOutCurrentEnabled, bool& bOutOthersEnabled)
-{
-	// Check to see which viewports have this enabled (current, non-current)
-	const bool bEnabled = IsStatEnabled(InName);
-	if (GStatProcessingViewportClient == this)
-	{
-		// Only if realtime and stats are also enabled should we show the stat as visible
-		bOutCurrentEnabled = IsRealtime() && ShouldShowStats() && bEnabled;
-	}
-	else
-	{
-		bOutOthersEnabled |= bEnabled;
-	}
-}
-
-void FLevelEditorViewportClient::HandleViewportStatEnabled(const TCHAR* InName)
-{
-	// Just enable this on the active viewport
-	if (GStatProcessingViewportClient == this)
-	{
-		SetShowStats(true);
-		SetRealtime(true);
-		SetStatEnabled(InName, true);
-	}
-}
-
-void FLevelEditorViewportClient::HandleViewportStatDisabled(const TCHAR* InName)
-{
-	// Just disable this on the active viewport
-	if (GStatProcessingViewportClient == this)
-	{
-		if (SetStatEnabled(InName, false) == 0)
-		{
-			SetShowStats(false);
-			// Note: we can't disable realtime as we don't know the setting it was previously
-		}
-	}
-}
-
-void FLevelEditorViewportClient::HandleViewportStatDisableAll(const bool bInAnyViewport)
-{
-	// Disable all on either all or the current viewport (depending on the flag)
-	if (bInAnyViewport || GStatProcessingViewportClient == this)
-	{
-		SetShowStats(false);
-		// Note: we can't disable realtime as we don't know the setting it was previously
-		SetStatEnabled(NULL, false, true);
-	}
-}
-
 #undef LOCTEXT_NAMESPACE
 
 // Doxygen cannot parse these correctly since the declarations are made in Editor, not UnrealEd
diff --git a/Engine/Source/Editor/UnrealEd/Private/SEditorViewport.cpp b/Engine/Source/Editor/UnrealEd/Private/SEditorViewport.cpp
--- a/Engine/Source/Editor/UnrealEd/Private/SEditorViewport.cpp
+++ b/Engine/Source/Editor/UnrealEd/Private/SEditorViewport.cpp
@@ -347,12 +347,34 @@ void SEditorViewport::OnToggleStats()
 	}
 }
 
+void SEditorViewport::ToggleStatCommand(FString CommandName)
+{
+	GEngine->ExecEngineStat(GetWorld(), Client.Get(), *CommandName);
+
+	// Invalidate the client to render once in case the click was on the checkbox itself (which doesn't dismiss the menu)
+	Client->Invalidate();
+}
+
 bool SEditorViewport::IsStatCommandVisible(FString CommandName) const
 {
 	// Only if realtime and stats are also enabled should we show the stat as visible
 	return Client->IsRealtime() && Client->ShouldShowStats() && Client->IsStatEnabled(*CommandName);
 }
 
+void SEditorViewport::ToggleShowFlag(uint32 EngineShowFlagIndex)
+{
+	bool bOldState = Client->EngineShowFlags.GetSingleFlag(EngineShowFlagIndex);
+	Client->EngineShowFlags.SetSingleFlag(EngineShowFlagIndex, !bOldState);
+
+	// Invalidate clients which aren't real-time so we see the changes
+	Client->Invalidate();
+}
+
+bool SEditorViewport::IsShowFlagEnabled(uint32 EngineShowFlagIndex) const
+{
+	return Client->EngineShowFlags.GetSingleFlag(EngineShowFlagIndex);
+}
+
 void SEditorViewport::ChangeExposureSetting( int32 ID )
 {
 	Client->ExposureSettings.bFixed = (ID != FEditorViewportCommands::AutoExposureRadioID);
@@ -457,6 +479,12 @@ void SEditorViewport::OnCycleCoordinateSystem()
 	Client->SetWidgetCoordSystemSpace( (ECoordSystem)CoordSystemAsInt );
 }
 
+UWorld* SEditorViewport::GetWorld() const
+{
+	return Client->GetWorld();
+}
+
+
 void SEditorViewport::OnToggleSurfaceSnap()
 {
 	auto* Settings = GetMutableDefault<ULevelEditorViewportSettings>();
diff --git a/Engine/Source/Editor/UnrealEd/Public/EditorViewportClient.h b/Engine/Source/Editor/UnrealEd/Public/EditorViewportClient.h
--- a/Engine/Source/Editor/UnrealEd/Public/EditorViewportClient.h
+++ b/Engine/Source/Editor/UnrealEd/Public/EditorViewportClient.h
@@ -411,8 +411,6 @@ public:
 
 	void SetInitialViewTransform( const FVector& ViewLocation, const FRotator& ViewRotation, float InOrthoZoom );
 
-	virtual void ProcessScreenShots(FViewport* Viewport) override;
-
 	void TakeHighResScreenShot();
 
 	/** Called when an editor mode has been (de)activated */
@@ -423,6 +421,7 @@ public:
 	virtual void Draw(FViewport* Viewport,FCanvas* Canvas) override;
 
 	/** FViewportClient interface */
+	virtual void ProcessScreenShots(FViewport* Viewport) override;
 	virtual void RedrawRequested(FViewport* Viewport) override;
 	virtual void RequestInvalidateHitProxy(FViewport* Viewport) override;
 	virtual bool InputKey(FViewport* Viewport, int32 ControllerId, FKey Key, EInputEvent Event, float AmountDepressed = 1.f, bool bGamepad=false) override;
@@ -437,7 +436,11 @@ public:
 	virtual void CapturedMouseMove( FViewport* InViewport, int32 InMouseX, int32 InMouseY ) override;
 	virtual bool IsOrtho() const override;
 	virtual void LostFocus(FViewport* Viewport) override;
-
+	virtual FStatUnitData* GetStatUnitData() const override;
+	virtual FStatHitchesData* GetStatHitchesData() const override;
+	virtual const TArray<FString>* GetEnabledStats() const override;
+	virtual void SetEnabledStats(const TArray<FString>& InEnabledStats) override;
+	virtual bool IsStatEnabled(const TCHAR* InName) const override;
 
 	/** FGCObject interface */
 	virtual void AddReferencedObjects( FReferenceCollector& Collector ) override;
@@ -1161,6 +1164,23 @@ private:
 		return &EngineShowFlags; 
 	}
 
+	/**
+	 * Set a specific stat to either enabled or disabled (returns true if there are any stats enabled)
+	 */
+	int32 SetStatEnabled(const TCHAR* InName, const bool bEnable, const bool bAll = false);
+
+	/** Delegate handler to see if a stat is enabled on this viewport */
+	void HandleViewportStatCheckEnabled(const TCHAR* InName, bool& bOutCurrentEnabled, bool& bOutOthersEnabled);
+
+	/** Delegate handler for when stats are enabled in a viewport */
+	void HandleViewportStatEnabled(const TCHAR* InName);
+
+	/** Delegate handler for when stats are disabled in a viewport */
+	void HandleViewportStatDisabled(const TCHAR* InName);
+
+	/** Delegate handler for when all stats are disabled in a viewport */
+	void HandleViewportStatDisableAll(const bool bInAnyViewport);
+
 public:
 	static const uint32 MaxCameraSpeeds;
 
@@ -1352,6 +1372,16 @@ public:
 	/* Default view mode for orthographic viewports */
 	static const EViewModeIndex DefaultOrthoViewMode;
 
+protected:
+	/** Data needed to display per-frame stat tracking when STAT UNIT is enabled */
+	mutable FStatUnitData StatUnitData;
+
+	/** Data needed to display per-frame stat tracking when STAT HITCHES is enabled */
+	mutable FStatHitchesData StatHitchesData;
+
+	/** A list of all the stat names which are enabled for this viewport */
+	TArray<FString> EnabledStats;
+
 private:
 	/* View mode to set when this viewport is of type LVT_Perspective */
 	EViewModeIndex PerspViewModeIndex;
@@ -1369,6 +1399,9 @@ private:
 	bool bInGameViewMode;
 };
 
+
+
+
 class UNREALED_API FEditorViewportStats
 {
 public:
diff --git a/Engine/Source/Editor/UnrealEd/Public/LevelEditorViewport.h b/Engine/Source/Editor/UnrealEd/Public/LevelEditorViewport.h
--- a/Engine/Source/Editor/UnrealEd/Public/LevelEditorViewport.h
+++ b/Engine/Source/Editor/UnrealEd/Public/LevelEditorViewport.h
@@ -477,46 +477,6 @@ public:
 		return ActorLockedByMatinee.IsValid();
 	}
 
-	/**
-	 * Get a ptr to the stat unit data for this viewport
-	 */
-	virtual FStatUnitData* GetStatUnitData() const override
-	{
-		return StatUnitData;
-	}
-
-	/**
-	 * Get a ptr to the stat unit data for this viewport
-	 */
-	virtual FStatHitchesData* GetStatHitchesData() const override
-	{
-		return StatHitchesData;
-	}
-
-	/**
-	 * Get a ptr to the enabled stats list
-	 */
-	virtual const TArray<FString>* GetEnabledStats() const override
-	{
-		return &EnabledStats;
-	}
-
-	/**
-	 * Sets all the stats that should be enabled for the viewport
-	 */
-	virtual void SetEnabledStats(const TArray<FString>& InEnabledStats) override
-	{
-		EnabledStats = InEnabledStats;
-	}
-
-	/**
-	 * Check whether a specific stat is enabled for this viewport
-	 */
-	virtual bool IsStatEnabled(const TCHAR* InName) const override
-	{
-		return EnabledStats.Contains(InName);
-	}
-
 	/**
 	 * Get the sound stat flags enabled for this viewport
 	 */
@@ -533,31 +493,6 @@ public:
 		SoundShowFlags = InSoundShowFlags;
 	}
 
-private:
-	/**
-	 * Set a specific stat to either enabled or disabled (returns true if there are any stats enabled)
-	 */
-	int32 SetStatEnabled(const TCHAR* InName, const bool bEnable, const bool bAll = false)
-	{
-		if (bEnable)
-		{
-			check(!bAll);	// Not possible to enable all
-			EnabledStats.AddUnique(InName);
-		}
-		else
-		{
-			if (bAll)
-			{
-				EnabledStats.Empty();
-			}
-			else
-			{
-				EnabledStats.Remove(InName);
-			}
-		}
-		return EnabledStats.Num();
-	}
-
 protected:
 	/** 
 	 * Checks the viewport to see if the given blueprint asset can be dropped on the viewport.
@@ -678,18 +613,6 @@ private:
 	/** Project the specified actors into the world according to the current drag parameters */
 	void ProjectActorsIntoWorld(const TArray<AActor*>& Actors, FViewport* Viewport, const FVector& Drag, const FRotator& Rot);
 
-	/** Delegate handler to see if a stat is enabled on this viewport */
-	void HandleViewportStatCheckEnabled(const TCHAR* InName, bool& bOutCurrentEnabled, bool& bOutOthersEnabled);
-
-	/** Delegate handler for when stats are enabled in a viewport */
-	void HandleViewportStatEnabled(const TCHAR* InName);
-
-	/** Delegate handler for when stats are disabled in a viewport */
-	void HandleViewportStatDisabled(const TCHAR* InName);
-
-	/** Delegate handler for when all stats are disabled in a viewport */
-	void HandleViewportStatDisableAll(const bool bInAnyViewport);
-
 public:
 	/** Static: List of objects we're hovering over */
 	static TSet< FViewportHoverTarget > HoveredObjects;
@@ -777,15 +700,6 @@ private:
 	TWeakObjectPtr<AActor>	ActorLockedByMatinee;
 	TWeakObjectPtr<AActor>	ActorLockedToCamera;
 
-	/** Data needed to display perframe stat tracking when STAT UNIT is enabled */
-	FStatUnitData*			StatUnitData;
-
-	/** Data needed to display perframe stat tracking when STAT HITCHES is enabled */
-	FStatHitchesData*		StatHitchesData;
-
-	/** A list of all the stat names which are enabled for this viewport */
-	TArray<FString>			EnabledStats;
-
 	/** Those sound stat flags which are enabled on this viewport */
 	ESoundShowFlags::Type	SoundShowFlags;
 };
diff --git a/Engine/Source/Editor/UnrealEd/Public/SEditorViewport.h b/Engine/Source/Editor/UnrealEd/Public/SEditorViewport.h
--- a/Engine/Source/Editor/UnrealEd/Public/SEditorViewport.h
+++ b/Engine/Source/Editor/UnrealEd/Public/SEditorViewport.h
@@ -63,7 +63,7 @@ protected:
 	 *
 	 * @param CommandName				Name of the command
 	 */
-	virtual void ToggleStatCommand(FString CommandName){}
+	virtual void ToggleStatCommand(FString CommandName);
 
 	/**
 	 * Checks if Stat command is visible in this viewport
@@ -72,6 +72,21 @@ protected:
 	 */
 	virtual bool IsStatCommandVisible(FString CommandName) const;
 
+	/**
+	 * Toggles a show flag in this viewport
+	 *
+	 * @param EngineShowFlagIndex	the ID to toggle
+	 */
+	void ToggleShowFlag( uint32 EngineShowFlagIndex );
+
+	/**
+	 * Checks if a show flag is enabled in this viewport
+	 *
+	 * @param EngineShowFlagIndex	the ID to check
+	 * @return true if the show flag is enabled, false otherwise
+	 */
+	bool IsShowFlagEnabled( uint32 EngineShowFlagIndex ) const;
+
 	/**
 	 * Changes the exposure setting for this viewport
 	 *
@@ -125,6 +140,8 @@ protected:
 	 */
 	virtual void OnFocusViewportToSelection(){}
 
+	/** Gets the world this viewport is for */
+	virtual UWorld* GetWorld() const;
 
 	/**
 	 * Called when surface snapping has been enabled/disabled
