Improved BSP interaction with the scaling widget (in non-percentage-based scaling mode).

diff --git a/Engine/Source/Editor/UnrealEd/Classes/Builders/EditorBrushBuilder.h b/Engine/Source/Editor/UnrealEd/Classes/Builders/EditorBrushBuilder.h
--- a/Engine/Source/Editor/UnrealEd/Classes/Builders/EditorBrushBuilder.h
+++ b/Engine/Source/Editor/UnrealEd/Classes/Builders/EditorBrushBuilder.h
@@ -14,9 +14,9 @@ public:
 	/** UBrushBuilder interface */
 	virtual void BeginBrush( bool InMergeCoplanars, FName InLayer ) override;
 	virtual bool EndBrush( UWorld* InWorld, ABrush* InBrush ) override;
-	virtual int32 GetVertexCount() override;
-	virtual FVector GetVertex( int32 i ) override;
-	virtual int32 GetPolyCount() override;
+	virtual int32 GetVertexCount() const override;
+	virtual FVector GetVertex( int32 i ) const override;
+	virtual int32 GetPolyCount() const override;
 	virtual bool BadParameters( const FText& msg ) override;
 	virtual int32 Vertexv( FVector v ) override;
 	virtual int32 Vertex3f( float X, float Y, float Z ) override;
diff --git a/Engine/Source/Editor/UnrealEd/Private/Editor.cpp b/Engine/Source/Editor/UnrealEd/Private/Editor.cpp
--- a/Engine/Source/Editor/UnrealEd/Private/Editor.cpp
+++ b/Engine/Source/Editor/UnrealEd/Private/Editor.cpp
@@ -2293,10 +2293,68 @@ void UEditorEngine::ApplyDeltaToActor(AActor* InActor,
 			}
 
 			// If the actor is a brush, update the vertices.
-			if( Brush )
+			if( Brush != nullptr)
 			{
+				FVector DeltaScale = ModifiedScale;
+
+				AActor::bUsePercentageBasedScaling = GEditor->UsePercentageBasedScaling();
+				if (!AActor::bUsePercentageBasedScaling)
+				{
+					// Attempt to deduce current implied scale by comparing the bounding box of the source builder brush with the current bounding box.
+					// Not a perfect method, but it feels OK.
+					// Also disallow singularities by not permitting a zero scale on any of the axes (otherwise we lose information).
+					const UBrushBuilder* BrushBuilder = Brush->GetBrushBuilder();
+					if (BrushBuilder != nullptr)
+					{
+						// Get builder brush bounding box
+						FBox BuildBrushBBox(0);
+
+						for (int32 Index = 0, NumVerts = BrushBuilder->GetVertexCount(); Index < NumVerts; Index++)
+						{
+							BuildBrushBBox += BrushBuilder->GetVertex(Index);
+						}
+
+						if (Brush->Brush != nullptr && Brush->Brush->Polys != nullptr)
+						{
+							// Get current bounding box
+							FBox CurrentBBox(0);
+
+							for (const auto& Poly : Brush->Brush->Polys->Element)
+							{
+								for (const auto& Vertex : Poly.Vertices)
+								{
+									CurrentBBox += Vertex;
+								}
+							}
+
+							// Estimate implied scale (can only be positive), avoiding zero scale components
+							const FVector ImpliedScale = ((CurrentBBox.Max - CurrentBBox.Min) / (BuildBrushBBox.Max - BuildBrushBBox.Min))
+								.ComponentMax(FVector(SMALL_NUMBER, SMALL_NUMBER, SMALL_NUMBER));
+							
+							// Calculate delta scale vector which, when applied as a scale transform and added to the original vector, will scale by this additional amount.
+							// A scale of 0.0 implies no change, a scale of -1.0 will collapse component to zero, a scale of +1.0 will double the size.
+							const FVector NewScaleVector = DeltaScale / ImpliedScale;
+
+							// Protect against singularities by never allowing a component scale close to 0.
+							const float Limit = -1.0f + 0.0625f;
+							if (NewScaleVector.X > Limit &&
+								NewScaleVector.Y > Limit &&
+								NewScaleVector.Z > Limit)
+							{
+								// Apply calculated scale delta
+								DeltaScale = NewScaleVector;
+							}
+							else
+							{
+								// Apply zero scale delta
+								DeltaScale = FVector::ZeroVector;
+							}
+						}
+					}
+				}
+
 				// Scale all of the polygons of the brush.
-				const FScaleMatrix matrix( FVector( ModifiedScale.X , ModifiedScale.Y, ModifiedScale.Z ) );
+				const FScaleMatrix Matrix(DeltaScale);
 				
 				if(Brush->GetBrushComponent()->Brush && Brush->GetBrushComponent()->Brush->Polys)
 				{
@@ -2323,7 +2381,7 @@ void UEditorEngine::ApplyDeltaToActor(AActor* InActor,
 						{
 							FVector Wk = BrushActorToWorld.TransformPosition( Poly->Vertices[vertex] );
 							Wk -= GLevelEditorModeTools().PivotLocation;
-							Wk += matrix.TransformPosition( Wk );
+							Wk += Matrix.TransformPosition( Wk );
 							Wk += GLevelEditorModeTools().PivotLocation;
 							Poly->Vertices[vertex] = BrushActorToWorld.InverseTransformPosition( Wk );
 						}
@@ -2336,7 +2394,7 @@ void UEditorEngine::ApplyDeltaToActor(AActor* InActor,
 
 						FVector Wk = BrushActorToWorld.TransformPosition( Poly->Base );
 						Wk -= GLevelEditorModeTools().PivotLocation;
-						Wk += matrix.TransformPosition( Wk );
+						Wk += Matrix.TransformPosition( Wk );
 						Wk += GLevelEditorModeTools().PivotLocation;
 						Poly->Base = BrushActorToWorld.InverseTransformPosition( Wk );
 
diff --git a/Engine/Source/Editor/UnrealEd/Private/EditorBrushBuilder.cpp b/Engine/Source/Editor/UnrealEd/Private/EditorBrushBuilder.cpp
--- a/Engine/Source/Editor/UnrealEd/Private/EditorBrushBuilder.cpp
+++ b/Engine/Source/Editor/UnrealEd/Private/EditorBrushBuilder.cpp
@@ -144,17 +144,17 @@ bool UEditorBrushBuilder::EndBrush( UWorld* InWorld, ABrush* InBrush )
 	return true;
 }
 
-int32 UEditorBrushBuilder::GetVertexCount()
+int32 UEditorBrushBuilder::GetVertexCount() const
 {
 	return Vertices.Num();
 }
 
-FVector UEditorBrushBuilder::GetVertex(int32 i)
+FVector UEditorBrushBuilder::GetVertex(int32 i) const
 {
 	return Vertices.IsValidIndex(i) ? Vertices[i] : FVector::ZeroVector;
 }
 
-int32 UEditorBrushBuilder::GetPolyCount()
+int32 UEditorBrushBuilder::GetPolyCount() const
 {
 	return Polys.Num();
 }
diff --git a/Engine/Source/Runtime/Engine/Classes/Engine/BrushBuilder.h b/Engine/Source/Runtime/Engine/Classes/Engine/BrushBuilder.h
--- a/Engine/Source/Runtime/Engine/Classes/Engine/BrushBuilder.h
+++ b/Engine/Source/Runtime/Engine/Classes/Engine/BrushBuilder.h
@@ -79,9 +79,9 @@ public:
 	// @ todo document all below
 	virtual void BeginBrush( bool InMergeCoplanars, FName InLayer ) {}
 	virtual bool EndBrush( UWorld* InWorld, ABrush* InBrush ) { return false; }
-	virtual int32 GetVertexCount() { return 0; }
-	virtual FVector GetVertex( int32 i ) { return FVector::ZeroVector; }
-	virtual int32 GetPolyCount() { return 0; }
+	virtual int32 GetVertexCount() const { return 0; }
+	virtual FVector GetVertex( int32 i ) const { return FVector::ZeroVector; }
+	virtual int32 GetPolyCount() const { return 0; }
 	virtual bool BadParameters( const FText& msg ) { return false; }
 	virtual int32 Vertexv( FVector v ) { return 0; }
 	virtual int32 Vertex3f( float X, float Y, float Z ) { return 0; }
