Changed FBX import to provide more options regarding vertex color import.  Now vertex colors in the imported mesh can be taken from the FBX file, preserved from the current mesh, or set to a constant color.

diff --git a/Engine/Config/BaseEditorUserSettings.ini b/Engine/Config/BaseEditorUserSettings.ini
--- a/Engine/Config/BaseEditorUserSettings.ini
+++ b/Engine/Config/BaseEditorUserSettings.ini
@@ -560,6 +560,8 @@ bImportMesh=True
 bOneConvexHullPerUCX=True
 bImportMeshLODs=False
 NormalImportMethod=FBXNIM_ImportNormals
+VertexColorImportOption=EVertexColorImportOption::Ignore
+VertexOverrideColor=(R=255,G=255,B=255,A=255)
 
 
 [/Script/UnrealEd.FbxSkeletalMeshImportData]
diff --git a/Engine/Source/Editor/DetailCustomizations/Private/FbxImportUIDetails.cpp b/Engine/Source/Editor/DetailCustomizations/Private/FbxImportUIDetails.cpp
--- a/Engine/Source/Editor/DetailCustomizations/Private/FbxImportUIDetails.cpp
+++ b/Engine/Source/Editor/DetailCustomizations/Private/FbxImportUIDetails.cpp
@@ -118,6 +118,14 @@ void FFbxImportUIDetails::CustomizeDetails( IDetailLayoutBuilder& DetailBuilder
 					{
 						SetStaticMeshLODGroupWidget(PropertyRow, Handle);
 					}
+
+					if (Property->GetFName() == GET_MEMBER_NAME_CHECKED(UFbxStaticMeshImportData, VertexOverrideColor))
+					{
+						// Cache the VertexColorImportOption property
+						VertexColorImportOptionHandle = StaticMeshDataProp->GetChildHandle(GET_MEMBER_NAME_CHECKED(UFbxStaticMeshImportData, VertexColorImportOption));
+
+						PropertyRow.IsEnabled(TAttribute<bool>(this, &FFbxImportUIDetails::GetVertexOverrideColorEnabledState));
+					}
 				}
 			}
 		}
@@ -242,8 +250,10 @@ void FFbxImportUIDetails::SetStaticMeshLODGroupWidget(IDetailPropertyRow& Proper
 		.ValueContent()
 		.MinDesiredWidth(Row.ValueWidget.MinWidth)
 		.MaxDesiredWidth(Row.ValueWidget.MaxWidth)
+		.VAlign(VAlign_Center)
 		[
 			SNew(STextComboBox)
+			.Font(IDetailLayoutBuilder::GetDetailFont())
 			.OptionsSource(&LODGroupOptions)
 			.InitiallySelectedItem(LODGroupOptions[GroupIndex])
 			.OnSelectionChanged(this, &FFbxImportUIDetails::OnLODGroupChanged, HandlePtr)
@@ -261,6 +271,14 @@ void FFbxImportUIDetails::OnLODGroupChanged(TSharedPtr<FString> NewValue, ESelec
 	}
 }
 
+bool FFbxImportUIDetails::GetVertexOverrideColorEnabledState() const
+{
+	uint8 VertexColorImportOption;
+	check(VertexColorImportOptionHandle.IsValid())
+	ensure(VertexColorImportOptionHandle->GetValue(VertexColorImportOption) == FPropertyAccess::Success);
+
+	return (VertexColorImportOption == EVertexColorImportOption::Override);
+}
 
 void FFbxImportUIDetails::CollectChildPropertiesRecursive(TSharedPtr<IPropertyHandle> Node, TArray<TSharedPtr<IPropertyHandle>>& OutProperties)
 {
diff --git a/Engine/Source/Editor/DetailCustomizations/Private/FbxImportUIDetails.h b/Engine/Source/Editor/DetailCustomizations/Private/FbxImportUIDetails.h
--- a/Engine/Source/Editor/DetailCustomizations/Private/FbxImportUIDetails.h
+++ b/Engine/Source/Editor/DetailCustomizations/Private/FbxImportUIDetails.h
@@ -40,10 +40,16 @@ private:
 	/** Called when the StaticMeshLODGroup spinbox is changed */
 	void OnLODGroupChanged(TSharedPtr<FString> NewValue, ESelectInfo::Type SelectInfo, TWeakPtr<IPropertyHandle> HandlePtr);
 
+	/** Called to determine the visibility of the VertexOverrideColor property */
+	bool GetVertexOverrideColorEnabledState() const;
+
 	/** LOD group options. */
 	TArray<FName> LODGroupNames;
 	TArray<TSharedPtr<FString>> LODGroupOptions;
 
 	/** Cached StaticMeshLODGroup property handle */
 	TSharedPtr<IPropertyHandle> StaticMeshLODGroupPropertyHandle;
+
+	/** Cached VertexColorImportOption property handle */
+	TSharedPtr<IPropertyHandle> VertexColorImportOptionHandle;
 };
diff --git a/Engine/Source/Editor/UnrealEd/Classes/Factories/FbxStaticMeshImportData.h b/Engine/Source/Editor/UnrealEd/Classes/Factories/FbxStaticMeshImportData.h
--- a/Engine/Source/Editor/UnrealEd/Classes/Factories/FbxStaticMeshImportData.h
+++ b/Engine/Source/Editor/UnrealEd/Classes/Factories/FbxStaticMeshImportData.h
@@ -7,6 +7,20 @@
 #pragma once
 #include "FbxStaticMeshImportData.generated.h"
 
+UENUM()
+namespace EVertexColorImportOption
+{
+	enum Type
+	{
+		/** Import the static mesh using the vertex colors from the FBX file */
+		Replace,
+		/** Ignore vertex colors from the FBX file, and keep the existing mesh vertex colors */
+		Ignore,
+		/** Override all vertex colors with the specified color */
+		Override
+	};
+}
+
 UCLASS(config=EditorUserSettings, AutoExpandCategories=(Options), MinimalAPI)
 class UFbxStaticMeshImportData : public UFbxMeshImportData
 {
@@ -16,9 +30,13 @@ class UFbxStaticMeshImportData : public UFbxMeshImportData
 	UPROPERTY(EditAnywhere, AdvancedDisplay, Category=ImportSettings, meta=(ImportType="StaticMesh"))
 	FName StaticMeshLODGroup;
 
-	/** If true, will replace the vertex colors on an existing mesh with the vertex colors from the FBX file */
-	UPROPERTY(EditAnywhere, AdvancedDisplay, config, Category=ImportSettings, meta=(OBJRestrict="true", ImportType="StaticMesh", ToolTip="If enabled, existing vertex colors will be replaced with ones from the imported file"))
-	uint32 bReplaceVertexColors:1;
+	/** Specify how vertex colors should be imported */
+	UPROPERTY(EditAnywhere, AdvancedDisplay, config, Category=ImportSettings, meta=(OBJRestrict="true", ImportType="StaticMesh"))
+	TEnumAsByte<EVertexColorImportOption::Type> VertexColorImportOption;
+
+	/** Specify override color in the case that VertexColorImportOption is set to Override */
+	UPROPERTY(EditAnywhere, AdvancedDisplay, config, Category=ImportSettings, meta=(OBJRestrict="true", ImportType="StaticMesh"))
+	FColor VertexOverrideColor;
 
 	/** Disabling this option will keep degenerate triangles found.  In general you should leave this option on. */
 	UPROPERTY(EditAnywhere, AdvancedDisplay, Category = ImportSettings, meta = (ImportType = "StaticMesh"))
diff --git a/Engine/Source/Editor/UnrealEd/Private/Fbx/FbxMainImport.cpp b/Engine/Source/Editor/UnrealEd/Private/Fbx/FbxMainImport.cpp
--- a/Engine/Source/Editor/UnrealEd/Private/Fbx/FbxMainImport.cpp
+++ b/Engine/Source/Editor/UnrealEd/Private/Fbx/FbxMainImport.cpp
@@ -226,7 +226,8 @@ void ApplyImportUIToImportOptions(UFbxImportUI* ImportUI, FBXImportOptions& InOu
 	InOutImportOptions.bPreserveSmoothingGroups = ImportUI->SkeletalMeshImportData->bPreserveSmoothingGroups;
 	InOutImportOptions.bKeepOverlappingVertices = ImportUI->SkeletalMeshImportData->bKeepOverlappingVertices;
 	InOutImportOptions.bCombineToSingle = ImportUI->bCombineMeshes;
-	InOutImportOptions.bReplaceVertexColors = ImportUI->StaticMeshImportData->bReplaceVertexColors;
+	InOutImportOptions.VertexColorImportOption = ImportUI->StaticMeshImportData->VertexColorImportOption;
+	InOutImportOptions.VertexOverrideColor = ImportUI->StaticMeshImportData->VertexOverrideColor;
 	InOutImportOptions.bRemoveDegenerates = ImportUI->StaticMeshImportData->bRemoveDegenerates;
 	InOutImportOptions.bGenerateLightmapUVs = ImportUI->StaticMeshImportData->bGenerateLightmapUVs;
 	InOutImportOptions.bOneConvexHullPerUCX = ImportUI->StaticMeshImportData->bOneConvexHullPerUCX;
diff --git a/Engine/Source/Editor/UnrealEd/Private/Fbx/FbxStaticMeshImport.cpp b/Engine/Source/Editor/UnrealEd/Private/Fbx/FbxStaticMeshImport.cpp
--- a/Engine/Source/Editor/UnrealEd/Private/Fbx/FbxStaticMeshImport.cpp
+++ b/Engine/Source/Editor/UnrealEd/Private/Fbx/FbxStaticMeshImport.cpp
@@ -91,7 +91,8 @@ UStaticMesh* UnFbx::FFbxImporter::ImportStaticMesh(UObject* InParent, FbxNode* N
 	return ImportStaticMeshAsSingle(InParent, MeshNodeArray, Name, Flags, ImportData, InStaticMesh, LODIndex);
 }
 
-bool UnFbx::FFbxImporter::BuildStaticMeshFromGeometry(FbxMesh* Mesh, UStaticMesh* StaticMesh, TArray<FFbxMaterial>& MeshMaterials, int LODIndex, TMap<FVector, FColor>* ExistingVertexColorData)
+bool UnFbx::FFbxImporter::BuildStaticMeshFromGeometry(FbxMesh* Mesh, UStaticMesh* StaticMesh, TArray<FFbxMaterial>& MeshMaterials, int LODIndex,
+													  EVertexColorImportOption::Type VertexColorImportOption, const TMap<FVector, FColor>& ExistingVertexColorData, const FColor& VertexOverrideColor)
 {
 	check(StaticMesh->SourceModels.IsValidIndex(LODIndex));
 
@@ -462,7 +463,7 @@ bool UnFbx::FFbxImporter::BuildStaticMeshFromGeometry(FbxMesh* Mesh, UStaticMesh
 		RawMesh.WedgeTangentZ.AddZeroed(WedgeOffset + WedgeCount - RawMesh.WedgeTangentZ.Num());
 	}
 
-	if (LayerElementVertexColor || ExistingVertexColorData || RawMesh.WedgeColors.Num() )
+	if (LayerElementVertexColor || VertexColorImportOption != EVertexColorImportOption::Replace || RawMesh.WedgeColors.Num() )
 	{
 		int32 NumNewColors = WedgeOffset + WedgeCount - RawMesh.WedgeColors.Num();
 		int32 FirstNewColor = RawMesh.WedgeColors.Num();
@@ -571,35 +572,44 @@ bool UnFbx::FFbxImporter::BuildStaticMeshFromGeometry(FbxMesh* Mesh, UStaticMesh
 			//
 			// vertex colors
 			//
-			if (LayerElementVertexColor && !ExistingVertexColorData)
+			if (VertexColorImportOption == EVertexColorImportOption::Replace)
 			{
-				int32 VertexColorMappingIndex = (VertexColorMappingMode == FbxLayerElement::eByControlPoint) ? 
-					Mesh->GetPolygonVertex(TriangleIndex,CornerIndex) : (TriangleIndex*3+CornerIndex);
+				if (LayerElementVertexColor)
+				{
+					int32 VertexColorMappingIndex = (VertexColorMappingMode == FbxLayerElement::eByControlPoint) ?
+						Mesh->GetPolygonVertex(TriangleIndex, CornerIndex) : (TriangleIndex * 3 + CornerIndex);
 
-				int32 VectorColorIndex = (VertexColorReferenceMode == FbxLayerElement::eDirect) ? 
+					int32 VectorColorIndex = (VertexColorReferenceMode == FbxLayerElement::eDirect) ?
 					VertexColorMappingIndex : LayerElementVertexColor->GetIndexArray().GetAt(VertexColorMappingIndex);
 
-				FbxColor VertexColor = LayerElementVertexColor->GetDirectArray().GetAt(VectorColorIndex);
+					FbxColor VertexColor = LayerElementVertexColor->GetDirectArray().GetAt(VectorColorIndex);
 
-				RawMesh.WedgeColors[WedgeIndex] = FColor(
-					uint8(255.f*VertexColor.mRed),
-					uint8(255.f*VertexColor.mGreen),
-					uint8(255.f*VertexColor.mBlue),
-					uint8(255.f*VertexColor.mAlpha)
-					);
+					RawMesh.WedgeColors[WedgeIndex] = FColor(
+						uint8(255.f*VertexColor.mRed),
+						uint8(255.f*VertexColor.mGreen),
+						uint8(255.f*VertexColor.mBlue),
+						uint8(255.f*VertexColor.mAlpha)
+						);
+				}
 			}
-			else if (ExistingVertexColorData)
+			else if (VertexColorImportOption == EVertexColorImportOption::Ignore)
 			{
 				// try to match this triangles current vertex with one that existed in the previous mesh.
 				// This is a find in a tmap which uses a fast hash table lookup.
 				FVector Position = RawMesh.VertexPositions[RawMesh.WedgeIndices[WedgeIndex]];
-				FColor* PaintedColor = ExistingVertexColorData->Find(Position);
+				const FColor* PaintedColor = ExistingVertexColorData.Find(Position);
 				if (PaintedColor)
 				{
 					// A matching color for this vertex was found
 					RawMesh.WedgeColors[WedgeIndex] = *PaintedColor;
 				}
 			}
+			else
+			{
+				// set the triangle's vertex color to a constant override
+				check(VertexColorImportOption == EVertexColorImportOption::Override);
+				RawMesh.WedgeColors[WedgeIndex] = VertexOverrideColor;
+			}
 		}
 
 		//
@@ -871,8 +881,8 @@ UStaticMesh* UnFbx::FFbxImporter::ImportStaticMeshAsSingle(UObject* InParent, TA
 
 	// A mapping of vertex positions to their color in the existing static mesh
 	TMap<FVector, FColor>		ExistingVertexColorData;
-	// If we should get vertex colors. Defaults to the checkbox value available to the user in the import property window
-	bool bGetVertexColors = ImportOptions->bReplaceVertexColors;
+
+	EVertexColorImportOption::Type VertexColorImportOption = ImportOptions->VertexColorImportOption;
 	FString NewPackageName;
 
 	if( InStaticMesh == NULL || LODIndex == 0 )
@@ -892,8 +902,11 @@ UStaticMesh* UnFbx::FFbxImporter::ImportStaticMeshAsSingle(UObject* InParent, TA
 
 		if (0 == ExistingVertexColorData.Num())
 		{
-			// If there were no vertex colors, automatically take vertex colors from the file.
-			bGetVertexColors = true;
+			// If there were no vertex colors and we specified to ignore FBX vertex colors, automatically take vertex colors from the file anyway.
+			if (VertexColorImportOption == EVertexColorImportOption::Ignore)
+			{
+				VertexColorImportOption = EVertexColorImportOption::Replace;
+			}
 		}
 
 		// Free any RHI resources for existing mesh before we re-create in place.
@@ -922,12 +935,18 @@ UStaticMesh* UnFbx::FFbxImporter::ImportStaticMeshAsSingle(UObject* InParent, TA
 		}
 
 		// Vertex colors should be copied always if there is no existing static mesh.
-		bGetVertexColors = true;
+		if (VertexColorImportOption == EVertexColorImportOption::Ignore)
+		{
+			VertexColorImportOption = EVertexColorImportOption::Replace;
+		}
 	}
 	else
 	{ 
 		// Vertex colors should be copied always if there is no existing static mesh.
-		bGetVertexColors = true;
+		if (VertexColorImportOption == EVertexColorImportOption::Ignore)
+		{
+			VertexColorImportOption = EVertexColorImportOption::Replace;
+		}
 	}
 	
 	if( InStaticMesh != NULL && LODIndex > 0 )
@@ -971,7 +990,8 @@ UStaticMesh* UnFbx::FFbxImporter::ImportStaticMeshAsSingle(UObject* InParent, TA
 
 		if (Node->GetMesh())
 		{
-			if (!BuildStaticMeshFromGeometry(Node->GetMesh(), StaticMesh, MeshMaterials, LODIndex, bGetVertexColors? NULL: &ExistingVertexColorData))
+			if (!BuildStaticMeshFromGeometry(Node->GetMesh(), StaticMesh, MeshMaterials, LODIndex,
+											 VertexColorImportOption, ExistingVertexColorData, ImportOptions->VertexOverrideColor))
 			{
 				bBuildStatus = false;
 				break;
diff --git a/Engine/Source/Editor/UnrealEd/Private/Fbx/FbxStaticMeshImportData.cpp b/Engine/Source/Editor/UnrealEd/Private/Fbx/FbxStaticMeshImportData.cpp
--- a/Engine/Source/Editor/UnrealEd/Private/Fbx/FbxStaticMeshImportData.cpp
+++ b/Engine/Source/Editor/UnrealEd/Private/Fbx/FbxStaticMeshImportData.cpp
@@ -10,6 +10,7 @@ UFbxStaticMeshImportData::UFbxStaticMeshImportData(const FObjectInitializer& Obj
 	bGenerateLightmapUVs = true;
 	bOneConvexHullPerUCX = true;
 	bAutoGenerateCollision = true;
+	VertexOverrideColor = FColor(255, 255, 255, 255);
 }
 
 UFbxStaticMeshImportData* UFbxStaticMeshImportData::GetImportDataForStaticMesh(UStaticMesh* StaticMesh, UFbxStaticMeshImportData* TemplateForCreation)
diff --git a/Engine/Source/Editor/UnrealEd/Private/FbxImporter.h b/Engine/Source/Editor/UnrealEd/Private/FbxImporter.h
--- a/Engine/Source/Editor/UnrealEd/Private/FbxImporter.h
+++ b/Engine/Source/Editor/UnrealEd/Private/FbxImporter.h
@@ -114,7 +114,8 @@ struct FBXImportOptions
 	EFBXNormalImportMethod NormalImportMethod;
 	// Static Mesh options
 	bool bCombineToSingle;
-	bool bReplaceVertexColors;
+	EVertexColorImportOption::Type VertexColorImportOption;
+	FColor VertexOverrideColor;
 	bool bRemoveDegenerates;
 	bool bGenerateLightmapUVs;
 	bool bOneConvexHullPerUCX;
@@ -689,7 +690,8 @@ protected:
 	 * @param LODIndex	LOD level to set up for StaticMesh
 	 * @return bool true if set up successfully
 	 */
-	bool BuildStaticMeshFromGeometry(FbxMesh* FbxMesh, UStaticMesh* StaticMesh, TArray<FFbxMaterial>& MeshMaterials, int LODIndex, TMap<FVector, FColor>* ExistingVertexColorData);
+	bool BuildStaticMeshFromGeometry(FbxMesh* FbxMesh, UStaticMesh* StaticMesh, TArray<FFbxMaterial>& MeshMaterials, int LODIndex,
+									 EVertexColorImportOption::Type VertexColorImportOption, const TMap<FVector, FColor>& ExistingVertexColorData, const FColor& VertexOverrideColor);
 	
 	/**
 	 * Clean up for destroy the Importer.
diff --git a/Engine/Source/Runtime/Slate/Private/Widgets/Input/STextComboBox.cpp b/Engine/Source/Runtime/Slate/Private/Widgets/Input/STextComboBox.cpp
--- a/Engine/Source/Runtime/Slate/Private/Widgets/Input/STextComboBox.cpp
+++ b/Engine/Source/Runtime/Slate/Private/Widgets/Input/STextComboBox.cpp
@@ -8,6 +8,7 @@ void STextComboBox::Construct( const FArguments& InArgs )
 {
 	SelectionChanged = InArgs._OnSelectionChanged;
 	GetTextLabelForItem = InArgs._OnGetTextLabelForItem;
+	Font = InArgs._Font;
 
 	// Then make widget
 	this->ChildSlot
@@ -23,6 +24,7 @@ void STextComboBox::Construct( const FArguments& InArgs )
 			SNew(STextBlock)
 				.ColorAndOpacity(InArgs._ColorAndOpacity)
 				.Text(this, &STextComboBox::GetSelectedTextLabel)
+				.Font(InArgs._Font)
 		]
 	];
 	SelectedItem = StringCombo->GetSelectedItem();
@@ -51,7 +53,8 @@ TSharedRef<SWidget> STextComboBox::MakeItemWidget( TSharedPtr<FString> StringIte
 	check( StringItem.IsValid() );
 
 	return SNew(STextBlock)
-		.Text(this, &STextComboBox::GetItemTextLabel, StringItem);
+		.Text(this, &STextComboBox::GetItemTextLabel, StringItem)
+		.Font(Font);
 }
 
 void STextComboBox::OnSelectionChanged (TSharedPtr<FString> Selection, ESelectInfo::Type SelectInfo)
diff --git a/Engine/Source/Runtime/Slate/Public/Widgets/Input/STextComboBox.h b/Engine/Source/Runtime/Slate/Public/Widgets/Input/STextComboBox.h
--- a/Engine/Source/Runtime/Slate/Public/Widgets/Input/STextComboBox.h
+++ b/Engine/Source/Runtime/Slate/Public/Widgets/Input/STextComboBox.h
@@ -24,6 +24,9 @@ public:
 		/** Text color and opacity */
 		SLATE_ATTRIBUTE( FSlateColor, ColorAndOpacity )
 
+		/** Sets the font used to draw the text */
+		SLATE_ATTRIBUTE(FSlateFontInfo, Font)
+
 		/** Visual padding of the button content for the combobox */
 		SLATE_ATTRIBUTE( FMargin, ContentPadding )
 
@@ -86,4 +89,7 @@ private:
 
 	/** Forwarding Delegate */
 	FOnTextSelectionChanged SelectionChanged;
+
+	/** Sets the font used to draw the text */
+	TAttribute< FSlateFontInfo > Font;
 };
