AnimInstance needs reinstancing even if nothing changed because the all instance values need to be re-created and re-set.
- Fix crash when compiling multiple times of the anim BP without adding/deleting node

diff --git a/Engine/Source/Editor/UnrealEd/Private/Kismet2/KismetReinstanceUtilities.cpp b/Engine/Source/Editor/UnrealEd/Private/Kismet2/KismetReinstanceUtilities.cpp
--- a/Engine/Source/Editor/UnrealEd/Private/Kismet2/KismetReinstanceUtilities.cpp
+++ b/Engine/Source/Editor/UnrealEd/Private/Kismet2/KismetReinstanceUtilities.cpp
@@ -181,9 +181,10 @@ void FBlueprintCompileReinstancer::ReinstanceObjects(bool bAlwaysReinstance)
 			auto BPClassA = Cast<const UBlueprintGeneratedClass>(DuplicatedClass);
 			auto BPClassB = Cast<const UBlueprintGeneratedClass>(ClassToReinstance);
 			auto BP = Cast<const UBlueprint>(ClassToReinstance->ClassGeneratedBy);
+			const bool bBPSubClass = (BP && BP->GetClass() != UBlueprint::StaticClass());
 			static const FBoolConfigValueHelper ChangeDefaultValueWithoutReinstancing(TEXT("Kismet"), TEXT("bChangeDefaultValueWithoutReinstancing"), GEngineIni);
 			const bool bTheSameDefaultValues = BP && ClassToReinstanceDefaultValuesCRC && (BP->CrcPreviousCompiledCDO == ClassToReinstanceDefaultValuesCRC);
-			const bool bTheSame = (ChangeDefaultValueWithoutReinstancing || bTheSameDefaultValues) && BPClassA && BPClassB && FStructUtils::TheSameLayout(BPClassA, BPClassB, true);
+			const bool bTheSame = !bBPSubClass && (ChangeDefaultValueWithoutReinstancing || bTheSameDefaultValues) && BPClassA && BPClassB && FStructUtils::TheSameLayout(BPClassA, BPClassB, true);
 			if (bTheSame)
 			{
 				UE_LOG(LogBlueprint, Log, TEXT("BlueprintCompileReinstancer: class '%s' is replaced without reinstancing (the optimized way)."), *GetPathNameSafe(ClassToReinstance));
