Fixed cook failure caused by enabling "Save on Compile" in the BP editor

diff --git a/Engine/Source/Editor/UnrealEd/Private/Kismet2/Kismet2.cpp b/Engine/Source/Editor/UnrealEd/Private/Kismet2/Kismet2.cpp
--- a/Engine/Source/Editor/UnrealEd/Private/Kismet2/Kismet2.cpp
+++ b/Engine/Source/Editor/UnrealEd/Private/Kismet2/Kismet2.cpp
@@ -492,8 +492,10 @@ void FKismetEditorUtilities::CompileBlueprint(UBlueprint* BlueprintObj, bool bIs
 		BlueprintPackage->SetDirtyFlag(bStartedWithUnsavedChanges);
 
 		UBlueprintEditorSettings* Settings = GetMutableDefault<UBlueprintEditorSettings>();
-		if ( (Settings->SaveOnCompile == SoC_Always) || 
-			((Settings->SaveOnCompile == SoC_SuccessOnly) && (BlueprintObj->Status == BS_UpToDate)) )
+		const bool bShouldSaveOnCompile = ( Settings->SaveOnCompile == SoC_Always ) || ( ( Settings->SaveOnCompile == SoC_SuccessOnly ) && ( BlueprintObj->Status == BS_UpToDate ) );
+
+		// Only try to save on compile if we aren't running a commandlet (i.e. cooking a blueprint shouldn't try to save it)
+		if ( !IsRunningCommandlet() && bShouldSaveOnCompile )
 		{
 			bool const bIsLevelPackage = (UWorld::FindWorldInPackage(BlueprintPackage) != nullptr);
 			// we don't want to save the entire level (especially if this 
