Prepare cpp files for header cleanup.

diff --git a/Engine/Source/Editor/UnrealEd/Private/Factories/ActorFactory.cpp b/Engine/Source/Editor/UnrealEd/Private/Factories/ActorFactory.cpp
--- a/Engine/Source/Editor/UnrealEd/Private/Factories/ActorFactory.cpp
+++ b/Engine/Source/Editor/UnrealEd/Private/Factories/ActorFactory.cpp
@@ -37,6 +37,33 @@ ActorFactory.cpp:
 #include "Engine/TextRenderActor.h"
 
 #include "Engine/DestructibleMesh.h"
+#include "Engine/BlueprintGeneratedClass.h"
+#include "Components/DecalComponent.h"
+#include "GameFramework/PlayerStart.h"
+#include "PhysicsEngine/PhysicsAsset.h"
+#include "Camera/CameraActor.h"
+#include "Sound/SoundBase.h"
+#include "Sound/AmbientSound.h"
+#include "Engine/DirectionalLight.h"
+#include "Engine/SpotLight.h"
+#include "Engine/PointLight.h"
+#include "Engine/SkyLight.h"
+#include "Engine/SphereReflectionCapture.h"
+#include "Engine/BoxReflectionCapture.h"
+#include "Engine/PlaneReflectionCapture.h"
+#include "Atmosphere/AtmosphericFog.h"
+#include "Engine/ExponentialHeightFog.h"
+#include "Engine/DecalActor.h"
+#include "Engine/TargetPoint.h"
+#include "Components/AudioComponent.h"
+#include "PhysicsEngine/DestructibleActor.h"
+#include "Engine/Note.h"
+#include "Components/DestructibleComponent.h"
+#include "VectorField/VectorFieldVolume.h"
+#include "Components/BrushComponent.h"
+#include "Engine/Polys.h"
+#include "Components/VectorFieldComponent.h"
+#include "Animation/AnimBlueprintGeneratedClass.h"
 
 DEFINE_LOG_CATEGORY_STATIC(LogActorFactory, Log, All);
 
diff --git a/Engine/Source/Editor/UnrealEd/Private/Factories/AnimBlueprintFactory.cpp b/Engine/Source/Editor/UnrealEd/Private/Factories/AnimBlueprintFactory.cpp
--- a/Engine/Source/Editor/UnrealEd/Private/Factories/AnimBlueprintFactory.cpp
+++ b/Engine/Source/Editor/UnrealEd/Private/Factories/AnimBlueprintFactory.cpp
@@ -23,6 +23,7 @@
 #include "ClassViewerModule.h"
 #include "ClassViewerFilter.h"
 #include "AssetRegistryModule.h"
+#include "Animation/AnimBlueprintGeneratedClass.h"
 
 #define LOCTEXT_NAMESPACE "AnimBlueprintFactory"
 
diff --git a/Engine/Source/Editor/UnrealEd/Private/Factories/AnimCompositeFactory.cpp b/Engine/Source/Editor/UnrealEd/Private/Factories/AnimCompositeFactory.cpp
--- a/Engine/Source/Editor/UnrealEd/Private/Factories/AnimCompositeFactory.cpp
+++ b/Engine/Source/Editor/UnrealEd/Private/Factories/AnimCompositeFactory.cpp
@@ -8,6 +8,7 @@
 
 #include "AssetData.h"
 #include "ContentBrowserModule.h"
+#include "Animation/AnimComposite.h"
 
 #define LOCTEXT_NAMESPACE "AnimCompositeFactory"
 
diff --git a/Engine/Source/Editor/UnrealEd/Private/Factories/CSVImportFactory.cpp b/Engine/Source/Editor/UnrealEd/Private/Factories/CSVImportFactory.cpp
--- a/Engine/Source/Editor/UnrealEd/Private/Factories/CSVImportFactory.cpp
+++ b/Engine/Source/Editor/UnrealEd/Private/Factories/CSVImportFactory.cpp
@@ -7,6 +7,11 @@
 #include "DirectoryWatcherModule.h"
 #include "../../../DataTableEditor/Public/IDataTableEditor.h"
 #include "Curves/CurveVector.h"
+#include "Curves/CurveFloat.h"
+#include "Curves/CurveLinearColor.h"
+#include "Engine/DataTable.h"
+#include "Engine/CurveTable.h"
+#include "Engine/UserDefinedStruct.h"
 DEFINE_LOG_CATEGORY(LogCSVImportFactory);
 
 #define LOCTEXT_NAMESPACE "CSVImportFactory"
diff --git a/Engine/Source/Editor/UnrealEd/Private/Factories/EditorFactories.cpp b/Engine/Source/Editor/UnrealEd/Private/Factories/EditorFactories.cpp
--- a/Engine/Source/Editor/UnrealEd/Private/Factories/EditorFactories.cpp
+++ b/Engine/Source/Editor/UnrealEd/Private/Factories/EditorFactories.cpp
@@ -61,6 +61,45 @@
 #if WITH_EDITOR
 #include "CubemapUnwrapUtils.h"			// FMipLevelBatchedElementParameters
 #endif
+#include "GameFramework/WorldSettings.h"
+#include "Materials/MaterialInstanceConstant.h"
+#include "Engine/Polys.h"
+#include "GameFramework/DefaultPhysicsVolume.h"
+#include "Components/BrushComponent.h"
+#include "Materials/MaterialParameterCollection.h"
+#include "Sound/ReverbEffect.h"
+#include "Sound/SoundCue.h"
+#include "Sound/SoundWave.h"
+#include "EngineUtils.h"
+#include "Engine/Selection.h"
+#include "Sound/DialogueVoice.h"
+#include "Sound/SoundMix.h"
+#include "PhysicalMaterials/PhysicalMaterial.h"
+#include "Engine/TextureRenderTarget2D.h"
+#include "Engine/TextureRenderTargetCube.h"
+#include "Engine/TextureCube.h"
+#include "Engine/Font.h"
+#include "Components/AudioComponent.h"
+#include "Sound/DialogueWave.h"
+#include "Engine/StaticMesh.h"
+#include "Engine/AssetUserData.h"
+#include "Engine/LevelScriptActor.h"
+#include "Engine/BlueprintGeneratedClass.h"
+#include "Curves/CurveFloat.h"
+#include "Engine/ObjectLibrary.h"
+#include "Engine/DataAsset.h"
+#include "Animation/BlendSpace1D.h"
+#include "AI/Navigation/NavCollision.h"
+#include "Kismet/BlueprintFunctionLibrary.h"
+#include "Animation/BlendSpace.h"
+#include "Animation/AimOffsetBlendSpace.h"
+#include "Animation/AimOffsetBlendSpace1D.h"
+#include "Engine/UserDefinedEnum.h"
+#include "Engine/UserDefinedStruct.h"
+#include "GameFramework/ForceFeedbackEffect.h"
+#include "Engine/SubsurfaceProfile.h"
+#include "Camera/CameraAnim.h"
+#include "GameFramework/TouchInterface.h"
 
 DEFINE_LOG_CATEGORY_STATIC(LogEditorFactories, Log, All);
 
diff --git a/Engine/Source/Editor/UnrealEd/Private/Factories/EditorMorphFactory.cpp b/Engine/Source/Editor/UnrealEd/Private/Factories/EditorMorphFactory.cpp
--- a/Engine/Source/Editor/UnrealEd/Private/Factories/EditorMorphFactory.cpp
+++ b/Engine/Source/Editor/UnrealEd/Private/Factories/EditorMorphFactory.cpp
@@ -7,6 +7,7 @@
 #include "UnrealEd.h"
 #include "Factories.h"
 #include "SkelImport.h"
+#include "Engine/StaticMesh.h"
 
 #define SET_MORPH_ERROR( ErrorType ) { if( Error ) *Error = ErrorType; }
 
diff --git a/Engine/Source/Editor/UnrealEd/Private/Factories/PackFactory.cpp b/Engine/Source/Editor/UnrealEd/Private/Factories/PackFactory.cpp
--- a/Engine/Source/Editor/UnrealEd/Private/Factories/PackFactory.cpp
+++ b/Engine/Source/Editor/UnrealEd/Private/Factories/PackFactory.cpp
@@ -13,6 +13,7 @@
 #include "ModuleManager.h"
 #include "GameProjectGenerationModule.h"
 #include "Factories/PackFactory.h"
+#include "GameFramework/InputSettings.h"
 
 DEFINE_LOG_CATEGORY_STATIC(LogPackFactory, Log, All);
 
diff --git a/Engine/Source/Editor/UnrealEd/Private/Factories/SkeletalMeshImport.cpp b/Engine/Source/Editor/UnrealEd/Private/Factories/SkeletalMeshImport.cpp
--- a/Engine/Source/Editor/UnrealEd/Private/Factories/SkeletalMeshImport.cpp
+++ b/Engine/Source/Editor/UnrealEd/Private/Factories/SkeletalMeshImport.cpp
@@ -11,6 +11,7 @@
 #include "../../../../Source/Runtime/Engine/Classes/PhysicsEngine/PhysicsAsset.h"
 #include "FbxImporter.h"
 #include "FbxErrors.h"
+#include "Engine/SkeletalMeshSocket.h"
 
 DEFINE_LOG_CATEGORY_STATIC(LogSkeletalMeshImport, Log, All);
 
diff --git a/Engine/Source/Editor/UnrealEd/Private/Factories/TTFontImport.cpp b/Engine/Source/Editor/UnrealEd/Private/Factories/TTFontImport.cpp
--- a/Engine/Source/Editor/UnrealEd/Private/Factories/TTFontImport.cpp
+++ b/Engine/Source/Editor/UnrealEd/Private/Factories/TTFontImport.cpp
@@ -7,6 +7,7 @@
 #include "Factories.h"
 #include "Editor/MainFrame/Public/MainFrame.h"
 #include "DesktopPlatformModule.h"
+#include "Engine/Font.h"
 
 #ifndef WITH_FREETYPE
 	#define WITH_FREETYPE	0
@@ -29,6 +30,7 @@ namespace TTFConstants
 #include FT_FREETYPE_H
 #include FT_GLYPH_H
 #endif // USE_FREETYPE
+#include "Engine/FontImportOptions.h"
 
 DEFINE_LOG_CATEGORY_STATIC(LogTTFontImport, Log, All);
 
diff --git a/Engine/Source/Editor/UnrealEd/Private/Fbx/FbxFactory.cpp b/Engine/Source/Editor/UnrealEd/Private/Fbx/FbxFactory.cpp
--- a/Engine/Source/Editor/UnrealEd/Private/Fbx/FbxFactory.cpp
+++ b/Engine/Source/Editor/UnrealEd/Private/Fbx/FbxFactory.cpp
@@ -37,6 +37,7 @@
 
 #include "FbxErrors.h"
 #include "AssetRegistryModule.h"
+#include "Engine/StaticMesh.h"
 
 #define LOCTEXT_NAMESPACE "FBXFactory"
 
diff --git a/Engine/Source/Editor/UnrealEd/Private/Fbx/FbxMainExport.cpp b/Engine/Source/Editor/UnrealEd/Private/Fbx/FbxMainExport.cpp
--- a/Engine/Source/Editor/UnrealEd/Private/Fbx/FbxMainExport.cpp
+++ b/Engine/Source/Editor/UnrealEd/Private/Fbx/FbxMainExport.cpp
@@ -57,6 +57,14 @@
 #include "FbxExporter.h"
 #include "RawMesh.h"
 #include "Particles/Emitter.h"
+#include "Engine/Light.h"
+#include "Engine/Polys.h"
+#include "Engine/StaticMeshActor.h"
+#include "Components/BrushComponent.h"
+#include "Components/SpotLightComponent.h"
+#include "Components/PointLightComponent.h"
+#include "Camera/CameraActor.h"
+#include "Components/DirectionalLightComponent.h"
 
 namespace UnFbx
 {
diff --git a/Engine/Source/Editor/UnrealEd/Private/Fbx/FbxMatineeImport.cpp b/Engine/Source/Editor/UnrealEd/Private/Fbx/FbxMatineeImport.cpp
--- a/Engine/Source/Editor/UnrealEd/Private/Fbx/FbxMatineeImport.cpp
+++ b/Engine/Source/Editor/UnrealEd/Private/Fbx/FbxMatineeImport.cpp
@@ -41,6 +41,7 @@
 #include "Matinee/InterpTrackInstFloatProp.h"
 
 #include "FbxImporter.h"
+#include "Camera/CameraActor.h"
 
 namespace UnFbx {
 
diff --git a/Engine/Source/Editor/UnrealEd/Private/Fbx/FbxStaticMeshImportData.cpp b/Engine/Source/Editor/UnrealEd/Private/Fbx/FbxStaticMeshImportData.cpp
--- a/Engine/Source/Editor/UnrealEd/Private/Fbx/FbxStaticMeshImportData.cpp
+++ b/Engine/Source/Editor/UnrealEd/Private/Fbx/FbxStaticMeshImportData.cpp
@@ -1,6 +1,7 @@
 // Copyright 1998-2014 Epic Games, Inc. All Rights Reserved.
 
 #include "UnrealEd.h"
+#include "Engine/StaticMesh.h"
 
 UFbxStaticMeshImportData::UFbxStaticMeshImportData(const FObjectInitializer& ObjectInitializer)
 	: Super(ObjectInitializer)
diff --git a/Engine/Source/Editor/UnrealEd/Private/Kismet2/CompilerResultsLog.cpp b/Engine/Source/Editor/UnrealEd/Private/Kismet2/CompilerResultsLog.cpp
--- a/Engine/Source/Editor/UnrealEd/Private/Kismet2/CompilerResultsLog.cpp
+++ b/Engine/Source/Editor/UnrealEd/Private/Kismet2/CompilerResultsLog.cpp
@@ -7,6 +7,7 @@
 #include "UObjectToken.h"
 #include "SourceCodeNavigation.h"
 #include "Developer/HotReload/Public/IHotReload.h"
+#include "EngineLogs.h"
 
 #if WITH_EDITOR
 
diff --git a/Engine/Source/Editor/UnrealEd/Private/Kismet2/EnumEditorUtils.cpp b/Engine/Source/Editor/UnrealEd/Private/Kismet2/EnumEditorUtils.cpp
--- a/Engine/Source/Editor/UnrealEd/Private/Kismet2/EnumEditorUtils.cpp
+++ b/Engine/Source/Editor/UnrealEd/Private/Kismet2/EnumEditorUtils.cpp
@@ -5,6 +5,7 @@
 #include "ScopedTransaction.h"
 #include "BlueprintGraphDefinitions.h"
 #include "Kismet2/BlueprintEditorUtils.h"
+#include "Engine/UserDefinedEnum.h"
 
 #define LOCTEXT_NAMESPACE "Enum"
 
diff --git a/Engine/Source/Editor/UnrealEd/Private/Layers/Layers.cpp b/Engine/Source/Editor/UnrealEd/Private/Layers/Layers.cpp
--- a/Engine/Source/Editor/UnrealEd/Private/Layers/Layers.cpp
+++ b/Engine/Source/Editor/UnrealEd/Private/Layers/Layers.cpp
@@ -5,7 +5,8 @@
 #include "Layers.h"
 #include "ActorEditorUtils.h"
 #include "Layers/Layer.h"
-
+#include "Engine/Selection.h"
+#include "EngineUtils.h"
 
 FLayers::FLayers( const TWeakObjectPtr< UEditorEngine >& InEditor )
 {
diff --git a/Engine/Source/Editor/UnrealEd/Private/Lightmass/Lightmass.cpp b/Engine/Source/Editor/UnrealEd/Private/Lightmass/Lightmass.cpp
--- a/Engine/Source/Editor/UnrealEd/Private/Lightmass/Lightmass.cpp
+++ b/Engine/Source/Editor/UnrealEd/Private/Lightmass/Lightmass.cpp
@@ -19,6 +19,7 @@
 #include "Lightmass/PrecomputedVisibilityOverrideVolume.h"
 #include "ComponentReregisterContext.h"
 #include "ShaderCompiler.h"
+#include "Lightmass/Lightmass.h"
 
 extern FSwarmDebugOptions GSwarmDebugOptions;
 
diff --git a/Engine/Source/Editor/UnrealEd/Private/Lightmass/Lightmass.h b/Engine/Source/Editor/UnrealEd/Private/Lightmass/Lightmass.h
--- a/Engine/Source/Editor/UnrealEd/Private/Lightmass/Lightmass.h
+++ b/Engine/Source/Editor/UnrealEd/Private/Lightmass/Lightmass.h
@@ -19,6 +19,7 @@
 
 #include "Lightmass/LightmassImportanceVolume.h"
 #include "Lightmass/LightmassCharacterIndirectDetailVolume.h"
+#include "GameFramework/WorldSettings.h"
 
 /** Forward declarations of Lightmass types */
 namespace Lightmass
diff --git a/Engine/Source/Editor/UnrealEd/Private/Lightmass/LightmassRender.cpp b/Engine/Source/Editor/UnrealEd/Private/Lightmass/LightmassRender.cpp
--- a/Engine/Source/Editor/UnrealEd/Private/Lightmass/LightmassRender.cpp
+++ b/Engine/Source/Editor/UnrealEd/Private/Lightmass/LightmassRender.cpp
@@ -12,6 +12,8 @@
 #include "MaterialCompiler.h"
 #include "LightMap.h"
 #include "Engine/TextureRenderTarget2D.h"
+#include "Materials/MaterialInstance.h"
+#include "CanvasTypes.h"
 
 DEFINE_LOG_CATEGORY_STATIC(LogLightmassRender, Error, All);
 
diff --git a/Engine/Source/Editor/UnrealEd/Private/Settings/SettingsClasses.cpp b/Engine/Source/Editor/UnrealEd/Private/Settings/SettingsClasses.cpp
--- a/Engine/Source/Editor/UnrealEd/Private/Settings/SettingsClasses.cpp
+++ b/Engine/Source/Editor/UnrealEd/Private/Settings/SettingsClasses.cpp
@@ -2,6 +2,9 @@
 
 #include "UnrealEd.h"
 #include "ISourceControlModule.h"
+#include "Components/BillboardComponent.h"
+#include "AI/Navigation/NavigationSystem.h"
+#include "Components/ArrowComponent.h"
 
 /* UContentBrowserSettings interface
  *****************************************************************************/
diff --git a/Engine/Source/Editor/UnrealEd/Private/StaticLightingSystem/StaticLightingExport.cpp b/Engine/Source/Editor/UnrealEd/Private/StaticLightingSystem/StaticLightingExport.cpp
--- a/Engine/Source/Editor/UnrealEd/Private/StaticLightingSystem/StaticLightingExport.cpp
+++ b/Engine/Source/Editor/UnrealEd/Private/StaticLightingSystem/StaticLightingExport.cpp
@@ -16,6 +16,9 @@
 #include "Lightmass/Lightmass.h"
 #include "Lightmass/LightmappedSurfaceCollection.h"
 
+#include "Materials/MaterialInstanceConstant.h"
+
+
 // Doxygen cannot parse these definitions correctly since the declaration is in a header from a different module
 #if !UE_BUILD_DOCS
 /** 
diff --git a/Engine/Source/Editor/UnrealEd/Private/Tests/AutomationEditorCommon.cpp b/Engine/Source/Editor/UnrealEd/Private/Tests/AutomationEditorCommon.cpp
--- a/Engine/Source/Editor/UnrealEd/Private/Tests/AutomationEditorCommon.cpp
+++ b/Engine/Source/Editor/UnrealEd/Private/Tests/AutomationEditorCommon.cpp
@@ -11,6 +11,44 @@
 #include "BehaviorTree/BehaviorTree.h"
 #include "Engine/DestructibleMesh.h"
 #include "FileManagerGeneric.h"
+#include "GameFramework/WorldSettings.h"
+#include "Animation/AimOffsetBlendSpace.h"
+#include "Animation/AimOffsetBlendSpace1D.h"
+#include "Animation/AnimBlueprint.h"
+#include "Animation/AnimMontage.h"
+#include "Animation/AnimSequence.h"
+#include "Animation/BlendSpace.h"
+#include "Animation/BlendSpace1D.h"
+#include "Engine/Blueprint.h"
+#include "Sound/DialogueVoice.h"
+#include "Sound/DialogueWave.h"
+#include "Engine/Font.h"
+#include "Materials/Material.h"
+#include "Materials/MaterialFunction.h"
+#include "Materials/MaterialInstanceConstant.h"
+#include "Materials/MaterialParameterCollection.h"
+#include "Particles/ParticleSystem.h"
+#include "PhysicalMaterials/PhysicalMaterial.h"
+#include "PhysicsEngine/PhysicsAsset.h"
+#include "Sound/ReverbEffect.h"
+#include "Engine/SkeletalMesh.h"
+#include "Animation/Skeleton.h"
+#include "Slate/SlateBrushAsset.h"
+#include "SlateWidgetStyleAsset.h"
+#include "Sound/SoundAttenuation.h"
+#include "Sound/SoundClass.h"
+#include "Sound/SoundCue.h"
+#include "Sound/SoundMix.h"
+#include "Sound/SoundWave.h"
+#include "Engine/StaticMesh.h"
+#include "Engine/SubsurfaceProfile.h"
+#include "Engine/Texture.h"
+#include "Engine/Texture2D.h"
+#include "Engine/TextureCube.h"
+#include "Engine/TextureRenderTarget.h"
+#include "Engine/TextureRenderTarget2D.h"
+#include "Engine/UserDefinedEnum.h"
+#include "Engine/World.h"
 
 
 DEFINE_LOG_CATEGORY_STATIC(LogAutomationEditorCommon, Log, All);
diff --git a/Engine/Source/Editor/UnrealEd/Private/Tests/EditorAssetAutomationTests.cpp b/Engine/Source/Editor/UnrealEd/Private/Tests/EditorAssetAutomationTests.cpp
--- a/Engine/Source/Editor/UnrealEd/Private/Tests/EditorAssetAutomationTests.cpp
+++ b/Engine/Source/Editor/UnrealEd/Private/Tests/EditorAssetAutomationTests.cpp
@@ -24,6 +24,38 @@
 #include "SlateWordWrapper.h"
 #include "AutomationCommon.h"
 #include "AutomationEditorCommon.h"
+#include "Engine/Selection.h"
+#include "Animation/AimOffsetBlendSpace.h"
+#include "Animation/AimOffsetBlendSpace1D.h"
+#include "Animation/AnimBlueprint.h"
+#include "Animation/AnimComposite.h"
+#include "Animation/AnimMontage.h"
+#include "Animation/BlendSpace.h"
+#include "Animation/BlendSpace1D.h"
+#include "Engine/TextureRenderTargetCube.h"
+#include "Engine/Font.h"
+#include "Materials/MaterialInstanceConstant.h"
+#include "Materials/MaterialParameterCollection.h"
+#include "Engine/TextureRenderTarget2D.h"
+#include "Sound/DialogueVoice.h"
+#include "Sound/DialogueWave.h"
+#include "Sound/ReverbEffect.h"
+#include "Sound/SoundAttenuation.h"
+#include "Sound/SoundClass.h"
+#include "Sound/SoundCue.h"
+#include "Sound/SoundMix.h"
+#include "PhysicalMaterials/PhysicalMaterial.h"
+#include "SlateWidgetStyleAsset.h"
+#include "Engine/Blueprint.h"
+#include "Camera/CameraAnim.h"
+#include "Curves/CurveBase.h"
+#include "Curves/CurveFloat.h"
+#include "Engine/DataAsset.h"
+#include "Engine/UserDefinedStruct.h"
+#include "Engine/UserDefinedEnum.h"
+#include "GameFramework/ForceFeedbackEffect.h"
+#include "Engine/ObjectLibrary.h"
+#include "GameFramework/TouchInterface.h"
 
 #define LOCTEXT_NAMESPACE "EditorAssetAutomationTests"
 
diff --git a/Engine/Source/Editor/UnrealEd/Private/Tests/EditorAutomationTests.cpp b/Engine/Source/Editor/UnrealEd/Private/Tests/EditorAutomationTests.cpp
--- a/Engine/Source/Editor/UnrealEd/Private/Tests/EditorAutomationTests.cpp
+++ b/Engine/Source/Editor/UnrealEd/Private/Tests/EditorAutomationTests.cpp
@@ -24,6 +24,13 @@
 #include "AutomationTest.h"
 
 #include "PackageTools.h"
+#include "Engine/PointLight.h"
+#include "Engine/Selection.h"
+#include "Engine/DirectionalLight.h"
+#include "Engine/StaticMesh.h"
+#include "Engine/BlockingVolume.h"
+#include "Engine/StaticMeshActor.h"
+#include "EngineUtils.h"
 
 
 /**
diff --git a/Engine/Source/Editor/UnrealEd/Private/Tests/EditorOpenAssetAutomationTests.cpp b/Engine/Source/Editor/UnrealEd/Private/Tests/EditorOpenAssetAutomationTests.cpp
--- a/Engine/Source/Editor/UnrealEd/Private/Tests/EditorOpenAssetAutomationTests.cpp
+++ b/Engine/Source/Editor/UnrealEd/Private/Tests/EditorOpenAssetAutomationTests.cpp
@@ -14,6 +14,25 @@
 #include "Slate/SlateBrushAsset.h"
 #include "BehaviorTree/BehaviorTree.h"
 #include "Engine/DestructibleMesh.h"
+#include "Animation/AimOffsetBlendSpace.h"
+#include "Animation/AimOffsetBlendSpace1D.h"
+#include "Animation/BlendSpace1D.h"
+#include "Engine/UserDefinedEnum.h"
+#include "Sound/DialogueVoice.h"
+#include "Sound/DialogueWave.h"
+#include "Materials/MaterialInstanceConstant.h"
+#include "Materials/MaterialParameterCollection.h"
+#include "PhysicalMaterials/PhysicalMaterial.h"
+#include "PhysicsEngine/PhysicsAsset.h"
+#include "Sound/ReverbEffect.h"
+#include "Sound/SoundAttenuation.h"
+#include "Engine/StaticMesh.h"
+#include "Sound/SoundClass.h"
+#include "Sound/SoundCue.h"
+#include "Sound/SoundMix.h"
+#include "Sound/SoundWave.h"
+#include "Engine/SubsurfaceProfile.h"
+#include "Engine/Font.h"
 
 #define ASSET_OPEN_INTERVAL 1.5f
 
diff --git a/Engine/Source/Editor/UnrealEd/Private/ThumbnailRendering/BlendSpaceThumbnailRenderer.cpp b/Engine/Source/Editor/UnrealEd/Private/ThumbnailRendering/BlendSpaceThumbnailRenderer.cpp
--- a/Engine/Source/Editor/UnrealEd/Private/ThumbnailRendering/BlendSpaceThumbnailRenderer.cpp
+++ b/Engine/Source/Editor/UnrealEd/Private/ThumbnailRendering/BlendSpaceThumbnailRenderer.cpp
@@ -6,6 +6,7 @@
 #include "ThumbnailHelpers.h"
 #include "EngineModule.h"
 #include "RendererInterface.h"
+#include "Animation/BlendSpaceBase.h"
 
 UBlendSpaceThumbnailRenderer::UBlendSpaceThumbnailRenderer(const FObjectInitializer& ObjectInitializer)
 	: Super(ObjectInitializer)
diff --git a/Engine/Source/Editor/UnrealEd/Private/ThumbnailRendering/BlueprintThumbnailRenderer.cpp b/Engine/Source/Editor/UnrealEd/Private/ThumbnailRendering/BlueprintThumbnailRenderer.cpp
--- a/Engine/Source/Editor/UnrealEd/Private/ThumbnailRendering/BlueprintThumbnailRenderer.cpp
+++ b/Engine/Source/Editor/UnrealEd/Private/ThumbnailRendering/BlueprintThumbnailRenderer.cpp
@@ -6,6 +6,8 @@
 #include "ThumbnailHelpers.h"
 #include "EngineModule.h"
 #include "RendererInterface.h"
+#include "Engine/SimpleConstructionScript.h"
+#include "Engine/SCS_Node.h"
 
 UBlueprintThumbnailRenderer::UBlueprintThumbnailRenderer(const FObjectInitializer& ObjectInitializer)
 	: Super(ObjectInitializer)
diff --git a/Engine/Source/Editor/UnrealEd/Private/ThumbnailRendering/FontThumbnailRenderer.cpp b/Engine/Source/Editor/UnrealEd/Private/ThumbnailRendering/FontThumbnailRenderer.cpp
--- a/Engine/Source/Editor/UnrealEd/Private/ThumbnailRendering/FontThumbnailRenderer.cpp
+++ b/Engine/Source/Editor/UnrealEd/Private/ThumbnailRendering/FontThumbnailRenderer.cpp
@@ -1,6 +1,9 @@
 // Copyright 1998-2014 Epic Games, Inc. All Rights Reserved.
 
 #include "UnrealEd.h"
+#include "Engine/Font.h"
+#include "CanvasItem.h"
+#include "CanvasTypes.h"
 
 UFontThumbnailRenderer::UFontThumbnailRenderer(const FObjectInitializer& ObjectInitializer)
 	: Super(ObjectInitializer)
diff --git a/Engine/Source/Editor/UnrealEd/Private/ThumbnailRendering/ParticleSystemThumbnailRenderer.cpp b/Engine/Source/Editor/UnrealEd/Private/ThumbnailRendering/ParticleSystemThumbnailRenderer.cpp
--- a/Engine/Source/Editor/UnrealEd/Private/ThumbnailRendering/ParticleSystemThumbnailRenderer.cpp
+++ b/Engine/Source/Editor/UnrealEd/Private/ThumbnailRendering/ParticleSystemThumbnailRenderer.cpp
@@ -8,6 +8,7 @@
 #include "Particles/ParticleSystem.h"
 #include "EngineModule.h"
 #include "RendererInterface.h"
+#include "CanvasTypes.h"
 
 UParticleSystemThumbnailRenderer::UParticleSystemThumbnailRenderer(const FObjectInitializer& ObjectInitializer)
 	: Super(ObjectInitializer)
diff --git a/Engine/Source/Editor/UnrealEd/Private/ThumbnailRendering/SlateBrushThumbnailRenderer.cpp b/Engine/Source/Editor/UnrealEd/Private/ThumbnailRendering/SlateBrushThumbnailRenderer.cpp
--- a/Engine/Source/Editor/UnrealEd/Private/ThumbnailRendering/SlateBrushThumbnailRenderer.cpp
+++ b/Engine/Source/Editor/UnrealEd/Private/ThumbnailRendering/SlateBrushThumbnailRenderer.cpp
@@ -4,6 +4,8 @@
 #include "Slate/SlateBrushAsset.h"
 // FPreviewScene derived helpers for rendering
 #include "ThumbnailHelpers.h"
+#include "CanvasTypes.h"
+#include "CanvasItem.h"
 
 USlateBrushThumbnailRenderer::USlateBrushThumbnailRenderer(const FObjectInitializer& ObjectInitializer)
 	: Super(ObjectInitializer)
diff --git a/Engine/Source/Editor/UnrealEd/Private/ThumbnailRendering/SoundWaveThumbnailRenderer.cpp b/Engine/Source/Editor/UnrealEd/Private/ThumbnailRendering/SoundWaveThumbnailRenderer.cpp
--- a/Engine/Source/Editor/UnrealEd/Private/ThumbnailRendering/SoundWaveThumbnailRenderer.cpp
+++ b/Engine/Source/Editor/UnrealEd/Private/ThumbnailRendering/SoundWaveThumbnailRenderer.cpp
@@ -2,6 +2,8 @@
 
 #include "UnrealEd.h"
 #include "Sound/SoundWave.h"
+#include "CanvasItem.h"
+#include "Audio.h"
 
 USoundWaveThumbnailRenderer::USoundWaveThumbnailRenderer(const FObjectInitializer& ObjectInitializer)
 	: Super(ObjectInitializer)
diff --git a/Engine/Source/Editor/UnrealEd/Private/ThumbnailRendering/StaticMeshThumbnailRenderer.cpp b/Engine/Source/Editor/UnrealEd/Private/ThumbnailRendering/StaticMeshThumbnailRenderer.cpp
--- a/Engine/Source/Editor/UnrealEd/Private/ThumbnailRendering/StaticMeshThumbnailRenderer.cpp
+++ b/Engine/Source/Editor/UnrealEd/Private/ThumbnailRendering/StaticMeshThumbnailRenderer.cpp
@@ -6,6 +6,7 @@
 #include "ThumbnailHelpers.h"
 #include "EngineModule.h"
 #include "RendererInterface.h"
+#include "Engine/StaticMesh.h"
 
 UStaticMeshThumbnailRenderer::UStaticMeshThumbnailRenderer(const FObjectInitializer& ObjectInitializer)
 	: Super(ObjectInitializer)
diff --git a/Engine/Source/Editor/UnrealEd/Private/ThumbnailRendering/SubsurfaceProfileRenderer.cpp b/Engine/Source/Editor/UnrealEd/Private/ThumbnailRendering/SubsurfaceProfileRenderer.cpp
--- a/Engine/Source/Editor/UnrealEd/Private/ThumbnailRendering/SubsurfaceProfileRenderer.cpp
+++ b/Engine/Source/Editor/UnrealEd/Private/ThumbnailRendering/SubsurfaceProfileRenderer.cpp
@@ -2,6 +2,7 @@
 
 #include "UnrealEd.h"
 #include "Engine/SubsurfaceProfile.h"
+#include "CanvasTypes.h"
 
 USubsurfaceProfileRenderer::USubsurfaceProfileRenderer(const FObjectInitializer& ObjectInitializer)
 	: Super(ObjectInitializer)
diff --git a/Engine/Source/Editor/UnrealEd/Private/ThumbnailRendering/TextureThumbnailRenderer.cpp b/Engine/Source/Editor/UnrealEd/Private/ThumbnailRendering/TextureThumbnailRenderer.cpp
--- a/Engine/Source/Editor/UnrealEd/Private/ThumbnailRendering/TextureThumbnailRenderer.cpp
+++ b/Engine/Source/Editor/UnrealEd/Private/ThumbnailRendering/TextureThumbnailRenderer.cpp
@@ -4,6 +4,10 @@
 
 #include "CubemapUnwrapUtils.h"			// FMipLevelBatchedElementParameters
 #include "NormalMapPreview.h"
+#include "CanvasItem.h"
+#include "CanvasTypes.h"
+#include "Engine/TextureCube.h"
+#include "Engine/TextureRenderTargetCube.h"
 
 UTextureThumbnailRenderer::UTextureThumbnailRenderer(const FObjectInitializer& ObjectInitializer)
 	: Super(ObjectInitializer)
diff --git a/Engine/Source/Editor/UnrealEd/Public/DragAndDrop/LevelDragDropOp.h b/Engine/Source/Editor/UnrealEd/Public/DragAndDrop/LevelDragDropOp.h
--- a/Engine/Source/Editor/UnrealEd/Public/DragAndDrop/LevelDragDropOp.h
+++ b/Engine/Source/Editor/UnrealEd/Public/DragAndDrop/LevelDragDropOp.h
@@ -1,5 +1,7 @@
 // Copyright 1998-2014 Epic Games, Inc. All Rights Reserved.
 #pragma once
+#include "DragAndDrop.h"
+#include "Engine/LevelStreaming.h"
 
 class FLevelDragDropOp : public FDragDropOperation
 {
diff --git a/Engine/Source/Editor/UnrealEd/Public/EditorModeInterpolation.h b/Engine/Source/Editor/UnrealEd/Public/EditorModeInterpolation.h
--- a/Engine/Source/Editor/UnrealEd/Public/EditorModeInterpolation.h
+++ b/Engine/Source/Editor/UnrealEd/Public/EditorModeInterpolation.h
@@ -6,6 +6,9 @@
 =============================================================================*/
 
 #pragma once
+#include "EditorModeTools.h"
+
+class UInterpGroup;
 
 //////////////////////////////////////////////////////////////////////////
 // FEdModeInterpEdit
diff --git a/Engine/Source/Editor/UnrealEd/Public/EditorModeRegistry.h b/Engine/Source/Editor/UnrealEd/Public/EditorModeRegistry.h
--- a/Engine/Source/Editor/UnrealEd/Public/EditorModeRegistry.h
+++ b/Engine/Source/Editor/UnrealEd/Public/EditorModeRegistry.h
@@ -1,6 +1,8 @@
 // Copyright 1998-2014 Epic Games, Inc. All Rights Reserved.
 
 #pragma once
+#include "Editor.h"
+#include "SlateIcon.h"
 
 /**
  *	Class responsible for maintaining a list of registered editor mode types.
diff --git a/Engine/Source/Editor/UnrealEd/Public/EditorModeTools.h b/Engine/Source/Editor/UnrealEd/Public/EditorModeTools.h
--- a/Engine/Source/Editor/UnrealEd/Public/EditorModeTools.h
+++ b/Engine/Source/Editor/UnrealEd/Public/EditorModeTools.h
@@ -23,7 +23,7 @@ enum EGeomSelectionStatus
 
 
 
-enum EModeTools
+enum EModeTools : int8
 {
 	MT_None,
 	MT_InterpEdit,
diff --git a/Engine/Source/Editor/UnrealEd/Public/EditorModes.h b/Engine/Source/Editor/UnrealEd/Public/EditorModes.h
--- a/Engine/Source/Editor/UnrealEd/Public/EditorModes.h
+++ b/Engine/Source/Editor/UnrealEd/Public/EditorModes.h
@@ -3,6 +3,15 @@
 #pragma once
 
 #include "EditorModeRegistry.h"
+#include "EditorComponents.h"
+#include "UnrealWidget.h"
+#include "EngineGlobals.h"
+
+enum EModeTools : int8;
+class FEditorViewportClient;
+class HHitProxy;
+struct FViewportClick;
+class FModeTool;
 
 /** Outcomes when determining whether it's possible to perform an action on the edit modes*/
 namespace EEditAction
diff --git a/Engine/Source/Editor/UnrealEd/Public/Factories.h b/Engine/Source/Editor/UnrealEd/Public/Factories.h
--- a/Engine/Source/Editor/UnrealEd/Public/Factories.h
+++ b/Engine/Source/Editor/UnrealEd/Public/Factories.h
@@ -6,6 +6,7 @@
 
 #ifndef __FACTORIES_H__
 #define __FACTORIES_H__
+#include "Animation/VertexAnim/MorphTarget.h"
 
 /**  
  * This class is a simple customizable object factory driven from a text buffer.  
diff --git a/Engine/Source/Editor/UnrealEd/Public/GameModeInfoCustomizer.h b/Engine/Source/Editor/UnrealEd/Public/GameModeInfoCustomizer.h
--- a/Engine/Source/Editor/UnrealEd/Public/GameModeInfoCustomizer.h
+++ b/Engine/Source/Editor/UnrealEd/Public/GameModeInfoCustomizer.h
@@ -4,6 +4,8 @@
 
 #include "IDocumentation.h"
 #include "KismetEditorUtilities.h"
+#include "GameFramework/GameMode.h"
+#include "Engine/BlueprintGeneratedClass.h"
 
 #define LOCTEXT_NAMESPACE "FGameModeInfoCustomizer"
 
diff --git a/Engine/Source/Editor/UnrealEd/Public/LevelEditorViewport.h b/Engine/Source/Editor/UnrealEd/Public/LevelEditorViewport.h
--- a/Engine/Source/Editor/UnrealEd/Public/LevelEditorViewport.h
+++ b/Engine/Source/Editor/UnrealEd/Public/LevelEditorViewport.h
@@ -4,6 +4,7 @@
 #pragma once
 
 #include "EditorViewportClient.h"
+#include "Camera/CameraComponent.h"
 
 // Forward declarations.
 
diff --git a/Engine/Source/Editor/UnrealEd/Public/MatineeExporter.h b/Engine/Source/Editor/UnrealEd/Public/MatineeExporter.h
--- a/Engine/Source/Editor/UnrealEd/Public/MatineeExporter.h
+++ b/Engine/Source/Editor/UnrealEd/Public/MatineeExporter.h
@@ -15,6 +15,9 @@ class UInterpData;
 class UInterpTrackMove;
 class UInterpTrackFloatProp;
 class UInterpTrackVectorProp;
+class AActor;
+class AMatineeActor;
+class ALight;
 
 /**
  * Main Matinee Exporter class.
diff --git a/Engine/Source/Editor/UnrealEd/Public/SCurveEditor.h b/Engine/Source/Editor/UnrealEd/Public/SCurveEditor.h
--- a/Engine/Source/Editor/UnrealEd/Public/SCurveEditor.h
+++ b/Engine/Source/Editor/UnrealEd/Public/SCurveEditor.h
@@ -4,6 +4,8 @@
 #ifndef __SCurveEditor_h__
 #define __SCurveEditor_h__
 
+#include "EditorUndoClient.h"
+
 //////////////////////////////////////////////////////////////////////////
 // FTrackScaleInfo
 
