Prepare cpp files for header cleanup.

diff --git a/Engine/Source/Editor/Layers/Private/ActorLayerCollectionViewModel.h b/Engine/Source/Editor/Layers/Private/ActorLayerCollectionViewModel.h
--- a/Engine/Source/Editor/Layers/Private/ActorLayerCollectionViewModel.h
+++ b/Engine/Source/Editor/Layers/Private/ActorLayerCollectionViewModel.h
@@ -1,5 +1,6 @@
 // Copyright 1998-2014 Epic Games, Inc. All Rights Reserved.
 #pragma once
+#include "EditorUndoClient.h"
 
 /**
  * The non-UI solution specific presentation logic for a collection of layers that belong to a collection of actors
diff --git a/Engine/Source/Editor/Layers/Private/ActorLayerViewModel.cpp b/Engine/Source/Editor/Layers/Private/ActorLayerViewModel.cpp
--- a/Engine/Source/Editor/Layers/Private/ActorLayerViewModel.cpp
+++ b/Engine/Source/Editor/Layers/Private/ActorLayerViewModel.cpp
@@ -1,7 +1,7 @@
 // Copyright 1998-2014 Epic Games, Inc. All Rights Reserved.
 
 #include "LayersPrivatePCH.h"
-
+#include "Layers/Layer.h"
 #include "ScopedTransaction.h"
 
 #define LOCTEXT_NAMESPACE "Layer"
diff --git a/Engine/Source/Editor/Layers/Private/ActorLayerViewModel.h b/Engine/Source/Editor/Layers/Private/ActorLayerViewModel.h
--- a/Engine/Source/Editor/Layers/Private/ActorLayerViewModel.h
+++ b/Engine/Source/Editor/Layers/Private/ActorLayerViewModel.h
@@ -1,5 +1,6 @@
 // Copyright 1998-2014 Epic Games, Inc. All Rights Reserved.
 #pragma once
+#include "EditorUndoClient.h"
 
 /**
  * The non-UI solution specific presentation logic for a single Layer
diff --git a/Engine/Source/Editor/Layers/Private/LayerViewModel.cpp b/Engine/Source/Editor/Layers/Private/LayerViewModel.cpp
--- a/Engine/Source/Editor/Layers/Private/LayerViewModel.cpp
+++ b/Engine/Source/Editor/Layers/Private/LayerViewModel.cpp
@@ -1,6 +1,7 @@
 // Copyright 1998-2014 Epic Games, Inc. All Rights Reserved.
 
 #include "LayersPrivatePCH.h"
+#include "Layers/Layer.h"
 
 #define LOCTEXT_NAMESPACE "Layer"
 
diff --git a/Engine/Source/Editor/Layers/Private/LayerViewModel.h b/Engine/Source/Editor/Layers/Private/LayerViewModel.h
--- a/Engine/Source/Editor/Layers/Private/LayerViewModel.h
+++ b/Engine/Source/Editor/Layers/Private/LayerViewModel.h
@@ -1,5 +1,7 @@
 // Copyright 1998-2014 Epic Games, Inc. All Rights Reserved.
 #pragma once
+#include "EditorUndoClient.h"
+#include "Layers/Layer.h"
 
 
 /**
diff --git a/Engine/Source/Editor/LevelEditor/Private/LevelEditorActions.cpp b/Engine/Source/Editor/LevelEditor/Private/LevelEditorActions.cpp
--- a/Engine/Source/Editor/LevelEditor/Private/LevelEditorActions.cpp
+++ b/Engine/Source/Editor/LevelEditor/Private/LevelEditorActions.cpp
@@ -51,6 +51,13 @@
 #include "ISourceControlModule.h"
 #include "NotificationManager.h"
 #include "SNotificationList.h"
+#include "Engine/Selection.h"
+#include "EngineUtils.h"
+#include "Engine/StaticMeshActor.h"
+#include "Engine/Polys.h"
+#include "Components/LightComponent.h"
+#include "Engine/StaticMesh.h"
+#include "Engine/Light.h"
 
 DEFINE_LOG_CATEGORY_STATIC(LevelEditorActions, Log, All);
 
diff --git a/Engine/Source/Editor/LevelEditor/Private/LevelEditorContextMenu.cpp b/Engine/Source/Editor/LevelEditor/Private/LevelEditorContextMenu.cpp
--- a/Engine/Source/Editor/LevelEditor/Private/LevelEditorContextMenu.cpp
+++ b/Engine/Source/Editor/LevelEditor/Private/LevelEditorContextMenu.cpp
@@ -29,6 +29,7 @@
 #include "SourceCodeNavigation.h"
 #include "EditorClassUtils.h"
 #include "GenericCommands.h"
+#include "Engine/Selection.h"
 
 #define LOCTEXT_NAMESPACE "LevelViewportContextMenu"
 
diff --git a/Engine/Source/Editor/LevelEditor/Private/LevelEditorCreateActorMenu.cpp b/Engine/Source/Editor/LevelEditor/Private/LevelEditorCreateActorMenu.cpp
--- a/Engine/Source/Editor/LevelEditor/Private/LevelEditorCreateActorMenu.cpp
+++ b/Engine/Source/Editor/LevelEditor/Private/LevelEditorCreateActorMenu.cpp
@@ -10,6 +10,8 @@
 #include "LevelEditorActions.h"
 #include "IPlacementModeModule.h"
 #include "Engine/TriggerBase.h"
+#include "GameFramework/Volume.h"
+#include "Engine/BlockingVolume.h"
 
 class SMenuThumbnail : public SCompoundWidget
 {
diff --git a/Engine/Source/Editor/LevelEditor/Private/LevelEditorToolBar.cpp b/Engine/Source/Editor/LevelEditor/Private/LevelEditorToolBar.cpp
--- a/Engine/Source/Editor/LevelEditor/Private/LevelEditorToolBar.cpp
+++ b/Engine/Source/Editor/LevelEditor/Private/LevelEditorToolBar.cpp
@@ -27,6 +27,14 @@
 #include "ISourceControlModule.h"
 #include "SVolumeControl.h"
 #include "ModuleManager.h"
+#include "GameFramework/GameMode.h"
+#include "GameFramework/GameState.h"
+#include "GameFramework/HUD.h"
+#include "GameFramework/Pawn.h"
+#include "GameFramework/PlayerController.h"
+#include "GameFramework/WorldSettings.h"
+#include "EngineUtils.h"
+#include "GameMapsSettings.h"
 
 
 namespace LevelEditorActionHelpers
diff --git a/Engine/Source/Editor/LevelEditor/Private/LevelViewportActions.cpp b/Engine/Source/Editor/LevelEditor/Private/LevelViewportActions.cpp
--- a/Engine/Source/Editor/LevelEditor/Private/LevelViewportActions.cpp
+++ b/Engine/Source/Editor/LevelEditor/Private/LevelViewportActions.cpp
@@ -6,6 +6,7 @@
 #include "EditorShowFlags.h"
 #include "StatsData.h"
 #include "BufferVisualizationData.h"
+#include "GameFramework/WorldSettings.h"
 
 #define LOCTEXT_NAMESPACE "LevelViewportActions"
 
diff --git a/Engine/Source/Editor/LevelEditor/Private/LightmapResRatioAdjust.cpp b/Engine/Source/Editor/LevelEditor/Private/LightmapResRatioAdjust.cpp
--- a/Engine/Source/Editor/LevelEditor/Private/LightmapResRatioAdjust.cpp
+++ b/Engine/Source/Editor/LevelEditor/Private/LightmapResRatioAdjust.cpp
@@ -10,6 +10,9 @@
 #include "SurfaceIterators.h"
 #include "MessageLog.h"
 #include "ActorEditorUtils.h"
+#include "Engine/StaticMesh.h"
+#include "Engine/Polys.h"
+#include "Engine/LevelStreaming.h"
 
 #define LOCTEXT_NAMESPACE "LightmapResRatioAdjustSettings"
 
diff --git a/Engine/Source/Editor/LevelEditor/Private/SLevelViewport.cpp b/Engine/Source/Editor/LevelEditor/Private/SLevelViewport.cpp
--- a/Engine/Source/Editor/LevelEditor/Private/SLevelViewport.cpp
+++ b/Engine/Source/Editor/LevelEditor/Private/SLevelViewport.cpp
@@ -34,6 +34,12 @@
 #include "SNotificationList.h"
 #include "NotificationManager.h"
 #include "SLevelViewportControlsPopup.h"
+#include "Camera/CameraActor.h"
+#include "GameFramework/WorldSettings.h"
+#include "Engine/LocalPlayer.h"
+#include "Engine/Selection.h"
+#include "GameFramework/PlayerInput.h"
+#include "GameFramework/PlayerController.h"
 
 static const FName LevelEditorName("LevelEditor");
 
diff --git a/Engine/Source/Editor/LevelEditor/Private/SSurfaceProperties.cpp b/Engine/Source/Editor/LevelEditor/Private/SSurfaceProperties.cpp
--- a/Engine/Source/Editor/LevelEditor/Private/SSurfaceProperties.cpp
+++ b/Engine/Source/Editor/LevelEditor/Private/SSurfaceProperties.cpp
@@ -8,6 +8,8 @@
 #include "SurfaceIterators.h"
 #include "SNumericEntryBox.h"
 #include "SHyperlink.h"
+#include "Engine/Polys.h"
+#include "Lightmass/LightmassPrimitiveSettingsObject.h"
 
 #define LOCTEXT_NAMESPACE "SSurfaceProperties"
 
diff --git a/Engine/Source/Editor/LevelEditor/Private/SSurfaceProperties.h b/Engine/Source/Editor/LevelEditor/Private/SSurfaceProperties.h
--- a/Engine/Source/Editor/LevelEditor/Private/SSurfaceProperties.h
+++ b/Engine/Source/Editor/LevelEditor/Private/SSurfaceProperties.h
@@ -1,5 +1,6 @@
 // Copyright 1998-2014 Epic Games, Inc. All Rights Reserved.
 #pragma once
+#include "Lightmass/LightmassPrimitiveSettingsObject.h"
 
 /** 
  * Surface Properties
diff --git a/Engine/Source/Editor/Levels/Private/LevelCollectionViewModel.cpp b/Engine/Source/Editor/Levels/Private/LevelCollectionViewModel.cpp
--- a/Engine/Source/Editor/Levels/Private/LevelCollectionViewModel.cpp
+++ b/Engine/Source/Editor/Levels/Private/LevelCollectionViewModel.cpp
@@ -14,6 +14,11 @@
 #include "ISourceControlModule.h"
 #include "ISourceControlRevision.h"
 #include "SourceControlWindows.h"
+#include "Engine/LevelStreamingKismet.h"
+#include "Engine/LevelStreaming.h"
+#include "Engine/Selection.h"
+#include "Engine/LevelStreamingVolume.h"
+#include "Engine/LevelStreamingAlwaysLoaded.h"
 
 #define LOCTEXT_NAMESPACE "LevelsView"
 
diff --git a/Engine/Source/Editor/Levels/Private/LevelCollectionViewModel.h b/Engine/Source/Editor/Levels/Private/LevelCollectionViewModel.h
--- a/Engine/Source/Editor/Levels/Private/LevelCollectionViewModel.h
+++ b/Engine/Source/Editor/Levels/Private/LevelCollectionViewModel.h
@@ -2,8 +2,10 @@
 #pragma once
 
 #include "Developer/AssetTools/Public/IAssetTypeActions.h"
+#include "EditorUndoClient.h"
 
 class FAssetData;
+class ALevelStreamingVolume;
 
 typedef TFilterCollection< const TSharedPtr< FLevelViewModel >& > LevelFilterCollection;
 typedef IFilter< const TSharedPtr< FLevelViewModel >& > LevelFilter;
diff --git a/Engine/Source/Editor/Levels/Private/LevelEdMode.cpp b/Engine/Source/Editor/Levels/Private/LevelEdMode.cpp
--- a/Engine/Source/Editor/Levels/Private/LevelEdMode.cpp
+++ b/Engine/Source/Editor/Levels/Private/LevelEdMode.cpp
@@ -4,6 +4,9 @@
 #include "LevelEdMode.h"
 #include "LevelUtils.h"
 #include "ActorEditorUtils.h"
+#include "Materials/MaterialInstanceDynamic.h"
+#include "Engine/LevelStreaming.h"
+#include "GameFramework/WorldSettings.h"
 
 #define LOCTEXT_NAMESPACE "LevelTransform_Mode"
 
diff --git a/Engine/Source/Editor/Levels/Private/LevelViewModel.cpp b/Engine/Source/Editor/Levels/Private/LevelViewModel.cpp
--- a/Engine/Source/Editor/Levels/Private/LevelViewModel.cpp
+++ b/Engine/Source/Editor/Levels/Private/LevelViewModel.cpp
@@ -9,6 +9,8 @@
 #include "Toolkits/AssetEditorManager.h"
 #include "EditorSupportDelegates.h"
 #include "SColorPicker.h"
+#include "Engine/LevelStreaming.h"
+#include "Engine/Selection.h"
 
 #define LOCTEXT_NAMESPACE "Level"
 
diff --git a/Engine/Source/Editor/Levels/Private/LevelViewModel.h b/Engine/Source/Editor/Levels/Private/LevelViewModel.h
--- a/Engine/Source/Editor/Levels/Private/LevelViewModel.h
+++ b/Engine/Source/Editor/Levels/Private/LevelViewModel.h
@@ -1,5 +1,7 @@
 // Copyright 1998-2014 Epic Games, Inc. All Rights Reserved.
 #pragma once
+#include "EditorUndoClient.h"
+#include "Layers/Layer.h"
 
 namespace ELevelsAction
 {
diff --git a/Engine/Source/Editor/Levels/Private/SLevelsViewRow.h b/Engine/Source/Editor/Levels/Private/SLevelsViewRow.h
--- a/Engine/Source/Editor/Levels/Private/SLevelsViewRow.h
+++ b/Engine/Source/Editor/Levels/Private/SLevelsViewRow.h
@@ -6,6 +6,7 @@
 #include "ISourceControlModule.h"
 #include "SVectorInputBox.h"
 #include "LevelUtils.h"
+#include "Engine/LevelStreaming.h"
 
 #define LOCTEXT_NAMESPACE "LevelsView"
 
diff --git a/Engine/Source/Editor/MaterialEditor/Private/MaterialEditorInstanceDetailCustomization.cpp b/Engine/Source/Editor/MaterialEditor/Private/MaterialEditorInstanceDetailCustomization.cpp
--- a/Engine/Source/Editor/MaterialEditor/Private/MaterialEditorInstanceDetailCustomization.cpp
+++ b/Engine/Source/Editor/MaterialEditor/Private/MaterialEditorInstanceDetailCustomization.cpp
@@ -12,6 +12,9 @@
 #include "MaterialEditorInstanceDetailCustomization.h"
 #include "ScopedTransaction.h"
 #include "DetailWidgetRow.h"
+#include "AssetData.h"
+#include "Materials/MaterialInstance.h"
+#include "Materials/MaterialInstanceConstant.h"
 
 #define LOCTEXT_NAMESPACE "MaterialInstanceEditor"
 
@@ -445,3 +448,4 @@ EVisibility FMaterialInstanceParameterDetails::ShouldShowSubsurfaceProfile() con
 }
 
 #undef LOCTEXT_NAMESPACE
+
diff --git a/Engine/Source/Editor/MaterialEditor/Private/MaterialEditorModule.cpp b/Engine/Source/Editor/MaterialEditor/Private/MaterialEditorModule.cpp
--- a/Engine/Source/Editor/MaterialEditor/Private/MaterialEditorModule.cpp
+++ b/Engine/Source/Editor/MaterialEditor/Private/MaterialEditorModule.cpp
@@ -7,6 +7,7 @@
 #include "MaterialEditorUtilities.h"
 #include "MaterialInstanceEditor.h"
 #include "ModuleManager.h"
+#include "Materials/MaterialInstance.h"
 
 const FName MaterialEditorAppIdentifier = FName(TEXT("MaterialEditorApp"));
 const FName MaterialInstanceEditorAppIdentifier = FName(TEXT("MaterialInstanceEditorApp"));
diff --git a/Engine/Source/Editor/MaterialEditor/Private/MaterialEditorUtilities.cpp b/Engine/Source/Editor/MaterialEditor/Private/MaterialEditorUtilities.cpp
--- a/Engine/Source/Editor/MaterialEditor/Private/MaterialEditorUtilities.cpp
+++ b/Engine/Source/Editor/MaterialEditor/Private/MaterialEditorUtilities.cpp
@@ -20,6 +20,7 @@
 #include "MaterialEditorUtilities.h"
 #include "Toolkits/ToolkitManager.h"
 #include "MaterialExpressionClasses.h"
+#include "Materials/MaterialInstance.h"
 
 #define LOCTEXT_NAMESPACE "MaterialEditorUtilities"
 
diff --git a/Engine/Source/Editor/MaterialEditor/Private/MaterialInstanceEditor.cpp b/Engine/Source/Editor/MaterialEditor/Private/MaterialInstanceEditor.cpp
--- a/Engine/Source/Editor/MaterialEditor/Private/MaterialInstanceEditor.cpp
+++ b/Engine/Source/Editor/MaterialEditor/Private/MaterialInstanceEditor.cpp
@@ -18,6 +18,8 @@
 #include "Editor/WorkspaceMenuStructure/Public/WorkspaceMenuStructureModule.h"
 #include "EditorViewportCommands.h"
 #include "SDockTab.h"
+#include "Materials/MaterialInstanceConstant.h"
+#include "CanvasTypes.h"
 
 #define LOCTEXT_NAMESPACE "MaterialInstanceEditor"
 
diff --git a/Engine/Source/Editor/MaterialEditor/Private/SMaterialEditorViewport.cpp b/Engine/Source/Editor/MaterialEditor/Private/SMaterialEditorViewport.cpp
--- a/Engine/Source/Editor/MaterialEditor/Private/SMaterialEditorViewport.cpp
+++ b/Engine/Source/Editor/MaterialEditor/Private/SMaterialEditorViewport.cpp
@@ -12,7 +12,9 @@
 #include "SMaterialEditorViewportToolBar.h"
 #include "ComponentReregisterContext.h"
 #include "SDockTab.h"
-
+#include "Engine/StaticMesh.h"
+#include "Engine/Selection.h"
+#include "Engine/TextureCube.h"
 /** Viewport Client for the preview viewport */
 class FMaterialEditorViewportClient : public FEditorViewportClient
 {
diff --git a/Engine/Source/Editor/Matinee/Private/Matinee.cpp b/Engine/Source/Editor/Matinee/Private/Matinee.cpp
--- a/Engine/Source/Editor/Matinee/Private/Matinee.cpp
+++ b/Engine/Source/Editor/Matinee/Private/Matinee.cpp
@@ -42,6 +42,12 @@
 #include "SDockTab.h"
 #include "STextComboBox.h"
 #include "GenericCommands.h"
+#include "CanvasTypes.h"
+#include "Engine/InterpCurveEdSetup.h"
+#include "Engine/Light.h"
+#include "CanvasItem.h"
+#include "Camera/CameraActor.h"
+#include "Camera/CameraAnim.h"
 
 
 DEFINE_LOG_CATEGORY(LogSlateMatinee);
diff --git a/Engine/Source/Editor/Matinee/Private/Matinee.h b/Engine/Source/Editor/Matinee/Private/Matinee.h
--- a/Engine/Source/Editor/Matinee/Private/Matinee.h
+++ b/Engine/Source/Editor/Matinee/Private/Matinee.h
@@ -20,6 +20,7 @@
 
 DECLARE_LOG_CATEGORY_EXTERN(LogSlateMatinee, Log, All);
 
+class UInterpTrackMove;
 
 class FMatinee : public IMatinee, public FGCObject, public FCurveEdNotifyInterface
 {
diff --git a/Engine/Source/Editor/Matinee/Private/MatineeCurveEd.cpp b/Engine/Source/Editor/Matinee/Private/MatineeCurveEd.cpp
--- a/Engine/Source/Editor/Matinee/Private/MatineeCurveEd.cpp
+++ b/Engine/Source/Editor/Matinee/Private/MatineeCurveEd.cpp
@@ -8,7 +8,7 @@
 #include "Matinee/InterpTrackColorProp.h"
 
 #include "Matinee.h"
-
+#include "Engine/InterpCurveEdSetup.h"
 
 /** Ensure the curve editor is synchronised with the track editor. */
 void FMatinee::SyncCurveEdView()
diff --git a/Engine/Source/Editor/Matinee/Private/MatineeDraw.cpp b/Engine/Source/Editor/Matinee/Private/MatineeDraw.cpp
--- a/Engine/Source/Editor/Matinee/Private/MatineeDraw.cpp
+++ b/Engine/Source/Editor/Matinee/Private/MatineeDraw.cpp
@@ -20,7 +20,11 @@
 #include "Materials/MaterialInstanceActor.h"
 #include "Animation/SkeletalMeshActor.h"
 #include "Particles/Emitter.h"
-
+#include "CanvasTypes.h"
+#include "Engine/InterpCurveEdSetup.h"
+#include "Engine/StaticMeshActor.h"
+#include "Camera/CameraActor.h"
+#include "Engine/Light.h"
 
 static const int32 GroupHeadHeight = 24;
 static const int32 TrackHeight = 24;
diff --git a/Engine/Source/Editor/Matinee/Private/MatineeMenus.cpp b/Engine/Source/Editor/Matinee/Private/MatineeMenus.cpp
--- a/Engine/Source/Editor/Matinee/Private/MatineeMenus.cpp
+++ b/Engine/Source/Editor/Matinee/Private/MatineeMenus.cpp
@@ -52,6 +52,13 @@
 #include "NotificationManager.h"
 #include "SNumericEntryBox.h"
 #include "STextComboPopup.h"
+#include "Camera/CameraActor.h"
+#include "Engine/Light.h"
+#include "Engine/Selection.h"
+#include "Engine/InterpCurveEdSetup.h"
+#include "Engine/Selection.h"
+#include "Sound/SoundBase.h"
+#include "Camera/CameraAnim.h"
 
 
 #define LOCTEXT_NAMESPACE "MatineeMenus"
diff --git a/Engine/Source/Editor/Matinee/Private/MatineeTools.cpp b/Engine/Source/Editor/Matinee/Private/MatineeTools.cpp
--- a/Engine/Source/Editor/Matinee/Private/MatineeTools.cpp
+++ b/Engine/Source/Editor/Matinee/Private/MatineeTools.cpp
@@ -38,6 +38,11 @@
 #include "EngineAnalytics.h"
 #include "SNotificationList.h"
 #include "NotificationManager.h"
+#include "Camera/CameraActor.h"
+#include "Camera/CameraAnim.h"
+#include "GameFramework/WorldSettings.h"
+#include "Engine/Selection.h"
+#include "Engine/InterpCurveEdSetup.h"
 
 ///// UTILS
 
diff --git a/Engine/Source/Editor/Matinee/Private/MatineeTrackHelpers.cpp b/Engine/Source/Editor/Matinee/Private/MatineeTrackHelpers.cpp
--- a/Engine/Source/Editor/Matinee/Private/MatineeTrackHelpers.cpp
+++ b/Engine/Source/Editor/Matinee/Private/MatineeTrackHelpers.cpp
@@ -24,6 +24,8 @@
 #include "Matinee/MatineeAnimInterface.h"
 #include "STextComboPopup.h"
 #include "STextEntryPopup.h"
+#include "Engine/Selection.h"
+#include "Sound/SoundBase.h"
 
 #define LOCTEXT_NAMESPACE "MatineeTrackHelpers"
 
diff --git a/Engine/Source/Editor/MeshPaint/Private/MeshPaintEdMode.cpp b/Engine/Source/Editor/MeshPaint/Private/MeshPaintEdMode.cpp
--- a/Engine/Source/Editor/MeshPaint/Private/MeshPaintEdMode.cpp
+++ b/Engine/Source/Editor/MeshPaint/Private/MeshPaintEdMode.cpp
@@ -28,6 +28,12 @@
 #include "Runtime/Engine/Classes/PhysicsEngine/BodySetup.h"
 #include "SMeshPaint.h"
 #include "ComponentReregisterContext.h"
+#include "CanvasTypes.h"
+#include "Engine/Selection.h"
+#include "Engine/TextureRenderTarget2D.h"
+#include "EngineUtils.h"
+#include "Engine/StaticMeshActor.h"
+#include "Materials/MaterialInstanceConstant.h"
 
 #define LOCTEXT_NAMESPACE "MeshPaint_Mode"
 
diff --git a/Engine/Source/Editor/MeshPaint/Private/SMeshPaint.cpp b/Engine/Source/Editor/MeshPaint/Private/SMeshPaint.cpp
--- a/Engine/Source/Editor/MeshPaint/Private/SMeshPaint.cpp
+++ b/Engine/Source/Editor/MeshPaint/Private/SMeshPaint.cpp
@@ -10,7 +10,9 @@
 #include "DesktopPlatformModule.h"
 #include "ISourceControlModule.h"
 #include "SColorPicker.h"
-
+#include "Engine/StaticMeshActor.h"
+#include "Engine/Selection.h"
+#include "Engine/StaticMesh.h"
 
 #define LOCTEXT_NAMESPACE "MeshPaint_Mode"
 
diff --git a/Engine/Source/Editor/MeshPaint/Public/MeshPaintEdMode.h b/Engine/Source/Editor/MeshPaint/Public/MeshPaintEdMode.h
--- a/Engine/Source/Editor/MeshPaint/Public/MeshPaintEdMode.h
+++ b/Engine/Source/Editor/MeshPaint/Public/MeshPaintEdMode.h
@@ -3,6 +3,8 @@
 #pragma once
 
 #include "GenericOctree.h"
+#include "GenericOctreePublic.h"
+#include "Engine/StaticMesh.h"
 
 
 /** Mesh paint resource types */
diff --git a/Engine/Source/Editor/Persona/Private/AnimationEditorViewportClient.cpp b/Engine/Source/Editor/Persona/Private/AnimationEditorViewportClient.cpp
--- a/Engine/Source/Editor/Persona/Private/AnimationEditorViewportClient.cpp
+++ b/Engine/Source/Editor/Persona/Private/AnimationEditorViewportClient.cpp
@@ -16,6 +16,18 @@
 #include "UnrealWidget.h"
 #include "MouseDeltaTracker.h"
 #include "ScopedTransaction.h"
+#include "Components/DirectionalLightComponent.h"
+#include "Components/ExponentialHeightFogComponent.h"
+#include "CanvasTypes.h"
+#include "Engine/TextureCube.h"
+#include "PhysicsEngine/PhysicsAsset.h"
+#include "Engine/CollisionProfile.h"
+#include "Engine/SkeletalMeshSocket.h"
+#include "EngineUtils.h"
+#include "GameFramework/WorldSettings.h"
+#include "PhysicsEngine/PhysicsSettings.h"
+#include "Components/WindDirectionalSourceComponent.h"
+#include "Engine/StaticMesh.h"
 
 namespace {
 	// Value from UE3
diff --git a/Engine/Source/Editor/Persona/Private/AnimationEditorViewportClient.h b/Engine/Source/Editor/Persona/Private/AnimationEditorViewportClient.h
--- a/Engine/Source/Editor/Persona/Private/AnimationEditorViewportClient.h
+++ b/Engine/Source/Editor/Persona/Private/AnimationEditorViewportClient.h
@@ -3,6 +3,7 @@
 #pragma once
 
 #include "EditorViewportClient.h"
+#include "Engine/WindDirectionalSource.h"
 
 //////////////////////////////////////////////////////////////////////////
 // ELocalAxesMode
diff --git a/Engine/Source/Editor/Persona/Private/SAnimBlueprintParentPlayerList.cpp b/Engine/Source/Editor/Persona/Private/SAnimBlueprintParentPlayerList.cpp
--- a/Engine/Source/Editor/Persona/Private/SAnimBlueprintParentPlayerList.cpp
+++ b/Engine/Source/Editor/Persona/Private/SAnimBlueprintParentPlayerList.cpp
@@ -4,6 +4,7 @@
 #include "SAnimBlueprintParentPlayerList.h"
 #include "PropertyEditorModule.h"
 #include "AnimGraphNodeDetails.h"
+#include "Animation/AnimBlueprintGeneratedClass.h"
 
 SAnimBlueprintParentPlayerList::SAnimBlueprintParentPlayerList()
 	: RootBlueprint(NULL)
diff --git a/Engine/Source/Editor/Persona/Private/SAnimationBlendSpace.cpp b/Engine/Source/Editor/Persona/Private/SAnimationBlendSpace.cpp
--- a/Engine/Source/Editor/Persona/Private/SAnimationBlendSpace.cpp
+++ b/Engine/Source/Editor/Persona/Private/SAnimationBlendSpace.cpp
@@ -10,6 +10,7 @@
 #include "DragAndDrop/AssetDragDropOp.h"
 #include "ScopedTransaction.h"
 #include "SNotificationList.h"
+#include "Animation/BlendSpaceBase.h"
 
 DEFINE_LOG_CATEGORY_STATIC(LogBlendSpace, Log, All);
 
diff --git a/Engine/Source/Editor/Persona/Private/SAnimationBlendSpace1D.cpp b/Engine/Source/Editor/Persona/Private/SAnimationBlendSpace1D.cpp
--- a/Engine/Source/Editor/Persona/Private/SAnimationBlendSpace1D.cpp
+++ b/Engine/Source/Editor/Persona/Private/SAnimationBlendSpace1D.cpp
@@ -10,6 +10,7 @@
 #include "DragAndDrop/AssetDragDropOp.h"
 #include "ScopedTransaction.h"
 #include "SNotificationList.h"
+#include "Animation/BlendSpaceBase.h"
 
 DEFINE_LOG_CATEGORY_STATIC(LogBlendSpace1D, Log, All);
 
diff --git a/Engine/Source/Editor/Persona/Private/SAnimationBlendSpace1D.h b/Engine/Source/Editor/Persona/Private/SAnimationBlendSpace1D.h
--- a/Engine/Source/Editor/Persona/Private/SAnimationBlendSpace1D.h
+++ b/Engine/Source/Editor/Persona/Private/SAnimationBlendSpace1D.h
@@ -4,6 +4,9 @@
 
 #include "SAnimationBlendSpaceBase.h"
 #include "Sorting.h"
+#include "Animation/BlendSpace1D.h"
+
+struct FEditorElement;
 
 //////////////////////////////////////////////////////////////////////////
 // SAnimationBlendSpace1D
diff --git a/Engine/Source/Editor/Persona/Private/SAnimationBlendSpaceBase.cpp b/Engine/Source/Editor/Persona/Private/SAnimationBlendSpaceBase.cpp
--- a/Engine/Source/Editor/Persona/Private/SAnimationBlendSpaceBase.cpp
+++ b/Engine/Source/Editor/Persona/Private/SAnimationBlendSpaceBase.cpp
@@ -12,6 +12,9 @@
 #include "AnimPreviewInstance.h"
 #include "SExpandableArea.h"
 #include "SNotificationList.h"
+#include "Animation/BlendSpaceBase.h"
+#include "Animation/AimOffsetBlendSpace.h"
+#include "Animation/AimOffsetBlendSpace1D.h"
 
 DEFINE_LOG_CATEGORY_STATIC(LogBlendSpaceBase, Log, All);
 
diff --git a/Engine/Source/Editor/Persona/Private/SMorphTargetViewer.cpp b/Engine/Source/Editor/Persona/Private/SMorphTargetViewer.cpp
--- a/Engine/Source/Editor/Persona/Private/SMorphTargetViewer.cpp
+++ b/Engine/Source/Editor/Persona/Private/SMorphTargetViewer.cpp
@@ -7,6 +7,7 @@
 #include "AssetRegistryModule.h"
 #include "ScopedTransaction.h"
 #include "SSearchBox.h"
+#include "Animation/VertexAnim/MorphTarget.h"
 
 #define LOCTEXT_NAMESPACE "SMorphTargetViewer"
 
diff --git a/Engine/Source/Editor/Persona/Private/SRigPicker.h b/Engine/Source/Editor/Persona/Private/SRigPicker.h
--- a/Engine/Source/Editor/Persona/Private/SRigPicker.h
+++ b/Engine/Source/Editor/Persona/Private/SRigPicker.h
@@ -2,6 +2,7 @@
 #pragma once
 
 #include "Editor/ContentBrowser/Public/ContentBrowserDelegates.h"
+#include "AssetData.h"
 
 /** Called when an asset is selected in the asset view */
 DECLARE_DELEGATE_OneParam(FOnAssetSelected, const class FAssetData& /*AssetData*/);
diff --git a/Engine/Source/Editor/Persona/Public/SocketDragDropOp.h b/Engine/Source/Editor/Persona/Public/SocketDragDropOp.h
--- a/Engine/Source/Editor/Persona/Public/SocketDragDropOp.h
+++ b/Engine/Source/Editor/Persona/Public/SocketDragDropOp.h
@@ -1,5 +1,7 @@
 // Copyright 1998-2014 Epic Games, Inc. All Rights Reserved.
 #pragma once
+#include "DragAndDrop.h"
+#include "Engine/SkeletalMeshSocket.h"
 
 //////////////////////////////////////////////////////////////////////////
 // FSocketDragDropOp
diff --git a/Engine/Source/Editor/PhAT/Private/PhAT.cpp b/Engine/Source/Editor/PhAT/Private/PhAT.cpp
--- a/Engine/Source/Editor/PhAT/Private/PhAT.cpp
+++ b/Engine/Source/Editor/PhAT/Private/PhAT.cpp
@@ -26,7 +26,12 @@
 #include "EngineAnalytics.h"
 #include "Runtime/Analytics/Analytics/Public/Interfaces/IAnalyticsProvider.h"
 #include "SDockTab.h"
-
+#include "PhysicsEngine/BodySetup.h"
+#include "PhysicsEngine/PhysicsConstraintTemplate.h"
+#include "Engine/Selection.h"
+#include "Engine/StaticMesh.h"
+#include "EngineLogs.h"
+#include "PhysicalMaterials/PhysicalMaterial.h"
 DEFINE_LOG_CATEGORY(LogPhAT);
 #define LOCTEXT_NAMESPACE "PhAT"
 
diff --git a/Engine/Source/Editor/PhAT/Private/PhAT.h b/Engine/Source/Editor/PhAT/Private/PhAT.h
--- a/Engine/Source/Editor/PhAT/Private/PhAT.h
+++ b/Engine/Source/Editor/PhAT/Private/PhAT.h
@@ -4,6 +4,8 @@
 
 #include "Toolkits/AssetEditorToolkit.h"
 #include "PhATSharedData.h"
+#include "EditorUndoClient.h"
+#include "PhysicsEngine/BodySetup.h"
 
 class SPhATPreviewViewport;
 class FPhATTreeInfo;
diff --git a/Engine/Source/Editor/PhAT/Private/PhATPreviewViewportClient.cpp b/Engine/Source/Editor/PhAT/Private/PhATPreviewViewportClient.cpp
--- a/Engine/Source/Editor/PhAT/Private/PhATPreviewViewportClient.cpp
+++ b/Engine/Source/Editor/PhAT/Private/PhATPreviewViewportClient.cpp
@@ -10,7 +10,14 @@
 #include "PhATSharedData.h"
 #include "PhATPreviewViewportClient.h"
 #include "SPhATPreviewViewport.h"
-
+#include "GameFramework/WorldSettings.h"
+#include "CanvasTypes.h"
+#include "PhysicsEngine/BodySetup.h"
+#include "Engine/Font.h"
+#include "PhysicsEngine/PhysicsConstraintTemplate.h"
+#include "PhysicsEngine/PhysicsSettings.h"
+#include "PhysicsEngine/PhysicsHandleComponent.h"
+#include "DrawDebugHelpers.h"
 
 FPhATEdPreviewViewportClient::FPhATEdPreviewViewportClient(TWeakPtr<FPhAT> InPhAT, TSharedPtr<FPhATSharedData> Data)
 	: FEditorViewportClient(GLevelEditorModeTools(), &Data->PreviewScene)
diff --git a/Engine/Source/Editor/PhAT/Private/PhATRender.cpp b/Engine/Source/Editor/PhAT/Private/PhATRender.cpp
--- a/Engine/Source/Editor/PhAT/Private/PhATRender.cpp
+++ b/Engine/Source/Editor/PhAT/Private/PhATRender.cpp
@@ -4,7 +4,8 @@
 #include "PhATSharedData.h"
 #include "PhATHitProxies.h"
 #include "PhATEdSkeletalMeshComponent.h"
-
+#include "PhysicsEngine/BodySetup.h"
+#include "PhysicsEngine/PhysicsConstraintTemplate.h"
 
 UPhATEdSkeletalMeshComponent::UPhATEdSkeletalMeshComponent(const FObjectInitializer& ObjectInitializer)
 	: Super(ObjectInitializer)
diff --git a/Engine/Source/Editor/PhAT/Private/PhATSharedData.cpp b/Engine/Source/Editor/PhAT/Private/PhATSharedData.cpp
--- a/Engine/Source/Editor/PhAT/Private/PhATSharedData.cpp
+++ b/Engine/Source/Editor/PhAT/Private/PhATSharedData.cpp
@@ -8,6 +8,12 @@
 #include "PhATEdSkeletalMeshComponent.h"
 #include "PhATSharedData.h"
 #include "Developer/MeshUtilities/Public/MeshUtilities.h"
+#include "PhysicsEngine/BodySetup.h"
+#include "Engine/CollisionProfile.h"
+#include "PhysicsEngine/PhysicsConstraintTemplate.h"
+#include "PhysicsEngine/PhysicsHandleComponent.h"
+#include "Vehicles/WheeledVehicleMovementComponent.h"
+#include "Engine/StaticMesh.h"
 
 #define LOCTEXT_NAMESPACE "PhATShared"
 
diff --git a/Engine/Source/Editor/PlacementMode/Private/PlacementMode.cpp b/Engine/Source/Editor/PlacementMode/Private/PlacementMode.cpp
--- a/Engine/Source/Editor/PlacementMode/Private/PlacementMode.cpp
+++ b/Engine/Source/Editor/PlacementMode/Private/PlacementMode.cpp
@@ -12,6 +12,7 @@
 #include "Editor/UnrealEd/Public/Toolkits/ToolkitManager.h"
 #include "Editor/LevelEditor/Public/LevelEditor.h"
 #include "Editor/LevelEditor/Public/ILevelViewport.h"
+#include "Engine/Selection.h"
 
 FPlacementMode::FPlacementMode()
 	: AssetsToPlace()
diff --git a/Engine/Source/Editor/PlacementMode/Private/SPlacementModeTools.cpp b/Engine/Source/Editor/PlacementMode/Private/SPlacementModeTools.cpp
--- a/Engine/Source/Editor/PlacementMode/Private/SPlacementModeTools.cpp
+++ b/Engine/Source/Editor/PlacementMode/Private/SPlacementModeTools.cpp
@@ -11,6 +11,8 @@
 #include "EditorClassUtils.h"
 #include "SSearchBox.h"
 #include "SWidgetSwitcher.h"
+#include "GameFramework/Volume.h"
+#include "Engine/PostProcessVolume.h"
 
 /**
  * These are the tab indexes, if the tabs are reorganized you need to adjust the
diff --git a/Engine/Source/Editor/PropertyEditor/Private/Presentation/PropertyTable/PropertyTableObjectNameColumn.h b/Engine/Source/Editor/PropertyEditor/Private/Presentation/PropertyTable/PropertyTableObjectNameColumn.h
--- a/Engine/Source/Editor/PropertyEditor/Private/Presentation/PropertyTable/PropertyTableObjectNameColumn.h
+++ b/Engine/Source/Editor/PropertyEditor/Private/Presentation/PropertyTable/PropertyTableObjectNameColumn.h
@@ -9,6 +9,7 @@
 #include "PropertyTableObjectNameCell.h"
 #include "DataSource.h"
 #include "PropertyPath.h"
+#include "Engine/BlueprintGeneratedClass.h"
 
 #define LOCTEXT_NAMESPACE "ObjectNameColumnHeader"
 
diff --git a/Engine/Source/Editor/PropertyEditor/Private/PropertyEditorModule.cpp b/Engine/Source/Editor/PropertyEditor/Private/PropertyEditorModule.cpp
--- a/Engine/Source/Editor/PropertyEditor/Private/PropertyEditorModule.cpp
+++ b/Engine/Source/Editor/PropertyEditor/Private/PropertyEditorModule.cpp
@@ -23,6 +23,8 @@
 #include "PropertyTableConstants.h"
 #include "SStructureDetailsView.h"
 #include "SColorPicker.h"
+#include "EngineUtils.h"
+#include "Engine/UserDefinedStruct.h"
 
 
 IMPLEMENT_MODULE( FPropertyEditorModule, PropertyEditor );
diff --git a/Engine/Source/Editor/PropertyEditor/Private/PropertyEditorToolkit.h b/Engine/Source/Editor/PropertyEditor/Private/PropertyEditorToolkit.h
--- a/Engine/Source/Editor/PropertyEditor/Private/PropertyEditorToolkit.h
+++ b/Engine/Source/Editor/PropertyEditor/Private/PropertyEditorToolkit.h
@@ -2,6 +2,7 @@
 #pragma once
 
 #include "Toolkits/AssetEditorToolkit.h"
+#include "TimerManager.h"
 
 class FPropertyEditorToolkit : public FAssetEditorToolkit
 {
diff --git a/Engine/Source/Editor/PropertyEditor/Private/PropertyHandleImpl.cpp b/Engine/Source/Editor/PropertyEditor/Private/PropertyHandleImpl.cpp
--- a/Engine/Source/Editor/PropertyEditor/Private/PropertyHandleImpl.cpp
+++ b/Engine/Source/Editor/PropertyEditor/Private/PropertyHandleImpl.cpp
@@ -10,6 +10,8 @@
 #include "Editor/UnrealEd/Public/Kismet2/KismetEditorUtilities.h"
 #include "IPropertyUtilities.h"
 #include "PropertyEditor.h"
+#include "Engine/Selection.h"
+#include "Engine/UserDefinedStruct.h"
 
 DEFINE_LOG_CATEGORY_STATIC(LogPropertyView, Log, All);
 
diff --git a/Engine/Source/Editor/PropertyEditor/Private/UserInterface/PropertyEditor/SPropertyEditorAsset.cpp b/Engine/Source/Editor/PropertyEditor/Private/UserInterface/PropertyEditor/SPropertyEditorAsset.cpp
--- a/Engine/Source/Editor/PropertyEditor/Private/UserInterface/PropertyEditor/SPropertyEditorAsset.cpp
+++ b/Engine/Source/Editor/PropertyEditor/Private/UserInterface/PropertyEditor/SPropertyEditorAsset.cpp
@@ -14,6 +14,8 @@
 #include "SAssetDropTarget.h"
 #include "AssetRegistryModule.h"
 #include "Particles/ParticleSystem.h"
+#include "Engine/Selection.h"
+#include "Engine/StaticMesh.h"
 
 #define LOCTEXT_NAMESPACE "PropertyEditor"
 
diff --git a/Engine/Source/Editor/PropertyEditor/Private/UserInterface/PropertyEditor/SPropertyEditorColor.cpp b/Engine/Source/Editor/PropertyEditor/Private/UserInterface/PropertyEditor/SPropertyEditorColor.cpp
--- a/Engine/Source/Editor/PropertyEditor/Private/UserInterface/PropertyEditor/SPropertyEditorColor.cpp
+++ b/Engine/Source/Editor/PropertyEditor/Private/UserInterface/PropertyEditor/SPropertyEditorColor.cpp
@@ -10,6 +10,7 @@
 #include "ScopedTransaction.h"
 #include "SPropertyEditorColor.h"
 #include "SColorPicker.h"
+#include "Components/LightComponent.h"
 
 void SPropertyEditorColor::Construct( const FArguments& InArgs, const TSharedRef<FPropertyEditor>& InPropertyEditor, const TSharedRef<IPropertyUtilities>& InPropertyUtilities )
 {
diff --git a/Engine/Source/Editor/Sequencer/Private/SAnimationOutlinerView.cpp b/Engine/Source/Editor/Sequencer/Private/SAnimationOutlinerView.cpp
--- a/Engine/Source/Editor/Sequencer/Private/SAnimationOutlinerView.cpp
+++ b/Engine/Source/Editor/Sequencer/Private/SAnimationOutlinerView.cpp
@@ -6,7 +6,7 @@
 #include "ScopedTransaction.h"
 #include "Sequencer.h"
 #include "MovieScene.h"
-
+#include "Engine/Selection.h"
 
 #define LOCTEXT_NAMESPACE "AnimationOutliner"
 
diff --git a/Engine/Source/Editor/Sequencer/Private/Sequencer.cpp b/Engine/Source/Editor/Sequencer/Private/Sequencer.cpp
--- a/Engine/Source/Editor/Sequencer/Private/Sequencer.cpp
+++ b/Engine/Source/Editor/Sequencer/Private/Sequencer.cpp
@@ -34,6 +34,8 @@
 #include "SnappingUtils.h"
 #include "STextEntryPopup.h"
 #include "GenericCommands.h"
+#include "Engine/BlueprintGeneratedClass.h"
+#include "Engine/Selection.h"
 
 #define LOCTEXT_NAMESPACE "Sequencer"
 
diff --git a/Engine/Source/Editor/StaticMeshEditor/Private/SStaticMeshEditorViewport.cpp b/Engine/Source/Editor/StaticMeshEditor/Private/SStaticMeshEditorViewport.cpp
--- a/Engine/Source/Editor/StaticMeshEditor/Private/SStaticMeshEditorViewport.cpp
+++ b/Engine/Source/Editor/StaticMeshEditor/Private/SStaticMeshEditorViewport.cpp
@@ -30,6 +30,7 @@
 #if WITH_PHYSX
 #include "Editor/UnrealEd/Private/EditorPhysXSupport.h"
 #endif
+#include "Engine/StaticMeshSocket.h"
 
 #define HITPROXY_SOCKET	1
 
diff --git a/Engine/Source/Editor/StaticMeshEditor/Private/StaticMeshAutomationTests.cpp b/Engine/Source/Editor/StaticMeshEditor/Private/StaticMeshAutomationTests.cpp
--- a/Engine/Source/Editor/StaticMeshEditor/Private/StaticMeshAutomationTests.cpp
+++ b/Engine/Source/Editor/StaticMeshEditor/Private/StaticMeshAutomationTests.cpp
@@ -5,6 +5,7 @@
 #include "AutomationCommon.h"
 #include "SStaticMeshEditorViewport.h"
 #include "StaticMeshEditorViewportClient.h"
+#include "Engine/StaticMesh.h"
 
 namespace EditorViewButtonHelper
 {
diff --git a/Engine/Source/Editor/StaticMeshEditor/Private/StaticMeshEditorTools.h b/Engine/Source/Editor/StaticMeshEditor/Private/StaticMeshEditorTools.h
--- a/Engine/Source/Editor/StaticMeshEditor/Private/StaticMeshEditorTools.h
+++ b/Engine/Source/Editor/StaticMeshEditor/Private/StaticMeshEditorTools.h
@@ -4,6 +4,7 @@
 
 #include "AssetThumbnail.h"
 #include "PropertyEditing.h"
+#include "Engine/StaticMesh.h"
 
 enum ECreationModeChoice
 {
diff --git a/Engine/Source/Editor/StaticMeshEditor/Private/StaticMeshEditorViewportClient.cpp b/Engine/Source/Editor/StaticMeshEditor/Private/StaticMeshEditorViewportClient.cpp
--- a/Engine/Source/Editor/StaticMeshEditor/Private/StaticMeshEditorViewportClient.cpp
+++ b/Engine/Source/Editor/StaticMeshEditor/Private/StaticMeshEditorViewportClient.cpp
@@ -23,6 +23,12 @@
 #if WITH_PHYSX
 #include "Editor/UnrealEd/Private/EditorPhysXSupport.h"
 #endif
+#include "CanvasTypes.h"
+#include "Engine/AssetUserData.h"
+#include "Engine/StaticMeshSocket.h"
+#include "CanvasItem.h"
+#include "Engine/Canvas.h"
+#include "Engine/TextureCube.h"
 
 #define LOCTEXT_NAMESPACE "FStaticMeshEditorViewportClient"
 
diff --git a/Engine/Source/Editor/TextureEditor/Private/Models/TextureEditorViewportClient.cpp b/Engine/Source/Editor/TextureEditor/Private/Models/TextureEditorViewportClient.cpp
--- a/Engine/Source/Editor/TextureEditor/Private/Models/TextureEditorViewportClient.cpp
+++ b/Engine/Source/Editor/TextureEditor/Private/Models/TextureEditorViewportClient.cpp
@@ -1,6 +1,11 @@
 // Copyright 1998-2014 Epic Games, Inc. All Rights Reserved.
 
 #include "TextureEditorPrivatePCH.h"
+#include "CanvasTypes.h"
+#include "CanvasItem.h"
+#include "Engine/TextureCube.h"
+#include "Engine/TextureRenderTarget2D.h"
+#include "Engine/TextureRenderTargetCube.h"
 
 
 /* FTextureEditorViewportClient structors
diff --git a/Engine/Source/Editor/TextureEditor/Private/TextureEditorToolkit.cpp b/Engine/Source/Editor/TextureEditor/Private/TextureEditorToolkit.cpp
--- a/Engine/Source/Editor/TextureEditor/Private/TextureEditorToolkit.cpp
+++ b/Engine/Source/Editor/TextureEditor/Private/TextureEditorToolkit.cpp
@@ -5,6 +5,14 @@
 #include "ISettingsModule.h"
 #include "SDockTab.h"
 #include "SNumericEntryBox.h"
+#include "Engine/TextureCube.h"
+#include "Engine/TextureRenderTargetCube.h"
+#include "Engine/TextureRenderTarget2D.h"
+#include "Engine/Texture2D.h"
+#include "SystemSettings.h"
+#include "Engine/LightMapTexture2D.h"
+#include "Engine/ShadowMapTexture2D.h"
+#include "Engine/Texture2DDynamic.h"
 
 #define LOCTEXT_NAMESPACE "FTextureEditorToolkit"
 
diff --git a/Engine/Source/Editor/UnrealEd/Classes/Animation/EditorParentPlayerListObj.h b/Engine/Source/Editor/UnrealEd/Classes/Animation/EditorParentPlayerListObj.h
--- a/Engine/Source/Editor/UnrealEd/Classes/Animation/EditorParentPlayerListObj.h
+++ b/Engine/Source/Editor/UnrealEd/Classes/Animation/EditorParentPlayerListObj.h
@@ -1,10 +1,12 @@
 // Copyright 1998-2014 Epic Games, Inc. All Rights Reserved.
 
 #pragma once
+#include "Animation/AnimBlueprint.h"
 #include "EditorParentPlayerListObj.generated.h"
 
 struct FAnimParentNodeAssetOverride;
 class UAnimGraphNode_Base;
+class UAnimBlueprint;
 
 UCLASS()
 class UNREALED_API UEditorParentPlayerListObj : public UObject
diff --git a/Engine/Source/Editor/UnrealEd/Classes/SoundCueGraph/SoundCueGraphNode.h b/Engine/Source/Editor/UnrealEd/Classes/SoundCueGraph/SoundCueGraphNode.h
--- a/Engine/Source/Editor/UnrealEd/Classes/SoundCueGraph/SoundCueGraphNode.h
+++ b/Engine/Source/Editor/UnrealEd/Classes/SoundCueGraph/SoundCueGraphNode.h
@@ -2,6 +2,7 @@
 
 #pragma once
 
+#include "SoundCueGraphNode_Base.h"
 #include "SoundCueGraphNode.generated.h"
 
 UCLASS(MinimalAPI)
diff --git a/Engine/Source/Editor/UnrealEd/Classes/SoundCueGraph/SoundCueGraphNode_Base.h b/Engine/Source/Editor/UnrealEd/Classes/SoundCueGraph/SoundCueGraphNode_Base.h
--- a/Engine/Source/Editor/UnrealEd/Classes/SoundCueGraph/SoundCueGraphNode_Base.h
+++ b/Engine/Source/Editor/UnrealEd/Classes/SoundCueGraph/SoundCueGraphNode_Base.h
@@ -1,7 +1,7 @@
 // Copyright 1998-2014 Epic Games, Inc. All Rights Reserved.
 
 #pragma once
-
+#include "EdGraph/EdGraphNode.h"
 #include "SoundCueGraphNode_Base.generated.h"
 
 UCLASS(MinimalAPI)
diff --git a/Engine/Source/Editor/UnrealEd/Classes/SoundCueGraph/SoundCueGraphSchema.h b/Engine/Source/Editor/UnrealEd/Classes/SoundCueGraph/SoundCueGraphSchema.h
--- a/Engine/Source/Editor/UnrealEd/Classes/SoundCueGraph/SoundCueGraphSchema.h
+++ b/Engine/Source/Editor/UnrealEd/Classes/SoundCueGraph/SoundCueGraphSchema.h
@@ -1,6 +1,7 @@
 // Copyright 1998-2014 Epic Games, Inc. All Rights Reserved.
 
 #pragma once
+#include "EdGraph/EdGraphSchema.h"
 #include "SoundCueGraphSchema.generated.h"
 
 class USoundCue;
diff --git a/Engine/Source/Editor/UnrealEd/Private/AnimationEditorUtils.cpp b/Engine/Source/Editor/UnrealEd/Private/AnimationEditorUtils.cpp
--- a/Engine/Source/Editor/UnrealEd/Private/AnimationEditorUtils.cpp
+++ b/Engine/Source/Editor/UnrealEd/Private/AnimationEditorUtils.cpp
@@ -4,6 +4,11 @@
 #include "AnimationEditorUtils.h"
 #include "AssetToolsModule.h"
 #include "ContentBrowserModule.h"
+#include "Animation/AnimComposite.h"
+#include "Animation/BlendSpace.h"
+#include "Animation/BlendSpace1D.h"
+#include "Animation/AimOffsetBlendSpace.h"
+#include "Animation/AimOffsetBlendSpace1D.h"
 
 #define LOCTEXT_NAMESPACE "AnimationEditorUtils"
 
diff --git a/Engine/Source/Editor/UnrealEd/Private/AssetSelection.cpp b/Engine/Source/Editor/UnrealEd/Private/AssetSelection.cpp
--- a/Engine/Source/Editor/UnrealEd/Private/AssetSelection.cpp
+++ b/Engine/Source/Editor/UnrealEd/Private/AssetSelection.cpp
@@ -27,6 +27,12 @@
 #include "ObjectEditorUtils.h"
 #include "SNotificationList.h"
 #include "NotificationManager.h"
+#include "Engine/Light.h"
+#include "Engine/StaticMeshActor.h"
+#include "Engine/StaticMesh.h"
+#include "GameFramework/Pawn.h"
+#include "Engine/Selection.h"
+#include "Components/DecalComponent.h"
 
 
 namespace AssetSelectionUtils
diff --git a/Engine/Source/Editor/UnrealEd/Private/AutoReimportManager.cpp b/Engine/Source/Editor/UnrealEd/Private/AutoReimportManager.cpp
--- a/Engine/Source/Editor/UnrealEd/Private/AutoReimportManager.cpp
+++ b/Engine/Source/Editor/UnrealEd/Private/AutoReimportManager.cpp
@@ -5,6 +5,8 @@
 
 #include "Developer/DirectoryWatcher/Public/DirectoryWatcherModule.h"
 #include "Editor/DataTableEditor/Public/IDataTableEditor.h"
+#include "Engine/CurveTable.h"
+#include "Engine/DataTable.h"
 
 
 #define LOCTEXT_NAMESPACE "AutoReimportManager"
diff --git a/Engine/Source/Editor/UnrealEd/Private/Commandlets/ContentCommandlets.cpp b/Engine/Source/Editor/UnrealEd/Private/Commandlets/ContentCommandlets.cpp
--- a/Engine/Source/Editor/UnrealEd/Private/Commandlets/ContentCommandlets.cpp
+++ b/Engine/Source/Editor/UnrealEd/Private/Commandlets/ContentCommandlets.cpp
@@ -24,6 +24,8 @@ DEFINE_LOG_CATEGORY_STATIC(LogContentCommandlet, Log, All);
 #include "Particles/ParticleModuleRequired.h"
 #include "Particles/ParticleSystem.h"
 #include "Particles/TypeData/ParticleModuleTypeDataMesh.h"
+#include "Engine/LevelStreaming.h"
+#include "Engine/StaticMesh.h"
 
 
 /**-----------------------------------------------------------------------------
diff --git a/Engine/Source/Editor/UnrealEd/Private/Commandlets/DiffPackagesCommandlet.cpp b/Engine/Source/Editor/UnrealEd/Private/Commandlets/DiffPackagesCommandlet.cpp
--- a/Engine/Source/Editor/UnrealEd/Private/Commandlets/DiffPackagesCommandlet.cpp
+++ b/Engine/Source/Editor/UnrealEd/Private/Commandlets/DiffPackagesCommandlet.cpp
@@ -6,6 +6,7 @@
 =============================================================================*/
 
 #include "UnrealEd.h"
+#include "EditorCommandlets.h"
 
 DEFINE_LOG_CATEGORY_STATIC(LogDiffPackagesCommandlet, Log, All);
 
diff --git a/Engine/Source/Editor/UnrealEd/Private/Commandlets/DumpBlueprintsInfoCommandlet.cpp b/Engine/Source/Editor/UnrealEd/Private/Commandlets/DumpBlueprintsInfoCommandlet.cpp
--- a/Engine/Source/Editor/UnrealEd/Private/Commandlets/DumpBlueprintsInfoCommandlet.cpp
+++ b/Engine/Source/Editor/UnrealEd/Private/Commandlets/DumpBlueprintsInfoCommandlet.cpp
@@ -30,6 +30,11 @@
 #include "PackageHelperFunctions.h" // for NORMALIZE_ExcludeMapPackages
 #include "K2Node_FunctionEntry.h"
 #include "Engine/LevelScriptActor.h"
+#include "Engine/Selection.h"
+#include "Engine/SimpleConstructionScript.h"
+#include "Engine/SCS_Node.h"
+#include "Engine/UserDefinedEnum.h"
+#include "Engine/UserDefinedStruct.h"
 
 DEFINE_LOG_CATEGORY_STATIC(LogBlueprintInfoDump, Log, All);
 
diff --git a/Engine/Source/Editor/UnrealEd/Private/Commandlets/DumpHiddenCategoriesCommandlet.cpp b/Engine/Source/Editor/UnrealEd/Private/Commandlets/DumpHiddenCategoriesCommandlet.cpp
--- a/Engine/Source/Editor/UnrealEd/Private/Commandlets/DumpHiddenCategoriesCommandlet.cpp
+++ b/Engine/Source/Editor/UnrealEd/Private/Commandlets/DumpHiddenCategoriesCommandlet.cpp
@@ -6,6 +6,7 @@
 #include "KismetEditorUtilities.h"
 #include "EditorCategoryUtils.h"
 #include "ScopedTimers.h"
+#include "Kismet/BlueprintFunctionLibrary.h"
 
 /*******************************************************************************
  * Static Helpers
diff --git a/Engine/Source/Editor/UnrealEd/Private/Commandlets/GatherTextFromAssetsCommandlet.cpp b/Engine/Source/Editor/UnrealEd/Private/Commandlets/GatherTextFromAssetsCommandlet.cpp
--- a/Engine/Source/Editor/UnrealEd/Private/Commandlets/GatherTextFromAssetsCommandlet.cpp
+++ b/Engine/Source/Editor/UnrealEd/Private/Commandlets/GatherTextFromAssetsCommandlet.cpp
@@ -8,6 +8,7 @@
 #include "PackageHelperFunctions.h"
 #include "InternationalizationMetadata.h"
 #include "Sound/DialogueWave.h"
+#include "Sound/DialogueVoice.h"
 
 DEFINE_LOG_CATEGORY_STATIC(LogGatherTextFromAssetsCommandlet, Log, All);
 
diff --git a/Engine/Source/Editor/UnrealEd/Private/Commandlets/GenerateBlueprintAPICommandlet.cpp b/Engine/Source/Editor/UnrealEd/Private/Commandlets/GenerateBlueprintAPICommandlet.cpp
--- a/Engine/Source/Editor/UnrealEd/Private/Commandlets/GenerateBlueprintAPICommandlet.cpp
+++ b/Engine/Source/Editor/UnrealEd/Private/Commandlets/GenerateBlueprintAPICommandlet.cpp
@@ -23,6 +23,7 @@
 #include "BlueprintActionDatabase.h"
 #include "BlueprintActionMenuUtils.h"
 #include "EditorCategoryUtils.h"
+#include "Engine/LevelScriptActor.h"
 
 DEFINE_LOG_CATEGORY_STATIC(LogBlueprintInfoDump, Log, All);
 
diff --git a/Engine/Source/Editor/UnrealEd/Private/Commandlets/PackageUtilities.cpp b/Engine/Source/Editor/UnrealEd/Private/Commandlets/PackageUtilities.cpp
--- a/Engine/Source/Editor/UnrealEd/Private/Commandlets/PackageUtilities.cpp
+++ b/Engine/Source/Editor/UnrealEd/Private/Commandlets/PackageUtilities.cpp
@@ -16,7 +16,8 @@
 #include "AnimationCompression.h"
 
 #include "CollectionManagerModule.h"
-
+#include "GameFramework/WorldSettings.h"
+#include "EngineUtils.h"
 
 DEFINE_LOG_CATEGORY(LogPackageHelperFunctions);
 DEFINE_LOG_CATEGORY_STATIC(LogPackageUtilities, Log, All);
diff --git a/Engine/Source/Editor/UnrealEd/Private/Commandlets/ParticleSystemAuditCommandlet.cpp b/Engine/Source/Editor/UnrealEd/Private/Commandlets/ParticleSystemAuditCommandlet.cpp
--- a/Engine/Source/Editor/UnrealEd/Private/Commandlets/ParticleSystemAuditCommandlet.cpp
+++ b/Engine/Source/Editor/UnrealEd/Private/Commandlets/ParticleSystemAuditCommandlet.cpp
@@ -9,6 +9,8 @@
 #include "Runtime/Engine/Classes/Particles/Spawn/ParticleModuleSpawn.h"
 #include "Runtime/AssetRegistry/Public/AssetRegistryModule.h"
 #include "Particles/ParticleLODLevel.h"
+#include "Distributions/DistributionFloatConstant.h"
+#include "Distributions/DistributionVectorConstant.h"
 
 DEFINE_LOG_CATEGORY_STATIC(LogParticleSystemAuditCommandlet, Log, All);
 
diff --git a/Engine/Source/Editor/UnrealEd/Private/ComponentAssetBroker.cpp b/Engine/Source/Editor/UnrealEd/Private/ComponentAssetBroker.cpp
--- a/Engine/Source/Editor/UnrealEd/Private/ComponentAssetBroker.cpp
+++ b/Engine/Source/Editor/UnrealEd/Private/ComponentAssetBroker.cpp
@@ -11,6 +11,9 @@
 #include "Engine/DestructibleMesh.h"
 #include "Components/DestructibleComponent.h"
 #include "Sound/SoundBase.h"
+#include "Engine/StaticMesh.h"
+#include "Components/AudioComponent.h"
+#include "Components/ChildActorComponent.h"
 
 //////////////////////////////////////////////////////////////////////////
 // FStaticMeshComponentBroker
diff --git a/Engine/Source/Editor/UnrealEd/Private/CookOnTheFlyServer.cpp b/Engine/Source/Editor/UnrealEd/Private/CookOnTheFlyServer.cpp
--- a/Engine/Source/Editor/UnrealEd/Private/CookOnTheFlyServer.cpp
+++ b/Engine/Source/Editor/UnrealEd/Private/CookOnTheFlyServer.cpp
@@ -31,6 +31,8 @@
 
 // shader compiler processAsyncResults
 #include "ShaderCompiler.h"
+#include "Engine/LevelStreaming.h"
+#include "TextureLODSettings.h"
 
 DEFINE_LOG_CATEGORY_STATIC(LogCookOnTheFly, Log, All);
 
diff --git a/Engine/Source/Editor/UnrealEd/Private/CubemapUnwrapUtils.cpp b/Engine/Source/Editor/UnrealEd/Private/CubemapUnwrapUtils.cpp
--- a/Engine/Source/Editor/UnrealEd/Private/CubemapUnwrapUtils.cpp
+++ b/Engine/Source/Editor/UnrealEd/Private/CubemapUnwrapUtils.cpp
@@ -11,6 +11,7 @@
 #include "Engine/TextureRenderTarget2D.h"
 #include "Engine/TextureCube.h"
 #include "Engine/TextureRenderTargetCube.h"
+#include "CanvasTypes.h"
 
 IMPLEMENT_SHADER_TYPE(,FCubemapTexturePropertiesVS,TEXT("SimpleElementVertexShader"),TEXT("Main"),SF_Vertex);
 IMPLEMENT_SHADER_TYPE(template<>,FCubemapTexturePropertiesPS<false>,TEXT("SimpleElementPixelShader"),TEXT("CubemapTextureProperties"),SF_Pixel);
diff --git a/Engine/Source/Editor/UnrealEd/Private/DebugToolExec.cpp b/Engine/Source/Editor/UnrealEd/Private/DebugToolExec.cpp
--- a/Engine/Source/Editor/UnrealEd/Private/DebugToolExec.cpp
+++ b/Engine/Source/Editor/UnrealEd/Private/DebugToolExec.cpp
@@ -6,6 +6,11 @@
 #include "Editor/PropertyEditor/Public/PropertyEditorModule.h"
 #include "Editor/PropertyEditor/Public/IDetailsView.h"
 #include "Engine/GameEngine.h"
+#include "GameFramework/PlayerController.h"
+#include "EngineUtils.h"
+#include "GameFramework/Pawn.h"
+#include "CollisionQueryParams.h"
+#include "Engine/EngineTypes.h"
 
 
 /**
diff --git a/Engine/Source/Editor/UnrealEd/Private/Dialogs/SMeshMergingDialog.cpp b/Engine/Source/Editor/UnrealEd/Private/Dialogs/SMeshMergingDialog.cpp
--- a/Engine/Source/Editor/UnrealEd/Private/Dialogs/SMeshMergingDialog.cpp
+++ b/Engine/Source/Editor/UnrealEd/Private/Dialogs/SMeshMergingDialog.cpp
@@ -9,7 +9,10 @@
 #include "AssetRegistryModule.h"
 #include "Editor/LevelEditor/Public/LevelEditor.h"
 #include "STextComboBox.h"
-
+#include "Engine/StaticMeshActor.h"
+#include "Engine/StaticMesh.h"
+#include "Engine/Selection.h"
+#include "SystemSettings.h"
 
 #define LOCTEXT_NAMESPACE "SMeshMergingDialog"
 
diff --git a/Engine/Source/Editor/UnrealEd/Private/Dialogs/SMeshProxyDialog.cpp b/Engine/Source/Editor/UnrealEd/Private/Dialogs/SMeshProxyDialog.cpp
--- a/Engine/Source/Editor/UnrealEd/Private/Dialogs/SMeshProxyDialog.cpp
+++ b/Engine/Source/Editor/UnrealEd/Private/Dialogs/SMeshProxyDialog.cpp
@@ -10,7 +10,8 @@
 #include "SNumericEntryBox.h"
 #include "STextComboBox.h"
 #include "Engine/StaticMeshActor.h"
-
+#include "Engine/StaticMesh.h"
+#include "Engine/Selection.h"
 #define LOCTEXT_NAMESPACE "SMeshProxyDialog"
 
 
diff --git a/Engine/Source/Editor/UnrealEd/Private/DragTool_BoxSelect.cpp b/Engine/Source/Editor/UnrealEd/Private/DragTool_BoxSelect.cpp
--- a/Engine/Source/Editor/UnrealEd/Private/DragTool_BoxSelect.cpp
+++ b/Engine/Source/Editor/UnrealEd/Private/DragTool_BoxSelect.cpp
@@ -5,6 +5,10 @@
 #include "DragTool_BoxSelect.h"
 #include "ActorEditorUtils.h"
 #include "ScopedTransaction.h"
+#include "Engine/LevelStreaming.h"
+#include "EngineUtils.h"
+#include "CanvasItem.h"
+#include "CanvasTypes.h"
 
 ///////////////////////////////////////////////////////////////////////////////////////////////////////////////////////
 //
diff --git a/Engine/Source/Editor/UnrealEd/Private/DragTool_FrustumSelect.cpp b/Engine/Source/Editor/UnrealEd/Private/DragTool_FrustumSelect.cpp
--- a/Engine/Source/Editor/UnrealEd/Private/DragTool_FrustumSelect.cpp
+++ b/Engine/Source/Editor/UnrealEd/Private/DragTool_FrustumSelect.cpp
@@ -9,6 +9,10 @@
 #include "HModel.h"
 #include "Components/DrawSphereComponent.h"
 #include "AI/Navigation/NavLinkRenderingComponent.h"
+#include "Engine/Selection.h"
+#include "EngineUtils.h"
+#include "CanvasItem.h"
+#include "CanvasTypes.h"
 
 ///////////////////////////////////////////////////////////////////////////////
 //
diff --git a/Engine/Source/Editor/UnrealEd/Private/DragTool_Measure.cpp b/Engine/Source/Editor/UnrealEd/Private/DragTool_Measure.cpp
--- a/Engine/Source/Editor/UnrealEd/Private/DragTool_Measure.cpp
+++ b/Engine/Source/Editor/UnrealEd/Private/DragTool_Measure.cpp
@@ -4,6 +4,8 @@
 #include "UnrealEd.h"
 #include "DragTool_Measure.h"
 #include "SnappingUtils.h"
+#include "CanvasItem.h"
+#include "CanvasTypes.h"
 
 ///////////////////////////////////////////////////////////////////////////////////////////////////////////////////////
 //
diff --git a/Engine/Source/Editor/UnrealEd/Private/Editor/ActorPositioning.cpp b/Engine/Source/Editor/UnrealEd/Private/Editor/ActorPositioning.cpp
--- a/Engine/Source/Editor/UnrealEd/Private/Editor/ActorPositioning.cpp
+++ b/Engine/Source/Editor/UnrealEd/Private/Editor/ActorPositioning.cpp
@@ -5,6 +5,9 @@
 #include "SnappingUtils.h"
 #include "LandscapeHeightfieldCollisionComponent.h"
 #include "LandscapeComponent.h"
+#include "GameFramework/Volume.h"
+#include "Components/ShapeComponent.h"
+#include "Components/ModelComponent.h"
 
 FActorPositionTraceResult FActorPositioning::TraceWorldForPositionWithDefault(const FViewportCursorLocation& Cursor, const FSceneView& View, const TArray<AActor*>* IgnoreActors)
 {
diff --git a/Engine/Source/Editor/UnrealEd/Private/EditorActor.cpp b/Engine/Source/Editor/UnrealEd/Private/EditorActor.cpp
--- a/Engine/Source/Editor/UnrealEd/Private/EditorActor.cpp
+++ b/Engine/Source/Editor/UnrealEd/Private/EditorActor.cpp
@@ -19,6 +19,14 @@
 #include "Particles/ParticleSystemComponent.h"
 #include "UnrealExporter.h"
 #include "Engine/Polys.h"
+#include "Engine/Selection.h"
+#include "Engine/LevelStreaming.h"
+#include "GameFramework/WorldSettings.h"
+#include "EngineUtils.h"
+#include "Engine/LevelScriptActor.h"
+#include "Engine/Light.h"
+#include "GameFramework/Pawn.h"
+#include "Engine/StaticMeshActor.h"
 
 #define LOCTEXT_NAMESPACE "UnrealEd.EditorActor"
 
diff --git a/Engine/Source/Editor/UnrealEd/Private/EditorActorFolders.cpp b/Engine/Source/Editor/UnrealEd/Private/EditorActorFolders.cpp
--- a/Engine/Source/Editor/UnrealEd/Private/EditorActorFolders.cpp
+++ b/Engine/Source/Editor/UnrealEd/Private/EditorActorFolders.cpp
@@ -3,6 +3,8 @@
 #include "UnrealEd.h"
 #include "EditorActorFolders.h"
 #include "ScopedTransaction.h"
+#include "EngineUtils.h"
+#include "Engine/Selection.h"
 
 #define LOCTEXT_NAMESPACE "FActorFolders"
 
diff --git a/Engine/Source/Editor/UnrealEd/Private/EditorBuildUtils.cpp b/Engine/Source/Editor/UnrealEd/Private/EditorBuildUtils.cpp
--- a/Engine/Source/Editor/UnrealEd/Private/EditorBuildUtils.cpp
+++ b/Engine/Source/Editor/UnrealEd/Private/EditorBuildUtils.cpp
@@ -17,6 +17,9 @@
 #include "MainFrame.h"
 #include "AssetToolsModule.h"
 #include "MessageLog.h"
+#include "Engine/LevelStreaming.h"
+#include "GameFramework/WorldSettings.h"
+#include "AI/Navigation/NavigationSystem.h"
 
 DEFINE_LOG_CATEGORY_STATIC(LogEditorBuildUtils, Log, All);
 
diff --git a/Engine/Source/Editor/UnrealEd/Private/EditorComponents.cpp b/Engine/Source/Editor/UnrealEd/Private/EditorComponents.cpp
--- a/Engine/Source/Editor/UnrealEd/Private/EditorComponents.cpp
+++ b/Engine/Source/Editor/UnrealEd/Private/EditorComponents.cpp
@@ -3,7 +3,9 @@
 
 #include "UnrealEd.h"
 #include "MouseDeltaTracker.h"
-
+#include "Materials/MaterialInstanceDynamic.h"
+#include "Engine/Selection.h"
+#include "GameFramework/WorldSettings.h"
 
 static TAutoConsoleVariable<int32> CVarEditorNewLevelGrid(
 	TEXT("r.Editor.NewLevelGrid"),
diff --git a/Engine/Source/Editor/UnrealEd/Private/EditorCsg.cpp b/Engine/Source/Editor/UnrealEd/Private/EditorCsg.cpp
--- a/Engine/Source/Editor/UnrealEd/Private/EditorCsg.cpp
+++ b/Engine/Source/Editor/UnrealEd/Private/EditorCsg.cpp
@@ -5,6 +5,10 @@
 #include "SurfaceIterators.h"
 #include "BSPOps.h"
 #include "ActorEditorUtils.h"
+#include "Engine/Polys.h"
+#include "EngineUtils.h"
+#include "Engine/Brush.h"
+#include "Engine/Selection.h"
 
 // Globals:
 static TArray<uint8*> GFlags1;
diff --git a/Engine/Source/Editor/UnrealEd/Private/EditorExporters.cpp b/Engine/Source/Editor/UnrealEd/Private/EditorExporters.cpp
--- a/Engine/Source/Editor/UnrealEd/Private/EditorExporters.cpp
+++ b/Engine/Source/Editor/UnrealEd/Private/EditorExporters.cpp
@@ -26,6 +26,17 @@
 #include "EngineModule.h"
 #include "InstancedFoliage.h"
 #include "RendererInterface.h"
+#include "Sound/SoundWave.h"
+#include "Engine/Polys.h"
+#include "GameFramework/PhysicsVolume.h"
+#include "Engine/StaticMeshActor.h"
+#include "Engine/TextureRenderTarget2D.h"
+#include "CanvasTypes.h"
+#include "Engine/Selection.h"
+#include "AI/Navigation/NavigationSystem.h"
+#include "Engine/TextureCube.h"
+#include "EngineUtils.h"
+#include "TextureLODSettings.h"
 
 DEFINE_LOG_CATEGORY_STATIC(LogEditorExporters, Log, All);
 
diff --git a/Engine/Source/Editor/UnrealEd/Private/EditorLevelUtils.cpp b/Engine/Source/Editor/UnrealEd/Private/EditorLevelUtils.cpp
--- a/Engine/Source/Editor/UnrealEd/Private/EditorLevelUtils.cpp
+++ b/Engine/Source/Editor/UnrealEd/Private/EditorLevelUtils.cpp
@@ -21,6 +21,10 @@
 
 #include "Runtime/AssetRegistry/Public/AssetRegistryModule.h"
 #include "Engine/LevelStreamingVolume.h"
+#include "Engine/LevelStreaming.h"
+#include "GameFramework/WorldSettings.h"
+#include "Engine/Selection.h"
+#include "Components/ModelComponent.h"
 
 DEFINE_LOG_CATEGORY(LogLevelTools);
 
diff --git a/Engine/Source/Editor/UnrealEd/Private/EditorModes.cpp b/Engine/Source/Editor/UnrealEd/Private/EditorModes.cpp
--- a/Engine/Source/Editor/UnrealEd/Private/EditorModes.cpp
+++ b/Engine/Source/Editor/UnrealEd/Private/EditorModes.cpp
@@ -18,6 +18,13 @@
 #include "ComponentVisualizer.h"
 #include "SNotificationList.h"
 #include "NotificationManager.h"
+#include "Engine/Selection.h"
+#include "EngineUtils.h"
+#include "CanvasItem.h"
+#include "CanvasTypes.h"
+#include "Engine/Polys.h"
+#include "Engine/StaticMeshActor.h"
+#include "Engine/LevelStreaming.h"
 
 DEFINE_LOG_CATEGORY_STATIC(LogEditorModes, Log, All);
 
diff --git a/Engine/Source/Editor/UnrealEd/Private/EditorObject.cpp b/Engine/Source/Editor/UnrealEd/Private/EditorObject.cpp
--- a/Engine/Source/Editor/UnrealEd/Private/EditorObject.cpp
+++ b/Engine/Source/Editor/UnrealEd/Private/EditorObject.cpp
@@ -11,6 +11,8 @@
 #include "Foliage/InstancedFoliageActor.h"
 #include "Foliage/FoliageType.h"
 #include "InstancedFoliage.h"
+#include "Components/BrushComponent.h"
+#include "Engine/StaticMesh.h"
 
 DEFINE_LOG_CATEGORY_STATIC(LogEditorObject, Log, All);
 
diff --git a/Engine/Source/Editor/UnrealEd/Private/EditorTransaction.cpp b/Engine/Source/Editor/UnrealEd/Private/EditorTransaction.cpp
--- a/Engine/Source/Editor/UnrealEd/Private/EditorTransaction.cpp
+++ b/Engine/Source/Editor/UnrealEd/Private/EditorTransaction.cpp
@@ -3,6 +3,7 @@
 
 #include "UnrealEd.h"
 #include "BSPOps.h"
+#include "Engine/BlueprintGeneratedClass.h"
 
 DEFINE_LOG_CATEGORY_STATIC(LogEditorTransaction, Log, All);
 
diff --git a/Engine/Source/Editor/UnrealEd/Private/EditorViewportClient.cpp b/Engine/Source/Editor/UnrealEd/Private/EditorViewportClient.cpp
--- a/Engine/Source/Editor/UnrealEd/Private/EditorViewportClient.cpp
+++ b/Engine/Source/Editor/UnrealEd/Private/EditorViewportClient.cpp
@@ -18,6 +18,12 @@
 #include "RendererInterface.h"
 #include "SNotificationList.h"
 #include "NotificationManager.h"
+#include "CanvasItem.h"
+#include "CanvasTypes.h"
+#include "Components/LineBatchComponent.h"
+#include "Debug/DebugDrawService.h"
+#include "Components/BillboardComponent.h"
+#include "EngineUtils.h"
 
 #define LOCTEXT_NAMESPACE "EditorViewportClient"
 
diff --git a/Engine/Source/Editor/UnrealEd/Private/FbxExporter.h b/Engine/Source/Editor/UnrealEd/Private/FbxExporter.h
--- a/Engine/Source/Editor/UnrealEd/Private/FbxExporter.h
+++ b/Engine/Source/Editor/UnrealEd/Private/FbxExporter.h
@@ -37,6 +37,9 @@
 class ALandscapeProxy;
 struct FAnimControlTrackKey;
 class USplineMeshComponent;
+class ALight;
+class UInterpTrackInstMove;
+class AMatineeActor;
 
 namespace UnFbx
 {
diff --git a/Engine/Source/Editor/UnrealEd/Private/FbxImporter.h b/Engine/Source/Editor/UnrealEd/Private/FbxImporter.h
--- a/Engine/Source/Editor/UnrealEd/Private/FbxImporter.h
+++ b/Engine/Source/Editor/UnrealEd/Private/FbxImporter.h
@@ -32,6 +32,12 @@
 
 #include "Factories.h"
 
+class ALight;
+class UInterpGroupInst;
+class AMatineeActor;
+class AActor;
+class UInterpTrackMove;
+
 // Temporarily disable a few warnings due to virtual function abuse in FBX source files
 #pragma warning( push )
 
diff --git a/Engine/Source/Editor/UnrealEd/Private/FileHelpers.cpp b/Engine/Source/Editor/UnrealEd/Private/FileHelpers.cpp
--- a/Engine/Source/Editor/UnrealEd/Private/FileHelpers.cpp
+++ b/Engine/Source/Editor/UnrealEd/Private/FileHelpers.cpp
@@ -29,6 +29,8 @@
 #include "PackageTools.h"
 #include "SNotificationList.h"
 #include "NotificationManager.h"
+#include "Engine/LevelStreaming.h"
+#include "GameMapsSettings.h"
 
 DEFINE_LOG_CATEGORY_STATIC(LogFileHelpers, Log, All);
 
diff --git a/Engine/Source/Editor/UnrealEd/Private/GroupActor.cpp b/Engine/Source/Editor/UnrealEd/Private/GroupActor.cpp
--- a/Engine/Source/Editor/UnrealEd/Private/GroupActor.cpp
+++ b/Engine/Source/Editor/UnrealEd/Private/GroupActor.cpp
@@ -3,6 +3,8 @@
 #include "UnrealEd.h"
 #include "ScopedTransaction.h"
 #include "Layers/ILayers.h"
+#include "Components/BillboardComponent.h"
+#include "Engine/Selection.h"
 
 const FLinearColor BOXCOLOR_LOCKEDGROUPS( 0.0f, 1.0f, 0.0f );
 const FLinearColor BOXCOLOR_UNLOCKEDGROUPS( 1.0f, 0.0f, 0.0f );
diff --git a/Engine/Source/Editor/UnrealEd/Private/LevelViewportClickHandlers.cpp b/Engine/Source/Editor/UnrealEd/Private/LevelViewportClickHandlers.cpp
--- a/Engine/Source/Editor/UnrealEd/Private/LevelViewportClickHandlers.cpp
+++ b/Engine/Source/Editor/UnrealEd/Private/LevelViewportClickHandlers.cpp
@@ -13,6 +13,10 @@
 #include "ActorEditorUtils.h"
 #include "Editor/ActorPositioning.h"
 #include "Engine/PointLight.h"
+#include "Components/PointLightComponent.h"
+#include "Engine/Selection.h"
+#include "Engine/StaticMeshActor.h"
+#include "Engine/TargetPoint.h"
 
 #define LOCTEXT_NAMESPACE "ClickHandlers"
 
diff --git a/Engine/Source/Editor/UnrealEd/Private/NavigationBuildingNotification.cpp b/Engine/Source/Editor/UnrealEd/Private/NavigationBuildingNotification.cpp
--- a/Engine/Source/Editor/UnrealEd/Private/NavigationBuildingNotification.cpp
+++ b/Engine/Source/Editor/UnrealEd/Private/NavigationBuildingNotification.cpp
@@ -7,6 +7,7 @@
 #include "NavigationBuildingNotification.h"
 #include "SNotificationList.h"
 #include "NotificationManager.h"
+#include "AI/Navigation/NavigationSystem.h"
 
 void FNavigationBuildingNotificationImpl::BuildStarted()
 {
diff --git a/Engine/Source/Editor/UnrealEd/Private/PhysicsAssetUtils.cpp b/Engine/Source/Editor/UnrealEd/Private/PhysicsAssetUtils.cpp
--- a/Engine/Source/Editor/UnrealEd/Private/PhysicsAssetUtils.cpp
+++ b/Engine/Source/Editor/UnrealEd/Private/PhysicsAssetUtils.cpp
@@ -7,7 +7,8 @@
 #include "MessageLog.h"
 #include "PhysicsEngine/PhysicsAsset.h"
 #include "PhysicsEngine/PhysicsConstraintTemplate.h"
-
+#include "EngineLogs.h"
+#include "PhysicsEngine/BodySetup.h"
 
 void FPhysAssetCreateParams::Initialize()
 {
diff --git a/Engine/Source/Editor/UnrealEd/Private/PlayLevel.cpp b/Engine/Source/Editor/UnrealEd/Private/PlayLevel.cpp
--- a/Engine/Source/Editor/UnrealEd/Private/PlayLevel.cpp
+++ b/Engine/Source/Editor/UnrealEd/Private/PlayLevel.cpp
@@ -38,8 +38,26 @@
 #include "SNotificationList.h"
 #include "SDPIScaler.h"
 #include "NotificationManager.h"
+#include "Engine/Selection.h"
+#include "TimerManager.h"
+#include "AI/Navigation/NavigationSystem.h"
 
 #include "Runtime/HeadMountedDisplay/Public/HeadMountedDisplay.h"
+#include "Components/AudioComponent.h"
+#include "Engine/Note.h"
+#include "UnrealEngine.h"
+#include "GameFramework/GameMode.h"
+#include "Engine/NavigationObjectBase.h"
+#include "GameFramework/PlayerController.h"
+#include "GameFramework/PlayerStart.h"
+#include "GameFramework/PlayerState.h"
+#include "GameFramework/WorldSettings.h"
+#include "Engine/LevelStreaming.h"
+#include "Engine/LocalPlayer.h"
+#include "Components/ModelComponent.h"
+#include "EngineUtils.h"
+#include "GameMapsSettings.h"
+#include "GameFramework/Pawn.h"
 
 DEFINE_LOG_CATEGORY_STATIC(LogPlayLevel, Log, All);
 
diff --git a/Engine/Source/Editor/UnrealEd/Private/PreviewMaterial.cpp b/Engine/Source/Editor/UnrealEd/Private/PreviewMaterial.cpp
--- a/Engine/Source/Editor/UnrealEd/Private/PreviewMaterial.cpp
+++ b/Engine/Source/Editor/UnrealEd/Private/PreviewMaterial.cpp
@@ -3,6 +3,7 @@
 #include "UnrealEd.h"
 #include "Editor/MaterialEditor/Public/MaterialEditorModule.h"
 #include "ComponentReregisterContext.h"
+#include "Materials/MaterialInstanceConstant.h"
 
 /**
  * Class for rendering the material on the preview mesh in the Material Editor
diff --git a/Engine/Source/Editor/UnrealEd/Private/SSocketManager.cpp b/Engine/Source/Editor/UnrealEd/Private/SSocketManager.cpp
--- a/Engine/Source/Editor/UnrealEd/Private/SSocketManager.cpp
+++ b/Engine/Source/Editor/UnrealEd/Private/SSocketManager.cpp
@@ -14,6 +14,8 @@
 #include "EngineAnalytics.h"
 #include "SInlineEditableTextBlock.h"
 #include "GenericCommands.h"
+#include "Engine/StaticMesh.h"
+#include "Engine/StaticMeshSocket.h"
 
 #define LOCTEXT_NAMESPACE "SSCSSocketManagerEditor"
 
diff --git a/Engine/Source/Editor/UnrealEd/Private/SnappingUtils.cpp b/Engine/Source/Editor/UnrealEd/Private/SnappingUtils.cpp
--- a/Engine/Source/Editor/UnrealEd/Private/SnappingUtils.cpp
+++ b/Engine/Source/Editor/UnrealEd/Private/SnappingUtils.cpp
@@ -8,6 +8,10 @@
 #include "VertexSnapping.h"
 #include "ViewportSnappingModule.h"
 #include "GameFramework/PhysicsVolume.h"
+#include "GameFramework/WorldSettings.h"
+#include "Engine/Selection.h"
+#include "EngineUtils.h"
+#include "Engine/PostProcessVolume.h"
 
 //////////////////////////////////////////////////////////////////////////
 // FEditorViewportSnapping
diff --git a/Engine/Source/Editor/UnrealEd/Private/SoundClassGraph.cpp b/Engine/Source/Editor/UnrealEd/Private/SoundClassGraph.cpp
--- a/Engine/Source/Editor/UnrealEd/Private/SoundClassGraph.cpp
+++ b/Engine/Source/Editor/UnrealEd/Private/SoundClassGraph.cpp
@@ -5,6 +5,7 @@
 
 #include "UnrealEd.h"
 #include "GraphEditor.h"
+#include "Sound/SoundClass.h"
 
 USoundClassGraph::USoundClassGraph(const FObjectInitializer& ObjectInitializer)
 	: Super(ObjectInitializer)
diff --git a/Engine/Source/Editor/UnrealEd/Private/SoundClassGraphNode.cpp b/Engine/Source/Editor/UnrealEd/Private/SoundClassGraphNode.cpp
--- a/Engine/Source/Editor/UnrealEd/Private/SoundClassGraphNode.cpp
+++ b/Engine/Source/Editor/UnrealEd/Private/SoundClassGraphNode.cpp
@@ -5,6 +5,7 @@
 =============================================================================*/
 
 #include "UnrealEd.h"
+#include "Sound/SoundClass.h"
 
 #define LOCTEXT_NAMESPACE "SoundClassGraphNode"
 
diff --git a/Engine/Source/Editor/UnrealEd/Private/SoundCueGraphNode.cpp b/Engine/Source/Editor/UnrealEd/Private/SoundCueGraphNode.cpp
--- a/Engine/Source/Editor/UnrealEd/Private/SoundCueGraphNode.cpp
+++ b/Engine/Source/Editor/UnrealEd/Private/SoundCueGraphNode.cpp
@@ -12,6 +12,7 @@
 #include "ScopedTransaction.h"
 #include "GenericCommands.h"
 #include "Sound/SoundCue.h"
+#include "Engine/Font.h"
 
 #define LOCTEXT_NAMESPACE "SoundCueGraphNode"
 
diff --git a/Engine/Source/Editor/UnrealEd/Private/ThumbnailHelpers.cpp b/Engine/Source/Editor/UnrealEd/Private/ThumbnailHelpers.cpp
--- a/Engine/Source/Editor/UnrealEd/Private/ThumbnailHelpers.cpp
+++ b/Engine/Source/Editor/UnrealEd/Private/ThumbnailHelpers.cpp
@@ -8,6 +8,13 @@
 #include "Particles/ParticleSystemComponent.h"
 #include "ContentStreaming.h"
 #include "Components/DirectionalLightComponent.h"
+#include "Engine/StaticMeshActor.h"
+#include "Animation/AnimSingleNodeInstance.h"
+#include "Animation/BlendSpaceBase.h"
+#include "Engine/SCS_Node.h"
+#include "Engine/SimpleConstructionScript.h"
+#include "Engine/StaticMesh.h"
+#include "Engine/TextureCube.h"
 
 /*
 ***************************************************************
diff --git a/Engine/Source/Editor/UnrealEd/Private/ThumbnailManager.cpp b/Engine/Source/Editor/UnrealEd/Private/ThumbnailManager.cpp
--- a/Engine/Source/Editor/UnrealEd/Private/ThumbnailManager.cpp
+++ b/Engine/Source/Editor/UnrealEd/Private/ThumbnailManager.cpp
@@ -7,6 +7,8 @@
 #include "ImageUtils.h"
 
 #include "ISourceControlModule.h"
+#include "Engine/StaticMesh.h"
+#include "Engine/TextureCube.h"
 
 DEFINE_LOG_CATEGORY_STATIC(LogThumbnailManager, Log, All);
 
diff --git a/Engine/Source/Editor/UnrealEd/Private/UnrealEd.cpp b/Engine/Source/Editor/UnrealEd/Private/UnrealEd.cpp
--- a/Engine/Source/Editor/UnrealEd/Private/UnrealEd.cpp
+++ b/Engine/Source/Editor/UnrealEd/Private/UnrealEd.cpp
@@ -21,6 +21,7 @@
 #include "AssetEditorManager.h"
 
 #include "EditorActorFolders.h"
+#include "UnrealEngine.h"
 
 UUnrealEdEngine* GUnrealEd;
 
diff --git a/Engine/Source/Editor/UnrealEd/Private/UnrealEdSrv.cpp b/Engine/Source/Editor/UnrealEd/Private/UnrealEdSrv.cpp
--- a/Engine/Source/Editor/UnrealEd/Private/UnrealEdSrv.cpp
+++ b/Engine/Source/Editor/UnrealEd/Private/UnrealEdSrv.cpp
@@ -40,6 +40,15 @@
 #include "Engine/DestructibleMesh.h"
 #include "NotificationManager.h"
 #include "SNotificationList.h"
+#include "Engine/Polys.h"
+#include "Engine/Selection.h"
+#include "Components/BoxComponent.h"
+#include "Components/CapsuleComponent.h"
+#include "Components/PointLightComponent.h"
+#include "PhysicsEngine/RadialForceComponent.h"
+#include "Components/SphereComponent.h"
+#include "Engine/StaticMeshActor.h"
+#include "EngineUtils.h"
 
 
 DEFINE_LOG_CATEGORY_STATIC(LogUnrealEdSrv, Log, All);
diff --git a/Engine/Source/Editor/UnrealEd/Private/VertexSnapping.cpp b/Engine/Source/Editor/UnrealEd/Private/VertexSnapping.cpp
--- a/Engine/Source/Editor/UnrealEd/Private/VertexSnapping.cpp
+++ b/Engine/Source/Editor/UnrealEd/Private/VertexSnapping.cpp
@@ -4,6 +4,9 @@
 #include "VertexSnapping.h"
 #include "StaticMeshResources.h"
 #include "Components/BrushComponent.h"
+#include "Engine/Polys.h"
+#include "Engine/Selection.h"
+#include "EngineUtils.h"
 
 namespace VertexSnappingConstants
 {
