UnrealEd - UE-4499 - UDN: UE4.5 crash after running test (UE4Editor-Core.dll!TStaticIndirectArrayThreadSafeRead)

diff --git a/Engine/Source/Editor/UnrealEd/Private/FileHelpers.cpp b/Engine/Source/Editor/UnrealEd/Private/FileHelpers.cpp
--- a/Engine/Source/Editor/UnrealEd/Private/FileHelpers.cpp
+++ b/Engine/Source/Editor/UnrealEd/Private/FileHelpers.cpp
@@ -2520,6 +2520,11 @@ bool FEditorFileUtils::SaveDirtyPackages(const bool bPromptUserToSave, const boo
 		*bOutPackagesNeededSaving = false;
 	}
 
+	if( bSaveContentPackages )
+	{
+		CollectGarbage(GARBAGE_COLLECTION_KEEPFLAGS);
+	}
+
 	// A list of all packages that need to be saved
 	TArray<UPackage*> PackagesToSave;
 
@@ -2528,11 +2533,9 @@ bool FEditorFileUtils::SaveDirtyPackages(const bool bPromptUserToSave, const boo
 		GetDirtyWorldPackages(PackagesToSave);
 	}
 
-	// Don't iterate through content packages if we dont plan on saving them
+	// Don't iterate through content packages if we don't plan on saving them
 	if( bSaveContentPackages )
 	{
-		CollectGarbage(GARBAGE_COLLECTION_KEEPFLAGS);
-
 		GetDirtyContentPackages(PackagesToSave);
 	}
 
