Preventing save on compile setting from taking effect while running a commandlet.

diff --git a/Engine/Source/Editor/UnrealEd/Private/Kismet2/Kismet2.cpp b/Engine/Source/Editor/UnrealEd/Private/Kismet2/Kismet2.cpp
--- a/Engine/Source/Editor/UnrealEd/Private/Kismet2/Kismet2.cpp
+++ b/Engine/Source/Editor/UnrealEd/Private/Kismet2/Kismet2.cpp
@@ -502,7 +502,7 @@ void FKismetEditorUtilities::CompileBlueprint(UBlueprint* BlueprintObj, bool bIs
 			// compile was already kicked off as a result of a level save, as it
 			// could cause a recursive save)... let the "SaveOnCompile" setting 
 			// only save blueprint assets
-			if (!bIsLevelPackage)
+			if (!bIsLevelPackage && !IsRunningCommandlet())
 			{
 				TArray<UPackage*> PackagesToSave;
 				PackagesToSave.Add(BlueprintPackage);
