Only add a dependency from UnrealEd to AndroidPlatformEditor on Windows

diff --git a/Engine/Source/Editor/UnrealEd/UnrealEd.Build.cs b/Engine/Source/Editor/UnrealEd/UnrealEd.Build.cs
--- a/Engine/Source/Editor/UnrealEd/UnrealEd.Build.cs
+++ b/Engine/Source/Editor/UnrealEd/UnrealEd.Build.cs
@@ -188,11 +188,15 @@ public class UnrealEd : ModuleRules
 				"ReferenceViewer",
 				"EditorLiveStreaming",
 				"HotReload",
-                "AndroidPlatformEditor",
                 "IOSPlatformEditor"
 			}
 		);
 
+		if (Target.Platform == UnrealTargetPlatform.Win64 || Target.Platform == UnrealTargetPlatform.Win32)
+		{
+			DynamicallyLoadedModuleNames.Add("AndroidPlatformEditor");
+		}
+
 		CircularlyReferencedDependentModules.AddRange(
 			new string[] 
 			{
