Fix comment typo

diff --git a/Engine/Source/Editor/UnrealEd/Classes/Settings/ProjectPackagingSettings.h b/Engine/Source/Editor/UnrealEd/Classes/Settings/ProjectPackagingSettings.h
--- a/Engine/Source/Editor/UnrealEd/Classes/Settings/ProjectPackagingSettings.h
+++ b/Engine/Source/Editor/UnrealEd/Classes/Settings/ProjectPackagingSettings.h
@@ -72,7 +72,7 @@ public:
 
 	/**
 	 * If enabled, a distribution build will be created and the shipping configuration will be used
-	 * If disabled, an development build will be created
+	 * If disabled, a development build will be created
 	 * Distribution builds are for publishing to the App Store
 	 */
 	UPROPERTY(config, EditAnywhere, Category=Project, AdvancedDisplay)
