Don't write out manifest and media files in the pak file

diff --git a/Engine/Source/Editor/UnrealEd/Private/Factories/PackFactory.cpp b/Engine/Source/Editor/UnrealEd/Private/Factories/PackFactory.cpp
--- a/Engine/Source/Editor/UnrealEd/Private/Factories/PackFactory.cpp
+++ b/Engine/Source/Editor/UnrealEd/Private/Factories/PackFactory.cpp
@@ -381,6 +381,12 @@ UObject* UPackFactory::FactoryCreateBinary
 				continue;
 			}
 
+			// Media and manifest files don't get written out as part of the install
+			if (It.Filename().Contains(TEXT("manifest.json")) || It.Filename().StartsWith(TEXT("Media/")) || It.Filename().Contains(TEXT("/Media/")))
+			{
+				continue;
+			}
+
 			const FPakEntry& Entry = It.Info();
 			PakReader.Seek(Entry.Offset);
 			FPakEntry EntryInfo;
