Editor: Add accessor GetModeTools() to FEditorViewportClient

diff --git a/Engine/Source/Editor/UnrealEd/Public/EditorViewportClient.h b/Engine/Source/Editor/UnrealEd/Public/EditorViewportClient.h
--- a/Engine/Source/Editor/UnrealEd/Public/EditorViewportClient.h
+++ b/Engine/Source/Editor/UnrealEd/Public/EditorViewportClient.h
@@ -1181,6 +1181,12 @@ private:
 	/** Delegate handler for when all stats are disabled in a viewport */
 	void HandleViewportStatDisableAll(const bool bInAnyViewport);
 
+	/** Editor mode tool manager being used for this viewport client */
+	FEditorModeTools* GetModeTools() const
+	{
+		return ModeTools;
+	}
+
 public:
 	static const uint32 MaxCameraSpeeds;
 
