Editor: Remove dead command 'SKELETALMESH CHARBITS'

diff --git a/Engine/Source/Editor/UnrealEd/Classes/Editor/UnrealEdEngine.h b/Engine/Source/Editor/UnrealEd/Classes/Editor/UnrealEdEngine.h
--- a/Engine/Source/Editor/UnrealEd/Classes/Editor/UnrealEdEngine.h
+++ b/Engine/Source/Editor/UnrealEd/Classes/Editor/UnrealEdEngine.h
@@ -297,7 +297,6 @@ class UNREALED_API UUnrealEdEngine : public UEditorEngine, public FNotifyHook
 	bool Exec_Pivot( const TCHAR* Str, FOutputDevice& Ar );
 	bool Exec_Actor( UWorld* InWorld, const TCHAR* Str, FOutputDevice& Ar );
 	bool Exec_Mode( const TCHAR* Str, FOutputDevice& Ar );
-	bool Exec_SkeletalMesh( const TCHAR* Str, FOutputDevice& Ar );
 	bool Exec_Group( const TCHAR* Str, FOutputDevice& Ar );
 
 
diff --git a/Engine/Source/Editor/UnrealEd/Private/UnrealEdSrv.cpp b/Engine/Source/Editor/UnrealEd/Private/UnrealEdSrv.cpp
--- a/Engine/Source/Editor/UnrealEd/Private/UnrealEdSrv.cpp
+++ b/Engine/Source/Editor/UnrealEd/Private/UnrealEdSrv.cpp
@@ -802,13 +802,6 @@ bool UUnrealEdEngine::Exec( UWorld* InWorld, const TCHAR* Stream, FOutputDevice&
 		return Exec_Actor( InWorld, Str, Ar );
 	}
 	//------------------------------------------------------------------------------------
-	// SKELETALMESH: SkeletalMesh-related functions
-	//
-	else if(FParse::Command(&Str, TEXT("SKELETALMESH")))
-	{
-		return Exec_SkeletalMesh(Str, Ar);
-	}
-	//------------------------------------------------------------------------------------
 	// MODE management (Global EDITOR mode):
 	//
 	else if( FParse::Command(&Str,TEXT("MODE")) )
@@ -1952,49 +1945,6 @@ FPoly* CreateHugeTrianglePolygonOnPlane( const FPlane* InPlane )
 	return Triangle;
 }
 
-bool UUnrealEdEngine::Exec_SkeletalMesh( const TCHAR* Str, FOutputDevice& Ar )
-{
-	//This command sets the offset and orientation for all skeletal meshes within the set of currently selected packages
-	if(FParse::Command(&Str, TEXT("CHARBITS"))) //SKELETALMESH CHARBITS
-	{
-		FVector Offset = FVector::ZeroVector;
-		FRotator Orientation = FRotator::ZeroRotator;
-		bool bHasOffset = GetFVECTOR(Str, TEXT("OFFSET="), Offset);
-		
-		TCHAR TempChars[80];
-		bool bHasOrientation = GetSUBSTRING(Str, TEXT("ORIENTATION="), TempChars, 80);
-		
-		//If orientation is present do custom parsing to allow for a proper conversion from a floating point representation of degress
-		//to its integer representation in FRotator. GetFROTATOR() does not allow us to do this.
-		if(bHasOrientation)
-		{
-			float Value = 0.0f;
-			
-			if(FParse::Value(TempChars, TEXT("YAW="), Value))
-			{
-				Value = FMath::Fmod(Value, 360.0f); //Make sure it's in the range 0-360
-				Orientation.Yaw = Value;
-			}
-
-			if(FParse::Value(TempChars, TEXT("PITCH="), Value))
-			{
-				Value = FMath::Fmod(Value, 360.0f); //Make sure it's in the range 0-360
-				Orientation.Pitch = Value;
-			}
-
-			if(FParse::Value(TempChars, TEXT("ROLL="), Value))
-			{
-				Value = FMath::Fmod(Value, 360.0f); //Make sure it's in the range 0-360
-				Orientation.Roll = Value;
-			}
-		}
-
-		return true;
-	}
-
-	return false;
-}
-
 bool UUnrealEdEngine::Exec_Actor( UWorld* InWorld, const TCHAR* Str, FOutputDevice& Ar )
 {
 	// Keep a pointer to the beginning of the string to use for message displaying purposes
