Editor: Add tracking of the associated toolkit host for a FEditorModeTools instance and reduced dependencies on the level editor in individual modes

diff --git a/Engine/Plugins/2D/Paper2D/Source/Paper2DEditor/Private/TileMapEditing/EdModeTileMap.cpp b/Engine/Plugins/2D/Paper2D/Source/Paper2DEditor/Private/TileMapEditing/EdModeTileMap.cpp
--- a/Engine/Plugins/2D/Paper2D/Source/Paper2DEditor/Private/TileMapEditing/EdModeTileMap.cpp
+++ b/Engine/Plugins/2D/Paper2D/Source/Paper2DEditor/Private/TileMapEditing/EdModeTileMap.cpp
@@ -106,15 +106,13 @@ bool FEdModeTileMap::UsesToolkits() const
 
 void FEdModeTileMap::Enter()
 {
+	FEdMode::Enter();
+
 	if (!Toolkit.IsValid())
 	{
-		// @todo: Remove this assumption when we make modes per level editor instead of global
-		auto ToolkitHost = FModuleManager::LoadModuleChecked< FLevelEditorModule >( "LevelEditor" ).GetFirstLevelEditor();
-		Toolkit = MakeShareable(new FTileMapEdModeToolkit);
-		Toolkit->Init(ToolkitHost);
+		Toolkit = MakeShareable(new FTileMapEdModeToolkit(this));
+		Toolkit->Init(Owner->GetToolkitHost());
 	}
-
-	FEdMode::Enter();
 }
 
 void FEdModeTileMap::Exit()
diff --git a/Engine/Source/Editor/FoliageEdit/Private/FoliageEdMode.cpp b/Engine/Source/Editor/FoliageEdit/Private/FoliageEdMode.cpp
--- a/Engine/Source/Editor/FoliageEdit/Private/FoliageEdMode.cpp
+++ b/Engine/Source/Editor/FoliageEdit/Private/FoliageEdMode.cpp
@@ -115,10 +115,8 @@ void FEdModeFoliage::Enter()
 
 	if (!Toolkit.IsValid())
 	{
-		// @todo: Remove this assumption when we make modes per level editor instead of global
-		auto ToolkitHost = FModuleManager::LoadModuleChecked< FLevelEditorModule >("LevelEditor").GetFirstLevelEditor();
 		Toolkit = MakeShareable(new FFoliageEdModeToolkit);
-		Toolkit->Init(ToolkitHost);
+		Toolkit->Init(Owner->GetToolkitHost());
 	}
 
 	// Fixup any broken clusters
diff --git a/Engine/Source/Editor/GeometryMode/Private/GeometryEdMode.cpp b/Engine/Source/Editor/GeometryMode/Private/GeometryEdMode.cpp
--- a/Engine/Source/Editor/GeometryMode/Private/GeometryEdMode.cpp
+++ b/Engine/Source/Editor/GeometryMode/Private/GeometryEdMode.cpp
@@ -119,10 +119,8 @@ void FEdModeGeometry::Enter()
 	
 	if (!Toolkit.IsValid())
 	{
-		// @todo: Remove this assumption when we make modes per level editor instead of global
-		auto ToolkitHost = FModuleManager::LoadModuleChecked< FLevelEditorModule >( "LevelEditor" ).GetFirstLevelEditor();
 		Toolkit = MakeShareable(new FGeometryMode);
-		Toolkit->Init(ToolkitHost);
+		Toolkit->Init(Owner->GetToolkitHost());
 	}
 
 	GetFromSource();
diff --git a/Engine/Source/Editor/LandscapeEditor/Private/LandscapeEdMode.cpp b/Engine/Source/Editor/LandscapeEditor/Private/LandscapeEdMode.cpp
--- a/Engine/Source/Editor/LandscapeEditor/Private/LandscapeEdMode.cpp
+++ b/Engine/Source/Editor/LandscapeEditor/Private/LandscapeEdMode.cpp
@@ -429,10 +429,8 @@ void FEdModeLandscape::Enter()
 	// Create the landscape editor window
 	if (!Toolkit.IsValid())
 	{
-		// @todo: Remove this assumption when we make modes per level editor instead of global
-		auto ToolkitHost = FModuleManager::LoadModuleChecked< FLevelEditorModule >("LevelEditor").GetFirstLevelEditor();
 		Toolkit = MakeShareable(new FLandscapeToolKit);
-		Toolkit->Init(ToolkitHost);
+		Toolkit->Init(Owner->GetToolkitHost());
 	}
 
 	// Force real-time viewports.  We'll back up the current viewport state so we can restore it when the
diff --git a/Engine/Source/Editor/LevelEditor/Private/LevelEditor.cpp b/Engine/Source/Editor/LevelEditor/Private/LevelEditor.cpp
--- a/Engine/Source/Editor/LevelEditor/Private/LevelEditor.cpp
+++ b/Engine/Source/Editor/LevelEditor/Private/LevelEditor.cpp
@@ -393,6 +393,7 @@ TSharedPtr<FTabManager> FLevelEditorModule::GetLevelEditorTabManager() const
 void FLevelEditorModule::SetLevelEditorInstance( TWeakPtr<SLevelEditor> LevelEditor )
 {
 	LevelEditorInstancePtr = LevelEditor;
+	GLevelEditorModeTools().SetToolkitHost(LevelEditorInstancePtr.Pin().ToSharedRef());
 }
 
 void FLevelEditorModule::SetLevelEditorInstanceTab( TWeakPtr<SDockTab> LevelEditorTab )
diff --git a/Engine/Source/Editor/MeshPaint/Private/MeshPaintEdMode.cpp b/Engine/Source/Editor/MeshPaint/Private/MeshPaintEdMode.cpp
--- a/Engine/Source/Editor/MeshPaint/Private/MeshPaintEdMode.cpp
+++ b/Engine/Source/Editor/MeshPaint/Private/MeshPaintEdMode.cpp
@@ -151,9 +151,8 @@ void FEdModeMeshPaint::Enter()
 	
 	if (!Toolkit.IsValid())
 	{
-		auto ToolkitHost = FModuleManager::LoadModuleChecked< FLevelEditorModule >( "LevelEditor" ).GetFirstLevelEditor();
 		Toolkit = MakeShareable(new FMeshPaintToolKit);
-		Toolkit->Init(ToolkitHost);
+		Toolkit->Init(Owner->GetToolkitHost());
 	}
 
 	// Change the engine to draw selected objects without a color boost, but unselected objects will
diff --git a/Engine/Source/Editor/PlacementMode/Private/PlacementMode.cpp b/Engine/Source/Editor/PlacementMode/Private/PlacementMode.cpp
--- a/Engine/Source/Editor/PlacementMode/Private/PlacementMode.cpp
+++ b/Engine/Source/Editor/PlacementMode/Private/PlacementMode.cpp
@@ -45,9 +45,8 @@ void FPlacementMode::Enter()
 
 	if ( !Toolkit.IsValid() )
 	{
-		TSharedPtr< ILevelEditor > ToolkitHost = FModuleManager::LoadModuleChecked< FLevelEditorModule >( "LevelEditor" ).GetFirstLevelEditor();
 		Toolkit = MakeShareable( new FPlacementModeToolkit );
-		Toolkit->Init( ToolkitHost );
+		Toolkit->Init(Owner->GetToolkitHost());
 	}
 }
 
diff --git a/Engine/Source/Editor/UnrealEd/Private/EditorModes.cpp b/Engine/Source/Editor/UnrealEd/Private/EditorModes.cpp
--- a/Engine/Source/Editor/UnrealEd/Private/EditorModes.cpp
+++ b/Engine/Source/Editor/UnrealEd/Private/EditorModes.cpp
@@ -1229,6 +1229,12 @@ void FEditorModeTools::SaveConfig(void)
 	SaveWidgetSettings();
 }
 
+void FEditorModeTools::SetToolkitHost(TSharedRef<class IToolkitHost> InHost)
+{
+	checkf(!ToolkitHost.IsValid(), TEXT("SetToolkitHost can only be called once"));
+	ToolkitHost = InHost;
+}
+
 void FEditorModeTools::OnEditorSelectionChanged(UObject* NewSelection)
 {
 	// If selecting an actor, move the pivot location.
diff --git a/Engine/Source/Editor/UnrealEd/Public/EditorModes.h b/Engine/Source/Editor/UnrealEd/Public/EditorModes.h
--- a/Engine/Source/Editor/UnrealEd/Public/EditorModes.h
+++ b/Engine/Source/Editor/UnrealEd/Public/EditorModes.h
@@ -784,6 +784,15 @@ public:
 	/** Draws in the top level corner of all FEditorViewportClient windows (can be used to relay info to the user). */
 	FString InfoString;
 
+	/** Sets the host for toolkits created via modes from this mode manager (can only be called once) */
+	void SetToolkitHost(TSharedRef<class IToolkitHost> Host);
+
+	/** Returns the host for toolkits created via modes from this mode manager */
+	TSharedPtr<class IToolkitHost> GetToolkitHost() const
+	{
+		check(ToolkitHost.IsValid());
+		return ToolkitHost;
+	}
 protected:
 	/** 
 	 * Delegate handlers
@@ -797,6 +806,9 @@ protected:
 	/** A list of active editor modes. */
 	TArray< TSharedPtr<FEdMode> > Modes;
 
+	/** The host of the toolkits created by these modes */
+	TSharedPtr<IToolkitHost> ToolkitHost;
+
 	/** A list of previously active editor modes that we will potentially recycle */
 	TMap< FEditorModeID, TSharedPtr<FEdMode> > RecycledModes;
 
