Editor: Replace some GLevelEditorModeTools() uses with ModeTools in editor viewport client subclasses

diff --git a/Engine/Source/Editor/Matinee/Private/MatineeViewportClient.cpp b/Engine/Source/Editor/Matinee/Private/MatineeViewportClient.cpp
--- a/Engine/Source/Editor/Matinee/Private/MatineeViewportClient.cpp
+++ b/Engine/Source/Editor/Matinee/Private/MatineeViewportClient.cpp
@@ -920,7 +920,7 @@ void FMatineeViewportClient::MouseMove(FViewport* Viewport, int32 X, int32 Y)
 		}
 		else if(bPanning)
 		{
-			const bool bInvertPanning = GLevelEditorModeTools().GetInterpPanInvert();
+			const bool bInvertPanning = ModeTools->GetInterpPanInvert();
 
 			float DeltaTime = (bInvertPanning ? -DeltaX : DeltaX) / InterpEd->PixelsPerSec;
 			InterpEd->ViewStartTime -= DeltaTime;
diff --git a/Engine/Source/Editor/Persona/Private/AnimationEditorViewportClient.cpp b/Engine/Source/Editor/Persona/Private/AnimationEditorViewportClient.cpp
--- a/Engine/Source/Editor/Persona/Private/AnimationEditorViewportClient.cpp
+++ b/Engine/Source/Editor/Persona/Private/AnimationEditorViewportClient.cpp
@@ -1339,12 +1339,12 @@ FMatrix FAnimationViewportClient::GetWidgetCoordSystem() const
 
 ECoordSystem FAnimationViewportClient::GetWidgetCoordSystemSpace() const
 { 
-	return GLevelEditorModeTools().GetCoordSystem();
+	return ModeTools->GetCoordSystem();
 }
 
 void FAnimationViewportClient::SetWidgetCoordSystemSpace(ECoordSystem NewCoordSystem)
 {
-	GLevelEditorModeTools().SetCoordSystem(NewCoordSystem);
+	ModeTools->SetCoordSystem(NewCoordSystem);
 	Invalidate();
 }
 
diff --git a/Engine/Source/Editor/UnrealEd/Private/LevelEditorViewport.cpp b/Engine/Source/Editor/UnrealEd/Private/LevelEditorViewport.cpp
--- a/Engine/Source/Editor/UnrealEd/Private/LevelEditorViewport.cpp
+++ b/Engine/Source/Editor/UnrealEd/Private/LevelEditorViewport.cpp
@@ -1913,7 +1913,7 @@ void FLevelEditorViewportClient::ProjectActorsIntoWorld(const TArray<AActor*>& A
 	const FMatrix InputCoordSystem = GetWidgetCoordSystem();
 	const EAxisList::Type CurrentAxis = GetCurrentWidgetAxis();
 	
-	const FVector DeltaTranslation = (GLevelEditorModeTools().PivotLocation - GLevelEditorModeTools().CachedLocation) + Drag;
+	const FVector DeltaTranslation = (ModeTools->PivotLocation - ModeTools->CachedLocation) + Drag;
 
 
 	// Loop over all the actors and attempt to snap them along the drag axis normal
@@ -2020,14 +2020,14 @@ bool FLevelEditorViewportClient::InputWidgetDelta( FViewport* Viewport, EAxisLis
 	}
 
 	bool bHandled = false;
+
 	// Give the current editor mode a chance to use the input first.  If it does, don't apply it to anything else.
-	FEditorModeTools& Tools = GLevelEditorModeTools();
-	if( Tools.InputDelta( this, Viewport, Drag, Rot, Scale ) )
+	if (ModeTools->InputDelta(this, Viewport, Drag, Rot, Scale))
 	{
-		if( Tools.AllowWidgetMove() )
+		if (ModeTools->AllowWidgetMove())
 		{
-			Tools.PivotLocation += Drag;
-			Tools.SnappedLocation += Drag;
+			ModeTools->PivotLocation += Drag;
+			ModeTools->SnappedLocation += Drag;
 		}
 
 		// Update visuals of the rotate widget 
@@ -2096,12 +2096,12 @@ bool FLevelEditorViewportClient::InputWidgetDelta( FViewport* Viewport, EAxisLis
 				}
 				else
 				{
-					FSnappingUtils::SnapDragLocationToNearestVertex( Tools.PivotLocation, Drag, this );
+					FSnappingUtils::SnapDragLocationToNearestVertex( ModeTools->PivotLocation, Drag, this );
 					bPivotMovedIndependantly = true;
 				}
 
-				Tools.PivotLocation += Drag;
-				Tools.SnappedLocation += Drag;
+				ModeTools->PivotLocation += Drag;
+				ModeTools->SnappedLocation += Drag;
 
 				if( IsShiftPressed() )
 				{
@@ -2110,7 +2110,7 @@ bool FLevelEditorViewportClient::InputWidgetDelta( FViewport* Viewport, EAxisLis
 				}
 
 				TArray<FEdMode*> ActiveModes; 
-				Tools.GetActiveModes( ActiveModes );
+				ModeTools->GetActiveModes(ActiveModes);
 
 				for( int32 ModeIndex = 0; ModeIndex < ActiveModes.Num(); ++ModeIndex )
 				{
@@ -2151,14 +2151,14 @@ TSharedPtr<FDragTool> FLevelEditorViewportClient::MakeDragTool( EDragTool::Type
 void FLevelEditorViewportClient::UpdateMouseDelta()
 {
 	// Do nothing if a drag tool is being used.
-	if( MouseDeltaTracker->UsingDragTool() || GLevelEditorModeTools().DisallowMouseDeltaTracking() )
+	if (MouseDeltaTracker->UsingDragTool() || ModeTools->DisallowMouseDeltaTracking())
 	{
 		return;
 	}
 
 	// Stop tracking and do nothing else if we're tracking and the widget mode has changed mid-track.
 	// It can confuse the widget code that handles the mouse movements.
-	if (bIsTracking && MouseDeltaTracker->GetTrackingWidgetMode() != GLevelEditorModeTools().GetWidgetMode())
+	if (bIsTracking && MouseDeltaTracker->GetTrackingWidgetMode() != ModeTools->GetWidgetMode())
 	{
 		StopTracking();
 		return;
@@ -2271,8 +2271,9 @@ bool FLevelEditorViewportClient::InputKey(FViewport* Viewport, int32 ControllerI
 		return true;
 	}
 
+	
 	// Let the current mode have a look at the input before reacting to it.
-	if( GLevelEditorModeTools().InputKey(this, Viewport, Key, Event) )
+	if (ModeTools->InputKey(this, Viewport, Key, Event))
 	{
 		return true;
 	}
@@ -2326,11 +2327,11 @@ bool FLevelEditorViewportClient::InputKey(FViewport* Viewport, int32 ControllerI
 		&& InputState.IsButtonPressed(EKeys::RightMouseButton) 
 		&& IsOrtho() )
 	{
-		GLevelEditorModeTools().SetWidgetModeOverride( FWidget::WM_Rotate );
+		ModeTools->SetWidgetModeOverride(FWidget::WM_Rotate);
 	}
 	else
 	{
-		GLevelEditorModeTools().SetWidgetModeOverride( FWidget::WM_None );
+		ModeTools->SetWidgetModeOverride(FWidget::WM_None);
 	}
 
 	bHandled |= FEditorViewportClient::InputKey(Viewport,ControllerId,Key,Event,AmountDepressed,bGamepad);
@@ -2348,7 +2349,7 @@ bool FLevelEditorViewportClient::InputKey(FViewport* Viewport, int32 ControllerI
 void FLevelEditorViewportClient::TrackingStarted( const FInputEventState& InInputState, bool bIsDraggingWidget, bool bNudge )
 {
 	// Begin transacting.  Give the current editor mode an opportunity to do the transacting.
-	const bool bTrackingHandledExternally = GLevelEditorModeTools().StartTracking(this, Viewport);
+	const bool bTrackingHandledExternally = ModeTools->StartTracking(this, Viewport);
 
 	TrackingTransaction.End();
 
@@ -2488,7 +2489,7 @@ void FLevelEditorViewportClient::TrackingStopped()
 	bool bDidAnythingActuallyChange = false;
 
 	// Stop transacting.  Give the current editor mode an opportunity to do the transacting.
-	const bool bTransactingHandledByEditorMode = GLevelEditorModeTools().EndTracking(this, Viewport);
+	const bool bTransactingHandledByEditorMode = ModeTools->EndTracking(this, Viewport);
 	if( !bTransactingHandledByEditorMode )
 	{
 		if( TrackingTransaction.TransCount > 0 )
@@ -2563,7 +2564,7 @@ void FLevelEditorViewportClient::TrackingStopped()
 	}
 
 	TArray<FEdMode*> ActiveModes; 
-	GLevelEditorModeTools().GetActiveModes( ActiveModes );
+	ModeTools->GetActiveModes(ActiveModes);
 	for( int32 ModeIndex = 0; ModeIndex < ActiveModes.Num(); ++ModeIndex )
 	{
 		// Also notify the current editing modes if they are interested.
@@ -2720,7 +2721,7 @@ bool FLevelEditorViewportClient::InputAxis(FViewport* Viewport, int32 Controller
 
 	FScopedSetCurrentViewportClient( this );
 	// Let the current mode have a look at the input before reacting to it.
-	if ( GLevelEditorModeTools().InputAxis(this, Viewport, ControllerId, Key, Delta, DeltaTime) )
+	if (ModeTools->InputAxis(this, Viewport, ControllerId, Key, Delta, DeltaTime))
 	{
 		return true;
 	}
@@ -2773,9 +2774,9 @@ bool FLevelEditorViewportClient::IsVolumeVisibleInViewport( const AActor& Volume
 
 void FLevelEditorViewportClient::SetWidgetMode( FWidget::EWidgetMode ActivatedMode )
 {
-	if( !GLevelEditorModeTools().IsTracking() && !IsFlightCameraActive() )
+	if (!ModeTools->IsTracking() && !IsFlightCameraActive())
 	{
-		GLevelEditorModeTools().SetWidgetMode( ActivatedMode );
+		ModeTools->SetWidgetMode(ActivatedMode);
 
 		// force an invalidation (non-deferred) of the hit proxy here, otherwise we will
 		// end up checking against an incorrect hit proxy if the cursor is not moved
@@ -2783,7 +2784,7 @@ void FLevelEditorViewportClient::SetWidgetMode( FWidget::EWidgetMode ActivatedMo
 		bShouldCheckHitProxy = true;
 
 		// Fire event delegate
-		GLevelEditorModeTools().BroadcastWidgetModeChanged(ActivatedMode);
+		ModeTools->BroadcastWidgetModeChanged(ActivatedMode);
 	}
 	
 	// Invalidate all viewports, so the new gizmo is rendered in each one
@@ -2798,18 +2799,18 @@ void FLevelEditorViewportClient::SetWidgetMode( FWidget::EWidgetMode ActivatedMo
 
 bool FLevelEditorViewportClient::CanSetWidgetMode( FWidget::EWidgetMode NewMode ) const
 {
-	return GLevelEditorModeTools().GetShowWidget() == true;
+	return ModeTools->GetShowWidget() == true;
 }
 
 void FLevelEditorViewportClient::SetWidgetCoordSystemSpace( ECoordSystem NewCoordSystem )
 {
-	GLevelEditorModeTools().SetCoordSystem(NewCoordSystem);
+	ModeTools->SetCoordSystem(NewCoordSystem);
 	Invalidate();
 }
 
 FWidget::EWidgetMode FLevelEditorViewportClient::GetWidgetMode() const
 {
-	return GLevelEditorModeTools().GetWidgetMode();
+	return ModeTools->GetWidgetMode();
 }
 
 FVector FLevelEditorViewportClient::GetWidgetLocation() const
@@ -2820,7 +2821,7 @@ FVector FLevelEditorViewportClient::GetWidgetLocation() const
 		return ComponentVisWidgetLocation;
 	}
 
-	return GLevelEditorModeTools().GetWidgetLocation();
+	return ModeTools->GetWidgetLocation();
 }
 
 FMatrix FLevelEditorViewportClient::GetWidgetCoordSystem() const 
@@ -2831,12 +2832,12 @@ FMatrix FLevelEditorViewportClient::GetWidgetCoordSystem() const
 		return ComponentVisWidgetCoordSystem;
 	}
 
-	return GLevelEditorModeTools().GetCustomInputCoordinateSystem();
+	return ModeTools->GetCustomInputCoordinateSystem();
 }
 
 ECoordSystem FLevelEditorViewportClient::GetWidgetCoordSystemSpace() const
 {
-	return GLevelEditorModeTools().GetCoordSystem();
+	return ModeTools->GetCoordSystem();
 }
 
 
@@ -3237,21 +3238,21 @@ void FLevelEditorViewportClient::ApplyDeltaToRotateWidget( const FRotator& InRot
 	//apply rotation to translate rotate widget
 	if (!InRot.IsZero())
 	{
-		FRotator TranslateRotateWidgetRotation(0, GLevelEditorModeTools().TranslateRotateXAxisAngle, 0);
+		FRotator TranslateRotateWidgetRotation(0, ModeTools->TranslateRotateXAxisAngle, 0);
 		TranslateRotateWidgetRotation += InRot;
-		GLevelEditorModeTools().TranslateRotateXAxisAngle = TranslateRotateWidgetRotation.Yaw;
+		ModeTools->TranslateRotateXAxisAngle = TranslateRotateWidgetRotation.Yaw;
 	}
 
 }
 
 void FLevelEditorViewportClient::MouseEnter( FViewport* Viewport,int32 x, int32 y )
 {
-	GLevelEditorModeTools().MouseEnter(this, Viewport, x, y);
+	ModeTools->MouseEnter(this, Viewport, x, y);
 }
 
 void FLevelEditorViewportClient::MouseLeave( FViewport* Viewport ) 
 {
-	GLevelEditorModeTools().MouseLeave(this, Viewport);
+	ModeTools->MouseLeave(this, Viewport);
 
 	FEditorViewportClient::MouseLeave(Viewport);
 }
@@ -3261,7 +3262,7 @@ void FLevelEditorViewportClient::MouseMove(FViewport* Viewport,int32 x, int32 y)
 	FEditorViewportClient::MouseMove(Viewport,x,y);
 
 	// Let the current editor mode know about the mouse movement.
-	if( IsLevelEditorClient() && GLevelEditorModeTools().MouseMove(this, Viewport, x, y) )
+	if (IsLevelEditorClient() && ModeTools->MouseMove(this, Viewport, x, y))
 	{
 		return;
 	}
@@ -3325,7 +3326,7 @@ void FLevelEditorViewportClient::CapturedMouseMove( FViewport* InViewport, int32
 	TrackingTransaction.PromotePendingToActive();
 
 	// Let the current editor mode know about the mouse movement.
-	if( IsLevelEditorClient() && GLevelEditorModeTools().CapturedMouseMove( this, InViewport, InMouseX, InMouseY ) )
+	if (IsLevelEditorClient() && ModeTools->CapturedMouseMove(this, InViewport, InMouseX, InMouseY))
 	{
 		return;
 	}
@@ -3471,7 +3472,7 @@ bool FLevelEditorViewportClient::GetActiveSafeFrame(float& OutAspectRatio) const
 void FLevelEditorViewportClient::SetCurrentWidgetAxis( EAxisList::Type NewAxis )
 {
 	FEditorViewportClient::SetCurrentWidgetAxis( NewAxis );
-	GLevelEditorModeTools().SetCurrentWidgetAxis( NewAxis );
+	ModeTools->SetCurrentWidgetAxis(NewAxis);
 }
 
 /** 
@@ -3553,10 +3554,10 @@ void FLevelEditorViewportClient::Draw(const FSceneView* View,FPrimitiveDrawInter
 	}
 
 
-	GLevelEditorModeTools().DrawActiveModes(View, PDI);
+	ModeTools->DrawActiveModes(View, PDI);
 
 	// Draw the current editor mode.
-	GLevelEditorModeTools().Render( View, Viewport, PDI );
+	ModeTools->Render(View, Viewport, PDI);
 
 	// Determine if a view frustum should be rendered in the viewport.
 	// The frustum should definitely be rendered if the viewport has a view parent.
@@ -3718,7 +3719,7 @@ void FLevelEditorViewportClient::SetupViewForRendering( FSceneViewFamily& ViewFa
 		}
 	}
 
-	if(GLevelEditorModeTools().GetActiveMode(FBuiltinEditorModes::EM_InterpEdit) == 0 || !AllowMatineePreview())
+	if (ModeTools->GetActiveMode(FBuiltinEditorModes::EM_InterpEdit) == 0 || !AllowMatineePreview())
 	{
 		// in the editor, disable camera motion blur and other rendering features that rely on the former frame
 		// unless the view port is Matinee controlled
@@ -3751,9 +3752,9 @@ void FLevelEditorViewportClient::DrawCanvas( FViewport& InViewport, FSceneView&
 	}
 
 	// Information string
-	Canvas.DrawShadowedString( 4,4, *GLevelEditorModeTools().InfoString, GEngine->GetSmallFont(), FColor(255,255,255) );
+	Canvas.DrawShadowedString(4, 4, *ModeTools->InfoString, GEngine->GetSmallFont(), FColor(255, 255, 255));
 
-	GLevelEditorModeTools().DrawHUD(this,&InViewport,&View,&Canvas);
+	ModeTools->DrawHUD(this, &InViewport, &View, &Canvas);
 
 	// Testbed
 	FCanvasItemTestbed TestBed;
@@ -3936,13 +3937,13 @@ void FLevelEditorViewportClient::SetCameraSpeedSetting(int32 SpeedSetting)
 void FLevelEditorViewportClient::ReceivedFocus(FViewport* Viewport)
 {
 	FEditorViewportClient::ReceivedFocus(Viewport);
-	GLevelEditorModeTools().ReceivedFocus(this, Viewport);
+	ModeTools->ReceivedFocus(this, Viewport);
 }
 
 void FLevelEditorViewportClient::LostFocus(FViewport* Viewport)
 {
 	FEditorViewportClient::LostFocus(Viewport);
-	GLevelEditorModeTools().LostFocus(this, Viewport);
+	ModeTools->LostFocus(this, Viewport);
 }
 
 bool FLevelEditorViewportClient::OverrideHighResScreenshotCaptureRegion(FIntRect& OutCaptureRegion)
diff --git a/Engine/Source/Editor/UnrealEd/Private/VertexSnapping.cpp b/Engine/Source/Editor/UnrealEd/Private/VertexSnapping.cpp
--- a/Engine/Source/Editor/UnrealEd/Private/VertexSnapping.cpp
+++ b/Engine/Source/Editor/UnrealEd/Private/VertexSnapping.cpp
@@ -718,7 +718,7 @@ bool FVertexSnappingImpl::SnapDraggedActorsToNearestVertex( FVector& DragDelta,
 
 		FSceneView* View = ViewportClient->CalcSceneView( &ViewFamily );
 
-		FVector StartLocation = GLevelEditorModeTools().PivotLocation;
+		FVector StartLocation = ViewportClient->GetModeTools()->PivotLocation;
 
 		FVector DesiredUnsnappedLocation = StartLocation+DragDelta;
 					
