Changed the Illumination warning on line 1256 to a Log instead.

diff --git a/Engine/Source/Editor/UnrealEd/Private/StaticLightingSystem/StaticLightingSystem.cpp b/Engine/Source/Editor/UnrealEd/Private/StaticLightingSystem/StaticLightingSystem.cpp
--- a/Engine/Source/Editor/UnrealEd/Private/StaticLightingSystem/StaticLightingSystem.cpp
+++ b/Engine/Source/Editor/UnrealEd/Private/StaticLightingSystem/StaticLightingSystem.cpp
@@ -1253,7 +1253,7 @@ void FStaticLightingSystem::ReportStatistics()
 	}
 	else	//if ( GLightmassStatsMode)
 	{
-		UE_LOG(LogStaticLightingSystem, Warning, TEXT("Illumination: %s (%s encoding lightmaps)"), *FPlatformTime::PrettyTime(LightmassStatistics.TotalTime), *FPlatformTime::PrettyTime(LightmassStatistics.EncodingTime) );
+		UE_LOG(LogStaticLightingSystem, Log, TEXT("Illumination: %s (%s encoding lightmaps)"), *FPlatformTime::PrettyTime(LightmassStatistics.TotalTime), *FPlatformTime::PrettyTime(LightmassStatistics.EncodingTime) );
 	}
 	UE_LOG(LogStaticLightingSystem, Log, TEXT("Lightmap texture memory:  %.1f MB (%.1f MB streaming, %.1f MB non-streaming), %d textures"),
 		GLightmapTotalSize/1024.0f/1024.0f,
