Changing variable types no longer requires all uses in graphs to be removed before it is possible. Nodes will reconstruct and disconnect, but otherwise not be removed (in all open BPs). Suppressible dialog confirms this action by the user.

diff --git a/Engine/Source/Editor/Kismet/Private/BlueprintDetailsCustomization.cpp b/Engine/Source/Editor/Kismet/Private/BlueprintDetailsCustomization.cpp
--- a/Engine/Source/Editor/Kismet/Private/BlueprintDetailsCustomization.cpp
+++ b/Engine/Source/Editor/Kismet/Private/BlueprintDetailsCustomization.cpp
@@ -795,53 +795,12 @@ void FBlueprintVarActionDetails::OnVarNameCommitted(const FText& InNewText, ETex
 
 bool FBlueprintVarActionDetails::GetVariableTypeChangeEnabled() const
 {
-	FEdGraphSchemaAction_K2Var* VarAction = MyBlueprintSelectionAsVar();
-	if (VarAction)
-	{
-		TArray<UK2Node_Variable*> VariableNodes;
-		FBlueprintEditorUtils::GetAllNodesOfClass<UK2Node_Variable>(GetBlueprintObj(), VariableNodes);
-
-		bool bNodesPendingDeletion = false;
-		for( TArray<UK2Node_Variable*>::TConstIterator NodeIt(VariableNodes); NodeIt; ++NodeIt )
-		{
-			UK2Node_Variable* CurrentNode = *NodeIt;
-			if (VarAction->GetVariableName() == CurrentNode->GetVarName())
-			{
-				bNodesPendingDeletion = true;
-				break;
-			}
-		}
-		
-		bool bIsAVarInThisBlueprint = FBlueprintEditorUtils::FindNewVariableIndex(GetBlueprintObj(), VarAction->GetVariableName()) != INDEX_NONE;
-		return !bNodesPendingDeletion && bIsAVarInThisBlueprint;
-	}
-
-	FEdGraphSchemaAction_K2LocalVar* LocalVarAction = MyBlueprintSelectionAsLocalVar();
-	if (LocalVarAction)
-	{
-		TArray<UK2Node_Variable*> VariableNodes;
-		FBlueprintEditorUtils::GetAllNodesOfClass<UK2Node_Variable>(GetBlueprintObj(), VariableNodes);
-
-		bool bNodesPendingDeletion = false;
-		for( TArray<UK2Node_Variable*>::TConstIterator NodeIt(VariableNodes); NodeIt; ++NodeIt )
-		{
-			UK2Node_Variable* CurrentNode = *NodeIt;
-			if (LocalVarAction->GetVariableName() == CurrentNode->GetVarName())
-			{
-				bNodesPendingDeletion = true;
-				break;
-			}
-		}
-
-		bool bIsAVarInThisBlueprint = FBlueprintEditorUtils::FindLocalVariable(GetBlueprintObj(), LocalVarAction->GetVariableScope(), LocalVarAction->GetVariableName()) != NULL;
-		return !bNodesPendingDeletion && bIsAVarInThisBlueprint;
-	}
-	FEdGraphSchemaAction_K2LocalVar* VarLocalAction = MyBlueprintSelectionAsLocalVar();
-	if(VarLocalAction)
+	UProperty* VariableProperty = SelectionAsProperty();
+	if(VariableProperty)
 	{
-		return true;
+		return GetBlueprintObj()->SkeletonGeneratedClass->GetAuthoritativeClass() == VariableProperty->GetOwnerClass()->GetAuthoritativeClass();
 	}
-	return false;
+	return true;
 }
 
 bool FBlueprintVarActionDetails::GetVariableCategoryChangeEnabled() const
@@ -888,16 +847,14 @@ void FBlueprintVarActionDetails::OnVarTypeChanged(const FEdGraphPinType& NewPinT
 			// Set the MyBP tab's last pin type used as this, for adding lots of variables of the same type
 			MyBlueprint.Pin()->GetLastPinTypeUsed() = NewPinType;
 
-			FEdGraphSchemaAction_K2Var* VarAction = MyBlueprintSelectionAsVar();
-			if(VarAction)
+			UProperty* VariableProperty = SelectionAsProperty();
+			if(IsALocalVariable(VariableProperty))
 			{
-				FBlueprintEditorUtils::ChangeMemberVariableType(GetBlueprintObj(), VarName, NewPinType);
+				FBlueprintEditorUtils::ChangeLocalVariableType(GetBlueprintObj(), GetLocalVariableScope(VariableProperty), VarName, NewPinType);
 			}
-
-			FEdGraphSchemaAction_K2LocalVar* VarLocalAction = MyBlueprintSelectionAsLocalVar();
-			if(VarLocalAction)
+			else
 			{
-				FBlueprintEditorUtils::ChangeLocalVariableType(GetBlueprintObj(), VarLocalAction->GetVariableScope(), VarName, NewPinType);
+				FBlueprintEditorUtils::ChangeMemberVariableType(GetBlueprintObj(), VarName, NewPinType);
 			}
 		}
 	}
diff --git a/Engine/Source/Editor/Kismet/Private/BlueprintEditor.cpp b/Engine/Source/Editor/Kismet/Private/BlueprintEditor.cpp
--- a/Engine/Source/Editor/Kismet/Private/BlueprintEditor.cpp
+++ b/Engine/Source/Editor/Kismet/Private/BlueprintEditor.cpp
@@ -1834,6 +1834,11 @@ void FBlueprintEditor::FocusInspectorOnGraphSelection(const FGraphPanelSelection
 
 		Inspector->ShowDetailsForObjects(NewSelection.Array(), ShowDetailsOptions);
 	}
+	else
+	{
+		// Clear the Inspector if nothing is selected, the MyBlueprints window will set the selection if it means to, but we need to make sure invalid nodes aren't still appearing
+		Inspector->ShowDetailsForSingleObject(nullptr);
+	}
 }
 
 void FBlueprintEditor::CreateDefaultTabContents(const TArray<UBlueprint*>& InBlueprints)
diff --git a/Engine/Source/Editor/UnrealEd/Private/Kismet2/BlueprintEditorUtils.cpp b/Engine/Source/Editor/UnrealEd/Private/Kismet2/BlueprintEditorUtils.cpp
--- a/Engine/Source/Editor/UnrealEd/Private/Kismet2/BlueprintEditorUtils.cpp
+++ b/Engine/Source/Editor/UnrealEd/Private/Kismet2/BlueprintEditorUtils.cpp
@@ -3537,6 +3537,80 @@ void FBlueprintEditorUtils::RenameMemberVariable(UBlueprint* Blueprint, const FN
 	}
 }
 
+TArray<UK2Node_Variable*> FBlueprintEditorUtils::GetNodesForVariable(const FName& InVarName, const UBlueprint* InBlueprint, const UStruct* InScope/* = nullptr*/)
+{
+	TArray<UK2Node_Variable*> ReturnNodes;
+	TArray<UK2Node_Variable*> VariableNodes;
+	GetAllNodesOfClass<UK2Node_Variable>(InBlueprint, VariableNodes);
+
+	bool bNodesPendingDeletion = false;
+	for( TArray<UK2Node_Variable*>::TConstIterator NodeIt(VariableNodes); NodeIt; ++NodeIt )
+	{
+		UK2Node_Variable* CurrentNode = *NodeIt;
+		if (InVarName == CurrentNode->GetVarName())
+		{
+			if(InScope && CurrentNode->VariableReference.GetMemberScopeName() != InScope->GetName())
+			{
+				// Variables are not in the same scope
+				continue;
+			}
+			ReturnNodes.Add(CurrentNode);
+		}
+	}
+	return ReturnNodes;
+}
+
+bool FBlueprintEditorUtils::VerifyUserWantsVariableTypeChanged(const FName& InVarName)
+{
+	FFormatNamedArguments Args;
+	Args.Add(TEXT("VariableName"), FText::FromName(InVarName));
+
+	FText ConfirmDelete = FText::Format(LOCTEXT( "ConfirmChangeVarType",
+		"This could break connections, do you want to search all Variable '{VariableName}' instances, change its type, and recompile?"), Args );
+
+	// Warn the user that this may result in data loss
+	FSuppressableWarningDialog::FSetupInfo Info( ConfirmDelete, LOCTEXT("ChangeVariableType", "Change Variable Type"), "ChangeVariableType_Warning" );
+	Info.ConfirmText = LOCTEXT( "ChangeVariableType_Yes", "Change Variable Type");
+	Info.CancelText = LOCTEXT( "ChangeVariableType_No", "Do Nothing");	
+
+	FSuppressableWarningDialog ChangeVariableType( Info );
+
+	FSuppressableWarningDialog::EResult RetCode = ChangeVariableType.ShowModal();
+	return RetCode == FSuppressableWarningDialog::Confirm || RetCode == FSuppressableWarningDialog::Suppressed;
+}
+
+void FBlueprintEditorUtils::GetLoadedChildBlueprints(UBlueprint* InBlueprint, TArray<UBlueprint*>& OutBlueprints)
+{
+	// Iterate over currently-loaded Blueprints and potentially adjust their variable names if they conflict with the parent
+	for(TObjectIterator<UBlueprint> BlueprintIt; BlueprintIt; ++BlueprintIt)
+	{
+		UBlueprint* ChildBP = *BlueprintIt;
+		if(ChildBP != NULL && ChildBP->ParentClass != NULL)
+		{
+			TArray<UBlueprint*> ParentBPArray;
+			// Get the parent hierarchy
+			UBlueprint::GetBlueprintHierarchyFromClass(ChildBP->ParentClass, ParentBPArray);
+
+			// Also get any BP interfaces we use
+			TArray<UClass*> ImplementedInterfaces;
+			FindImplementedInterfaces(ChildBP, true, ImplementedInterfaces);
+			for(auto InterfaceIt(ImplementedInterfaces.CreateConstIterator()); InterfaceIt; InterfaceIt++)
+			{
+				UBlueprint* BlueprintInterfaceClass = UBlueprint::GetBlueprintFromClass(*InterfaceIt);
+				if(BlueprintInterfaceClass != NULL)
+				{
+					ParentBPArray.Add(BlueprintInterfaceClass);
+				}
+			}
+
+			if(ParentBPArray.Contains(InBlueprint))
+			{
+				OutBlueprints.Add(ChildBP);
+			}
+		}
+	}
+}
+
 void FBlueprintEditorUtils::ChangeMemberVariableType(UBlueprint* Blueprint, const FName& VariableName, const FEdGraphPinType& NewPinType)
 {
 	if (VariableName != NAME_None)
@@ -3550,15 +3624,36 @@ void FBlueprintEditorUtils::ChangeMemberVariableType(UBlueprint* Blueprint, cons
 			// Update the variable type only if it is different
 			if (Variable.VarType != NewPinType)
 			{
+				TArray<UBlueprint*> ChildBPs;
+				GetLoadedChildBlueprints(Blueprint, ChildBPs);
+
+				TArray<UK2Node_Variable*> AllVariableNodes = GetNodesForVariable(VariableName, Blueprint);
+				for(UBlueprint* ChildBP : ChildBPs)
+				{
+					TArray<UK2Node_Variable*> VariableNodes = GetNodesForVariable(VariableName, ChildBP);
+					AllVariableNodes.Append(VariableNodes);
+				}
+
+				// TRUE if the user might be breaking variable connections
+				bool bBreakingVariableConnections = false;
+
+				// If there are variable nodes in place, warn the user of the consequences using a suppressible dialog
+				if(AllVariableNodes.Num())
+				{
+					if(!VerifyUserWantsVariableTypeChanged(VariableName))
+					{
+						// User has decided to cancel changing the variable member type
+						return;
+					}
+					bBreakingVariableConnections = true;
+				}
+
 				const FScopedTransaction Transaction( LOCTEXT("ChangeVariableType", "Change Variable Type") );
 				Blueprint->Modify();
 
 				/** Only change the variable type if type selection is valid, some unloaded Blueprints will turn out to be bad */
 				bool bChangeVariableType = true;
 
-				// Destroy all our nodes, because the pin types could be incorrect now (as will the links)
-				RemoveVariableNodes(Blueprint, VariableName);
-
 				if ((NewPinType.PinCategory == K2Schema->PC_Object) || (NewPinType.PinCategory == K2Schema->PC_Interface))
 				{
 					// if it's a PC_Object, then it should have an associated UClass object
@@ -3603,6 +3698,39 @@ void FBlueprintEditorUtils::ChangeMemberVariableType(UBlueprint* Blueprint, cons
 				if(bChangeVariableType)
 				{
 					Variable.VarType = NewPinType;
+
+					// Compile the Blueprint even if no loaded BPs are using the variable, this ensures that BPs refresh correctly when they are loaded
+					FKismetEditorUtilities::CompileBlueprint(Blueprint, false, true);
+
+					if(bBreakingVariableConnections)
+					{
+						for(UBlueprint* ChildBP : ChildBPs)
+						{
+							// Mark the Blueprint as structurally modified so we can reconstruct the node successfully
+							FBlueprintEditorUtils::MarkBlueprintAsStructurallyModified(ChildBP);
+						}
+
+						// Reconstruct all variable nodes that reference the changing variable
+						for(UK2Node_Variable* VariableNode : AllVariableNodes)
+						{
+							K2Schema->ReconstructNode(*VariableNode, true);
+						}
+
+						TSharedPtr<IToolkit> FoundAssetEditor = FToolkitManager::Get().FindEditorForAsset(Blueprint);
+						if (FoundAssetEditor.IsValid())
+						{
+							TSharedRef<IBlueprintEditor> BlueprintEditor = StaticCastSharedRef<IBlueprintEditor>(FoundAssetEditor.ToSharedRef());
+
+							const bool bSetFindWithinBlueprint = false;
+							const bool bSelectFirstResult = false;
+							BlueprintEditor->SummonSearchUI(bSetFindWithinBlueprint, VariableName.ToString(), bSelectFirstResult);
+						}
+					}
+					else
+					{
+						// And recompile
+						FBlueprintEditorUtils::MarkBlueprintAsStructurallyModified(Blueprint);
+					}
 				}
 
 				// And recompile
@@ -3954,7 +4082,7 @@ void FBlueprintEditorUtils::ChangeLocalVariableType(UBlueprint* InBlueprint, con
 	{
 		FString ActionCategory;
 		const UEdGraphSchema_K2* K2Schema = GetDefault<UEdGraphSchema_K2>();
-		
+
 		UK2Node_FunctionEntry* FunctionEntry = NULL;
 		FBPVariableDescription* VariablePtr = FindLocalVariable(InBlueprint, InScope, InVariableName, &FunctionEntry);
 
@@ -3965,17 +4093,45 @@ void FBlueprintEditorUtils::ChangeLocalVariableType(UBlueprint* InBlueprint, con
 			// Update the variable type only if it is different
 			if (Variable.VarName == InVariableName && Variable.VarType != NewPinType)
 			{
+				TArray<UK2Node_Variable*> VariableNodes = GetNodesForVariable(InVariableName, InBlueprint, InScope);
+
+				// If there are variable nodes in place, warn the user of the consequences using a suppressible dialog
+				if(VariableNodes.Num())
+				{
+					if(!VerifyUserWantsVariableTypeChanged(InVariableName))
+					{
+						// User has decided to cancel changing the variable member type
+						return;
+					}
+				}
+
 				const FScopedTransaction Transaction( LOCTEXT("ChangeLocalVariableType", "Change Local Variable Type") );
 				InBlueprint->Modify();
 				FunctionEntry->Modify();
 
 				Variable.VarType = NewPinType;
 
-				// Destroy all our nodes, because the pin types could be incorrect now (as will the links)
-				RemoveVariableNodes(InBlueprint, InVariableName);
+				// Mark the Blueprint as structurally modified so we can reconstruct the node successfully
+				FBlueprintEditorUtils::MarkBlueprintAsStructurallyModified(InBlueprint);
+
+				// Reconstruct all local variables referencing the modified one
+				for(UK2Node_Variable* VariableNode : VariableNodes)
+				{
+					K2Schema->ReconstructNode(*VariableNode, true);
+				}
 
 				// And recompile
 				FBlueprintEditorUtils::MarkBlueprintAsStructurallyModified(InBlueprint);
+
+				TSharedPtr<IToolkit> FoundAssetEditor = FToolkitManager::Get().FindEditorForAsset(InBlueprint);
+				if (FoundAssetEditor.IsValid())
+				{
+					TSharedRef<IBlueprintEditor> BlueprintEditor = StaticCastSharedRef<IBlueprintEditor>(FoundAssetEditor.ToSharedRef());
+
+					const bool bSetFindWithinBlueprint = true;
+					const bool bSelectFirstResult = false;
+					BlueprintEditor->SummonSearchUI(bSetFindWithinBlueprint, InVariableName.ToString(), bSelectFirstResult);
+				}
 			}
 		}
 	}
diff --git a/Engine/Source/Editor/UnrealEd/Public/Kismet2/BlueprintEditorUtils.h b/Engine/Source/Editor/UnrealEd/Public/Kismet2/BlueprintEditorUtils.h
--- a/Engine/Source/Editor/UnrealEd/Public/Kismet2/BlueprintEditorUtils.h
+++ b/Engine/Source/Editor/UnrealEd/Public/Kismet2/BlueprintEditorUtils.h
@@ -1086,6 +1086,32 @@ protected:
 	 */
 	static bool CheckIfNodeConnectsToSelection(UEdGraphNode* InNode, const TSet<UEdGraphNode*>& InSelectionSet);
 
+	/**
+	 * Returns an array of variables Get/Set nodes of the current variable
+	 *
+	 * @param InVarName		Variable to check for being in use
+	 * @param InBlueprint	Blueprint to check within
+	 * @param InScope		Option scope for local variables
+	 * @return				Array of variable nodes
+	 */
+	static TArray<UK2Node_Variable*> GetNodesForVariable(const FName& InVarName, const UBlueprint* InBlueprint, const UStruct* InScope = nullptr);
+
+	/**
+	 * Helper function to warn user of the results of changing var type by displaying a suppressible dialog
+	 *
+	 * @param InVarName		Variable name to display in the dialog message
+	 * @return				TRUE if the user wants to change the variable type
+	 */
+	static bool VerifyUserWantsVariableTypeChanged(const FName& InVarName);
+
+	/**
+	 * Helper function to get all loaded Blueprints that are children (or using as an interface) the passed Blueprint
+	 *
+	 * @param InBlueprint		Blueprint to find children of
+	 * @param OutBlueprints		Filled out with child Blueprints
+	 */
+	static void GetLoadedChildBlueprints(UBlueprint* InBlueprint, TArray<UBlueprint*>& OutBlueprints);
+
 public:
 	static FName GetFunctionNameFromClassByGuid(const UClass* InClass, const FGuid FunctionGuid);
 	static bool GetFunctionGuidFromClassByFieldName(const UClass* InClass, const FName FunctionName, FGuid& FunctionGuid);
