Editor: Deleting some obsolete/unused LOD code (FLODSettingsWindow, FLODUtilities::ImportMeshLOD, and FLODUtilities::CreateLODSettingsDialog)

diff --git a/Engine/Source/Editor/UnrealEd/Private/AnimSetViewer/LODSettingsWindow.cpp b/Engine/Source/Editor/UnrealEd/Private/AnimSetViewer/LODSettingsWindow.cpp
deleted file mode 100644
diff --git a/Engine/Source/Editor/UnrealEd/Private/AnimSetViewer/LODSettingsWindow.h b/Engine/Source/Editor/UnrealEd/Private/AnimSetViewer/LODSettingsWindow.h
deleted file mode 100644
diff --git a/Engine/Source/Editor/UnrealEd/Private/LODUtilities.cpp b/Engine/Source/Editor/UnrealEd/Private/LODUtilities.cpp
--- a/Engine/Source/Editor/UnrealEd/Private/LODUtilities.cpp
+++ b/Engine/Source/Editor/UnrealEd/Private/LODUtilities.cpp
@@ -4,13 +4,6 @@
 
 #include "LODUtilities.h"
 
-#include "AnimSetViewer/LODSettingsWindow.h"
-
-#if WITH_FBX
-	#include "FbxImporter.h"
-	#include "FbxExporter.h"
-#endif // WITH_FBX
-
 #include "MeshUtilities.h"
 
 #if WITH_APEX_CLOTHING
@@ -78,211 +71,6 @@ void FLODUtilities::RemoveLOD(FSkeletalMeshUpdateContext& UpdateContext, int32 D
 	}
 }
 
-void FLODUtilities::ImportMeshLOD( FSkeletalMeshUpdateContext& UpdateContext, int32 DesiredLOD)
-{
-	USkeletalMesh* SkeletalMesh = UpdateContext.SkeletalMesh;
-
-#if WITH_FBX
-	if(!SkeletalMesh)
-	{
-		return;
-	}
-
-	FString ExtensionStr;
-
-	ExtensionStr += TEXT("FBX files|*.fbx|");
-
-	ExtensionStr += TEXT("All files|*.*");
-
-	// First, display the file open dialog for selecting the file.
-	TArray<FString> OpenFilenames;
-	IDesktopPlatform* DesktopPlatform = FDesktopPlatformModule::Get();
-	bool bOpen = false;
-	if ( DesktopPlatform )
-	{
-		void* ParentWindowWindowHandle = NULL;
-
-		IMainFrameModule& MainFrameModule = FModuleManager::LoadModuleChecked<IMainFrameModule>(TEXT("MainFrame"));
-		const TSharedPtr<SWindow>& MainFrameParentWindow = MainFrameModule.GetParentWindow();
-		if ( MainFrameParentWindow.IsValid() && MainFrameParentWindow->GetNativeWindow().IsValid() )
-		{
-			ParentWindowWindowHandle = MainFrameParentWindow->GetNativeWindow()->GetOSWindowHandle();
-		}
-
-		bOpen = DesktopPlatform->SaveFileDialog(
-			ParentWindowWindowHandle,
-		NSLOCTEXT("ImportSkeletalMeshLOD", "ImportMeshLOD", "Import Mesh LOD...").ToString(), 
-		*(FEditorDirectories::Get().GetLastDirectory(ELastDirectory::FBX)),
-		TEXT(""),
-		*ExtensionStr,
-			EFileDialogFlags::None,
-			OpenFilenames
-			);
-	}
-
-	// Only continue if we pressed OK and have only one file selected.
-	if( bOpen )
-	{
-		if(OpenFilenames.Num() == 0)
-		{
-			FMessageDialog::Open( EAppMsgType::Ok, NSLOCTEXT("ImportSkeletalMeshLOD", "NoFileSelectedForLOD", "No file was selected for the LOD.") );
-		}
-		else if(OpenFilenames.Num() > 1)
-		{
-			FMessageDialog::Open( EAppMsgType::Ok, NSLOCTEXT("ImportSkeletalMeshLOD", "MultipleFilesSelectedForLOD", "You may only select one file for the LOD.") );
-		}
-		else
-		{
-
-			FString Filename = OpenFilenames[0];
-			FEditorDirectories::Get().SetLastDirectory(ELastDirectory::FBX, FPaths::GetPath(Filename)); // Save path as default for next time.
-
-			// Check the file extension for FBX. Anything that isn't .FBX is rejected
-			const FString FileExtension = FPaths::GetExtension(Filename);
-			const bool bIsFBX = FCString::Stricmp(*FileExtension, TEXT("FBX")) == 0;
-
-			if (bIsFBX)
-			{
-				UnFbx::FFbxImporter* FFbxImporter = UnFbx::FFbxImporter::GetInstance();
-				// don't import material and animation
-				UnFbx::FBXImportOptions* ImportOptions = FFbxImporter->GetImportOptions();
-				ImportOptions->bImportMaterials = false;
-				ImportOptions->bImportTextures = false;
-				ImportOptions->bImportAnimations = false;
-
-				if ( !FFbxImporter->ImportFromFile( *Filename ) )
-				{
-					// Log the error message and fail the import.
-					//UE_LOG(LogAnimSetViewer, Log, TEXT("FBX file parse failed"));
-					//Warn->Log(ELogVerbosity::Error, FFbxImporter->GetErrorMessage() );
-				}
-				else
-				{
-					bool bUseLODs = true;
-					int32 MaxLODLevel = 0;
-					TArray< TArray<FbxNode*>* > MeshArray;
-					TArray<FString> LODStrings;
-					TArray<FbxNode*>* MeshObject = NULL;;
-
-					// Populate the mesh array
-					FFbxImporter->FillFbxSkelMeshArrayInScene(FFbxImporter->Scene->GetRootNode(), MeshArray, false);
-
-					// Nothing found, error out
-					if (MeshArray.Num() == 0)
-					{
-						FMessageDialog::Open( EAppMsgType::Ok, NSLOCTEXT("ImportMeshLOD", "Prompt_NoMeshFound", "No meshes were found in file.") );
-						FFbxImporter->ReleaseScene();
-						return;
-					}
-
-					MeshObject = MeshArray[0];
-
-					// check if there is LODGroup for this skeletal mesh
-					for (int32 j = 0; j < MeshObject->Num(); j++)
-					{
-						FbxNode* Node = (*MeshObject)[j];
-						if (Node->GetNodeAttribute() && Node->GetNodeAttribute()->GetAttributeType() == FbxNodeAttribute::eLODGroup)
-						{
-							// get max LODgroup level
-							if (MaxLODLevel < (Node->GetChildCount() - 1))
-							{
-								MaxLODLevel = Node->GetChildCount() - 1;
-							}
-						}
-					}
-
-					// No LODs found, switch to supporting a mesh array containing meshes instead of LODs
-					if (MaxLODLevel == 0)
-					{
-						bUseLODs = false;
-						MaxLODLevel = SkeletalMesh->LODInfo.Num();
-					}
-
-					// Create LOD dropdown strings
-					LODStrings.AddZeroed(MaxLODLevel + 1);
-					LODStrings[0] = FString::Printf( TEXT("Base") );
-					for(int32 i = 1; i < MaxLODLevel + 1; i++)
-					{
-						LODStrings[i] = FString::Printf(TEXT("%d"), i);
-					}
-
-
-					int32 SelectedLOD = DesiredLOD;
-					if (SelectedLOD > SkeletalMesh->LODInfo.Num())
-					{
-						// Make sure they don't manage to select a bad LOD index
-						FMessageDialog::Open( EAppMsgType::Ok, FText::Format(NSLOCTEXT("ImportMeshLOD", "Prompt_InvalidLODIndex", "Invalid mesh LOD index {0}, as no prior LOD index exists!"), FText::AsNumber(SelectedLOD)) );
-					}
-					else
-					{
-						// Find the LOD node to import (if using LODs)
-						TArray<FbxNode*> SkelMeshNodeArray;
-						if (bUseLODs)
-						{
-							for (int32 j = 0; j < MeshObject->Num(); j++)
-							{
-								FbxNode* Node = (*MeshObject)[j];
-								if (Node->GetNodeAttribute() && Node->GetNodeAttribute()->GetAttributeType() == FbxNodeAttribute::eLODGroup)
-								{
-									if (Node->GetChildCount() > SelectedLOD)
-									{
-										SkelMeshNodeArray.Add(Node->GetChild(SelectedLOD));
-									}
-									else // in less some LODGroups have less level, use the last level
-									{
-										SkelMeshNodeArray.Add(Node->GetChild(Node->GetChildCount() - 1));
-									}
-								}
-								else
-								{
-									SkelMeshNodeArray.Add(Node);
-								}
-							}
-						}
-
-						// Import mesh
-						USkeletalMesh* TempSkelMesh = NULL;
-						UFbxSkeletalMeshImportData* TempImportData = NULL;
-						TempSkelMesh = (USkeletalMesh*)FFbxImporter->ImportSkeletalMesh(GetTransientPackage(), bUseLODs? SkelMeshNodeArray: *MeshObject, NAME_None, (EObjectFlags)0, TempImportData, FPaths::GetBaseFilename(Filename));
-
-						// Add imported mesh to existing model
-						bool bImportSucceeded = false;
-						if( TempSkelMesh )
-						{
-							bImportSucceeded = FFbxImporter->ImportSkeletalMeshLOD(TempSkelMesh, UpdateContext.SkeletalMesh, SelectedLOD, true, UpdateContext.AssociatedComponents);
-
-							//Notify calling system of change
-							UpdateContext.OnLODChanged.ExecuteIfBound();
-
-							// Mark package containing skeletal mesh as dirty.
-							SkeletalMesh->MarkPackageDirty();
-						}
-
-						if (bImportSucceeded)
-						{
-							// Pop up a box saying it worked.
-							FMessageDialog::Open( EAppMsgType::Ok, FText::Format(NSLOCTEXT("ImportMeshLOD", "LODImportSuccessful", "Mesh for LOD {0} imported successfully!"), FText::AsNumber(SelectedLOD)) );
-						}
-						else
-						{
-							// Import failed
-							FMessageDialog::Open( EAppMsgType::Ok, FText::Format(NSLOCTEXT("ImportMeshLOD", "LODImportFail", "Failed to import mesh for LOD {0}!"), FText::AsNumber(SelectedLOD)) );
-						}
-					}
-
-					// Cleanup
-					for (int32 i=0; i<MeshArray.Num(); i++)
-					{
-						delete MeshArray[i];
-					}					
-				}
-				FFbxImporter->ReleaseScene();
-			}
-		}
-	}
-#endif
-}
-
 void FLODUtilities::SimplifySkeletalMesh( FSkeletalMeshUpdateContext& UpdateContext, TArray<FSkeletalMeshOptimizationSettings> &InSettings )
 {
 	USkeletalMesh* SkeletalMesh = UpdateContext.SkeletalMesh;
@@ -351,13 +139,3 @@ void FLODUtilities::RefreshLODChange(const USkeletalMesh* SkeletalMesh)
 		}
 	}
 }
-
-void FLODUtilities::CreateLODSettingsDialog(FSkeletalMeshUpdateContext& UpdateContext)
-{
-	if ( UpdateContext.SkeletalMesh )
-	{
-		// Display the LOD Settings dialog window.
-		FDlgSkeletalMeshLODSettings CreateLODsDialog( UpdateContext );
-		CreateLODsDialog.ShowModal();
-	}
-}
diff --git a/Engine/Source/Editor/UnrealEd/Public/LODUtilities.h b/Engine/Source/Editor/UnrealEd/Public/LODUtilities.h
--- a/Engine/Source/Editor/UnrealEd/Public/LODUtilities.h
+++ b/Engine/Source/Editor/UnrealEd/Public/LODUtilities.h
@@ -29,14 +29,6 @@ public:
 	*/
 	static void RemoveLOD( FSkeletalMeshUpdateContext& UpdateContext, int32 DesiredLOD );
 
-	/**
-	* Import a .FBX file as an LOD of the SkeletalMesh.
-	*
-	* @param UpdateContext - The skeletal mesh and actor components to operate on.
-	* @param DesiredLOD   - the LOD index to import into. A new LOD entry is created if one doesn't exist
-	*/
-	static void ImportMeshLOD( FSkeletalMeshUpdateContext& UpdateContext, int32 DesiredLOD);
-
 	/**
 	 *	Simplifies the static mesh based upon various user settings.
 	 *
@@ -55,15 +47,6 @@ public:
 	 */
 	static void RefreshLODChange(const USkeletalMesh* SkeletalMesh);
 
-	/**
-	 * Create LOD Settings Dialog
-	 * 
-	 * Spawns a Slate dialog for editing LOD settings on a USkeletalMesh
-	 *
-	 * @param	UpdateContext	Context to be acted upon
-	 */
-	static void CreateLODSettingsDialog(FSkeletalMeshUpdateContext& UpdateContext);
-
 private:
 	FLODUtilities() {}
 };
\ No newline at end of file
