Fixed compile error with clang

diff --git a/Engine/Source/Editor/UnrealEd/Private/AnimationEditorUtils.cpp b/Engine/Source/Editor/UnrealEd/Private/AnimationEditorUtils.cpp
--- a/Engine/Source/Editor/UnrealEd/Private/AnimationEditorUtils.cpp
+++ b/Engine/Source/Editor/UnrealEd/Private/AnimationEditorUtils.cpp
@@ -3,7 +3,6 @@
 #include "UnrealEd.h"
 #include "AnimationEditorUtils.h"
 #include "AssetToolsModule.h"
-#include "ContentBrowserModule.h"
 #include "Animation/AnimComposite.h"
 #include "Animation/BlendSpace.h"
 #include "Animation/BlendSpace1D.h"
diff --git a/Engine/Source/Editor/UnrealEd/Public/AnimationEditorUtils.h b/Engine/Source/Editor/UnrealEd/Public/AnimationEditorUtils.h
--- a/Engine/Source/Editor/UnrealEd/Public/AnimationEditorUtils.h
+++ b/Engine/Source/Editor/UnrealEd/Public/AnimationEditorUtils.h
@@ -3,6 +3,8 @@
 #ifndef __AnimationEditorUtils_h__
 #define __AnimationEditorUtils_h__
 
+#include "Editor/ContentBrowser/Public/ContentBrowserModule.h"
+
 /** dialog to prompt users to decide an animation asset name */
 class SCreateAnimationAssetDlg : public SWindow
 {
