Editor: WIP of standalone asset editor toolkits being capable of hosting other toolkits

diff --git a/Engine/Source/Editor/UnrealEd/Private/EditorModes.cpp b/Engine/Source/Editor/UnrealEd/Private/EditorModes.cpp
--- a/Engine/Source/Editor/UnrealEd/Private/EditorModes.cpp
+++ b/Engine/Source/Editor/UnrealEd/Private/EditorModes.cpp
@@ -1229,6 +1229,13 @@ void FEditorModeTools::SaveConfig(void)
 	SaveWidgetSettings();
 }
 
+TSharedPtr<class IToolkitHost> FEditorModeTools::GetToolkitHost() const
+{
+	TSharedPtr<class IToolkitHost> Result = ToolkitHost.Pin();
+	check(ToolkitHost.IsValid());
+	return Result;
+}
+
 void FEditorModeTools::SetToolkitHost(TSharedRef<class IToolkitHost> InHost)
 {
 	checkf(!ToolkitHost.IsValid(), TEXT("SetToolkitHost can only be called once"));
diff --git a/Engine/Source/Editor/UnrealEd/Private/Toolkits/SStandaloneAssetEditorToolkitHost.cpp b/Engine/Source/Editor/UnrealEd/Private/Toolkits/SStandaloneAssetEditorToolkitHost.cpp
--- a/Engine/Source/Editor/UnrealEd/Private/Toolkits/SStandaloneAssetEditorToolkitHost.cpp
+++ b/Engine/Source/Editor/UnrealEd/Private/Toolkits/SStandaloneAssetEditorToolkitHost.cpp
@@ -87,7 +87,7 @@ void SStandaloneAssetEditorToolkitHost::SetupInitialContent( const TSharedRef<FT
 
 		TSharedPtr<FExtender> MenuExtender = MakeShareable(new FExtender());
 
-		auto AssetEditorToolkit = StaticCastSharedRef< FAssetEditorToolkit >( HostedToolkit.ToSharedRef() );
+		auto AssetEditorToolkit = HostedAssetEditorToolkit.ToSharedRef();
 
 		// Add asset-specific menu items to the top of the "File" menu
 		MenuExtender->AddMenuExtension( "FileLoadAndSave", EExtensionHook::First, AssetEditorToolkit->GetToolkitCommands(), FMenuExtensionDelegate::CreateStatic( &Local::FillFileMenu, TWeakPtr< FAssetEditorToolkit >( AssetEditorToolkit ) ) );
@@ -207,6 +207,8 @@ SStandaloneAssetEditorToolkitHost::~SStandaloneAssetEditorToolkitHost()
 {
 	// Let the toolkit manager know that we're going away now
 	FToolkitManager::Get().OnToolkitHostDestroyed( this );
+	HostedToolkits.Reset();
+	HostedAssetEditorToolkit.Reset();
 }
 
 
@@ -231,36 +233,51 @@ TSharedRef< SDockTabStack > SStandaloneAssetEditorToolkitHost::GetTabSpot( const
 void SStandaloneAssetEditorToolkitHost::OnToolkitHostingStarted( const TSharedRef< class IToolkit >& Toolkit )
 {
 	// Keep track of the toolkit we're hosting
-	check( !HostedToolkit.IsValid() );
-	HostedToolkit = Toolkit;
+	HostedToolkits.Add(Toolkit);
 
 	// The tab manager needs to know how to spawn tabs from this toolkit
 	Toolkit->RegisterTabSpawners(MyTabManager.ToSharedRef());
+
+	if (!HostedAssetEditorToolkit.IsValid())
+	{
+		HostedAssetEditorToolkit = StaticCastSharedRef<FAssetEditorToolkit>(Toolkit);
+	}
+	else
+	{
+		HostedAssetEditorToolkit->OnToolkitHostingStarted(Toolkit);
+	}
 }
 
 
 void SStandaloneAssetEditorToolkitHost::OnToolkitHostingFinished( const TSharedRef< class IToolkit >& Toolkit )
 {
-	// No need to worry about the toolkit anymore
-	check( Toolkit == HostedToolkit );
-	
 	// The tab manager should forget how to spawn tabs from this toolkit
 	Toolkit->UnregisterTabSpawners(MyTabManager.ToSharedRef());
 
+	HostedToolkits.Remove(Toolkit);
+
 	// Standalone Asset Editors close by shutting down their major tab.
-	const TSharedPtr<SDockTab> HostTab = HostTabPtr.Pin();
-	if (HostTab.IsValid())
+	if (Toolkit == HostedAssetEditorToolkit)
 	{
-		HostTab->RequestCloseTab();
-	}
+		HostedAssetEditorToolkit.Reset();
 
-	HostedToolkit.Reset();
+		const TSharedPtr<SDockTab> HostTab = HostTabPtr.Pin();
+		if (HostTab.IsValid())
+		{
+			HostTab->RequestCloseTab();
+		}
+	}
+	else if (HostedAssetEditorToolkit.IsValid())
+	{
+		HostedAssetEditorToolkit->OnToolkitHostingFinished(Toolkit);
+	}
 }
 
 
 UWorld* SStandaloneAssetEditorToolkitHost::GetWorld() const 
 {
 	// Currently, standalone asset editors never have a world
+	UE_LOG(LogInit, Warning, TEXT("IToolkitHost::GetWorld() doesn't make sense in SStandaloneAssetEditorToolkitHost currently"));
 	return NULL;
 }
 
@@ -269,9 +286,59 @@ FReply SStandaloneAssetEditorToolkitHost::OnKeyDown( const FGeometry& MyGeometry
 {
 	// Check to see if any of the actions for the level editor can be processed by the current event
 	// If we are in debug mode do not process commands
-	if( HostedToolkit.IsValid() && HostedToolkit->ProcessCommandBindings( InKeyEvent ) )
+	if (FSlateApplication::Get().IsNormalExecution())
 	{
-		return FReply::Handled();
+		// Figure out if any of our toolkit's tabs is the active tab.  This is important because we want
+		// the toolkit to have it's own keybinds (which may overlap the level editor's keybinds or any
+		// other toolkit).  When a toolkit tab is active, we give that toolkit a chance to process
+		// commands instead of the level editor.
+		TSharedPtr< IToolkit > ActiveToolkit;
+		{
+			const TSharedPtr<SDockableTab> CurrentActiveTab;// = FSlateApplication::xxxGetGlobalTabManager()->GetActiveTab();
+
+			for (auto HostedToolkitIt = HostedToolkits.CreateConstIterator(); HostedToolkitIt && !ActiveToolkit.IsValid(); ++HostedToolkitIt)
+			{
+				const auto& CurToolkit = *HostedToolkitIt;
+				if (CurToolkit.IsValid())
+				{
+					// Iterate over this toolkits spawned tabs
+					const auto& ToolkitTabsInSpots = CurToolkit->GetToolkitTabsInSpots();
+
+					for (auto CurSpotIt(ToolkitTabsInSpots.CreateConstIterator()); CurSpotIt && !ActiveToolkit.IsValid(); ++CurSpotIt)
+					{
+						const auto& TabsForSpot = CurSpotIt.Value();
+						for (auto CurTabIt(TabsForSpot.CreateConstIterator()); CurTabIt; ++CurTabIt)
+						{
+							const auto& PinnedTab = CurTabIt->Pin();
+							if (PinnedTab.IsValid())
+							{
+								if (PinnedTab == CurrentActiveTab)
+								{
+									ActiveToolkit = CurToolkit;
+								}
+							}
+						}
+					}
+				}
+			}
+		}
+
+		if (ActiveToolkit.IsValid())
+		{
+			// A toolkit tab is active, so direct all command processing to it
+			if (ActiveToolkit->ProcessCommandBindings(InKeyEvent))
+			{
+				return FReply::Handled();
+			}
+		}
+		else
+		{
+			// No toolkit tab is active, so let the underlying asset editor have a chance at the keystroke
+			if (HostedAssetEditorToolkit->ProcessCommandBindings(InKeyEvent))
+			{
+				return FReply::Handled();
+			}
+		}
 	}
 
 	return SCompoundWidget::OnKeyDown(MyGeometry, InKeyEvent);
@@ -282,10 +349,9 @@ void SStandaloneAssetEditorToolkitHost::OnTabClosed(TSharedRef<SDockTab> TabClos
 {
 	check(TabClosed == HostTabPtr.Pin());
 
-	const TSharedPtr<FAssetEditorToolkit> AssetEditorToolkit = StaticCastSharedPtr<FAssetEditorToolkit>(HostedToolkit);
-	if(AssetEditorToolkit.IsValid())
+	if (HostedAssetEditorToolkit.IsValid())
 	{
-		const TArray<UObject*>* const ObjectsBeingEdited = AssetEditorToolkit->GetObjectsCurrentlyBeingEdited();
+		const TArray<UObject*>* const ObjectsBeingEdited = HostedAssetEditorToolkit->GetObjectsCurrentlyBeingEdited();
 		if(ObjectsBeingEdited)
 		{
 			const bool IsDockedAssetEditor = TabClosed->HasSiblingTab(FName("DockedToolkit"), false/*TreatIndexNoneAsWildcard*/);
diff --git a/Engine/Source/Editor/UnrealEd/Private/Toolkits/SStandaloneAssetEditorToolkitHost.h b/Engine/Source/Editor/UnrealEd/Private/Toolkits/SStandaloneAssetEditorToolkitHost.h
--- a/Engine/Source/Editor/UnrealEd/Private/Toolkits/SStandaloneAssetEditorToolkitHost.h
+++ b/Engine/Source/Editor/UnrealEd/Private/Toolkits/SStandaloneAssetEditorToolkitHost.h
@@ -96,8 +96,11 @@ private:
 	/** Name ID for this app */
 	FName AppName;
 
-	/** The single toolkit we're hosting */
-	TSharedPtr< IToolkit > HostedToolkit;
+	/** List of all of the toolkits we're currently hosting */
+	TArray< TSharedPtr< class IToolkit > > HostedToolkits;
+
+	/** The 'owning' asset editor toolkit we're hosting */
+	TSharedPtr< class FAssetEditorToolkit > HostedAssetEditorToolkit;
 
 	/** Delegate to be called to determine if we are allowed to close this toolkit host */
 	FRequestAssetEditorClose EditorCloseRequest;
diff --git a/Engine/Source/Editor/UnrealEd/Public/EditorModes.h b/Engine/Source/Editor/UnrealEd/Public/EditorModes.h
--- a/Engine/Source/Editor/UnrealEd/Public/EditorModes.h
+++ b/Engine/Source/Editor/UnrealEd/Public/EditorModes.h
@@ -788,11 +788,8 @@ public:
 	void SetToolkitHost(TSharedRef<class IToolkitHost> Host);
 
 	/** Returns the host for toolkits created via modes from this mode manager */
-	TSharedPtr<class IToolkitHost> GetToolkitHost() const
-	{
-		check(ToolkitHost.IsValid());
-		return ToolkitHost;
-	}
+	TSharedPtr<class IToolkitHost> GetToolkitHost() const;
+
 protected:
 	/** 
 	 * Delegate handlers
@@ -807,7 +804,7 @@ protected:
 	TArray< TSharedPtr<FEdMode> > Modes;
 
 	/** The host of the toolkits created by these modes */
-	TSharedPtr<IToolkitHost> ToolkitHost;
+	TWeakPtr<IToolkitHost> ToolkitHost;
 
 	/** A list of previously active editor modes that we will potentially recycle */
 	TMap< FEditorModeID, TSharedPtr<FEdMode> > RecycledModes;
diff --git a/Engine/Source/Editor/UnrealEd/Public/Toolkits/AssetEditorToolkit.h b/Engine/Source/Editor/UnrealEd/Public/Toolkits/AssetEditorToolkit.h
--- a/Engine/Source/Editor/UnrealEd/Public/Toolkits/AssetEditorToolkit.h
+++ b/Engine/Source/Editor/UnrealEd/Public/Toolkits/AssetEditorToolkit.h
@@ -117,6 +117,12 @@ public:
 	/** Called at the end of RegenerateMenusAndToolbars() */
 	virtual void PostRegenerateMenusAndToolbars() { }
 	
+	// Called when another toolkit (such as a ed mode toolkit) is being hosted in this asset editor toolkit
+	virtual void OnToolkitHostingStarted(const TSharedRef< class IToolkit >& Toolkit) {}
+
+	// Called when another toolkit (such as a ed mode toolkit) is no longer being hosted in this asset editor toolkit
+	virtual void OnToolkitHostingFinished(const TSharedRef< class IToolkit >& Toolkit) {}
+
 	/** Adds or removes extenders to the default menu or the toolbar menu this asset editor */
 	void AddMenuExtender(TSharedPtr<FExtender> Extender);
 	void RemoveMenuExtender(TSharedPtr<FExtender> Extender);
@@ -206,7 +212,7 @@ protected:
 
 	/** @return a pointer to the brush to use for the tab icon */
 	virtual const FSlateBrush* GetDefaultTabIcon() const;
-	
+
 private:
 	/** Spawns the toolbar tab */
 	TSharedRef<SDockTab> SpawnTab_Toolbar(const FSpawnTabArgs& Args);
