Editor: Add DLL export to FWidget::SetUsesEditorModeTools and remove an unused variable

diff --git a/Engine/Source/Editor/UnrealEd/Public/UnrealWidget.h b/Engine/Source/Editor/UnrealEd/Public/UnrealWidget.h
--- a/Engine/Source/Editor/UnrealEd/Public/UnrealWidget.h
+++ b/Engine/Source/Editor/UnrealEd/Public/UnrealWidget.h
@@ -26,7 +26,7 @@ public:
 	/**
 	 * Sets editor mode tools to use in this widget
 	 */
-	void SetUsesEditorModeTools( class FEditorModeTools* InEditorModeTools );
+	UNREALED_API void SetUsesEditorModeTools(class FEditorModeTools* InEditorModeTools);
 
 	/**
 	 * Renders any widget specific HUD text
@@ -354,9 +354,6 @@ private:
 	/** Any mode tools being used */
 	class FEditorModeTools* EditorModeTools;
 
-	/** Viewport client the widget is being drawn in */
-	class FEditorViewportClient* EditorViewportClient;
-
 	/**
 	 * An extra matrix to apply to the widget before drawing it (allows for local/custom coordinate systems).
 	 */
