Improved performance in looking up module names and disallowed header names from source code by initializing a database asynchronously upon editor startup.

diff --git a/Engine/Source/Editor/GameProjectGeneration/Private/SNewClassDialog.cpp b/Engine/Source/Editor/GameProjectGeneration/Private/SNewClassDialog.cpp
--- a/Engine/Source/Editor/GameProjectGeneration/Private/SNewClassDialog.cpp
+++ b/Engine/Source/Editor/GameProjectGeneration/Private/SNewClassDialog.cpp
@@ -104,22 +104,6 @@ void SNewClassDialog::Construct( const FArguments& InArgs )
 		{
 			AvailableModules.Emplace(MakeShareable(new FModuleContextInfo(ModuleInfo)));
 		}
-
-		// Get collection of disallowed header names
-		TArray<FString> HeaderFiles;
-		FindPublicEngineHeaderFiles(HeaderFiles, FPaths::EngineDir() / TEXT("Source") / TEXT("Developer"));
-		FindPublicEngineHeaderFiles(HeaderFiles, FPaths::EngineDir() / TEXT("Source") / TEXT("Editor"));
-		FindPublicEngineHeaderFiles(HeaderFiles, FPaths::EngineDir() / TEXT("Source") / TEXT("Runtime"));
-
-		for (const FString& HeaderFile : HeaderFiles)
-		{
-			DisallowedHeaderNames.Add(FPaths::GetBaseFilename(HeaderFile));
-		}
-
-		for (TObjectIterator<UClass> ClassIt; ClassIt; ++ClassIt)
-		{
-			DisallowedHeaderNames.Remove(ClassIt->GetName());
-		}
 	}
 
 	// If we have a runtime module with the same name as our project, then use that
@@ -1008,6 +992,7 @@ void SNewClassDialog::FinishClicked()
 	FString CppFilePath;
 
 	FText FailReason;
+	const TSet<FString>& DisallowedHeaderNames = FSourceCodeNavigation::GetSourceFileDatabase().GetDisallowedHeaderNames();
 	if (GameProjectUtils::AddCodeToProject(NewClassName, NewClassPath, *SelectedModuleInfo, ParentClassInfo, DisallowedHeaderNames, HeaderFilePath, CppFilePath, FailReason))
 	{
 		// Prevent periodic validity checks. This is to prevent a brief error message about the class already existing while you are exiting.
@@ -1238,6 +1223,7 @@ void SNewClassDialog::UpdateInputValidity()
 	// Validate the class name only if the path is valid
 	if ( bLastInputValidityCheckSuccessful )
 	{
+		const TSet<FString>& DisallowedHeaderNames = FSourceCodeNavigation::GetSourceFileDatabase().GetDisallowedHeaderNames();
 		bLastInputValidityCheckSuccessful = GameProjectUtils::IsValidClassNameForCreation(NewClassName, *SelectedModuleInfo, DisallowedHeaderNames, LastInputValidityErrorText);
 	}
 
diff --git a/Engine/Source/Editor/GameProjectGeneration/Private/SNewClassDialog.h b/Engine/Source/Editor/GameProjectGeneration/Private/SNewClassDialog.h
--- a/Engine/Source/Editor/GameProjectGeneration/Private/SNewClassDialog.h
+++ b/Engine/Source/Editor/GameProjectGeneration/Private/SNewClassDialog.h
@@ -222,7 +222,4 @@ private:
 
 	/** Information about the currently selected module; used for class validation */
 	TSharedPtr<FModuleContextInfo> SelectedModuleInfo;
-
-	/** Collection of disallowed header filenames */
-	TSet<FString> DisallowedHeaderNames;
 };
diff --git a/Engine/Source/Editor/UnrealEd/Private/SourceCodeNavigation.cpp b/Engine/Source/Editor/UnrealEd/Private/SourceCodeNavigation.cpp
--- a/Engine/Source/Editor/UnrealEd/Private/SourceCodeNavigation.cpp
+++ b/Engine/Source/Editor/UnrealEd/Private/SourceCodeNavigation.cpp
@@ -166,6 +166,73 @@ private:
 
 
 
+FSourceFileDatabase::FSourceFileDatabase()
+{
+	// Find all the build rules within the game and engine directories
+	FindRootFilesRecursive(ModuleNames, *(FPaths::EngineDir() / TEXT("Source") / TEXT("Developer")), TEXT("*.Build.cs"));
+	FindRootFilesRecursive(ModuleNames, *(FPaths::EngineDir() / TEXT("Source") / TEXT("Editor")), TEXT("*.Build.cs"));
+	FindRootFilesRecursive(ModuleNames, *(FPaths::EngineDir() / TEXT("Source") / TEXT("Runtime")), TEXT("*.Build.cs"));
+	FindRootFilesRecursive(ModuleNames, *(FPaths::GameDir() / TEXT("Source")), TEXT("*.Build.cs"));
+
+	// Find list of disallowed header names in native (non-plugin) directories
+	TArray<FString> HeaderFiles;
+	for (const FString& ModuleName : ModuleNames)
+	{
+		IFileManager::Get().FindFilesRecursive(HeaderFiles, *(FPaths::GetPath(ModuleName) / TEXT("Classes")), TEXT("*.h"), true, false, false);
+		IFileManager::Get().FindFilesRecursive(HeaderFiles, *(FPaths::GetPath(ModuleName) / TEXT("Public")), TEXT("*.h"), true, false, false);
+	}
+
+	for (const FString& HeaderFile : HeaderFiles)
+	{
+		DisallowedHeaderNames.Add(FPaths::GetBaseFilename(HeaderFile));
+	}
+
+	for (TObjectIterator<UClass> ClassIt; ClassIt; ++ClassIt)
+	{
+		DisallowedHeaderNames.Remove(ClassIt->GetName());
+	}
+
+	// Find all the plugin directories
+	TArray<FString> PluginNames;
+
+	FindRootFilesRecursive(PluginNames, *(FPaths::EngineDir() / TEXT("Plugins")), TEXT("*.uplugin"));
+	FindRootFilesRecursive(PluginNames, *(FPaths::GameDir() / TEXT("Plugins")), TEXT("*.uplugin"));
+
+	// Add all the files within plugin directories
+	for (const FString& PluginName : PluginNames)
+	{
+		FindRootFilesRecursive(ModuleNames, *(FPaths::GetPath(PluginName) / TEXT("Source")), TEXT("*.Build.cs"));
+	}
+}
+
+
+void FSourceFileDatabase::FindRootFilesRecursive(TArray<FString> &FileNames, const FString &BaseDirectory, const FString &Wildcard)
+{
+	// Find all the files within this directory
+	TArray<FString> BasedFileNames;
+	IFileManager::Get().FindFiles(BasedFileNames, *(BaseDirectory / Wildcard), true, false);
+
+	// Append to the result if we have any, otherwise recurse deeper
+	if (BasedFileNames.Num() == 0)
+	{
+		TArray<FString> DirectoryNames;
+		IFileManager::Get().FindFiles(DirectoryNames, *(BaseDirectory / TEXT("*")), false, true);
+
+		for (int32 Idx = 0; Idx < DirectoryNames.Num(); Idx++)
+		{
+			FindRootFilesRecursive(FileNames, BaseDirectory / DirectoryNames[Idx], Wildcard);
+		}
+	}
+	else
+	{
+		for (int32 Idx = 0; Idx < BasedFileNames.Num(); Idx++)
+		{
+			FileNames.Add(BaseDirectory / BasedFileNames[Idx]);
+		}
+	}
+}
+
+
 DECLARE_DELEGATE_RetVal( bool, FShouldAbortDelegate );
 
 
@@ -535,6 +602,45 @@ void FSourceCodeNavigationImpl::NavigateToFunctionSource( const FString& Functio
 }
 
 
+void FSourceCodeNavigation::Initialize()
+{
+	class FAsyncInitializeSourceFileDatabase : public FNonAbandonableTask
+	{
+	public:
+		/** Performs work on thread */
+		void DoWork()
+		{
+			FSourceCodeNavigation::GetSourceFileDatabase();
+		}
+
+		/** Returns true if the task should be aborted.  Called from within the task processing code itself via delegate */
+		bool ShouldAbort() const
+		{
+			return false;
+		}
+
+		/** @return Queries the name of this task for for external event viewers */
+		static const TCHAR* Name()
+		{
+			return TEXT("FAsyncInitializeSourceFileDatabase");
+		}
+	};
+
+	// Initialize SourceFileDatabase instance asynchronously
+	(new FAutoDeleteAsyncTask<FAsyncInitializeSourceFileDatabase>)->StartBackgroundTask();
+}
+
+
+const FSourceFileDatabase& FSourceCodeNavigation::GetSourceFileDatabase()
+{
+	// Lock so that nothing may proceed while the AsyncTask is constructing the FSourceFileDatabase for the first time
+	static FCriticalSection CriticalSection;
+	FScopeLock Lock(&CriticalSection);
+
+	static FSourceFileDatabase Instance;
+	return Instance;
+}
+
 
 void FSourceCodeNavigation::NavigateToFunctionSourceAsync( const FString& FunctionSymbolName, const FString& FunctionModuleName, const bool bIgnoreLineNumber )
 {
@@ -1383,36 +1489,14 @@ FSourceCodeNavigation::FOnCompilerNotFound& FSourceCodeNavigation::AccessOnCompi
 
 bool FSourceCodeNavigation::FindModulePath( const FString& ModuleName, FString &OutModulePath )
 {
-	// Find all the plugin directories
-	static TArray<FString> PluginFileNames;
-	static TArray<FString> FileNames;
-	static bool bHasScannedFiles = false;
-	if( !bHasScannedFiles )
-	{
-		FindRootFilesRecursive(PluginFileNames, *(FPaths::EngineDir() / TEXT("Plugins")), TEXT("*.uplugin"));
-		FindRootFilesRecursive(PluginFileNames, *(FPaths::GameDir() / TEXT("Plugins")), TEXT("*.uplugin"));
-
-		// Find all the build rules within the game and engine directories
-		FindRootFilesRecursive(FileNames, *(FPaths::EngineDir() / TEXT("Source")), TEXT("*.Build.cs"));
-		FindRootFilesRecursive(FileNames, *(FPaths::GameDir() / TEXT("Source")), TEXT("*.Build.cs"));
-
-		// Add all the files within plugin directories
-		for(int32 Idx = 0; Idx < PluginFileNames.Num(); Idx++)
-		{
-			FindRootFilesRecursive(FileNames, *(FPaths::GetPath(PluginFileNames[Idx]) / TEXT("Source")), TEXT("*.Build.cs"));
-		}
-
-		// Only do this once per session
-		bHasScannedFiles = true;
-	}
-
 	// Try to find a file matching the module name
+	const TArray<FString>& ModuleNames = GetSourceFileDatabase().GetModuleNames();
 	FString FindModuleSuffix = FString(TEXT("/")) + ModuleName + ".Build.cs";
-	for(int32 Idx = 0; Idx < FileNames.Num(); Idx++)
+	for (int32 Idx = 0; Idx < ModuleNames.Num(); Idx++)
 	{
-		if(FileNames[Idx].EndsWith(FindModuleSuffix))
+		if (ModuleNames[Idx].EndsWith(FindModuleSuffix))
 		{
-			OutModulePath = FileNames[Idx].Left(FileNames[Idx].Len() - FindModuleSuffix.Len());
+			OutModulePath = ModuleNames[Idx].Left(ModuleNames[Idx].Len() - FindModuleSuffix.Len());
 			return true;
 		}
 	}
@@ -1444,32 +1528,6 @@ bool FSourceCodeNavigation::FindClassHeaderPath( const UField *Field, FString &O
 	return false;
 }
 
-void FSourceCodeNavigation::FindRootFilesRecursive(TArray<FString> &FileNames, const FString &BaseDirectory, const FString &Wildcard)
-{
-	// Find all the files within this directory
-	TArray<FString> BasedFileNames;
-	IFileManager::Get().FindFiles(BasedFileNames, *(BaseDirectory / Wildcard), true, false);
-
-	// Append to the result if we have any, otherwise recurse deeper
-	if(BasedFileNames.Num() == 0)
-	{
-		TArray<FString> DirectoryNames;
-		IFileManager::Get().FindFiles(DirectoryNames, *(BaseDirectory / TEXT("*")), false, true);
-
-		for(int32 Idx = 0; Idx < DirectoryNames.Num(); Idx++)
-		{
-			FindRootFilesRecursive(FileNames, BaseDirectory / DirectoryNames[Idx], Wildcard);
-		}
-	}
-	else
-	{
-		for(int32 Idx = 0; Idx < BasedFileNames.Num(); Idx++)
-		{
-			FileNames.Add(BaseDirectory / BasedFileNames[Idx]);
-		}
-	}
-}
-
 void FSourceCodeNavigationImpl::TryToGatherFunctions( const FString& ModuleName, const FString& ClassName, TArray< FString >& OutFunctionSymbolNames, bool& OutIsCompleteList )
 {
 	FScopeLock ScopeLock( &SynchronizationObject );
diff --git a/Engine/Source/Editor/UnrealEd/Private/UnrealEdEngine.cpp b/Engine/Source/Editor/UnrealEd/Private/UnrealEdEngine.cpp
--- a/Engine/Source/Editor/UnrealEd/Private/UnrealEdEngine.cpp
+++ b/Engine/Source/Editor/UnrealEd/Private/UnrealEdEngine.cpp
@@ -28,6 +28,7 @@
 #include "BSPOps.h"
 #include "ComponentVisualizer.h"
 #include "Editor/EditorLiveStreaming/Public/IEditorLiveStreaming.h"
+#include "SourceCodeNavigation.h"
 
 DEFINE_LOG_CATEGORY_STATIC(LogUnrealEdEngine, Log, All);
 
@@ -36,6 +37,9 @@ void UUnrealEdEngine::Init(IEngineLoop* InEngineLoop)
 {
 	Super::Init(InEngineLoop);
 
+	// Build databases used by source code navigation
+	FSourceCodeNavigation::Initialize();
+
 	PackageAutoSaver.Reset(new FPackageAutoSaver);
 	PackageAutoSaver->LoadRestoreFile();
 
diff --git a/Engine/Source/Editor/UnrealEd/Public/SourceCodeNavigation.h b/Engine/Source/Editor/UnrealEd/Public/SourceCodeNavigation.h
--- a/Engine/Source/Editor/UnrealEd/Public/SourceCodeNavigation.h
+++ b/Engine/Source/Editor/UnrealEd/Public/SourceCodeNavigation.h
@@ -9,6 +9,31 @@
 DECLARE_LOG_CATEGORY_EXTERN(LogSelectionDetails, Log, All);
 
 
+/**
+ * Singleton holding database of module names and disallowed header names in the engine and current project.
+ */
+class FSourceFileDatabase
+{
+public:
+	/** Constructs database */
+	FSourceFileDatabase();
+
+	/** Return array of module names used by the engine and current project, including those in plugins */
+	const TArray<FString>& GetModuleNames() const { return ModuleNames; }
+
+	/** Return set of public header names used by engine modules, which are disallowed as project header names */
+	const TSet<FString>& GetDisallowedHeaderNames() const { return DisallowedHeaderNames; }
+
+private:
+
+	/** Return array of filenames matching the given wildcard, recursing into subdirectories if no results are yielded from the base directory */
+	void FindRootFilesRecursive(TArray<FString> &FileNames, const FString &BaseDirectory, const FString &Wildcard);
+
+	TArray<FString> ModuleNames;
+	TSet<FString> DisallowedHeaderNames;
+};
+
+
 /**
  * Source code navigation functionality
  */
@@ -50,6 +75,17 @@ public:
 		TWeakObjectPtr<UObject> ReferencedObject;
 	};
 
+
+	/**
+	 * Initializes FSourceCodeNavigation static class
+	 */
+	UNREALED_API static void Initialize();
+
+	/**
+	 * Retrieves the SourceFileDatabase instance
+	 */
+	UNREALED_API static const FSourceFileDatabase& GetSourceFileDatabase();
+
 	/**
 	 * Asynchronously locates the source file and line for a specific function in a specific module and navigates an external editing to that source line
 	 *
@@ -134,10 +170,6 @@ public:
 
 	/** Call this to access the multi-cast delegate that you can register a callback with */
 	UNREALED_API static FOnCompilerNotFound& AccessOnCompilerNotFound();
-
-private:
-	/** Recursively finds files which mark the root of a directory tree. Stops recursing once a file is found. */
-	static void FindRootFilesRecursive(TArray<FString> &FileNames, const FString &BaseDirectory, const FString &Wildcard);
 };
 
 
