Fixed various issues with snapping brushes to floor and brush/volume placement.

diff --git a/Engine/Source/Editor/UnrealEd/Private/AssetSelection.cpp b/Engine/Source/Editor/UnrealEd/Private/AssetSelection.cpp
--- a/Engine/Source/Editor/UnrealEd/Private/AssetSelection.cpp
+++ b/Engine/Source/Editor/UnrealEd/Private/AssetSelection.cpp
@@ -366,7 +366,13 @@ static AActor* PrivateAddActor( UObject* Asset, UActorFactory* Factory, bool Sel
 	AActor* NewActorTemplate = Factory->GetDefaultActor( Asset );
 	if ( !NewActorTemplate )
 	{
-		return NULL;
+		return nullptr;
+	}
+
+	// For Brushes/Volumes, use the default brush as the template rather than the factory default actor
+	if (NewActorTemplate->IsA(ABrush::StaticClass()) && GWorld->GetDefaultBrush() != nullptr)
+	{
+		NewActorTemplate = GWorld->GetDefaultBrush();
 	}
 
 	const FSnappedPositioningData PositioningData = FSnappedPositioningData(GCurrentLevelEditingViewportClient, GEditor->ClickLocation, GEditor->ClickPlane)
@@ -388,7 +394,7 @@ static AActor* PrivateAddActor( UObject* Asset, UActorFactory* Factory, bool Sel
 		FNotificationInfo Info( NSLOCTEXT("UnrealEd", "Error_OperationDisallowedOnLockedLevel", "The requested operation could not be completed because the level is locked.") );
 		Info.ExpireDuration = 3.0f;
 		FSlateNotificationManager::Get().AddNotification(Info);
-		return NULL;
+		return nullptr;
 	}
 
 	{
diff --git a/Engine/Source/Editor/UnrealEd/Private/EditorServer.cpp b/Engine/Source/Editor/UnrealEd/Private/EditorServer.cpp
--- a/Engine/Source/Editor/UnrealEd/Private/EditorServer.cpp
+++ b/Engine/Source/Editor/UnrealEd/Private/EditorServer.cpp
@@ -4462,10 +4462,17 @@ bool UEditorEngine::SnapActorTo( AActor* InActor, const bool InAlign, const bool
 		Direction.Normalize();
 	}
 
+	// In the case that we're about to do a line trace from a brush, move the start position so it's guaranteed to be very slightly
+	// outside of the brush bounds. The BSP geometry is double-sided which will give rise to an unwanted hit.
+	if (Brush)
+	{
+		const float fTinyOffset = 0.01f;
+		StartLocation.Z = InActor->GetRootComponent()->Bounds.Origin.Z - InActor->GetRootComponent()->Bounds.BoxExtent.Z - fTinyOffset;
+	}
 
 	// Do the actual actor->world check.  We try to collide against the world, straight down from our current position.
 	// If we hit anything, we will move the actor to a position that lets it rest on the floor.
-	FHitResult Hit(1.0f) ;
+	FHitResult Hit(1.0f);
 	FCollisionQueryParams Params(FName(TEXT("MoveActorToTrace")), false, InActor);
 	if ( InActor->GetWorld()->SweepSingle(Hit, StartLocation, StartLocation + Direction*WORLD_MAX, FQuat::Identity, FCollisionShape::MakeBox(Extent), Params, FCollisionObjectQueryParams(ECC_WorldStatic)))
 	{
diff --git a/Engine/Source/Runtime/Engine/Classes/Components/BrushComponent.h b/Engine/Source/Runtime/Engine/Classes/Components/BrushComponent.h
--- a/Engine/Source/Runtime/Engine/Classes/Components/BrushComponent.h
+++ b/Engine/Source/Runtime/Engine/Classes/Components/BrushComponent.h
@@ -35,6 +35,8 @@ class UBrushComponent : public UPrimitiveComponent
 	// Begin USceneComponent interface
 	virtual FBoxSphereBounds CalcBounds(const FTransform& LocalToWorld) const override;
 	virtual FVector GetCustomLocation() const override;
+	virtual bool ShouldCollideWhenPlacing() const override { return true; }
+
 protected:
 	
 	/** Calculate the new ComponentToWorld transform for this component.
diff --git a/Engine/Source/Runtime/Engine/Private/Actor.cpp b/Engine/Source/Runtime/Engine/Private/Actor.cpp
--- a/Engine/Source/Runtime/Engine/Private/Actor.cpp
+++ b/Engine/Source/Runtime/Engine/Private/Actor.cpp
@@ -1576,9 +1576,10 @@ FVector AActor::GetPlacementExtent() const
 			}
 		}
 
-		FVector BoxExtent = ActorBox.GetExtent();
-		float CollisionRadius = FMath::Sqrt( (BoxExtent.X * BoxExtent.X) + (BoxExtent.Y * BoxExtent.Y) );
-		Extent = FVector(CollisionRadius, CollisionRadius, BoxExtent.Z);
+		// Get box extent, adjusting for any difference between the center of the box and the actor pivot
+		FVector AdjustedBoxExtent = ActorBox.GetExtent() - ActorBox.GetCenter();
+		float CollisionRadius = FMath::Sqrt((AdjustedBoxExtent.X * AdjustedBoxExtent.X) + (AdjustedBoxExtent.Y * AdjustedBoxExtent.Y));
+		Extent = FVector(CollisionRadius, CollisionRadius, AdjustedBoxExtent.Z);
 	}
 	return Extent;
 }
diff --git a/Engine/Source/Runtime/Engine/Private/Brush.cpp b/Engine/Source/Runtime/Engine/Private/Brush.cpp
--- a/Engine/Source/Runtime/Engine/Private/Brush.cpp
+++ b/Engine/Source/Runtime/Engine/Private/Brush.cpp
@@ -33,6 +33,7 @@ ABrush::ABrush(const FObjectInitializer& ObjectInitializer)
 	bHidden = true;
 	bNotForClientOrServer = false;
 	bCanBeDamaged = false;
+	bCollideWhenPlacing = true;
 }
 
 #if WITH_EDITOR
