Allow numeric translation and rotation for geometry brushes
#github Github pull request #516

diff --git a/Engine/Source/Editor/DetailCustomizations/Private/ActorDetails.cpp b/Engine/Source/Editor/DetailCustomizations/Private/ActorDetails.cpp
--- a/Engine/Source/Editor/DetailCustomizations/Private/ActorDetails.cpp
+++ b/Engine/Source/Editor/DetailCustomizations/Private/ActorDetails.cpp
@@ -420,14 +420,19 @@ void FActorDetails::AddTransformCategory( IDetailLayoutBuilder& DetailBuilder )
 	bool bIsOnlyWorldPropsSelected =  SelectedActors.Num() == 1 && SelectedActors[0].IsValid() && SelectedActors[0]->IsA<AWorldSettings>();
 	bool bLacksRootComponent = SelectedActors[0].IsValid() && (SelectedActors[0]->GetRootComponent()==NULL);
 
-	// Don't show the Transform details when brushes are selected, if the only actor selected is world properties, or if they have no RootComponent
-	if ( bAreBrushesSelected || bIsOnlyWorldPropsSelected || bLacksRootComponent )
+	// Don't show the Transform details if the only actor selected is world properties, or if they have no RootComponent
+	if ( bIsOnlyWorldPropsSelected || bLacksRootComponent )
 	{
 		return;
 	}
 	
 	TSharedRef<FComponentTransformDetails> TransformDetails = MakeShareable( new FComponentTransformDetails( DetailBuilder.GetDetailsView().GetSelectedObjects(), SelectedActorInfo, DetailBuilder ) );
 
+	if (bAreBrushesSelected)
+	{
+		TransformDetails->HideTransformField(ETransformField::Scale);
+	}
+
 	IDetailCategoryBuilder& TransformCategory = DetailBuilder.EditCategory( "TransformCommon", LOCTEXT("TransformCommonCategory", "Transform").ToString(), ECategoryPriority::Transform );
 
 	TransformCategory.AddCustomBuilder( TransformDetails );
diff --git a/Engine/Source/Editor/UnrealEd/Private/BSPOps.cpp b/Engine/Source/Editor/UnrealEd/Private/BSPOps.cpp
--- a/Engine/Source/Editor/UnrealEd/Private/BSPOps.cpp
+++ b/Engine/Source/Editor/UnrealEd/Private/BSPOps.cpp
@@ -1255,24 +1255,6 @@ void FBSPOps::RotateBrushVerts(ABrush* Brush, const FRotator& Rotation, bool bCl
 }
 
 
-void FBSPOps::HandleBrushRegistered(ABrush* Brush)
-{
-	// Only patch up the rotation if this is a level-placed actor (ie not a prefab).
-	if ( Brush && !Brush->GetActorRotation().IsZero() && Brush->GetOuter() && Brush->GetOuter()->IsA(ULevel::StaticClass()) )
-	{
-		if ( Brush->IsVolumeBrush() )
-		{
-			const FRotator OldRotation = Brush->GetActorRotation();
-			Brush->SetActorRotation(FRotator::ZeroRotator);
-			FBSPOps::RotateBrushVerts( Brush, OldRotation, false );
-		}
-		else
-		{
-			Brush->SetActorRotation(FRotator::ZeroRotator);
-		}
-	}
-}
-
 void FBSPOps::HandleVolumeShapeChanged(AVolume& Volume)
 {
 	// The default physics volume doesn't have an associated UModel, so we need to handle that case gracefully.
diff --git a/Engine/Source/Editor/UnrealEd/Private/Editor.cpp b/Engine/Source/Editor/UnrealEd/Private/Editor.cpp
--- a/Engine/Source/Editor/UnrealEd/Private/Editor.cpp
+++ b/Engine/Source/Editor/UnrealEd/Private/Editor.cpp
@@ -2173,54 +2173,44 @@ void UEditorEngine::ApplyDeltaToActor(AActor* InActor,
 		{
 			bTranslationOnly = false;
 
-			if( Brush )
-			{
-				FBSPOps::RotateBrushVerts( Brush, InDeltaRot, true );
-			}
-			else
+			if ( bDelta )
 			{
-				if ( bDelta )
+				if( InActor->GetRootComponent() != NULL )
 				{
-					if( InActor->GetRootComponent() != NULL )
-					{
-						const FRotator OriginalRotation = InActor->GetRootComponent()->GetComponentRotation();
+					const FRotator OriginalRotation = InActor->GetRootComponent()->GetComponentRotation();
 
-						InActor->EditorApplyRotation( InDeltaRot, bAltDown, bShiftDown, bControlDown );
+					InActor->EditorApplyRotation( InDeltaRot, bAltDown, bShiftDown, bControlDown );
 
-						// Check to see if we should transform the rigid body
-						UPrimitiveComponent* RootPrimitiveComponent = Cast< UPrimitiveComponent >( InActor->GetRootComponent() );
-						if( bIsSimulatingInEditor && GIsPlayInEditorWorld && RootPrimitiveComponent != NULL )
-						{
-							FRotator ActorRotWind, ActorRotRem;
-							OriginalRotation.GetWindingAndRemainder(ActorRotWind, ActorRotRem);
+					// Check to see if we should transform the rigid body
+					UPrimitiveComponent* RootPrimitiveComponent = Cast< UPrimitiveComponent >( InActor->GetRootComponent() );
+					if( bIsSimulatingInEditor && GIsPlayInEditorWorld && RootPrimitiveComponent != NULL )
+					{
+						FRotator ActorRotWind, ActorRotRem;
+						OriginalRotation.GetWindingAndRemainder(ActorRotWind, ActorRotRem);
 
-							const FQuat ActorQ = ActorRotRem.Quaternion();
-							const FQuat DeltaQ = InDeltaRot.Quaternion();
-							const FQuat ResultQ = DeltaQ * ActorQ;
+						const FQuat ActorQ = ActorRotRem.Quaternion();
+						const FQuat DeltaQ = InDeltaRot.Quaternion();
+						const FQuat ResultQ = DeltaQ * ActorQ;
 
-							const FRotator NewActorRotRem = FRotator( ResultQ );
-							FRotator DeltaRot = NewActorRotRem - ActorRotRem;
-							DeltaRot.Normalize();
+						const FRotator NewActorRotRem = FRotator( ResultQ );
+						FRotator DeltaRot = NewActorRotRem - ActorRotRem;
+						DeltaRot.Normalize();
 
-							// @todo SIE: Not taking into account possible offset between root component and actor
-							RootPrimitiveComponent->SetWorldRotation( OriginalRotation + DeltaRot );
-						}
+						// @todo SIE: Not taking into account possible offset between root component and actor
+						RootPrimitiveComponent->SetWorldRotation( OriginalRotation + DeltaRot );
 					}
 				}
-				else
-				{
-					InActor->SetActorRotation( InDeltaRot );
-				}
-			}
 
-			if ( bDelta )
-			{
 				FVector NewActorLocation = InActor->GetActorLocation();
 				NewActorLocation -= GLevelEditorModeTools().PivotLocation;
-				NewActorLocation = FRotationMatrix( InDeltaRot ).TransformPosition( NewActorLocation );
+				NewActorLocation = FRotationMatrix(InDeltaRot).TransformPosition(NewActorLocation);
 				NewActorLocation += GLevelEditorModeTools().PivotLocation;
 				NewActorLocation -= InActor->GetActorLocation();
-				InActor->EditorApplyTranslation( NewActorLocation, bAltDown, bShiftDown, bControlDown );
+				InActor->EditorApplyTranslation(NewActorLocation, bAltDown, bShiftDown, bControlDown);
+			}
+			else
+			{
+				InActor->SetActorRotation( InDeltaRot );
 			}
 		}
 	}
diff --git a/Engine/Source/Editor/UnrealEd/Private/EditorBsp.cpp b/Engine/Source/Editor/UnrealEd/Private/EditorBsp.cpp
--- a/Engine/Source/Editor/UnrealEd/Private/EditorBsp.cpp
+++ b/Engine/Source/Editor/UnrealEd/Private/EditorBsp.cpp
@@ -1183,6 +1183,7 @@ int UEditorEngine::bspBrushCSG
 			DestEdPoly.iLink = i;
 
 		// Transform it.
+		DestEdPoly.Rotate( Actor->GetPrePivot(), Actor->GetActorRotation() );
 		DestEdPoly.Transform( Actor->GetPrePivot(), Actor->GetActorLocation() );
 
 		// Add poly to the temp model.
diff --git a/Engine/Source/Editor/UnrealEd/Private/EditorServer.cpp b/Engine/Source/Editor/UnrealEd/Private/EditorServer.cpp
--- a/Engine/Source/Editor/UnrealEd/Private/EditorServer.cpp
+++ b/Engine/Source/Editor/UnrealEd/Private/EditorServer.cpp
@@ -4486,14 +4486,7 @@ bool UEditorEngine::SnapActorTo( AActor* InActor, const bool InAlign, const bool
 			//@todo: This doesn't take into account that rotating the actor changes LocationOffset.
 			FRotator NewRotation( Hit.Normal.Rotation() );
 			NewRotation.Pitch -= 90.f;
-			if( Brush )
-			{
-				FBSPOps::RotateBrushVerts( Brush, NewRotation, false );
-			}
-			else
-			{
-				InActor->SetActorRotation(NewRotation);
-			}
+			InActor->SetActorRotation(NewRotation);
 		}
 
 		// Switch to the pie world if we have one
diff --git a/Engine/Source/Editor/UnrealEd/Private/LevelEditorViewport.cpp b/Engine/Source/Editor/UnrealEd/Private/LevelEditorViewport.cpp
--- a/Engine/Source/Editor/UnrealEd/Private/LevelEditorViewport.cpp
+++ b/Engine/Source/Editor/UnrealEd/Private/LevelEditorViewport.cpp
@@ -2851,14 +2851,12 @@ void FLevelEditorViewportClient::MoveLockedActorToCamera() const
 		{
 			ActiveActorLock->SetActorLocation( GCurrentLevelEditingViewportClient->GetViewLocation(), false );
 		}
+		ActiveActorLock->SetActorRotation( GCurrentLevelEditingViewportClient->GetViewRotation() );
+
 		ABrush* Brush = Cast< ABrush >( ActiveActorLock.Get() );
 		if( Brush )
 		{
-			FBSPOps::RotateBrushVerts( (ABrush*)ActiveActorLock.Get(), GCurrentLevelEditingViewportClient->GetViewRotation(), true );
-		}
-		else
-		{
-			ActiveActorLock->SetActorRotation( GCurrentLevelEditingViewportClient->GetViewRotation() );
+			Brush->SetNeedRebuild( Brush->GetLevel() );
 		}
 
 		FScopedLevelDirtied LevelDirtyCallback;
diff --git a/Engine/Source/Editor/UnrealEd/Private/UnrealEdEngine.cpp b/Engine/Source/Editor/UnrealEd/Private/UnrealEdEngine.cpp
--- a/Engine/Source/Editor/UnrealEd/Private/UnrealEdEngine.cpp
+++ b/Engine/Source/Editor/UnrealEd/Private/UnrealEdEngine.cpp
@@ -68,8 +68,7 @@ void UUnrealEdEngine::Init(IEngineLoop* InEngineLoop)
 	// Initialize the snap manager
 	FSnappingUtils::InitEditorSnappingTools();
 
-	// Register for notification of brush/volume changes
-	ABrush::GetOnBrushRegisteredDelegate().AddStatic(&FBSPOps::HandleBrushRegistered);
+	// Register for notification of volume changes
 	AVolume::GetOnVolumeShapeChangedDelegate().AddStatic(&FBSPOps::HandleVolumeShapeChanged);
 	//
 	InitBuilderBrush( GWorld );
diff --git a/Engine/Source/Editor/UnrealEd/Public/BSPOps.h b/Engine/Source/Editor/UnrealEd/Public/BSPOps.h
--- a/Engine/Source/Editor/UnrealEd/Public/BSPOps.h
+++ b/Engine/Source/Editor/UnrealEd/Public/BSPOps.h
@@ -51,8 +51,6 @@ public:
 	 */
 	static void RotateBrushVerts(ABrush* Brush, const FRotator& Rotation, bool bClearComponents);
 
-	/** Called when an ABrush is registered */
-	static void HandleBrushRegistered(ABrush* Brush);
 	/** Called when an AVolume shape is changed*/
 	static void HandleVolumeShapeChanged(AVolume& Volume);
 
diff --git a/Engine/Source/Runtime/Engine/Classes/Engine/Polys.h b/Engine/Source/Runtime/Engine/Classes/Engine/Polys.h
--- a/Engine/Source/Runtime/Engine/Classes/Engine/Polys.h
+++ b/Engine/Source/Runtime/Engine/Classes/Engine/Polys.h
@@ -72,6 +72,11 @@ public:
 	 */
 	ENGINE_API void Transform(const FVector &PreSubtract,const FVector &PostAdd);
 
+	/**
+	 * Rotate an editor polygon with a pre-transformation addition.
+	 */
+	ENGINE_API void Rotate(const FVector &PreSubtract, const FRotator &Rotation);
+
 	/**
 	 * Fix up an editor poly by deleting vertices that are identical.  Sets
 	 * vertex count to zero if it collapses.  Returns number of vertices, 0 or >=3.
diff --git a/Engine/Source/Runtime/Engine/Private/Polygon.cpp b/Engine/Source/Runtime/Engine/Private/Polygon.cpp
--- a/Engine/Source/Runtime/Engine/Private/Polygon.cpp
+++ b/Engine/Source/Runtime/Engine/Private/Polygon.cpp
@@ -464,6 +464,30 @@ void FPoly::Transform
 }
 
 
+void FPoly::Rotate
+(
+	const FVector&		PreSubtract,
+	const FRotator&		Rotation
+)
+{
+	// Rotate the vertices.
+	for (int32 Vertex = 0; Vertex < Vertices.Num(); Vertex++)
+	{
+		Vertices[Vertex] = PreSubtract + FRotationMatrix(Rotation).TransformVector(Vertices[Vertex] - PreSubtract);
+	}
+
+	Base = PreSubtract + FRotationMatrix(Rotation).TransformVector(Base - PreSubtract);
+
+	// Rotate the texture vectors.
+	TextureU = FRotationMatrix(Rotation).TransformVector(TextureU);
+	TextureV = FRotationMatrix(Rotation).TransformVector(TextureV);
+
+	// Rotate the normal.
+	Normal = FRotationMatrix(Rotation).TransformVector(Normal);
+	Normal = Normal.SafeNormal();
+}
+
+
 int32 FPoly::RemoveColinears()
 {
 	FMemStack& LocalMemStack = FMemStack::Get();
