Allowed BSP scale to affect the transform scale, instead of immediately transforming the vertices.

diff --git a/Engine/Source/Editor/DetailCustomizations/Private/ActorDetails.cpp b/Engine/Source/Editor/DetailCustomizations/Private/ActorDetails.cpp
--- a/Engine/Source/Editor/DetailCustomizations/Private/ActorDetails.cpp
+++ b/Engine/Source/Editor/DetailCustomizations/Private/ActorDetails.cpp
@@ -428,11 +428,6 @@ void FActorDetails::AddTransformCategory( IDetailLayoutBuilder& DetailBuilder )
 	
 	TSharedRef<FComponentTransformDetails> TransformDetails = MakeShareable( new FComponentTransformDetails( DetailBuilder.GetDetailsView().GetSelectedObjects(), SelectedActorInfo, DetailBuilder ) );
 
-	if (bAreBrushesSelected)
-	{
-		TransformDetails->HideTransformField(ETransformField::Scale);
-	}
-
 	IDetailCategoryBuilder& TransformCategory = DetailBuilder.EditCategory( "TransformCommon", LOCTEXT("TransformCommonCategory", "Transform").ToString(), ECategoryPriority::Transform );
 
 	TransformCategory.AddCustomBuilder( TransformDetails );
diff --git a/Engine/Source/Editor/UnrealEd/Private/Editor.cpp b/Engine/Source/Editor/UnrealEd/Private/Editor.cpp
--- a/Engine/Source/Editor/UnrealEd/Private/Editor.cpp
+++ b/Engine/Source/Editor/UnrealEd/Private/Editor.cpp
@@ -2257,196 +2257,49 @@ void UEditorEngine::ApplyDeltaToActor(AActor* InActor,
 			// Note: With the new additive scaling method, this is handled in FLevelEditorViewportClient::ModifyScale
 			if( GEditor->UsePercentageBasedScaling() )
 			{
-			// Get actor box extents
-			const FBox BoundingBox = InActor->GetComponentsBoundingBox( true );
-			const FVector BoundsExtents = BoundingBox.GetExtent();
+				// Get actor box extents
+				const FBox BoundingBox = InActor->GetComponentsBoundingBox( true );
+				const FVector BoundsExtents = BoundingBox.GetExtent();
 
-			// Make sure scale on actors is clamped to a minimum and maximum size.
-			const float MinThreshold = 1.0f;
+				// Make sure scale on actors is clamped to a minimum and maximum size.
+				const float MinThreshold = 1.0f;
 
-			for (int32 Idx=0; Idx<3; Idx++)
-			{
-				if ( ( FMath::Pow(BoundsExtents[Idx], 2) ) > BIG_NUMBER)
-				{
-					ModifiedScale[Idx] = 0.0f;
-				}
-				else if (SMALL_NUMBER < BoundsExtents[Idx])
+				for (int32 Idx=0; Idx<3; Idx++)
 				{
-					const bool bBelowAllowableScaleThreshold = ((InDeltaScale[Idx] + 1.0f) * BoundsExtents[Idx]) < MinThreshold;
-
-					if(bBelowAllowableScaleThreshold)
+					if ( ( FMath::Pow(BoundsExtents[Idx], 2) ) > BIG_NUMBER)
 					{
-						ModifiedScale[Idx] = (MinThreshold / BoundsExtents[Idx]) - 1.0f;
+						ModifiedScale[Idx] = 0.0f;
 					}
-				}
-			}
-			}
-
-			// If the actor is a brush, update the vertices.
-			if( Brush != nullptr)
-			{
-				FVector DeltaScale = ModifiedScale;
-
-				AActor::bUsePercentageBasedScaling = GEditor->UsePercentageBasedScaling();
-				if (!AActor::bUsePercentageBasedScaling)
-				{
-					// Attempt to deduce current implied scale by comparing the bounding box of the source builder brush with the current bounding box.
-					// Not a perfect method, but it feels OK.
-					// Also disallow singularities by not permitting a zero scale on any of the axes (otherwise we lose information).
-					const UBrushBuilder* BrushBuilder = Brush->GetBrushBuilder();
-					if (BrushBuilder != nullptr)
+					else if (SMALL_NUMBER < BoundsExtents[Idx])
 					{
-						// Get builder brush bounding box
-						FBox BuildBrushBBox(0);
-
-						for (int32 Index = 0, NumVerts = BrushBuilder->GetVertexCount(); Index < NumVerts; Index++)
-						{
-							BuildBrushBBox += BrushBuilder->GetVertex(Index);
-						}
+						const bool bBelowAllowableScaleThreshold = ((InDeltaScale[Idx] + 1.0f) * BoundsExtents[Idx]) < MinThreshold;
 
-						if (Brush->Brush != nullptr && Brush->Brush->Polys != nullptr)
+						if(bBelowAllowableScaleThreshold)
 						{
-							// Get current bounding box
-							FBox CurrentBBox(0);
-
-							for (const auto& Poly : Brush->Brush->Polys->Element)
-							{
-								for (const auto& Vertex : Poly.Vertices)
-								{
-									CurrentBBox += Vertex;
-								}
-							}
-
-							// Estimate implied scale (can only be positive), avoiding zero scale components
-							const FVector ImpliedScale = ((CurrentBBox.Max - CurrentBBox.Min) / (BuildBrushBBox.Max - BuildBrushBBox.Min))
-								.ComponentMax(FVector(SMALL_NUMBER, SMALL_NUMBER, SMALL_NUMBER));
-							
-							// Calculate delta scale vector which, when applied as a scale transform and added to the original vector, will scale by this additional amount.
-							// A scale of 0.0 implies no change, a scale of -1.0 will collapse component to zero, a scale of +1.0 will double the size.
-							const FVector NewScaleVector = DeltaScale / ImpliedScale;
-
-							// Protect against singularities by never allowing a component scale close to 0.
-							const float Limit = -1.0f + 0.0625f;
-							if (NewScaleVector.X > Limit &&
-								NewScaleVector.Y > Limit &&
-								NewScaleVector.Z > Limit)
-							{
-								// Apply calculated scale delta
-								DeltaScale = NewScaleVector;
-							}
-							else
-							{
-								// Apply zero scale delta
-								DeltaScale = FVector::ZeroVector;
-							}
+							ModifiedScale[Idx] = (MinThreshold / BoundsExtents[Idx]) - 1.0f;
 						}
 					}
 				}
+			}
 
-				// Scale all of the polygons of the brush.
-				const FScaleMatrix Matrix(DeltaScale);
-				
-				if(Brush->GetBrushComponent()->Brush && Brush->GetBrushComponent()->Brush->Polys)
-				{
-					// @todo UE4 : verify this code
-					// Brush Transform doesn't seem to change from this code path
-					// So I'm changing to caching it, but if this actorToWorld supposed to re-calculate
-					// from each poly transformation, this will have to be fixed. 
-					// if this causes any issue, please contact LH
-					FTransform BrushActorToWorld = Brush->ActorToWorld();
-
-					for( int32 poly = 0 ; poly < Brush->GetBrushComponent()->Brush->Polys->Element.Num() ; poly++ )
-					{
-						FPoly* Poly = &(Brush->GetBrushComponent()->Brush->Polys->Element[poly]);
-
-						FBox bboxBefore(0);
-						for( int32 vertex = 0 ; vertex < Poly->Vertices.Num() ; vertex++ )
-						{
-							bboxBefore += BrushActorToWorld.TransformPosition( Poly->Vertices[vertex] );
-						}
-
-						// Scale the vertices
-
-						for( int32 vertex = 0 ; vertex < Poly->Vertices.Num() ; vertex++ )
-						{
-							FVector Wk = BrushActorToWorld.TransformPosition( Poly->Vertices[vertex] );
-							Wk -= GLevelEditorModeTools().PivotLocation;
-							Wk += Matrix.TransformPosition( Wk );
-							Wk += GLevelEditorModeTools().PivotLocation;
-							Poly->Vertices[vertex] = BrushActorToWorld.InverseTransformPosition( Wk );
-						}
-
-						FBox bboxAfter(0);
-						for( int32 vertex = 0 ; vertex < Poly->Vertices.Num() ; vertex++ )
-						{
-							bboxAfter += BrushActorToWorld.TransformPosition( Poly->Vertices[vertex] );
-						}
-
-						FVector Wk = BrushActorToWorld.TransformPosition( Poly->Base );
-						Wk -= GLevelEditorModeTools().PivotLocation;
-						Wk += Matrix.TransformPosition( Wk );
-						Wk += GLevelEditorModeTools().PivotLocation;
-						Poly->Base = BrushActorToWorld.InverseTransformPosition( Wk );
-
-						// Scale the texture vectors
-
-						for( int32 a = 0 ; a < 3 ; ++a )
-						{
-							const float Before = bboxBefore.GetExtent()[a];
-							const float After = bboxAfter.GetExtent()[a];
-
-							if( After != 0.0 )
-							{
-								const float Pct = Before / After;
-
-								if( Pct != 0.0 )
-								{
-									Poly->TextureU[a] *= Pct;
-									Poly->TextureV[a] *= Pct;
-								}
-							}
-						}
-
-						// Recalc the normal for the poly
-
-						Poly->Normal = FVector::ZeroVector;
-						Poly->Finalize((ABrush*)InActor,0);
-					}
-
-					Brush->GetBrushComponent()->Brush->BuildBound();
+			if ( bDelta )
+			{
+				// Flag actors to use old-style scaling or not
+				// @todo: Remove this hack once we have decided on the scaling method to use.
+				AActor::bUsePercentageBasedScaling = GEditor->UsePercentageBasedScaling();
 
-					if( !Brush->IsStaticBrush() )
-					{
-						FBSPOps::csgPrepMovingBrush( Brush );
-					}
-				}
+				InActor->EditorApplyScale( 
+					ModifiedScale,
+					&GLevelEditorModeTools().PivotLocation,
+					bAltDown,
+					bShiftDown,
+					bControlDown
+					);
 
-				// Brushes need to be re-registered so that the edges will line up with the vertices correctly.
-				// We need to invalidate the lighting cache BEFORE clearing the components!
-				InActor->InvalidateLightingCache();
-				InActor->ReregisterAllComponents();
 			}
-			else
+			else if( InActor->GetRootComponent() != NULL )
 			{
-				if ( bDelta )
-				{
-					// Flag actors to use old-style scaling or not
-					// @todo: Remove this hack once we have decided on the scaling method to use.
-					AActor::bUsePercentageBasedScaling = GEditor->UsePercentageBasedScaling();
-
-					InActor->EditorApplyScale( 
-						ModifiedScale,
-												&GLevelEditorModeTools().PivotLocation,
-												bAltDown,
-												bShiftDown,
-						bControlDown
-						);
-
-				}
-				else if( InActor->GetRootComponent() != NULL )
-				{
-					InActor->GetRootComponent()->SetRelativeScale3D( InDeltaScale );
-				}
+				InActor->GetRootComponent()->SetRelativeScale3D( InDeltaScale );
 			}
 		}
 	}
diff --git a/Engine/Source/Editor/UnrealEd/Private/EditorBsp.cpp b/Engine/Source/Editor/UnrealEd/Private/EditorBsp.cpp
--- a/Engine/Source/Editor/UnrealEd/Private/EditorBsp.cpp
+++ b/Engine/Source/Editor/UnrealEd/Private/EditorBsp.cpp
@@ -1183,6 +1183,7 @@ int UEditorEngine::bspBrushCSG
 			DestEdPoly.iLink = i;
 
 		// Transform it.
+		DestEdPoly.Scale( Actor->GetPrePivot(), Actor->GetActorScale() );
 		DestEdPoly.Rotate( Actor->GetPrePivot(), Actor->GetActorRotation() );
 		DestEdPoly.Transform( Actor->GetPrePivot(), Actor->GetActorLocation() );
 
diff --git a/Engine/Source/Runtime/Engine/Classes/Engine/Polys.h b/Engine/Source/Runtime/Engine/Classes/Engine/Polys.h
--- a/Engine/Source/Runtime/Engine/Classes/Engine/Polys.h
+++ b/Engine/Source/Runtime/Engine/Classes/Engine/Polys.h
@@ -77,6 +77,11 @@ public:
 	 */
 	ENGINE_API void Rotate(const FVector &PreSubtract, const FRotator &Rotation);
 
+	/**
+	 * Scale an editor polygon with a pre-transformation addition.
+	 */
+	ENGINE_API void Scale(const FVector &PreSubtract, const FVector &Scale);
+
 	/**
 	 * Fix up an editor poly by deleting vertices that are identical.  Sets
 	 * vertex count to zero if it collapses.  Returns number of vertices, 0 or >=3.
diff --git a/Engine/Source/Runtime/Engine/Private/Polygon.cpp b/Engine/Source/Runtime/Engine/Private/Polygon.cpp
--- a/Engine/Source/Runtime/Engine/Private/Polygon.cpp
+++ b/Engine/Source/Runtime/Engine/Private/Polygon.cpp
@@ -488,6 +488,29 @@ void FPoly::Rotate
 }
 
 
+void FPoly::Scale
+(
+	const FVector&		PreSubtract,
+	const FVector&		Scale
+)
+{
+	// Scale the vertices.
+	for (int32 Vertex = 0; Vertex < Vertices.Num(); Vertex++)
+	{
+		Vertices[Vertex] = PreSubtract + (Vertices[Vertex] - PreSubtract) * Scale;
+	}
+
+	Base = PreSubtract + (Base - PreSubtract) * Scale;
+
+	// Scale the texture vectors.
+	TextureU *= Scale;
+	TextureV *= Scale;
+
+	// Renormalize the normal.
+	Normal = Normal.SafeNormal();
+}
+
+
 int32 FPoly::RemoveColinears()
 {
 	FMemStack& LocalMemStack = FMemStack::Get();
