Fixed dead actor weak references after re-instancing

(checking in for JarekP)

diff --git a/Engine/Source/Editor/UnrealEd/Private/Kismet2/KismetReinstanceUtilities.cpp b/Engine/Source/Editor/UnrealEd/Private/Kismet2/KismetReinstanceUtilities.cpp
--- a/Engine/Source/Editor/UnrealEd/Private/Kismet2/KismetReinstanceUtilities.cpp
+++ b/Engine/Source/Editor/UnrealEd/Private/Kismet2/KismetReinstanceUtilities.cpp
@@ -276,6 +276,8 @@ void FBlueprintCompileReinstancer::ReplaceInstancesOfClass(UClass* OldClass, UCl
 	TArray<UObject*> ObjectsToReplace;
 	const bool bLogConversions = false; // for debugging
 
+	TMap<FStringAssetReference, UObject*> ReinstancedObjectsWeakReferenceMap;
+
 	// Map of old objects to new objects
 	TMap<UObject*, UObject*> OldToNewInstanceMap;
 	TMap<UClass*, UClass*> OldToNewClassMap;
@@ -417,6 +419,8 @@ void FBlueprintCompileReinstancer::ReplaceInstancesOfClass(UClass* OldClass, UCl
 					check(NewActor);
 					NewObject = NewActor;
 
+					ReinstancedObjectsWeakReferenceMap.Add(OldObject, NewObject);
+
 					OldActor->DestroyConstructedComponents(); // don't want to serialize components from the old actor
 					// Unregister native components so we don't copy any sub-components they generate for themselves (like UCameraComponent does)
 					OldActor->UnregisterAllComponents(); 
@@ -540,7 +544,12 @@ void FBlueprintCompileReinstancer::ReplaceInstancesOfClass(UClass* OldClass, UCl
 
 	// Now replace any pointers to the old archetypes/instances with pointers to the new one
 	TArray<UObject*> SourceObjects;
+	TArray<UObject*> DstObjects;
+
 	OldToNewInstanceMap.GenerateKeyArray(SourceObjects);
+	OldToNewInstanceMap.GenerateValueArray(DstObjects); // Also look for references in new spawned objects.
+
+	SourceObjects.Append(DstObjects);
 
 	// Add in the old class/CDO to this pass, so class/CDO references are fixed up
 	SourceObjects.Add(OldClass);
@@ -576,7 +585,45 @@ void FBlueprintCompileReinstancer::ReplaceInstancesOfClass(UClass* OldClass, UCl
 			UObject* Obj = *It;
 			if (!ObjectsToReplace.Contains(Obj)) // Don't bother trying to fix old objects, this would break them
 			{
-				FArchiveReplaceObjectRef<UObject> ReplaceAr(Obj, OldToNewInstanceMap, /*bNullPrivateRefs=*/ false, /*bIgnoreOuterRef=*/ false, /*bIgnoreArchetypeRef=*/ false);
+				// The class for finding and replacing weak references.
+				// We can't relay on "standard" weak references replacement as
+				// it depends on FStringAssetReference::ResolveObject, which
+				// tries to find the object with the stored path. It is
+				// impossible, cause above we deleted old actors (after
+				// spawning new ones), so during objects traverse we have to
+				// find FStringAssetReferences with the raw given path taken
+				// before deletion of old actors and fix them.
+				class ReferenceReplace : public FArchiveReplaceObjectRef<UObject>
+				{
+				public:
+					ReferenceReplace(UObject* InSearchObject, const TMap<UObject*, UObject*>& InReplacementMap, TMap<FStringAssetReference, UObject*> WeakReferencesMap)
+						: FArchiveReplaceObjectRef<UObject>(InSearchObject, InReplacementMap, false, false, false, true), WeakReferencesMap(WeakReferencesMap)
+					{
+						SerializeSearchObject();
+					}
+
+					FArchive& operator<<(FStringAssetReference& Ref) override
+					{
+						const UObject*const* PtrToObjPtr = WeakReferencesMap.Find(Ref);
+
+						if (PtrToObjPtr != nullptr)
+						{
+							Ref = *PtrToObjPtr;
+						}
+
+						return *this;
+					}
+
+					FArchive& operator<<(FAssetPtr& Ref) override
+					{
+						return operator<<(Ref.GetUniqueID());
+					}
+
+				private:
+					const TMap<FStringAssetReference, UObject*>& WeakReferencesMap;
+				};
+
+				ReferenceReplace ReplaceAr(Obj, OldToNewInstanceMap, ReinstancedObjectsWeakReferenceMap);
 			}
 		}
 	}
diff --git a/Engine/Source/Runtime/CoreUObject/Public/UObject/PersistentObjectPtr.h b/Engine/Source/Runtime/CoreUObject/Public/UObject/PersistentObjectPtr.h
--- a/Engine/Source/Runtime/CoreUObject/Public/UObject/PersistentObjectPtr.h
+++ b/Engine/Source/Runtime/CoreUObject/Public/UObject/PersistentObjectPtr.h
@@ -101,6 +101,14 @@ public:
 		return ObjectID;
 	}
 
+	/**
+	 * Non-const version of the above.
+	 */
+	FORCEINLINE TObjectID& GetUniqueID()
+	{
+		return ObjectID;
+	}
+
 	/**  
 	 * Dereference the lazy pointer, which may cause it to become valid again. Will not try to load pending outside of game thread
 	 * @return NULL if this object is gone or the lazy pointer was NULL, otherwise a valid UObject pointer
