3DWidget proper display name and minor fixes

diff --git a/Engine/Source/Editor/UnrealEd/Private/EditorModes.cpp b/Engine/Source/Editor/UnrealEd/Private/EditorModes.cpp
--- a/Engine/Source/Editor/UnrealEd/Private/EditorModes.cpp
+++ b/Engine/Source/Editor/UnrealEd/Private/EditorModes.cpp
@@ -1048,7 +1048,7 @@ static bool IsTransformProperty(UProperty* InProp)
 
 }
 
-void FEdMode::GetPropertyWidgetInfos(const UStruct* InStruct, const void* InContainer, TArray<FPropertyWidgetInfo>& OutInfos, FString PropertyNamePrefix) const
+void FEdMode::GetPropertyWidgetInfos(const UStruct* InStruct, const void* InContainer, TArray<FPropertyWidgetInfo>& OutInfos, FString PropertyNamePrefix, FString DisplayNamePrefix) const
 {
 	if(PropertyNamePrefix.Len() == 0)
 	{
@@ -1058,9 +1058,19 @@ void FEdMode::GetPropertyWidgetInfos(const UStruct* InStruct, const void* InCont
 	for (TFieldIterator<UProperty> PropertyIt(InStruct, EFieldIteratorFlags::IncludeSuper); PropertyIt; ++PropertyIt)
 	{
 		UProperty* CurrentProp = *PropertyIt;
+		check(CurrentProp);
+		FString DisplayName = CurrentProp->GetMetaData(TEXT("DisplayName"));
+		if (!PropertyNamePrefix.IsEmpty() && DisplayNamePrefix.IsEmpty()) // Display name is already invalid
+		{
+			DisplayName.Empty();
+		}
+		if (!DisplayName.IsEmpty()) //Display name cannot be only the prefix.
+		{
+			DisplayName = DisplayNamePrefix + DisplayName;
+		}
+
 		if(	ShouldCreateWidgetForProperty(CurrentProp) )
 		{
-			const FString DisplayName = CurrentProp->GetMetaData(TEXT("DisplayName"));
 			if( UArrayProperty* ArrayProp = Cast<UArrayProperty>(CurrentProp) )
 			{
 				check(InContainer != NULL);
@@ -1076,7 +1086,7 @@ void FEdMode::GetPropertyWidgetInfos(const UStruct* InStruct, const void* InCont
 
 					//fill it in with the struct name
 					WidgetInfo.PropertyName = PropertyNamePrefix + CurrentProp->GetFName().ToString();
-					WidgetInfo.DisplayName = DisplayName.IsEmpty() ? WidgetInfo.PropertyName : (PropertyNamePrefix + DisplayName);
+					WidgetInfo.DisplayName = DisplayName.IsEmpty() ? WidgetInfo.PropertyName : DisplayName;
 
 					//And see if we have any meta data that matches the MD_ValidateWidgetUsing name
 					WidgetInfo.PropertyValidationName = FName(*CurrentProp->GetMetaData(MD_ValidateWidgetUsing));
@@ -1098,7 +1108,7 @@ void FEdMode::GetPropertyWidgetInfos(const UStruct* InStruct, const void* InCont
 
 				//fill it in with the struct name
 				WidgetInfo.PropertyName = PropertyNamePrefix + CurrentProp->GetFName().ToString();
-				WidgetInfo.DisplayName = DisplayName.IsEmpty() ? WidgetInfo.PropertyName : (PropertyNamePrefix + DisplayName);
+				WidgetInfo.DisplayName = DisplayName.IsEmpty() ? WidgetInfo.PropertyName : DisplayName;
 
 				//And see if we have any meta data that matches the MD_ValidateWidgetUsing name
 				WidgetInfo.PropertyValidationName = FName(*CurrentProp->GetMetaData(MD_ValidateWidgetUsing));
@@ -1117,7 +1127,11 @@ void FEdMode::GetPropertyWidgetInfos(const UStruct* InStruct, const void* InCont
 			if(StructProp != NULL)
 			{
 				// Recursively traverse into structures, looking for additional vector properties to expose
-				GetPropertyWidgetInfos(StructProp->Struct, StructProp->ContainerPtrToValuePtr<void>(InContainer), OutInfos, StructProp->GetFName().ToString() + TEXT("."));
+				GetPropertyWidgetInfos(StructProp->Struct
+					, StructProp->ContainerPtrToValuePtr<void>(InContainer)
+					, OutInfos
+					, PropertyNamePrefix + StructProp->GetFName().ToString() + TEXT(".")
+					, !DisplayName.IsEmpty() ? (DisplayName + TEXT(".")) : FString());
 			}
 			else
 			{
@@ -1133,7 +1147,12 @@ void FEdMode::GetPropertyWidgetInfos(const UStruct* InStruct, const void* InCont
 						{
 							if(ArrayHelper.IsValidIndex(ArrayIndex))
 							{
-								GetPropertyWidgetInfos(StructProp->Struct, ArrayHelper.GetRawPtr(ArrayIndex), OutInfos, StructProp->GetFName().ToString() + FString::Printf(TEXT("[%d]"), ArrayIndex) + TEXT("."));
+								const FString ArrayPostfix = FString::Printf(TEXT("[%d]"), ArrayIndex) + TEXT(".");
+								GetPropertyWidgetInfos(StructProp->Struct
+									, ArrayHelper.GetRawPtr(ArrayIndex)
+									, OutInfos
+									, PropertyNamePrefix + ArrayProp->GetFName().ToString() + ArrayPostfix
+									, !DisplayName.IsEmpty() ? (DisplayName + ArrayPostfix) : FString());
 							}
 						}
 					}
diff --git a/Engine/Source/Editor/UnrealEd/Public/EditorModes.h b/Engine/Source/Editor/UnrealEd/Public/EditorModes.h
--- a/Engine/Source/Editor/UnrealEd/Public/EditorModes.h
+++ b/Engine/Source/Editor/UnrealEd/Public/EditorModes.h
@@ -402,7 +402,7 @@ protected:
 	 * @param OutInfos An array of widget info structures (output).
 	 * @param PropertyNamePrefix Optional prefix to use to filter the output.
 	 */
-	void GetPropertyWidgetInfos(const UStruct* InStruct, const void* InContainer, TArray<FPropertyWidgetInfo>& OutInfos, FString PropertyNamePrefix = TEXT("")) const;
+	void GetPropertyWidgetInfos(const UStruct* InStruct, const void* InContainer, TArray<FPropertyWidgetInfo>& OutInfos, FString PropertyNamePrefix = TEXT(""), FString DisplayNamePrefix = TEXT("")) const;
 
 	/** Name of the property currently being edited */
 	FString EditedPropertyName;
