Switched to using the cook on the fly server code path in the cook commandlet.

diff --git a/Engine/Source/Editor/UnrealEd/Private/Commandlets/CookCommandlet.cpp b/Engine/Source/Editor/UnrealEd/Private/Commandlets/CookCommandlet.cpp
--- a/Engine/Source/Editor/UnrealEd/Private/Commandlets/CookCommandlet.cpp
+++ b/Engine/Source/Editor/UnrealEd/Private/Commandlets/CookCommandlet.cpp
@@ -561,13 +561,13 @@ int32 UCookCommandlet::Main(const FString& CmdLineParams)
 		GenerateAssetRegistry(Platforms);
 
 		// new cook is better 
-		if ( Switches.Contains(TEXT("NEWCOOK")))
+		if ( Switches.Contains(TEXT("OLDCOOK")))
 		{
-			NewCook(Platforms, FilesInPath );
+			Cook(Platforms, FilesInPath);
 		}
 		else
 		{
-			Cook(Platforms, FilesInPath);
+			NewCook(Platforms, FilesInPath );
 		}
 	}
 	
