Fixed badly resolved file

diff --git a/Engine/Source/Editor/UnrealEd/Private/Toolkits/SStandaloneAssetEditorToolkitHost.cpp b/Engine/Source/Editor/UnrealEd/Private/Toolkits/SStandaloneAssetEditorToolkitHost.cpp
--- a/Engine/Source/Editor/UnrealEd/Private/Toolkits/SStandaloneAssetEditorToolkitHost.cpp
+++ b/Engine/Source/Editor/UnrealEd/Private/Toolkits/SStandaloneAssetEditorToolkitHost.cpp
@@ -349,17 +349,10 @@ void SStandaloneAssetEditorToolkitHost::OnTabClosed(TSharedRef<SDockTab> TabClos
 {
 	check(TabClosed == HostTabPtr.Pin());
 
->>>> ORIGINAL //depot/UE4/Engine/Source/Editor/UnrealEd/Private/Toolkits/SStandaloneAssetEditorToolkitHost.cpp#40
-	const TSharedPtr<FAssetEditorToolkit> AssetEditorToolkit = StaticCastSharedPtr<FAssetEditorToolkit>(HostedToolkit);
-	if(AssetEditorToolkit.IsValid())
-==== THEIRS //depot/UE4/Engine/Source/Editor/UnrealEd/Private/Toolkits/SStandaloneAssetEditorToolkitHost.cpp#41
-	if (HostedAssetEditorToolkit.IsValid())
-==== YOURS //matt_kuhlenschmidt_mac/UE4/Engine/Source/Editor/UnrealEd/Private/Toolkits/SStandaloneAssetEditorToolkitHost.cpp
 	MyTabManager->SetMenuMultiBox(nullptr);
 	
 	const TSharedPtr<FAssetEditorToolkit> AssetEditorToolkit = StaticCastSharedPtr<FAssetEditorToolkit>(HostedToolkit);
 	if(AssetEditorToolkit.IsValid())
-<<<<
 	{
 		const TArray<UObject*>* const ObjectsBeingEdited = HostedAssetEditorToolkit->GetObjectsCurrentlyBeingEdited();
 		if(ObjectsBeingEdited)
