Fixed another resolve issue.

diff --git a/Engine/Source/Editor/UnrealEd/Private/Toolkits/SStandaloneAssetEditorToolkitHost.cpp b/Engine/Source/Editor/UnrealEd/Private/Toolkits/SStandaloneAssetEditorToolkitHost.cpp
--- a/Engine/Source/Editor/UnrealEd/Private/Toolkits/SStandaloneAssetEditorToolkitHost.cpp
+++ b/Engine/Source/Editor/UnrealEd/Private/Toolkits/SStandaloneAssetEditorToolkitHost.cpp
@@ -351,8 +351,7 @@ void SStandaloneAssetEditorToolkitHost::OnTabClosed(TSharedRef<SDockTab> TabClos
 
 	MyTabManager->SetMenuMultiBox(nullptr);
 	
-	const TSharedPtr<FAssetEditorToolkit> AssetEditorToolkit = StaticCastSharedPtr<FAssetEditorToolkit>(HostedToolkit);
-	if(AssetEditorToolkit.IsValid())
+	if(HostedAssetEditorToolkit.IsValid())
 	{
 		const TArray<UObject*>* const ObjectsBeingEdited = HostedAssetEditorToolkit->GetObjectsCurrentlyBeingEdited();
 		if(ObjectsBeingEdited)
