Editor: Move some selection set references in mode tools to use the 'local' selection set (may still be the global one, depending on the instance)

diff --git a/Engine/Source/Editor/UnrealEd/Private/EditorModes.cpp b/Engine/Source/Editor/UnrealEd/Private/EditorModes.cpp
--- a/Engine/Source/Editor/UnrealEd/Private/EditorModes.cpp
+++ b/Engine/Source/Editor/UnrealEd/Private/EditorModes.cpp
@@ -297,7 +297,7 @@ bool FEdMode::InputKey(FEditorViewportClient* ViewportClient, FViewport* Viewpor
 	{
 		// Pass input up to selected actors if not in a tool mode
 		TArray<AActor*> SelectedActors;
-		GEditor->GetSelectedActors()->GetSelectedObjects<AActor>(SelectedActors);
+		Owner->GetSelectedActors()->GetSelectedObjects<AActor>(SelectedActors);
 
 		for( TArray<AActor*>::TIterator it(SelectedActors); it; ++it )
 		{
@@ -461,7 +461,7 @@ FVector FEdMode::GetWidgetLocation() const
 
 bool FEdMode::ShouldDrawWidget() const
 {
-	return (GEditor->GetSelectedActors()->GetTop<AActor>() != NULL);
+	return (Owner->GetSelectedActors()->GetTop<AActor>() != NULL);
 }
 
 /**
@@ -560,10 +560,9 @@ void FEdMode::Enter()
 {
 	// Update components for selected actors, in case the mode we just exited
 	// was hijacking selection events selection and not updating components.
-	for ( FSelectionIterator It( GEditor->GetSelectedActorIterator() ) ; It ; ++It )
+	for ( FSelectionIterator It( *Owner->GetSelectedActors() ) ; It ; ++It )
 	{
-		AActor* SelectedActor = static_cast<AActor*>( *It );
-		checkSlow( SelectedActor->IsA(AActor::StaticClass()) );
+		AActor* SelectedActor = CastChecked<AActor>( *It );
 		SelectedActor->MarkComponentsRenderStateDirty();
 	}
 
@@ -622,7 +621,7 @@ void FEdMode::Render(const FSceneView* View,FViewport* Viewport,FPrimitiveDrawIn
 			ABrush* Brush = *It;
 
 			// Brush->Brush is checked to safe from brushes that were created without having their brush members attached.
-			if( Brush->Brush && (FActorEditorUtils::IsABuilderBrush(Brush) || Brush->IsVolumeBrush()) && GEditor->GetSelectedActors()->IsSelected(Brush) )
+			if( Brush->Brush && (FActorEditorUtils::IsABuilderBrush(Brush) || Brush->IsVolumeBrush()) && Owner->GetSelectedActors()->IsSelected(Brush) )
 			{
 				// Build a mesh by basically drawing the triangles of each 
 				FDynamicMeshBuilder MeshBuilder;
@@ -670,7 +669,7 @@ void FEdMode::Render(const FSceneView* View,FViewport* Viewport,FPrimitiveDrawIn
 		const float TextureSizeX = VertexTexture->GetSizeX() * 0.170f;
 		const float TextureSizeY = VertexTexture->GetSizeY() * 0.170f;
 
-		for (FSelectionIterator It(GEditor->GetSelectedActorIterator()); It; ++It)
+		for (FSelectionIterator It(*Owner->GetSelectedActors()); It; ++It)
 		{
 			AActor* SelectedActor = static_cast<AActor*>(*It);
 			checkSlow(SelectedActor->IsA(AActor::StaticClass()));
@@ -819,7 +818,7 @@ void FEdMode::DrawHUD(FEditorViewportClient* ViewportClient,FViewport* Viewport,
 	// Temporaries.
 	TArray<FVector> Vertices;
 
-	for ( FSelectionIterator It( GEditor->GetSelectedActorIterator() ) ; It ; ++It )
+	for ( FSelectionIterator It( *Owner->GetSelectedActors() ) ; It ; ++It )
 	{
 		AActor* SelectedActor = static_cast<AActor*>( *It );
 		checkSlow( SelectedActor->IsA(AActor::StaticClass()) );
@@ -939,7 +938,7 @@ void FEdMode::DrawHUD(FEditorViewportClient* ViewportClient,FViewport* Viewport,
 // Draw brackets around all selected objects
 void FEdMode::DrawBrackets( FEditorViewportClient* ViewportClient, FViewport* Viewport, const FSceneView* View, FCanvas* Canvas )
 {
-	USelection& SelectedActors = *GEditor->GetSelectedActors();
+	USelection& SelectedActors = *Owner->GetSelectedActors();
 	for( int32 CurSelectedActorIndex = 0; CurSelectedActorIndex < SelectedActors.Num(); ++CurSelectedActorIndex )
 	{
 		AActor* SelectedActor = Cast<AActor>( SelectedActors.GetSelectedObject(CurSelectedActorIndex ) );
@@ -1009,12 +1008,7 @@ FVector FEdMode::GetWidgetNormalFromCurrentAxis( void* InData )
 
 AActor* FEdMode::GetFirstSelectedActorInstance() const
 {
-	if(GEditor->GetSelectedActorCount() > 0)
-	{
-		return GEditor->GetSelectedActors()->GetTop<AActor>();
-	}
-
-	return NULL;
+	return Owner->GetSelectedActors()->GetTop<AActor>();
 }
 
 bool FEdMode::CanCreateWidgetForStructure(const UStruct* InPropStruct)
@@ -1518,9 +1512,9 @@ FEdMode* FEditorModeTools::FindMode( FEditorModeID InID )
 
 FMatrix FEditorModeTools::GetCustomDrawingCoordinateSystem()
 {
-	FMatrix matrix = FMatrix::Identity;
+	FMatrix Matrix = FMatrix::Identity;
 
-	switch( GetCoordSystem() )
+	switch (GetCoordSystem())
 	{
 		case COORD_Local:
 		{
@@ -1528,29 +1522,29 @@ FMatrix FEditorModeTools::GetCustomDrawingCoordinateSystem()
 			// If it doesn't want to, create it by looking at the currently selected actors list.
 
 			bool CustomCoordinateSystemProvided = false;
-			for ( const auto& Mode : Modes)
+			for (const auto& Mode : Modes)
 			{
-				if( Mode->GetCustomDrawingCoordinateSystem( matrix, NULL ) )
+				if (Mode->GetCustomDrawingCoordinateSystem(Matrix, nullptr))
 				{
 					CustomCoordinateSystemProvided = true;
 					break;
 				}
 			}
 
-			if ( !CustomCoordinateSystemProvided )
+			if (!CustomCoordinateSystemProvided)
 			{
-				const int32 Num = GEditor->GetSelectedActors()->CountSelections<AActor>();
+				const int32 Num = GetSelectedActors()->CountSelections<AActor>();
 
 				// Coordinate system needs to come from the last actor selected
-				if( Num > 0 )
+				if (Num > 0)
 				{
-					matrix = FRotationMatrix( GEditor->GetSelectedActors()->GetBottom<AActor>()->GetActorRotation() );
+					Matrix = FRotationMatrix(GetSelectedActors()->GetBottom<AActor>()->GetActorRotation());
 				}
 			}
 
-			if(!matrix.Equals(FMatrix::Identity))
+			if (!Matrix.Equals(FMatrix::Identity))
 			{
-				matrix.RemoveScaling();
+				Matrix.RemoveScaling();
 			}
 		}
 		break;
@@ -1562,7 +1556,7 @@ FMatrix FEditorModeTools::GetCustomDrawingCoordinateSystem()
 			break;
 	}
 
-	return matrix;
+	return Matrix;
 }
 
 FMatrix FEditorModeTools::GetCustomInputCoordinateSystem()
@@ -1941,7 +1935,7 @@ void FEditorModeTools::CycleWidgetMode (void)
 		}
 	}
 
-	//only cycle when the mode is requesting the drawing of a widgeth
+	//only cycle when the mode is requesting the drawing of a widget
 	if( GetShowWidget() )
 	{
 		const int32 CurrentWk = GetWidgetMode();
