Editor: Split out FEdMode and FEditorModeTools into separate files

diff --git a/Engine/Source/Editor/UnrealEd/Private/EdMode.cpp b/Engine/Source/Editor/UnrealEd/Private/EdMode.cpp
new file mode 100644
--- /dev/null
+++ b/Engine/Source/Editor/UnrealEd/Private/EdMode.cpp
@@ -0,0 +1,1133 @@
+// Copyright 1998-2014 Epic Games, Inc. All Rights Reserved.
+
+#include "UnrealEd.h"
+#include "Engine/BookMark.h"
+#include "StaticMeshResources.h"
+#include "EditorSupportDelegates.h"
+#include "MouseDeltaTracker.h"
+#include "ScopedTransaction.h"
+#include "SurfaceIterators.h"
+#include "SoundDefinitions.h"
+#include "LevelEditor.h"
+#include "Toolkits/ToolkitManager.h"
+#include "EditorLevelUtils.h"
+#include "DynamicMeshBuilder.h"
+
+#include "ActorEditorUtils.h"
+#include "EditorStyle.h"
+#include "ComponentVisualizer.h"
+#include "SNotificationList.h"
+#include "NotificationManager.h"
+#include "Engine/Selection.h"
+#include "EngineUtils.h"
+#include "CanvasItem.h"
+#include "CanvasTypes.h"
+#include "Engine/Polys.h"
+#include "Engine/StaticMeshActor.h"
+#include "Engine/LevelStreaming.h"
+
+/** Hit proxy used for editable properties */
+struct HPropertyWidgetProxy : public HHitProxy
+{
+	DECLARE_HIT_PROXY();
+
+	/** Name of property this is the widget for */
+	FString	PropertyName;
+
+	/** If the property is an array property, the index into that array that this widget is for */
+	int32	PropertyIndex;
+
+	/** This property is a transform */
+	bool	bPropertyIsTransform;
+
+	HPropertyWidgetProxy(FString InPropertyName, int32 InPropertyIndex, bool bInPropertyIsTransform)
+		: HHitProxy(HPP_Foreground)
+		, PropertyName(InPropertyName)
+		, PropertyIndex(InPropertyIndex)
+		, bPropertyIsTransform(bInPropertyIsTransform)
+	{}
+
+	/** Show cursor as cross when over this handle */
+	virtual EMouseCursor::Type GetMouseCursor()
+	{
+		return EMouseCursor::Crosshairs;
+	}
+};
+
+IMPLEMENT_HIT_PROXY(HPropertyWidgetProxy, HHitProxy);
+
+
+namespace
+{
+	/**
+	 * Returns a reference to the named property value data in the given container.
+	 */
+	template<typename T>
+	T* GetPropertyValuePtrByName(const UStruct* InStruct, void* InContainer, FString PropertyName, int32 ArrayIndex)
+	{
+		T* ValuePtr = NULL;
+
+		// Extract the vector ptr recursively using the property name
+		int32 DelimPos = PropertyName.Find(TEXT("."));
+		if(DelimPos != INDEX_NONE)
+		{
+			// Parse the property name and (optional) array index
+			int32 SubArrayIndex = 0;
+			FString NameToken = PropertyName.Left(DelimPos);
+			int32 ArrayPos = NameToken.Find(TEXT("["));
+			if(ArrayPos != INDEX_NONE)
+			{
+				FString IndexToken = NameToken.RightChop(ArrayPos + 1).LeftChop(1);
+				SubArrayIndex = FCString::Atoi(*IndexToken);
+
+				NameToken = PropertyName.Left(ArrayPos);
+			}
+
+			// Obtain the property info from the given structure definition
+			UProperty* CurrentProp = FindField<UProperty>(InStruct, FName(*NameToken));
+
+			// Check first to see if this is a simple structure (i.e. not an array of structures)
+			UStructProperty* StructProp = Cast<UStructProperty>(CurrentProp);
+			if(StructProp != NULL)
+			{
+				// Recursively call back into this function with the structure property and container value
+				ValuePtr = GetPropertyValuePtrByName<T>(StructProp->Struct, StructProp->ContainerPtrToValuePtr<void>(InContainer), PropertyName.RightChop(DelimPos + 1), ArrayIndex);
+			}
+			else
+			{
+				// Check to see if this is an array
+				UArrayProperty* ArrayProp = Cast<UArrayProperty>(CurrentProp);
+				if(ArrayProp != NULL)
+				{
+					// It is an array, now check to see if this is an array of structures
+					StructProp = Cast<UStructProperty>(ArrayProp->Inner);
+					if(StructProp != NULL)
+					{
+						FScriptArrayHelper_InContainer ArrayHelper(ArrayProp, InContainer);
+						if(ArrayHelper.IsValidIndex(SubArrayIndex))
+						{
+							// Recursively call back into this function with the array element and container value
+							ValuePtr = GetPropertyValuePtrByName<T>(StructProp->Struct, ArrayHelper.GetRawPtr(SubArrayIndex), PropertyName.RightChop(DelimPos + 1), ArrayIndex);
+						}
+					}
+				}
+			}
+		}
+		else
+		{
+			UProperty* Prop = FindField<UProperty>(InStruct, FName(*PropertyName));
+			if(Prop != NULL)
+			{
+				if( UArrayProperty* ArrayProp = Cast<UArrayProperty>(Prop) )
+				{
+					check(ArrayIndex != INDEX_NONE);
+
+					// Property is an array property, so make sure we have a valid index specified
+					FScriptArrayHelper_InContainer ArrayHelper(ArrayProp, InContainer);
+					if( ArrayHelper.IsValidIndex(ArrayIndex) )
+					{
+						ValuePtr = (T*)ArrayHelper.GetRawPtr(ArrayIndex);
+					}
+				}
+				else
+				{
+					// Property is a vector property, so access directly
+					ValuePtr = Prop->ContainerPtrToValuePtr<T>(InContainer);
+				}
+			}
+		}
+
+		return ValuePtr;
+	}
+
+	/**
+	 * Returns the value of the property with the given name in the given Actor instance.
+	 */
+	template<typename T>
+	T GetPropertyValueByName(AActor* Actor, FString PropertyName, int32 PropertyIndex)
+	{
+		T Value;
+		T* ValuePtr = GetPropertyValuePtrByName<T>(Actor->GetClass(), Actor, PropertyName, PropertyIndex);
+		if(ValuePtr)
+		{
+			Value = *ValuePtr;
+		}
+		return Value;
+	}
+
+	/**
+	 * Sets the property with the given name in the given Actor instance to the given value.
+	 */
+	template<typename T>
+	void SetPropertyValueByName(AActor* Actor, FString PropertyName, int32 PropertyIndex, const T& InValue)
+	{
+		T* ValuePtr = GetPropertyValuePtrByName<T>(Actor->GetClass(), Actor, PropertyName, PropertyIndex);
+		if(ValuePtr)
+		{
+			*ValuePtr = InValue;
+		}
+	}
+}
+
+//////////////////////////////////
+// FEdMode
+
+const FName FEdMode::MD_MakeEditWidget(TEXT("MakeEditWidget"));
+const FName FEdMode::MD_ValidateWidgetUsing(TEXT("ValidateWidgetUsing"));
+
+FEdMode::FEdMode()
+	: bPendingDeletion( false )
+	, CurrentTool( NULL )
+	, EditedPropertyName(TEXT(""))
+	, EditedPropertyIndex( INDEX_NONE )
+	, bEditedPropertyIsTransform( false )
+{
+	bDrawKillZ = true;
+}
+
+FEdMode::~FEdMode()
+{
+}
+
+void FEdMode::OnModeUnregistered( FEditorModeID ModeID )
+{
+	if( ModeID == Info.ID )
+	{
+		// This should be synonymous with "delete this"
+		Owner->DestroyMode(ModeID);
+	}
+}
+
+bool FEdMode::MouseEnter( FEditorViewportClient* ViewportClient,FViewport* Viewport,int32 x, int32 y )
+{
+	if( GetCurrentTool() )
+	{
+		return GetCurrentTool()->MouseEnter( ViewportClient, Viewport, x, y );
+	}
+
+	return false;
+}
+
+bool FEdMode::MouseLeave( FEditorViewportClient* ViewportClient,FViewport* Viewport )
+{
+	if( GetCurrentTool() )
+	{
+		return GetCurrentTool()->MouseLeave( ViewportClient, Viewport );
+	}
+
+	return false;
+}
+
+bool FEdMode::MouseMove(FEditorViewportClient* ViewportClient,FViewport* Viewport,int32 x, int32 y)
+{
+	if( GetCurrentTool() )
+	{
+		return GetCurrentTool()->MouseMove( ViewportClient, Viewport, x, y );
+	}
+
+	return false;
+}
+
+bool FEdMode::ReceivedFocus(FEditorViewportClient* ViewportClient,FViewport* Viewport)
+{
+	if( GetCurrentTool() )
+	{
+		return GetCurrentTool()->ReceivedFocus( ViewportClient, Viewport );
+	}
+
+	return false;
+}
+
+bool FEdMode::LostFocus(FEditorViewportClient* ViewportClient,FViewport* Viewport)
+{
+	if( GetCurrentTool() )
+	{
+		return GetCurrentTool()->LostFocus( ViewportClient, Viewport );
+	}
+
+	return false;
+}
+
+bool FEdMode::CapturedMouseMove( FEditorViewportClient* InViewportClient, FViewport* InViewport, int32 InMouseX, int32 InMouseY )
+{
+	if( GetCurrentTool() )
+	{
+		return GetCurrentTool()->CapturedMouseMove( InViewportClient, InViewport, InMouseX, InMouseY );
+	}
+
+	return false;
+}
+
+bool FEdMode::InputKey(FEditorViewportClient* ViewportClient, FViewport* Viewport, FKey Key, EInputEvent Event)
+{
+	if( GetCurrentTool() && GetCurrentTool()->InputKey( ViewportClient, Viewport, Key, Event ) )
+	{
+		return true;
+	}
+	else
+	{
+		// Pass input up to selected actors if not in a tool mode
+		TArray<AActor*> SelectedActors;
+		Owner->GetSelectedActors()->GetSelectedObjects<AActor>(SelectedActors);
+
+		for( TArray<AActor*>::TIterator it(SelectedActors); it; ++it )
+		{
+			// Tell the object we've had a key press
+			(*it)->EditorKeyPressed(Key, Event);
+		}
+	}
+
+	return 0;
+}
+
+bool FEdMode::InputAxis(FEditorViewportClient* InViewportClient, FViewport* Viewport, int32 ControllerId, FKey Key, float Delta, float DeltaTime)
+{
+	FModeTool* Tool = GetCurrentTool();
+	if (Tool)
+	{
+		return Tool->InputAxis(InViewportClient, Viewport, ControllerId, Key, Delta, DeltaTime);
+	}
+
+	return false;
+}
+
+bool FEdMode::InputDelta(FEditorViewportClient* InViewportClient, FViewport* InViewport, FVector& InDrag, FRotator& InRot, FVector& InScale)
+{	
+	if(UsesPropertyWidgets())
+	{
+		AActor* SelectedActor = GetFirstSelectedActorInstance();
+		if(SelectedActor != NULL && InViewportClient->GetCurrentWidgetAxis() != EAxisList::None)
+		{
+			GEditor->NoteActorMovement();
+
+			if (EditedPropertyName != TEXT(""))
+			{
+				FTransform LocalTM = FTransform::Identity;
+
+				if(bEditedPropertyIsTransform)
+				{
+					LocalTM = GetPropertyValueByName<FTransform>(SelectedActor, EditedPropertyName, EditedPropertyIndex);
+				}
+				else
+				{					
+					FVector LocalPos = GetPropertyValueByName<FVector>(SelectedActor, EditedPropertyName, EditedPropertyIndex);
+					LocalTM = FTransform(LocalPos);
+				}
+
+				// Get actor transform (actor to world)
+				FTransform ActorTM = SelectedActor->ActorToWorld();
+				// Calculate world transform
+				FTransform WorldTM = LocalTM * ActorTM;
+				// Calc delta specified by drag
+				//FTransform DeltaTM(InRot.Quaternion(), InDrag);
+				// Apply delta in world space
+				WorldTM.SetTranslation(WorldTM.GetTranslation() + InDrag);
+				WorldTM.SetRotation(InRot.Quaternion() * WorldTM.GetRotation());
+				// Convert new world transform back into local space
+				LocalTM = WorldTM.GetRelativeTransform(ActorTM);
+				// Apply delta scale
+				LocalTM.SetScale3D(LocalTM.GetScale3D() + InScale);
+
+				SelectedActor->PreEditChange(NULL);
+
+				if(bEditedPropertyIsTransform)
+				{
+					SetPropertyValueByName<FTransform>(SelectedActor, EditedPropertyName, EditedPropertyIndex, LocalTM);
+				}
+				else
+				{
+					SetPropertyValueByName<FVector>(SelectedActor, EditedPropertyName, EditedPropertyIndex, LocalTM.GetLocation());
+				}
+
+				SelectedActor->PostEditChange();
+
+				return true;
+			}
+		}
+	}
+
+	if( GetCurrentTool() )
+	{
+		return GetCurrentTool()->InputDelta(InViewportClient,InViewport,InDrag,InRot,InScale);
+	}
+
+	return 0;
+}
+
+/**
+ * Lets each tool determine if it wants to use the editor widget or not.  If the tool doesn't want to use it, it will be
+ * fed raw mouse delta information (not snapped or altered in any way).
+ */
+
+bool FEdMode::UsesTransformWidget() const
+{
+	if( GetCurrentTool() )
+	{
+		return GetCurrentTool()->UseWidget();
+	}
+
+	return 1;
+}
+
+/**
+ * Lets each mode selectively exclude certain widget types.
+ */
+bool FEdMode::UsesTransformWidget(FWidget::EWidgetMode CheckMode) const
+{
+	if(UsesPropertyWidgets())
+	{
+		AActor* SelectedActor = GetFirstSelectedActorInstance();
+		if(SelectedActor != NULL)
+		{
+			// If editing a vector (not a transform)
+			if(EditedPropertyName != TEXT("") && !bEditedPropertyIsTransform)
+			{
+				return (CheckMode == FWidget::WM_Translate);
+			}
+		}
+	}
+
+	return true;
+}
+
+/**
+ * Allows each mode/tool to determine a good location for the widget to be drawn at.
+ */
+
+FVector FEdMode::GetWidgetLocation() const
+{
+	if(UsesPropertyWidgets())
+	{
+		AActor* SelectedActor = GetFirstSelectedActorInstance();
+		if(SelectedActor != NULL)
+		{
+			if(EditedPropertyName != TEXT(""))
+			{
+				FVector LocalPos = FVector::ZeroVector;
+
+				if(bEditedPropertyIsTransform)
+				{
+					FTransform LocalTM = GetPropertyValueByName<FTransform>(SelectedActor, EditedPropertyName, EditedPropertyIndex);
+					LocalPos = LocalTM.GetLocation();
+				}
+				else
+				{
+					LocalPos = GetPropertyValueByName<FVector>(SelectedActor, EditedPropertyName, EditedPropertyIndex);
+				}
+
+				FTransform ActorToWorld = SelectedActor->ActorToWorld();
+				FVector WorldPos = ActorToWorld.TransformPosition(LocalPos);
+				return WorldPos;
+			}
+		}
+	}
+
+	//UE_LOG(LogEditorModes, Log, TEXT("In FEdMode::GetWidgetLocation"));
+	return Owner->PivotLocation;
+}
+
+/**
+ * Lets the mode determine if it wants to draw the widget or not.
+ */
+
+bool FEdMode::ShouldDrawWidget() const
+{
+	return (Owner->GetSelectedActors()->GetTop<AActor>() != NULL);
+}
+
+/**
+ * Allows each mode to customize the axis pieces of the widget they want drawn.
+ *
+ * @param	InwidgetMode	The current widget mode
+ *
+ * @return	A bitfield comprised of AXIS_ values
+ */
+
+EAxisList::Type FEdMode::GetWidgetAxisToDraw( FWidget::EWidgetMode InWidgetMode ) const
+{
+	return EAxisList::All;
+}
+
+/**
+ * Lets each mode/tool handle box selection in its own way.
+ *
+ * @param	InBox	The selection box to use, in worldspace coordinates.
+ * @return		true if something was selected/deselected, false otherwise.
+ */
+bool FEdMode::BoxSelect( FBox& InBox, bool InSelect )
+{
+	bool bResult = false;
+	if( GetCurrentTool() )
+	{
+		bResult = GetCurrentTool()->BoxSelect( InBox, InSelect );
+	}
+	return bResult;
+}
+
+/**
+ * Lets each mode/tool handle frustum selection in its own way.
+ *
+ * @param	InFrustum	The selection frustum to use, in worldspace coordinates.
+ * @return	true if something was selected/deselected, false otherwise.
+ */
+bool FEdMode::FrustumSelect( const FConvexVolume& InFrustum, bool InSelect )
+{
+	bool bResult = false;
+	if( GetCurrentTool() )
+	{
+		bResult = GetCurrentTool()->FrustumSelect( InFrustum, InSelect );
+	}
+	return bResult;
+}
+
+void FEdMode::SelectNone()
+{
+	if( GetCurrentTool() )
+	{
+		GetCurrentTool()->SelectNone();
+	}
+}
+
+void FEdMode::Tick(FEditorViewportClient* ViewportClient,float DeltaTime)
+{
+	if( GetCurrentTool() )
+	{
+		GetCurrentTool()->Tick(ViewportClient,DeltaTime);
+	}
+}
+
+void FEdMode::ActorSelectionChangeNotify()
+{
+	EditedPropertyName = TEXT("");
+	EditedPropertyIndex = INDEX_NONE;
+	bEditedPropertyIsTransform = false;
+}
+
+bool FEdMode::HandleClick(FEditorViewportClient* InViewportClient, HHitProxy *HitProxy, const FViewportClick &Click)
+{
+	if(UsesPropertyWidgets() && HitProxy)
+	{
+		if( HitProxy->IsA(HPropertyWidgetProxy::StaticGetType()) )
+		{
+			HPropertyWidgetProxy* PropertyProxy = (HPropertyWidgetProxy*)HitProxy;
+			EditedPropertyName = PropertyProxy->PropertyName;
+			EditedPropertyIndex = PropertyProxy->PropertyIndex;
+			bEditedPropertyIsTransform = PropertyProxy->bPropertyIsTransform;
+			return true;
+		}
+		// Left clicking on an actor, stop editing a property
+		else if( HitProxy->IsA(HActor::StaticGetType()) )
+		{
+			EditedPropertyName = TEXT("");
+			EditedPropertyIndex = INDEX_NONE;
+			bEditedPropertyIsTransform = false;
+		}
+	}
+
+	return false;
+}
+
+void FEdMode::Enter()
+{
+	// Update components for selected actors, in case the mode we just exited
+	// was hijacking selection events selection and not updating components.
+	for ( FSelectionIterator It( *Owner->GetSelectedActors() ) ; It ; ++It )
+	{
+		AActor* SelectedActor = CastChecked<AActor>( *It );
+		SelectedActor->MarkComponentsRenderStateDirty();
+	}
+
+	bPendingDeletion = false;
+
+	FEditorDelegates::EditorModeEnter.Broadcast( this );
+	const bool bIsEnteringMode = true;
+	Owner->BroadcastEditorModeChanged( this, bIsEnteringMode );
+}
+
+void FEdMode::Exit()
+{
+	const bool bIsEnteringMode = false;
+	Owner->BroadcastEditorModeChanged( this, bIsEnteringMode );
+	FEditorDelegates::EditorModeExit.Broadcast( this );
+}
+
+void FEdMode::SetCurrentTool( EModeTools InID )
+{
+	CurrentTool = FindTool( InID );
+	check( CurrentTool );	// Tool not found!  This can't happen.
+
+	CurrentToolChanged();
+}
+
+void FEdMode::SetCurrentTool( FModeTool* InModeTool )
+{
+	CurrentTool = InModeTool;
+	check(CurrentTool);
+
+	CurrentToolChanged();
+}
+
+FModeTool* FEdMode::FindTool( EModeTools InID )
+{
+	for( int32 x = 0 ; x < Tools.Num() ; ++x )
+	{
+		if( Tools[x]->GetID() == InID )
+		{
+			return Tools[x];
+		}
+	}
+
+	UE_LOG(LogEditorModes, Fatal, TEXT("FEdMode::FindTool failed to find tool %d"), (int32)InID);
+	return NULL;
+}
+
+void FEdMode::Render(const FSceneView* View,FViewport* Viewport,FPrimitiveDrawInterface* PDI)
+{
+	if( GEditor->bShowBrushMarkerPolys )
+	{
+		// Draw translucent polygons on brushes and volumes
+
+		for( TActorIterator<ABrush> It(GetWorld()); It; ++ It )
+		{
+			ABrush* Brush = *It;
+
+			// Brush->Brush is checked to safe from brushes that were created without having their brush members attached.
+			if( Brush->Brush && (FActorEditorUtils::IsABuilderBrush(Brush) || Brush->IsVolumeBrush()) && Owner->GetSelectedActors()->IsSelected(Brush) )
+			{
+				// Build a mesh by basically drawing the triangles of each 
+				FDynamicMeshBuilder MeshBuilder;
+				int32 VertexOffset = 0;
+
+				for( int32 PolyIdx = 0 ; PolyIdx < Brush->Brush->Polys->Element.Num() ; ++PolyIdx )
+				{
+					const FPoly* Poly = &Brush->Brush->Polys->Element[PolyIdx];
+
+					if( Poly->Vertices.Num() > 2 )
+					{
+						const FVector Vertex0 = Poly->Vertices[0];
+						FVector Vertex1 = Poly->Vertices[1];
+
+						MeshBuilder.AddVertex(Vertex0, FVector2D::ZeroVector, FVector(1,0,0), FVector(0,1,0), FVector(0,0,1), FColor(255,255,255));
+						MeshBuilder.AddVertex(Vertex1, FVector2D::ZeroVector, FVector(1,0,0), FVector(0,1,0), FVector(0,0,1), FColor(255,255,255));
+
+						for( int32 VertexIdx = 2 ; VertexIdx < Poly->Vertices.Num() ; ++VertexIdx )
+						{
+							const FVector Vertex2 = Poly->Vertices[VertexIdx];
+							MeshBuilder.AddVertex(Vertex2, FVector2D::ZeroVector, FVector(1,0,0), FVector(0,1,0), FVector(0,0,1), FColor(255,255,255));
+							MeshBuilder.AddTriangle(VertexOffset,VertexOffset + VertexIdx,VertexOffset+VertexIdx-1);
+							Vertex1 = Vertex2;
+						}
+
+						// Increment the vertex offset so the next polygon uses the correct vertex indices.
+						VertexOffset += Poly->Vertices.Num();
+					}
+				}
+
+				// Allocate the material proxy and register it so it can be deleted properly once the rendering is done with it.
+				FDynamicColoredMaterialRenderProxy* MaterialProxy = new FDynamicColoredMaterialRenderProxy(GEngine->EditorBrushMaterial->GetRenderProxy(false),Brush->GetWireColor());
+				PDI->RegisterDynamicResource( MaterialProxy );
+
+				// Flush the mesh triangles.
+				MeshBuilder.Draw(PDI, Brush->ActorToWorld().ToMatrixWithScale(), MaterialProxy, SDPG_World, 0.f);
+			}
+		}
+	}
+
+	const bool bIsInGameView = !Viewport->GetClient() || Viewport->GetClient()->IsInGameView();
+	if (Owner->ShouldDrawBrushVertices() && !bIsInGameView)
+	{
+		UTexture2D* VertexTexture = GetVertexTexture();
+		const float TextureSizeX = VertexTexture->GetSizeX() * 0.170f;
+		const float TextureSizeY = VertexTexture->GetSizeY() * 0.170f;
+
+		for (FSelectionIterator It(*Owner->GetSelectedActors()); It; ++It)
+		{
+			AActor* SelectedActor = static_cast<AActor*>(*It);
+			checkSlow(SelectedActor->IsA(AActor::StaticClass()));
+
+			ABrush* Brush = Cast< ABrush >(SelectedActor);
+			if (Brush && Brush->Brush && !FActorEditorUtils::IsABuilderBrush(Brush))
+			{
+				for (int32 p = 0; p < Brush->Brush->Polys->Element.Num(); ++p)
+				{
+					FTransform BrushTransform = Brush->ActorToWorld();
+
+					FPoly* poly = &Brush->Brush->Polys->Element[p];
+					for (int32 VertexIndex = 0; VertexIndex < poly->Vertices.Num(); ++VertexIndex)
+					{
+						const FVector& PolyVertex = poly->Vertices[VertexIndex];
+						const FVector WorldLocation = BrushTransform.TransformPosition(PolyVertex);
+						
+						const float Scale = View->WorldToScreen( WorldLocation ).W * ( 4.0f / View->ViewRect.Width() / View->ViewMatrices.ProjMatrix.M[0][0] );
+
+						const FColor Color(Brush->GetWireColor());
+						PDI->SetHitProxy(new HBSPBrushVert(Brush, &poly->Vertices[VertexIndex]));
+
+						PDI->DrawSprite(WorldLocation, TextureSizeX * Scale, TextureSizeY * Scale, VertexTexture->Resource, Color, SDPG_World, 0.0f, 0.0f, 0.0f, 0.0f, SE_BLEND_Masked );
+
+						PDI->SetHitProxy(NULL);
+			
+					}
+				}
+			}
+		}
+	}
+
+	// Let the current mode tool render if it wants to
+	FModeTool* tool = GetCurrentTool();
+	if( tool )
+	{
+		tool->Render( View, Viewport, PDI );
+	}
+
+	AGroupActor::DrawBracketsForGroups(PDI, Viewport);
+
+	if(UsesPropertyWidgets())
+	{
+		bool bHitTesting = PDI->IsHitTesting();
+		AActor* SelectedActor = GetFirstSelectedActorInstance();
+		if (SelectedActor != NULL)
+		{
+			UClass* Class = SelectedActor->GetClass();
+			TArray<FPropertyWidgetInfo> WidgetInfos;
+			GetPropertyWidgetInfos(Class, SelectedActor, WidgetInfos);
+			FEditorScriptExecutionGuard ScriptGuard;
+			for(int32 i=0; i<WidgetInfos.Num(); i++)
+			{
+				FString WidgetName = WidgetInfos[i].PropertyName;
+				FName WidgetValidator = WidgetInfos[i].PropertyValidationName;
+				int32 WidgetIndex = WidgetInfos[i].PropertyIndex;
+				bool bIsTransform = WidgetInfos[i].bIsTransform;
+
+				bool bSelected = (WidgetName == EditedPropertyName) && (WidgetIndex == EditedPropertyIndex);
+				FColor WidgetColor = bSelected ? FColor(255,255,255) : FColor(128,128,255);
+
+				FVector LocalPos = FVector::ZeroVector;
+				if(bIsTransform)
+				{
+					FTransform LocalTM = GetPropertyValueByName<FTransform>(SelectedActor, WidgetName, WidgetIndex);
+					LocalPos = LocalTM.GetLocation();
+				}
+				else
+				{
+					LocalPos = GetPropertyValueByName<FVector>(SelectedActor, WidgetName, WidgetIndex);
+				}
+
+				FTransform ActorToWorld = SelectedActor->ActorToWorld();
+				FVector WorldPos = ActorToWorld.TransformPosition(LocalPos);
+
+				UFunction* ValidateFunc= NULL;
+				if(WidgetValidator != NAME_None && 
+					(ValidateFunc = SelectedActor->FindFunction(WidgetValidator)) != NULL)
+				{
+					FString ReturnText;
+					SelectedActor->ProcessEvent(ValidateFunc, &ReturnText);
+
+					//if we have a negative result, the widget color is red.
+					WidgetColor = ReturnText.IsEmpty() ? WidgetColor : FColor::Red;
+				}
+
+				FTranslationMatrix WidgetTM(WorldPos);
+
+				const float WidgetSize = 0.035f;
+				const float ZoomFactor = FMath::Min<float>(View->ViewMatrices.ProjMatrix.M[0][0], View->ViewMatrices.ProjMatrix.M[1][1]);
+				const float WidgetRadius = View->Project(WorldPos).W * (WidgetSize / ZoomFactor);
+
+				if(bHitTesting) PDI->SetHitProxy( new HPropertyWidgetProxy(WidgetName, WidgetIndex, bIsTransform) );
+				DrawWireDiamond(PDI, WidgetTM, WidgetRadius, WidgetColor, SDPG_Foreground );
+				if(bHitTesting) PDI->SetHitProxy( NULL );
+			}
+		}
+	}
+}
+
+void FEdMode::DrawHUD(FEditorViewportClient* ViewportClient,FViewport* Viewport,const FSceneView* View,FCanvas* Canvas)
+{
+	// Render the drag tool.
+	ViewportClient->RenderDragTool( View, Canvas );
+
+	// Let the current mode tool draw a HUD if it wants to
+	FModeTool* tool = GetCurrentTool();
+	if( tool )
+	{
+		tool->DrawHUD( ViewportClient, Viewport, View, Canvas );
+	}
+
+	if (ViewportClient->IsPerspective() && GetDefault<ULevelEditorViewportSettings>()->bHighlightWithBrackets)
+	{
+		DrawBrackets( ViewportClient, Viewport, View, Canvas );
+	}
+
+	// If this viewport doesn't show mode widgets or the mode itself doesn't want them, leave.
+	if( !(ViewportClient->EngineShowFlags.ModeWidgets) || !ShowModeWidgets() )
+	{
+		return;
+	}
+
+	// Clear Hit proxies
+	const bool bIsHitTesting = Canvas->IsHitTesting();
+	if ( !bIsHitTesting )
+	{
+		Canvas->SetHitProxy(NULL);
+	}
+
+	// Draw vertices for selected BSP brushes and static meshes if the large vertices show flag is set.
+	if ( !ViewportClient->bDrawVertices )
+	{
+		return;
+	}
+
+	const bool bLargeVertices		= View->Family->EngineShowFlags.LargeVertices;
+	const bool bShowBrushes			= View->Family->EngineShowFlags.Brushes;
+	const bool bShowBSP				= View->Family->EngineShowFlags.BSP;
+	const bool bShowBuilderBrush	= View->Family->EngineShowFlags.BuilderBrush != 0;
+
+	UTexture2D* VertexTexture = GetVertexTexture();
+	const float TextureSizeX		= VertexTexture->GetSizeX() * ( bLargeVertices ? 1.0f : 0.5f );
+	const float TextureSizeY		= VertexTexture->GetSizeY() * ( bLargeVertices ? 1.0f : 0.5f );
+
+	// Temporaries.
+	TArray<FVector> Vertices;
+
+	for ( FSelectionIterator It( *Owner->GetSelectedActors() ) ; It ; ++It )
+	{
+		AActor* SelectedActor = static_cast<AActor*>( *It );
+		checkSlow( SelectedActor->IsA(AActor::StaticClass()) );
+
+		if( bLargeVertices )
+		{
+			FCanvasItemTestbed::bTestState = !FCanvasItemTestbed::bTestState;
+
+			// Static mesh vertices
+			AStaticMeshActor* Actor = Cast<AStaticMeshActor>( SelectedActor );
+			if( Actor && Actor->GetStaticMeshComponent() && Actor->GetStaticMeshComponent()->StaticMesh
+				&& Actor->GetStaticMeshComponent()->StaticMesh->RenderData )
+			{
+				FTransform ActorToWorld = Actor->ActorToWorld();
+				Vertices.Empty();
+				const FPositionVertexBuffer& VertexBuffer = Actor->GetStaticMeshComponent()->StaticMesh->RenderData->LODResources[0].PositionVertexBuffer;
+				for( uint32 i = 0 ; i < VertexBuffer.GetNumVertices() ; i++ )
+				{
+					Vertices.AddUnique( ActorToWorld.TransformPosition( VertexBuffer.VertexPosition(i) ) );
+				}
+
+				FCanvasTileItem TileItem( FVector2D( 0.0f, 0.0f ), FVector2D( 0.0f, 0.0f ), FLinearColor::White );
+				TileItem.BlendMode = SE_BLEND_Translucent;
+				for( int32 VertexIndex = 0 ; VertexIndex < Vertices.Num() ; ++VertexIndex )
+				{				
+					const FVector& Vertex = Vertices[VertexIndex];
+					FVector2D PixelLocation;
+					if(View->ScreenToPixel(View->WorldToScreen(Vertex),PixelLocation))
+					{
+						const bool bOutside =
+							PixelLocation.X < 0.0f || PixelLocation.X > View->ViewRect.Width() ||
+							PixelLocation.Y < 0.0f || PixelLocation.Y > View->ViewRect.Height();
+						if( !bOutside )
+						{
+							const float X = PixelLocation.X - (TextureSizeX/2);
+							const float Y = PixelLocation.Y - (TextureSizeY/2);
+							if( bIsHitTesting ) 
+							{
+								Canvas->SetHitProxy( new HStaticMeshVert(Actor,Vertex) );
+							}
+							TileItem.Texture = VertexTexture->Resource;
+							
+							TileItem.Size = FVector2D( TextureSizeX, TextureSizeY );
+							Canvas->DrawItem( TileItem, FVector2D( X, Y ) );							
+							if( bIsHitTesting )
+							{
+								Canvas->SetHitProxy( NULL );
+							}
+						}
+					}
+				}
+			}
+		}
+	}
+
+	if(UsesPropertyWidgets())
+	{
+		AActor* SelectedActor = GetFirstSelectedActorInstance();
+		if (SelectedActor != NULL)
+		{
+			FEditorScriptExecutionGuard ScriptGuard;
+
+			const int32 HalfX = 0.5f * Viewport->GetSizeXY().X;
+			const int32 HalfY = 0.5f * Viewport->GetSizeXY().Y;
+
+			UClass* Class = SelectedActor->GetClass();		
+			TArray<FPropertyWidgetInfo> WidgetInfos;
+			GetPropertyWidgetInfos(Class, SelectedActor, WidgetInfos);
+			for(int32 i=0; i<WidgetInfos.Num(); i++)
+			{
+				FString WidgetName = WidgetInfos[i].PropertyName;
+				FName WidgetValidator = WidgetInfos[i].PropertyValidationName;
+				int32 WidgetIndex = WidgetInfos[i].PropertyIndex;
+				bool bIsTransform = WidgetInfos[i].bIsTransform;
+
+				FVector LocalPos = FVector::ZeroVector;
+				if(bIsTransform)
+				{
+					FTransform LocalTM = GetPropertyValueByName<FTransform>(SelectedActor, WidgetName, WidgetIndex);
+					LocalPos = LocalTM.GetLocation();
+				}
+				else
+				{
+					LocalPos = GetPropertyValueByName<FVector>(SelectedActor, WidgetName, WidgetIndex);
+				}
+
+				FTransform ActorToWorld = SelectedActor->ActorToWorld();
+				FVector WorldPos = ActorToWorld.TransformPosition(LocalPos);
+
+				UFunction* ValidateFunc = NULL;
+				FString FinalString;
+				if(WidgetValidator != NAME_None && 
+					(ValidateFunc = SelectedActor->FindFunction(WidgetValidator)) != NULL)
+				{
+					SelectedActor->ProcessEvent(ValidateFunc, &FinalString);
+				}
+
+				const FPlane Proj = View->Project( WorldPos );
+
+				//do some string fixing
+				const uint32 VectorIndex = WidgetInfos[i].PropertyIndex;
+				const FString WidgetDisplayName = WidgetInfos[i].DisplayName + ((VectorIndex != INDEX_NONE) ? FString::Printf(TEXT("[%d]"), VectorIndex) : TEXT(""));
+				FinalString = FinalString.IsEmpty() ? WidgetDisplayName : FinalString;
+
+				if(Proj.W > 0.f)
+				{
+					const int32 XPos = HalfX + ( HalfX * Proj.X );
+					const int32 YPos = HalfY + ( HalfY * (Proj.Y * -1.f) );
+					FCanvasTextItem TextItem( FVector2D( XPos + 5, YPos), FText::FromString( FinalString ), GEngine->GetSmallFont(), FLinearColor::White );
+					Canvas->DrawItem( TextItem );
+				}
+			}
+		}
+	}
+}
+
+// Draw brackets around all selected objects
+void FEdMode::DrawBrackets( FEditorViewportClient* ViewportClient, FViewport* Viewport, const FSceneView* View, FCanvas* Canvas )
+{
+	USelection& SelectedActors = *Owner->GetSelectedActors();
+	for( int32 CurSelectedActorIndex = 0; CurSelectedActorIndex < SelectedActors.Num(); ++CurSelectedActorIndex )
+	{
+		AActor* SelectedActor = Cast<AActor>( SelectedActors.GetSelectedObject(CurSelectedActorIndex ) );
+		if( SelectedActor != NULL )
+		{
+			// Draw a bracket for selected "paintable" static mesh actors
+			const bool bIsValidActor = ( Cast< AStaticMeshActor >( SelectedActor ) != NULL );
+
+			const FLinearColor SelectedActorBoxColor( 0.6f, 0.6f, 1.0f );
+			const bool bDrawBracket = bIsValidActor;
+			ViewportClient->DrawActorScreenSpaceBoundingBox( Canvas, View, Viewport, SelectedActor, SelectedActorBoxColor, bDrawBracket );
+		}
+	}
+}
+
+bool FEdMode::UsesToolkits() const
+{
+	return false;
+}
+
+bool FEdMode::StartTracking(FEditorViewportClient* InViewportClient, FViewport* InViewport)
+{
+	bool bResult = false;
+	if( GetCurrentTool() )
+	{
+		bResult = GetCurrentTool()->StartModify();
+	}
+	return bResult;
+}
+
+bool FEdMode::EndTracking(FEditorViewportClient* InViewportClient, FViewport* InViewport)
+{
+	bool bResult = false;
+	if( GetCurrentTool() )
+	{
+		bResult = GetCurrentTool()->EndModify();
+	}
+	return bResult;
+}
+
+FVector FEdMode::GetWidgetNormalFromCurrentAxis( void* InData )
+{
+	// Figure out the proper coordinate system.
+
+	FMatrix matrix = FMatrix::Identity;
+	if( Owner->GetCoordSystem() == COORD_Local )
+	{
+		GetCustomDrawingCoordinateSystem( matrix, InData );
+	}
+
+	// Get a base normal from the current axis.
+
+	FVector BaseNormal(1,0,0);		// Default to X axis
+	switch( CurrentWidgetAxis )
+	{
+		case EAxisList::Y:	BaseNormal = FVector(0,1,0);	break;
+		case EAxisList::Z:	BaseNormal = FVector(0,0,1);	break;
+		case EAxisList::XY:	BaseNormal = FVector(1,1,0);	break;
+		case EAxisList::XZ:	BaseNormal = FVector(1,0,1);	break;
+		case EAxisList::YZ:	BaseNormal = FVector(0,1,1);	break;
+		case EAxisList::XYZ:	BaseNormal = FVector(1,1,1);	break;
+	}
+
+	// Transform the base normal into the proper coordinate space.
+	return matrix.TransformPosition( BaseNormal );
+}
+
+AActor* FEdMode::GetFirstSelectedActorInstance() const
+{
+	return Owner->GetSelectedActors()->GetTop<AActor>();
+}
+
+bool FEdMode::CanCreateWidgetForStructure(const UStruct* InPropStruct)
+{
+	return InPropStruct && (InPropStruct->GetFName() == NAME_Vector || InPropStruct->GetFName() == NAME_Transform);
+}
+
+bool FEdMode::CanCreateWidgetForProperty(UProperty* InProp)
+{
+	UStructProperty* TestProperty = Cast<UStructProperty>(InProp);
+	if( !TestProperty )
+	{
+		UArrayProperty* ArrayProperty = Cast<UArrayProperty>(InProp);
+		if( ArrayProperty )
+		{
+			TestProperty = Cast<UStructProperty>(ArrayProperty->Inner);
+		}
+	}
+	return (TestProperty != NULL) && CanCreateWidgetForStructure(TestProperty->Struct);
+}
+
+bool FEdMode::ShouldCreateWidgetForProperty(UProperty* InProp)
+{
+	return CanCreateWidgetForProperty(InProp) && InProp->HasMetaData(MD_MakeEditWidget);
+}
+
+static bool IsTransformProperty(UProperty* InProp)
+{
+	UStructProperty* StructProp = Cast<UStructProperty>(InProp);
+	return (StructProp != NULL && StructProp->Struct->GetFName() == NAME_Transform);
+
+}
+
+void FEdMode::GetPropertyWidgetInfos(const UStruct* InStruct, const void* InContainer, TArray<FPropertyWidgetInfo>& OutInfos, FString PropertyNamePrefix, FString DisplayNamePrefix) const
+{
+	if(PropertyNamePrefix.Len() == 0)
+	{
+		OutInfos.Empty();
+	}
+
+	for (TFieldIterator<UProperty> PropertyIt(InStruct, EFieldIteratorFlags::IncludeSuper); PropertyIt; ++PropertyIt)
+	{
+		UProperty* CurrentProp = *PropertyIt;
+		check(CurrentProp);
+		FString DisplayName = CurrentProp->GetMetaData(TEXT("DisplayName"));
+		if (!PropertyNamePrefix.IsEmpty() && DisplayNamePrefix.IsEmpty()) // Display name is already invalid
+		{
+			DisplayName.Empty();
+		}
+		if (!DisplayName.IsEmpty()) //Display name cannot be only the prefix.
+		{
+			DisplayName = DisplayNamePrefix + DisplayName;
+		}
+
+		if(	ShouldCreateWidgetForProperty(CurrentProp) )
+		{
+			if( UArrayProperty* ArrayProp = Cast<UArrayProperty>(CurrentProp) )
+			{
+				check(InContainer != NULL);
+
+				FScriptArrayHelper_InContainer ArrayHelper(ArrayProp, InContainer);
+
+				// See how many widgets we need to make for the array property
+				uint32 ArrayDim = ArrayHelper.Num();
+				for( uint32 i = 0; i < ArrayDim; i++ )
+				{
+					//create a new widget info struct
+					FPropertyWidgetInfo WidgetInfo;
+
+					//fill it in with the struct name
+					WidgetInfo.PropertyName = PropertyNamePrefix + CurrentProp->GetFName().ToString();
+					WidgetInfo.DisplayName = DisplayName.IsEmpty() ? WidgetInfo.PropertyName : DisplayName;
+
+					//And see if we have any meta data that matches the MD_ValidateWidgetUsing name
+					WidgetInfo.PropertyValidationName = FName(*CurrentProp->GetMetaData(MD_ValidateWidgetUsing));
+
+					WidgetInfo.PropertyIndex = i;
+
+					// See if its a transform
+					WidgetInfo.bIsTransform = IsTransformProperty(ArrayProp->Inner);
+
+					//Add it to our out array
+					OutInfos.Add(WidgetInfo);
+				}
+			}
+			else
+			{
+
+				//create a new widget info struct
+				FPropertyWidgetInfo WidgetInfo;
+
+				//fill it in with the struct name
+				WidgetInfo.PropertyName = PropertyNamePrefix + CurrentProp->GetFName().ToString();
+				WidgetInfo.DisplayName = DisplayName.IsEmpty() ? WidgetInfo.PropertyName : DisplayName;
+
+				//And see if we have any meta data that matches the MD_ValidateWidgetUsing name
+				WidgetInfo.PropertyValidationName = FName(*CurrentProp->GetMetaData(MD_ValidateWidgetUsing));
+
+				// See if its a transform
+				WidgetInfo.bIsTransform = IsTransformProperty(CurrentProp);
+
+				//Add it to our out array
+				OutInfos.Add(WidgetInfo);
+			}
+
+		}
+		else
+		{
+			UStructProperty* StructProp = Cast<UStructProperty>(CurrentProp);
+			if(StructProp != NULL)
+			{
+				// Recursively traverse into structures, looking for additional vector properties to expose
+				GetPropertyWidgetInfos(StructProp->Struct
+					, StructProp->ContainerPtrToValuePtr<void>(InContainer)
+					, OutInfos
+					, PropertyNamePrefix + StructProp->GetFName().ToString() + TEXT(".")
+					, !DisplayName.IsEmpty() ? (DisplayName + TEXT(".")) : FString());
+			}
+			else
+			{
+				// Recursively traverse into arrays of structures, looking for additional vector properties to expose
+				UArrayProperty* ArrayProp = Cast<UArrayProperty>(CurrentProp);
+				if(ArrayProp != NULL)
+				{
+					StructProp = Cast<UStructProperty>(ArrayProp->Inner);
+					if(StructProp != NULL)
+					{
+						FScriptArrayHelper_InContainer ArrayHelper(ArrayProp, InContainer);
+						for(int32 ArrayIndex = 0; ArrayIndex < ArrayHelper.Num(); ++ArrayIndex)
+						{
+							if(ArrayHelper.IsValidIndex(ArrayIndex))
+							{
+								const FString ArrayPostfix = FString::Printf(TEXT("[%d]"), ArrayIndex) + TEXT(".");
+								GetPropertyWidgetInfos(StructProp->Struct
+									, ArrayHelper.GetRawPtr(ArrayIndex)
+									, OutInfos
+									, PropertyNamePrefix + ArrayProp->GetFName().ToString() + ArrayPostfix
+									, !DisplayName.IsEmpty() ? (DisplayName + ArrayPostfix) : FString());
+							}
+						}
+					}
+				}
+			}
+		}
+	}
+}
+
+bool FEdMode::IsSnapRotationEnabled()
+{
+	return GetDefault<ULevelEditorViewportSettings>()->RotGridEnabled;
+}
diff --git a/Engine/Source/Editor/UnrealEd/Private/EditorModeManager.cpp b/Engine/Source/Editor/UnrealEd/Private/EditorModeManager.cpp
new file mode 100644
--- /dev/null
+++ b/Engine/Source/Editor/UnrealEd/Private/EditorModeManager.cpp
@@ -0,0 +1,1102 @@
+// Copyright 1998-2014 Epic Games, Inc. All Rights Reserved.
+
+#include "UnrealEd.h"
+#include "Engine/BookMark.h"
+#include "StaticMeshResources.h"
+#include "EditorSupportDelegates.h"
+#include "MouseDeltaTracker.h"
+#include "ScopedTransaction.h"
+#include "SurfaceIterators.h"
+#include "SoundDefinitions.h"
+#include "LevelEditor.h"
+#include "Toolkits/ToolkitManager.h"
+#include "EditorLevelUtils.h"
+#include "DynamicMeshBuilder.h"
+
+#include "ActorEditorUtils.h"
+#include "EditorStyle.h"
+#include "ComponentVisualizer.h"
+#include "SNotificationList.h"
+#include "NotificationManager.h"
+#include "Engine/Selection.h"
+#include "EngineUtils.h"
+#include "CanvasItem.h"
+#include "CanvasTypes.h"
+#include "Engine/Polys.h"
+#include "Engine/StaticMeshActor.h"
+#include "Engine/LevelStreaming.h"
+
+/*------------------------------------------------------------------------------
+	FEditorModeTools.
+
+	The master class that handles tracking of the current mode.
+------------------------------------------------------------------------------*/
+
+FEditorModeTools::FEditorModeTools()
+	:	PivotShown( 0 )
+	,	Snapping( 0 )
+	,	SnappedActor( 0 )
+	,	TranslateRotateXAxisAngle(0)
+	,	DefaultID(FBuiltinEditorModes::EM_Default)
+	,	WidgetMode( FWidget::WM_Translate )
+	,	OverrideWidgetMode( FWidget::WM_None )
+	,	bShowWidget( 1 )
+	,	bHideViewportUI(false)
+	,	CoordSystem(COORD_World)
+	,	bIsTracking(false)
+{
+	// Load the last used settings
+	LoadConfig();
+
+	// Register our callback for actor selection changes
+	USelection::SelectNoneEvent.AddRaw(this, &FEditorModeTools::OnEditorSelectNone);
+	USelection::SelectionChangedEvent.AddRaw(this, &FEditorModeTools::OnEditorSelectionChanged);
+	USelection::SelectObjectEvent.AddRaw(this, &FEditorModeTools::OnEditorSelectionChanged);
+}
+
+FEditorModeTools::~FEditorModeTools()
+{
+	// Should we call Exit on any modes that are still active, or is it too late?
+	USelection::SelectionChangedEvent.RemoveAll(this);
+	USelection::SelectNoneEvent.RemoveAll(this);
+	USelection::SelectObjectEvent.RemoveAll(this);
+}
+
+/**
+ * Loads the state that was saved in the INI file
+ */
+void FEditorModeTools::LoadConfig(void)
+{
+	GConfig->GetBool(TEXT("FEditorModeTools"),TEXT("ShowWidget"),bShowWidget,
+		GEditorUserSettingsIni);
+
+	const bool bGetRawValue = true;
+	int32 Bogus = (int32)GetCoordSystem(bGetRawValue);
+	GConfig->GetInt(TEXT("FEditorModeTools"),TEXT("CoordSystem"),Bogus,
+		GEditorUserSettingsIni);
+	SetCoordSystem((ECoordSystem)Bogus);
+
+	
+	GConfig->GetBool( TEXT("UnEdViewport"), TEXT("InterpEdPanInvert"), bInterpPanInverted, GEditorUserSettingsIni );
+	LoadWidgetSettings();
+}
+
+/**
+ * Saves the current state to the INI file
+ */
+void FEditorModeTools::SaveConfig(void)
+{
+	GConfig->SetBool(TEXT("FEditorModeTools"),TEXT("ShowWidget"),bShowWidget,
+		GEditorUserSettingsIni);
+
+	const bool bGetRawValue = true;
+	GConfig->SetInt(TEXT("FEditorModeTools"),TEXT("CoordSystem"),(int32)GetCoordSystem(bGetRawValue),
+		GEditorUserSettingsIni);
+
+	GConfig->SetBool( TEXT("UnEdViewport"), TEXT("InterpEdPanInvert"), bInterpPanInverted, GEditorUserSettingsIni );
+
+	SaveWidgetSettings();
+}
+
+TSharedPtr<class IToolkitHost> FEditorModeTools::GetToolkitHost() const
+{
+	TSharedPtr<class IToolkitHost> Result = ToolkitHost.Pin();
+	check(ToolkitHost.IsValid());
+	return Result;
+}
+
+void FEditorModeTools::SetToolkitHost(TSharedRef<class IToolkitHost> InHost)
+{
+	checkf(!ToolkitHost.IsValid(), TEXT("SetToolkitHost can only be called once"));
+	ToolkitHost = InHost;
+}
+
+class USelection* FEditorModeTools::GetSelectedActors() const
+{
+	return GEditor->GetSelectedActors();
+}
+
+class USelection* FEditorModeTools::GetSelectedObjects() const
+{
+	return GEditor->GetSelectedObjects();
+}
+
+void FEditorModeTools::OnEditorSelectionChanged(UObject* NewSelection)
+{
+	// If selecting an actor, move the pivot location.
+	AActor* Actor = Cast<AActor>(NewSelection);
+	if (Actor != NULL)
+	{
+		//@fixme - why isn't this using UObject::IsSelected()?
+		if ( GEditor->GetSelectedActors()->IsSelected( Actor ) )
+		{
+			SetPivotLocation( Actor->GetActorLocation(), false );
+
+			// If this actor wasn't part of the original selection set during pie/sie, clear it now
+			if ( GEditor->ActorsThatWereSelected.Num() > 0 )
+			{
+				AActor* EditorActor = EditorUtilities::GetEditorWorldCounterpartActor( Actor );
+				if ( !EditorActor || !GEditor->ActorsThatWereSelected.Contains(EditorActor) )
+				{
+					GEditor->ActorsThatWereSelected.Empty();
+				}
+			}
+		}
+		else if ( GEditor->ActorsThatWereSelected.Num() > 0 )
+		{
+			// Clear the selection set
+			GEditor->ActorsThatWereSelected.Empty();
+		}
+	}
+
+	for (const auto& Pair : FEditorModeRegistry::Get().GetFactoryMap())
+	{
+		Pair.Value->OnSelectionChanged(*this, NewSelection);
+	}
+}
+
+void FEditorModeTools::OnEditorSelectNone()
+{
+	GEditor->SelectNone( false, true );
+	GEditor->ActorsThatWereSelected.Empty();
+}
+
+/** 
+ * Sets the pivot locations
+ * 
+ * @param Location 		The location to set
+ * @param bIncGridBase	Whether or not to also set the GridBase
+ */
+void FEditorModeTools::SetPivotLocation( const FVector& Location, const bool bIncGridBase )
+{
+	CachedLocation = PivotLocation = SnappedLocation = Location;
+	if ( bIncGridBase )
+	{
+		GridBase = Location;
+	}
+}
+
+ECoordSystem FEditorModeTools::GetCoordSystem(bool bGetRawValue)
+{
+	if (!bGetRawValue && GetWidgetMode() == FWidget::WM_Scale )
+	{
+		return COORD_Local;
+	}
+	else
+	{
+		return CoordSystem;
+	}
+}
+
+void FEditorModeTools::SetCoordSystem(ECoordSystem NewCoordSystem)
+{
+	CoordSystem = NewCoordSystem;
+}
+
+void FEditorModeTools::SetDefaultMode ( FEditorModeID InDefaultID )
+{
+	DefaultID = InDefaultID;
+}
+
+void FEditorModeTools::ActivateDefaultMode()
+{
+	ActivateMode( DefaultID );
+
+	check( IsModeActive( DefaultID ) );
+}
+
+void FEditorModeTools::DeactivateModeAtIndex(int32 InIndex)
+{
+	check( InIndex >= 0 && InIndex < Modes.Num() );
+
+	auto& Mode = Modes[InIndex];
+
+	Mode->Exit();
+	RecycledModes.Add( Mode->GetID(), Mode );
+	Modes.RemoveAt( InIndex );
+}
+
+/**
+ * Deactivates an editor mode. 
+ * 
+ * @param InID		The ID of the editor mode to deactivate.
+ */
+void FEditorModeTools::DeactivateMode( FEditorModeID InID )
+{
+	// Find the mode from the ID and exit it.
+	for( int32 Index = Modes.Num() - 1; Index >= 0; --Index )
+	{
+		auto& Mode = Modes[Index];
+		if( Mode->GetID() == InID )
+		{
+			DeactivateModeAtIndex(Index);
+			break;
+		}
+	}
+
+	if( Modes.Num() == 0 )
+	{
+		// Ensure the default mode is active if there are no active modes.
+		ActivateDefaultMode();
+	}
+}
+
+void FEditorModeTools::DeactivateAllModes()
+{
+	for( int32 Index = Modes.Num() - 1; Index >= 0; --Index )
+	{
+		DeactivateModeAtIndex(Index);
+	}
+}
+
+void FEditorModeTools::DestroyMode( FEditorModeID InID )
+{
+	// Find the mode from the ID and exit it.
+	for( int32 Index = Modes.Num() - 1; Index >= 0; --Index )
+	{
+		auto& Mode = Modes[Index];
+		if ( Mode->GetID() == InID )
+		{
+			// Deactivate and destroy
+			DeactivateModeAtIndex(Index);
+			break;
+		}
+	}
+
+	RecycledModes.Remove(InID);
+}
+
+/**
+ * Activates an editor mode. Shuts down all other active modes which cannot run with the passed in mode.
+ * 
+ * @param InID		The ID of the editor mode to activate.
+ * @param bToggle	true if the passed in editor mode should be toggled off if it is already active.
+ */
+void FEditorModeTools::ActivateMode( FEditorModeID InID, bool bToggle )
+{
+	if (InID == FBuiltinEditorModes::EM_Default)
+	{
+		InID = DefaultID;
+	}
+
+	// Check to see if the mode is already active
+	if( IsModeActive(InID) )
+	{
+		// The mode is already active toggle it off if we should toggle off already active modes.
+		if( bToggle )
+		{
+			DeactivateMode( InID );
+		}
+		// Nothing more to do
+		return;
+	}
+
+	// Recycle a mode or factory a new one
+	TSharedPtr<FEdMode> Mode = RecycledModes.FindRef( InID );
+
+	if ( Mode.IsValid() )
+	{
+		RecycledModes.Remove( InID );
+	}
+	else
+	{
+		Mode = FEditorModeRegistry::Get().CreateMode( InID, *this );
+	}
+
+	if( !Mode.IsValid() )
+	{
+		UE_LOG(LogEditorModes, Log, TEXT("FEditorModeTools::ActivateMode : Couldn't find mode '%s'."), *InID.ToString() );
+		// Just return and leave the mode list unmodified
+		return;
+	}
+
+	// Remove anything that isn't compatible with this mode
+	for( int32 ModeIndex = Modes.Num() - 1; ModeIndex >= 0; --ModeIndex )
+	{
+		const bool bModesAreCompatible = Mode->IsCompatibleWith( Modes[ModeIndex]->GetID() ) || Modes[ModeIndex]->IsCompatibleWith( Mode->GetID() );
+		if ( !bModesAreCompatible )
+		{
+			DeactivateModeAtIndex(ModeIndex);
+		}
+	}
+
+	Modes.Add( Mode );
+
+	// Enter the new mode
+	Mode->Enter();
+	
+	// Update the editor UI
+	FEditorSupportDelegates::UpdateUI.Broadcast();
+}
+
+bool FEditorModeTools::EnsureNotInMode(FEditorModeID ModeID, const FText& ErrorMsg, bool bNotifyUser) const
+{
+	// We're in a 'safe' mode if we're not in the specified mode.
+	const bool bInASafeMode = !IsModeActive(ModeID);
+	if( !bInASafeMode && !ErrorMsg.IsEmpty() )
+	{
+		// Do we want to display this as a notification or a dialog to the user
+		if ( bNotifyUser )
+		{
+			FNotificationInfo Info( ErrorMsg );
+			FSlateNotificationManager::Get().AddNotification( Info );
+		}
+		else
+		{
+			FMessageDialog::Open( EAppMsgType::Ok, ErrorMsg );
+		}		
+	}
+	return bInASafeMode;
+}
+
+FEdMode* FEditorModeTools::FindMode( FEditorModeID InID )
+{
+	for( auto& Mode : Modes )
+	{
+		if( Mode->GetID() == InID )
+		{
+			return Mode.Get();
+		}
+	}
+
+	return NULL;
+}
+
+/**
+ * Returns a coordinate system that should be applied on top of the worldspace system.
+ */
+
+FMatrix FEditorModeTools::GetCustomDrawingCoordinateSystem()
+{
+	FMatrix Matrix = FMatrix::Identity;
+
+	switch (GetCoordSystem())
+	{
+		case COORD_Local:
+		{
+			// Let the current mode have a shot at setting the local coordinate system.
+			// If it doesn't want to, create it by looking at the currently selected actors list.
+
+			bool CustomCoordinateSystemProvided = false;
+			for (const auto& Mode : Modes)
+			{
+				if (Mode->GetCustomDrawingCoordinateSystem(Matrix, nullptr))
+				{
+					CustomCoordinateSystemProvided = true;
+					break;
+				}
+			}
+
+			if (!CustomCoordinateSystemProvided)
+			{
+				const int32 Num = GetSelectedActors()->CountSelections<AActor>();
+
+				// Coordinate system needs to come from the last actor selected
+				if (Num > 0)
+				{
+					Matrix = FRotationMatrix(GetSelectedActors()->GetBottom<AActor>()->GetActorRotation());
+				}
+			}
+
+			if (!Matrix.Equals(FMatrix::Identity))
+			{
+				Matrix.RemoveScaling();
+			}
+		}
+		break;
+
+		case COORD_World:
+			break;
+
+		default:
+			break;
+	}
+
+	return Matrix;
+}
+
+FMatrix FEditorModeTools::GetCustomInputCoordinateSystem()
+{
+	return GetCustomDrawingCoordinateSystem();
+}
+
+/** Gets the widget axis to be drawn */
+EAxisList::Type FEditorModeTools::GetWidgetAxisToDraw( FWidget::EWidgetMode InWidgetMode ) const
+{
+	EAxisList::Type OutAxis = EAxisList::All;
+	for( int Index = Modes.Num() - 1; Index >= 0 ; Index-- )
+	{
+		if ( Modes[Index]->ShouldDrawWidget() )
+		{
+			OutAxis = Modes[Index]->GetWidgetAxisToDraw( InWidgetMode );
+			break;
+		}
+	}
+
+	return OutAxis;
+}
+
+/** Mouse tracking interface.  Passes tracking messages to all active modes */
+bool FEditorModeTools::StartTracking(FEditorViewportClient* InViewportClient, FViewport* InViewport)
+{
+	bIsTracking = true;
+	bool bTransactionHandled = false;
+
+	CachedLocation = PivotLocation;	// Cache the pivot location
+
+	for ( const auto& Mode : Modes)
+	{
+		bTransactionHandled |= Mode->StartTracking(InViewportClient, InViewport);
+	}
+
+	return bTransactionHandled;
+}
+
+/** Mouse tracking interface.  Passes tracking messages to all active modes */
+bool FEditorModeTools::EndTracking(FEditorViewportClient* InViewportClient, FViewport* InViewport)
+{
+	bIsTracking = false;
+	bool bTransactionHandled = false;
+
+	for ( const auto& Mode : Modes)
+	{
+		bTransactionHandled |= Mode->EndTracking(InViewportClient, InViewportClient->Viewport);
+	}
+
+	CachedLocation = PivotLocation;	// Clear the pivot location
+	
+	return bTransactionHandled;
+}
+
+bool FEditorModeTools::AllowsViewportDragTool() const
+{
+	bool bCanUseDragTool = false;
+	for (const TSharedPtr<FEdMode>& Mode : Modes)
+	{
+		bCanUseDragTool |= Mode->AllowsViewportDragTool();
+	}
+	return bCanUseDragTool;
+}
+
+/** Notifies all active modes that a map change has occured */
+void FEditorModeTools::MapChangeNotify()
+{
+	for ( const auto& Mode : Modes)
+	{
+		Mode->MapChangeNotify();
+	}
+}
+
+
+/** Notifies all active modes to empty their selections */
+void FEditorModeTools::SelectNone()
+{
+	for ( const auto& Mode : Modes)
+	{
+		Mode->SelectNone();
+	}
+}
+
+/** Notifies all active modes of box selection attempts */
+bool FEditorModeTools::BoxSelect( FBox& InBox, bool InSelect )
+{
+	bool bHandled = false;
+	for( const auto& Mode : Modes)
+	{
+		bHandled |= Mode->BoxSelect( InBox, InSelect );
+	}
+	return bHandled;
+}
+
+/** Notifies all active modes of frustum selection attempts */
+bool FEditorModeTools::FrustumSelect( const FConvexVolume& InFrustum, bool InSelect )
+{
+	bool bHandled = false;
+	for( const auto& Mode : Modes)
+	{
+		bHandled |= Mode->FrustumSelect( InFrustum, InSelect );
+	}
+	return bHandled;
+}
+
+
+/** true if any active mode uses a transform widget */
+bool FEditorModeTools::UsesTransformWidget() const
+{
+	bool bUsesTransformWidget = false;
+	for( const auto& Mode : Modes)
+	{
+		bUsesTransformWidget |= Mode->UsesTransformWidget();
+	}
+
+	return bUsesTransformWidget;
+}
+
+/** true if any active mode uses the passed in transform widget */
+bool FEditorModeTools::UsesTransformWidget( FWidget::EWidgetMode CheckMode ) const
+{
+	bool bUsesTransformWidget = false;
+	for( const auto& Mode : Modes)
+	{
+		bUsesTransformWidget |= Mode->UsesTransformWidget(CheckMode);
+	}
+
+	return bUsesTransformWidget;
+}
+
+/** Sets the current widget axis */
+void FEditorModeTools::SetCurrentWidgetAxis( EAxisList::Type NewAxis )
+{
+	for( const auto& Mode : Modes)
+	{
+		Mode->SetCurrentWidgetAxis( NewAxis );
+	}
+}
+
+/** Notifies all active modes of mouse click messages. */
+bool FEditorModeTools::HandleClick(FEditorViewportClient* InViewportClient,  HHitProxy *HitProxy, const FViewportClick& Click )
+{
+	bool bHandled = false;
+	for( const auto& Mode : Modes)
+	{
+		bHandled |= Mode->HandleClick(InViewportClient, HitProxy, Click);
+	}
+
+	return bHandled;
+}
+
+/** true if the passed in brush actor should be drawn in wireframe */	
+bool FEditorModeTools::ShouldDrawBrushWireframe( AActor* InActor ) const
+{
+	bool bShouldDraw = false;
+	for( const auto& Mode : Modes)
+	{
+		bShouldDraw |= Mode->ShouldDrawBrushWireframe( InActor );
+	}
+
+	if( Modes.Num() == 0 )
+	{
+		// We can get into a state where there are no active modes at editor startup if the builder brush is created before the default mode is activated.
+		// Ensure we can see the builder brush when no modes are active.
+		bShouldDraw = true;
+	}
+	return bShouldDraw;
+}
+
+/** true if brush vertices should be drawn */
+bool FEditorModeTools::ShouldDrawBrushVertices() const
+{
+	// Currently only geometry mode being active prevents vertices from being drawn.
+	return !IsModeActive( FBuiltinEditorModes::EM_Geometry );
+}
+
+/** Ticks all active modes */
+void FEditorModeTools::Tick( FEditorViewportClient* ViewportClient, float DeltaTime )
+{
+	// Remove anything pending destruction
+	for( int32 Index = Modes.Num() - 1; Index >= 0; --Index)
+	{
+		if (Modes[Index]->IsPendingDeletion())
+		{
+			DeactivateModeAtIndex(Index);
+		}
+	}
+	
+	if (Modes.Num() == 0)
+	{
+		// Ensure the default mode is active if there are no active modes.
+		ActivateDefaultMode();
+	}
+
+	for( const auto& Mode : Modes)
+	{
+		Mode->Tick( ViewportClient, DeltaTime );
+	}
+}
+
+/** Notifies all active modes of any change in mouse movement */
+bool FEditorModeTools::InputDelta( FEditorViewportClient* InViewportClient,FViewport* InViewport,FVector& InDrag,FRotator& InRot,FVector& InScale )
+{
+	bool bHandled = false;
+	for( const auto& Mode : Modes)
+	{
+		bHandled |= Mode->InputDelta( InViewportClient, InViewport, InDrag, InRot, InScale );
+	}
+	return bHandled;
+}
+
+/** Notifies all active modes of captured mouse movement */	
+bool FEditorModeTools::CapturedMouseMove( FEditorViewportClient* InViewportClient, FViewport* InViewport, int32 InMouseX, int32 InMouseY )
+{
+	bool bHandled = false;
+	for( const auto& Mode : Modes)
+	{
+		bHandled |= Mode->CapturedMouseMove( InViewportClient, InViewport, InMouseX, InMouseY );
+	}
+	return bHandled;
+}
+
+/** Notifies all active modes of keyboard input */
+bool FEditorModeTools::InputKey(FEditorViewportClient* InViewportClient, FViewport* Viewport, FKey Key, EInputEvent Event)
+{
+	bool bHandled = false;
+	for (const auto& Mode : Modes)
+	{
+		bHandled |= Mode->InputKey( InViewportClient, Viewport, Key, Event );
+	}
+	return bHandled;
+}
+
+/** Notifies all active modes of axis movement */
+bool FEditorModeTools::InputAxis(FEditorViewportClient* InViewportClient, FViewport* Viewport, int32 ControllerId, FKey Key, float Delta, float DeltaTime)
+{
+	bool bHandled = false;
+	for( const auto& Mode : Modes)
+	{
+		bHandled |= Mode->InputAxis( InViewportClient, Viewport, ControllerId, Key, Delta, DeltaTime );
+	}
+	return bHandled;
+}
+
+bool FEditorModeTools::MouseEnter( FEditorViewportClient* InViewportClient, FViewport* Viewport, int32 X, int32 Y )
+{
+	bool bHandled = false;
+	for( const auto& Mode : Modes)
+	{
+		bHandled |= Mode->MouseEnter( InViewportClient, Viewport, X, Y );
+	}
+	return bHandled;
+}
+
+bool FEditorModeTools::MouseLeave( FEditorViewportClient* InViewportClient, FViewport* Viewport )
+{
+	bool bHandled = false;
+	for( const auto& Mode : Modes)
+	{
+		bHandled |= Mode->MouseLeave( InViewportClient, Viewport );
+	}
+	return bHandled;
+}
+
+/** Notifies all active modes that the mouse has moved */
+bool FEditorModeTools::MouseMove( FEditorViewportClient* InViewportClient, FViewport* Viewport, int32 X, int32 Y )
+{
+	bool bHandled = false;
+	for( const auto& Mode : Modes)
+	{
+		bHandled |= Mode->MouseMove( InViewportClient, Viewport, X, Y );
+	}
+	return bHandled;
+}
+
+bool FEditorModeTools::ReceivedFocus( FEditorViewportClient* InViewportClient, FViewport* Viewport )
+{
+	bool bHandled = false;
+	for( const auto& Mode : Modes)
+	{
+		bHandled |= Mode->ReceivedFocus( InViewportClient, Viewport );
+	}
+	return bHandled;
+}
+
+bool FEditorModeTools::LostFocus( FEditorViewportClient* InViewportClient, FViewport* Viewport )
+{
+	bool bHandled = false;
+	for( const auto& Mode : Modes)
+	{
+		bHandled |= Mode->LostFocus( InViewportClient, Viewport );
+	}
+	return bHandled;
+}
+
+/** Draws all active mode components */	
+void FEditorModeTools::DrawActiveModes( const FSceneView* InView, FPrimitiveDrawInterface* PDI )
+{
+	for( const auto& Mode : Modes)
+	{
+		Mode->Draw( InView, PDI );
+	}
+}
+
+/** Renders all active modes */
+void FEditorModeTools::Render( const FSceneView* InView, FViewport* Viewport, FPrimitiveDrawInterface* PDI )
+{
+	for( const auto& Mode : Modes)
+	{
+		Mode->Render( InView, Viewport, PDI );
+	}
+}
+
+/** Draws the HUD for all active modes */
+void FEditorModeTools::DrawHUD( FEditorViewportClient* InViewportClient,FViewport* Viewport, const FSceneView* View, FCanvas* Canvas )
+{
+	for( const auto& Mode : Modes)
+	{
+		Mode->DrawHUD( InViewportClient, Viewport, View, Canvas );
+	}
+}
+
+/** Calls PostUndo on all active components */
+void FEditorModeTools::PostUndo()
+{
+	for( const auto& Mode : Modes)
+	{
+		Mode->PostUndo();
+	}
+}
+
+/** true if we should allow widget move */
+bool FEditorModeTools::AllowWidgetMove() const
+{
+	bool bAllow = false;
+	for( const auto& Mode : Modes)
+	{
+		bAllow |= Mode->AllowWidgetMove();
+	}
+	return bAllow;
+}
+
+bool FEditorModeTools::DisallowMouseDeltaTracking() const
+{
+	bool bDisallow = false;
+	for( const auto& Mode : Modes)
+	{
+		bDisallow |= Mode->DisallowMouseDeltaTracking();
+	}
+	return bDisallow;
+}
+
+bool FEditorModeTools::GetCursor(EMouseCursor::Type& OutCursor) const
+{
+	bool bHandled = false;
+	for( const auto& Mode : Modes)
+	{
+		bHandled |= Mode->GetCursor(OutCursor);
+	}
+	return bHandled;
+}
+
+/**
+ * Used to cycle widget modes
+ */
+void FEditorModeTools::CycleWidgetMode (void)
+{
+	//make sure we're not currently tracking mouse movement.  If we are, changing modes could cause a crash due to referencing an axis/plane that is incompatible with the widget
+	for(int32 ViewportIndex = 0;ViewportIndex < GEditor->LevelViewportClients.Num();ViewportIndex++)
+	{
+		FEditorViewportClient* ViewportClient = GEditor->LevelViewportClients[ ViewportIndex ];
+		if (ViewportClient->IsTracking())
+		{
+			return;
+		}
+	}
+
+	//only cycle when the mode is requesting the drawing of a widget
+	if( GetShowWidget() )
+	{
+		const int32 CurrentWk = GetWidgetMode();
+		int32 Wk = CurrentWk;
+		do
+		{
+			Wk++;
+			if ((Wk == FWidget::WM_TranslateRotateZ) && (!GetDefault<ULevelEditorViewportSettings>()->bAllowTranslateRotateZWidget))
+			{
+				Wk++;
+			}
+			// Roll back to the start if we go past FWidget::WM_Scale
+			if( Wk >= FWidget::WM_Max)
+			{
+				Wk -= FWidget::WM_Max;
+			}
+		}
+		while (!UsesTransformWidget((FWidget::EWidgetMode)Wk) && Wk != CurrentWk);
+		SetWidgetMode( (FWidget::EWidgetMode)Wk );
+		FEditorSupportDelegates::RedrawAllViewports.Broadcast();
+	}
+}
+
+/**Save Widget Settings to Ini file*/
+void FEditorModeTools::SaveWidgetSettings(void)
+{
+	GEditor->SaveEditorUserSettings();
+}
+
+/**Load Widget Settings from Ini file*/
+void FEditorModeTools::LoadWidgetSettings(void)
+{
+}
+
+/**
+ * Returns a good location to draw the widget at.
+ */
+
+FVector FEditorModeTools::GetWidgetLocation() const
+{
+	for (int Index = Modes.Num() - 1; Index >= 0 ; Index--)
+	{
+		if ( Modes[Index]->UsesTransformWidget() )
+		{
+			 return Modes[Index]->GetWidgetLocation();
+		}
+	}
+	
+	return FVector(EForceInit::ForceInitToZero);
+}
+
+/**
+ * Changes the current widget mode.
+ */
+
+void FEditorModeTools::SetWidgetMode( FWidget::EWidgetMode InWidgetMode )
+{
+	WidgetMode = InWidgetMode;
+}
+
+/**
+ * Allows you to temporarily override the widget mode.  Call this function again
+ * with FWidget::WM_None to turn off the override.
+ */
+
+void FEditorModeTools::SetWidgetModeOverride( FWidget::EWidgetMode InWidgetMode )
+{
+	OverrideWidgetMode = InWidgetMode;
+}
+
+/**
+ * Retrieves the current widget mode, taking overrides into account.
+ */
+
+FWidget::EWidgetMode FEditorModeTools::GetWidgetMode() const
+{
+	if( OverrideWidgetMode != FWidget::WM_None )
+	{
+		return OverrideWidgetMode;
+	}
+
+	return WidgetMode;
+}
+
+bool FEditorModeTools::GetShowFriendlyVariableNames() const
+{
+	return GetDefault<UEditorStyleSettings>()->bShowFriendlyNames;
+}
+
+/**
+ * Sets a bookmark in the levelinfo file, allocating it if necessary.
+ */
+
+void FEditorModeTools::SetBookmark( uint32 InIndex, FEditorViewportClient* InViewportClient )
+{
+	UWorld* World = InViewportClient->GetWorld();
+	if ( World )
+	{
+		AWorldSettings* WorldSettings = World->GetWorldSettings();
+
+		// Verify the index is valid for the bookmark
+		if ( WorldSettings && InIndex < AWorldSettings::MAX_BOOKMARK_NUMBER )
+		{
+			// If the index doesn't already have a bookmark in place, create a new one
+			if ( !WorldSettings->BookMarks[ InIndex ] )
+			{
+				WorldSettings->BookMarks[ InIndex ] = ConstructObject<UBookMark>( UBookMark::StaticClass(), WorldSettings );
+			}
+
+			UBookMark* CurBookMark = WorldSettings->BookMarks[ InIndex ];
+			check(CurBookMark);
+			check(InViewportClient);
+
+			// Use the rotation from the first perspective viewport can find.
+			FRotator Rotation(0,0,0);
+			if( !InViewportClient->IsOrtho() )
+			{
+				Rotation = InViewportClient->GetViewRotation();
+			}
+
+			CurBookMark->Location = InViewportClient->GetViewLocation();
+			CurBookMark->Rotation = Rotation;
+
+			// Keep a record of which levels were hidden so that we can restore these with the bookmark
+			CurBookMark->HiddenLevels.Empty();
+			for ( int32 LevelIndex = 0 ; LevelIndex < World->StreamingLevels.Num() ; ++LevelIndex )
+			{
+				ULevelStreaming* StreamingLevel = World->StreamingLevels[LevelIndex];
+				if ( StreamingLevel )
+				{
+					if( !StreamingLevel->bShouldBeVisibleInEditor )
+					{
+						CurBookMark->HiddenLevels.Add( StreamingLevel->GetFullName() );
+					}
+				}
+			}
+		}
+	}
+}
+
+/**
+ * Checks to see if a bookmark exists at a given index
+ */
+
+bool FEditorModeTools::CheckBookmark( uint32 InIndex, FEditorViewportClient* InViewportClient )
+{
+	UWorld* World = InViewportClient->GetWorld();
+	if ( World )
+	{
+		AWorldSettings* WorldSettings = World->GetWorldSettings();
+		if ( WorldSettings && InIndex < AWorldSettings::MAX_BOOKMARK_NUMBER && WorldSettings->BookMarks[ InIndex ] )
+		{
+			return ( WorldSettings->BookMarks[ InIndex ] ? true : false );
+		}
+	}
+
+	return false;
+}
+
+/**
+ * Retrieves a bookmark from the list.
+ */
+
+void FEditorModeTools::JumpToBookmark( uint32 InIndex, bool bShouldRestoreLevelVisibility, FEditorViewportClient* InViewportClient )
+{
+	UWorld* World = InViewportClient->GetWorld();
+	if ( World )
+	{
+		AWorldSettings* WorldSettings = World->GetWorldSettings();
+
+		// Can only jump to a pre-existing bookmark
+		if ( WorldSettings && InIndex < AWorldSettings::MAX_BOOKMARK_NUMBER && WorldSettings->BookMarks[ InIndex ] )
+		{
+			const UBookMark* CurBookMark = WorldSettings->BookMarks[ InIndex ];
+			check(CurBookMark);
+
+			// Set all level editing cameras to this bookmark
+			for( int32 v = 0 ; v < GEditor->LevelViewportClients.Num() ; v++ )
+			{
+				GEditor->LevelViewportClients[v]->SetViewLocation( CurBookMark->Location );
+				if( !GEditor->LevelViewportClients[v]->IsOrtho() )
+				{
+					GEditor->LevelViewportClients[v]->SetViewRotation( CurBookMark->Rotation );
+				}
+				GEditor->LevelViewportClients[v]->Invalidate();
+			}
+		}
+	}
+}
+
+
+/**
+ * Clears a bookmark
+ */
+void FEditorModeTools::ClearBookmark( uint32 InIndex, FEditorViewportClient* InViewportClient )
+{
+	UWorld* World = InViewportClient->GetWorld();
+	if( World )
+	{
+		AWorldSettings* pWorldSettings = World->GetWorldSettings();
+
+		// Verify the index is valid for the bookmark
+		if ( pWorldSettings && InIndex < AWorldSettings::MAX_BOOKMARK_NUMBER )
+		{
+			pWorldSettings->BookMarks[ InIndex ] = NULL;
+		}
+	}
+}
+
+/**
+* Clears all book marks
+*/
+void FEditorModeTools::ClearAllBookmarks( FEditorViewportClient* InViewportClient )
+{
+	for( int i = 0; i <  AWorldSettings::MAX_BOOKMARK_NUMBER; ++i )
+	{
+		ClearBookmark( i , InViewportClient );
+	}
+}
+
+/**
+ * Serializes the components for all modes.
+ */
+
+void FEditorModeTools::AddReferencedObjects( FReferenceCollector& Collector )
+{
+	for( int32 x = 0 ; x < Modes.Num() ; ++x )
+	{
+		Modes[x]->AddReferencedObjects( Collector );
+	}
+}
+
+/**
+ * Returns a pointer to an active mode specified by the passed in ID
+ * If the editor mode is not active, NULL is returned
+ */
+FEdMode* FEditorModeTools::GetActiveMode( FEditorModeID InID )
+{
+	for( auto& Mode : Modes )
+	{
+		if( Mode->GetID() == InID )
+		{
+			return Mode.Get();
+		}
+	}
+	return nullptr;
+}
+
+/**
+ * Returns a pointer to an active mode specified by the passed in ID
+ * If the editor mode is not active, NULL is returned
+ */
+const FEdMode* FEditorModeTools::GetActiveMode( FEditorModeID InID ) const
+{
+	for (const auto& Mode : Modes)
+	{
+		if (Mode->GetID() == InID)
+		{
+			return Mode.Get();
+		}
+	}
+
+	return nullptr;
+}
+
+/**
+ * Returns the active tool of the passed in editor mode.
+ * If the passed in editor mode is not active or the mode has no active tool, NULL is returned
+ */
+const FModeTool* FEditorModeTools::GetActiveTool( FEditorModeID InID ) const
+{
+	const FEdMode* ActiveMode = GetActiveMode( InID );
+	const FModeTool* Tool = NULL;
+	if( ActiveMode )
+	{
+		Tool = ActiveMode->GetCurrentTool();
+	}
+	return Tool;
+}
+
+/** 
+ * Returns true if the passed in editor mode is active 
+ */
+bool FEditorModeTools::IsModeActive( FEditorModeID InID ) const
+{
+	return GetActiveMode( InID ) != NULL;
+}
+
+/** 
+ * Returns true if the default editor mode is active 
+ */
+bool FEditorModeTools::IsDefaultModeActive() const
+{
+	return IsModeActive(DefaultID);
+}
+
+/** 
+ * Returns an array of all active modes
+ */
+void FEditorModeTools::GetActiveModes( TArray<FEdMode*>& OutActiveModes )
+{
+	OutActiveModes.Empty();
+	// Copy into an array.  Do not let users modify the active list directly.
+	for( auto& Mode : Modes)
+	{
+		OutActiveModes.Add(Mode.Get());
+	}
+}
diff --git a/Engine/Source/Editor/UnrealEd/Private/EditorModes.cpp b/Engine/Source/Editor/UnrealEd/Private/EditorModes.cpp
--- a/Engine/Source/Editor/UnrealEd/Private/EditorModes.cpp
+++ b/Engine/Source/Editor/UnrealEd/Private/EditorModes.cpp
@@ -26,7 +26,7 @@
 #include "Engine/StaticMeshActor.h"
 #include "Engine/LevelStreaming.h"
 
-DEFINE_LOG_CATEGORY_STATIC(LogEditorModes, Log, All);
+DEFINE_LOG_CATEGORY(LogEditorModes);
 
 // Builtin editor mode constants
 const FEditorModeID FBuiltinEditorModes::EM_None = NAME_None;
@@ -44,1123 +44,6 @@ const FEditorModeID FBuiltinEditorModes::EM_StreamingLevel(TEXT("EM_StreamingLev
 const FEditorModeID FBuiltinEditorModes::EM_Physics(TEXT("EM_Physics"));
 const FEditorModeID FBuiltinEditorModes::EM_ActorPicker(TEXT("EM_ActorPicker"));
 
-
-/** Hit proxy used for editable properties */
-struct HPropertyWidgetProxy : public HHitProxy
-{
-	DECLARE_HIT_PROXY();
-
-	/** Name of property this is the widget for */
-	FString	PropertyName;
-
-	/** If the property is an array property, the index into that array that this widget is for */
-	int32	PropertyIndex;
-
-	/** This property is a transform */
-	bool	bPropertyIsTransform;
-
-	HPropertyWidgetProxy(FString InPropertyName, int32 InPropertyIndex, bool bInPropertyIsTransform)
-		: HHitProxy(HPP_Foreground)
-		, PropertyName(InPropertyName)
-		, PropertyIndex(InPropertyIndex)
-		, bPropertyIsTransform(bInPropertyIsTransform)
-	{}
-
-	/** Show cursor as cross when over this handle */
-	virtual EMouseCursor::Type GetMouseCursor()
-	{
-		return EMouseCursor::Crosshairs;
-	}
-};
-
-IMPLEMENT_HIT_PROXY(HPropertyWidgetProxy, HHitProxy);
-
-const FName FEdMode::MD_MakeEditWidget(TEXT("MakeEditWidget"));
-const FName FEdMode::MD_ValidateWidgetUsing(TEXT("ValidateWidgetUsing"));
-
-namespace
-{
-	/**
-	 * Returns a reference to the named property value data in the given container.
-	 */
-	template<typename T>
-	T* GetPropertyValuePtrByName(const UStruct* InStruct, void* InContainer, FString PropertyName, int32 ArrayIndex)
-	{
-		T* ValuePtr = NULL;
-
-		// Extract the vector ptr recursively using the property name
-		int32 DelimPos = PropertyName.Find(TEXT("."));
-		if(DelimPos != INDEX_NONE)
-		{
-			// Parse the property name and (optional) array index
-			int32 SubArrayIndex = 0;
-			FString NameToken = PropertyName.Left(DelimPos);
-			int32 ArrayPos = NameToken.Find(TEXT("["));
-			if(ArrayPos != INDEX_NONE)
-			{
-				FString IndexToken = NameToken.RightChop(ArrayPos + 1).LeftChop(1);
-				SubArrayIndex = FCString::Atoi(*IndexToken);
-
-				NameToken = PropertyName.Left(ArrayPos);
-			}
-
-			// Obtain the property info from the given structure definition
-			UProperty* CurrentProp = FindField<UProperty>(InStruct, FName(*NameToken));
-
-			// Check first to see if this is a simple structure (i.e. not an array of structures)
-			UStructProperty* StructProp = Cast<UStructProperty>(CurrentProp);
-			if(StructProp != NULL)
-			{
-				// Recursively call back into this function with the structure property and container value
-				ValuePtr = GetPropertyValuePtrByName<T>(StructProp->Struct, StructProp->ContainerPtrToValuePtr<void>(InContainer), PropertyName.RightChop(DelimPos + 1), ArrayIndex);
-			}
-			else
-			{
-				// Check to see if this is an array
-				UArrayProperty* ArrayProp = Cast<UArrayProperty>(CurrentProp);
-				if(ArrayProp != NULL)
-				{
-					// It is an array, now check to see if this is an array of structures
-					StructProp = Cast<UStructProperty>(ArrayProp->Inner);
-					if(StructProp != NULL)
-					{
-						FScriptArrayHelper_InContainer ArrayHelper(ArrayProp, InContainer);
-						if(ArrayHelper.IsValidIndex(SubArrayIndex))
-						{
-							// Recursively call back into this function with the array element and container value
-							ValuePtr = GetPropertyValuePtrByName<T>(StructProp->Struct, ArrayHelper.GetRawPtr(SubArrayIndex), PropertyName.RightChop(DelimPos + 1), ArrayIndex);
-						}
-					}
-				}
-			}
-		}
-		else
-		{
-			UProperty* Prop = FindField<UProperty>(InStruct, FName(*PropertyName));
-			if(Prop != NULL)
-			{
-				if( UArrayProperty* ArrayProp = Cast<UArrayProperty>(Prop) )
-				{
-					check(ArrayIndex != INDEX_NONE);
-
-					// Property is an array property, so make sure we have a valid index specified
-					FScriptArrayHelper_InContainer ArrayHelper(ArrayProp, InContainer);
-					if( ArrayHelper.IsValidIndex(ArrayIndex) )
-					{
-						ValuePtr = (T*)ArrayHelper.GetRawPtr(ArrayIndex);
-					}
-				}
-				else
-				{
-					// Property is a vector property, so access directly
-					ValuePtr = Prop->ContainerPtrToValuePtr<T>(InContainer);
-				}
-			}
-		}
-
-		return ValuePtr;
-	}
-
-	/**
-	 * Returns the value of the property with the given name in the given Actor instance.
-	 */
-	template<typename T>
-	T GetPropertyValueByName(AActor* Actor, FString PropertyName, int32 PropertyIndex)
-	{
-		T Value;
-		T* ValuePtr = GetPropertyValuePtrByName<T>(Actor->GetClass(), Actor, PropertyName, PropertyIndex);
-		if(ValuePtr)
-		{
-			Value = *ValuePtr;
-		}
-		return Value;
-	}
-
-	/**
-	 * Sets the property with the given name in the given Actor instance to the given value.
-	 */
-	template<typename T>
-	void SetPropertyValueByName(AActor* Actor, FString PropertyName, int32 PropertyIndex, const T& InValue)
-	{
-		T* ValuePtr = GetPropertyValuePtrByName<T>(Actor->GetClass(), Actor, PropertyName, PropertyIndex);
-		if(ValuePtr)
-		{
-			*ValuePtr = InValue;
-		}
-	}
-}
-
-//////////////////////////////////
-// FEdMode
-
-FEdMode::FEdMode()
-	: bPendingDeletion( false )
-	, CurrentTool( NULL )
-	, EditedPropertyName(TEXT(""))
-	, EditedPropertyIndex( INDEX_NONE )
-	, bEditedPropertyIsTransform( false )
-{
-	bDrawKillZ = true;
-}
-
-FEdMode::~FEdMode()
-{
-}
-
-void FEdMode::OnModeUnregistered( FEditorModeID ModeID )
-{
-	if( ModeID == Info.ID )
-	{
-		// This should be synonymous with "delete this"
-		Owner->DestroyMode(ModeID);
-	}
-}
-
-bool FEdMode::MouseEnter( FEditorViewportClient* ViewportClient,FViewport* Viewport,int32 x, int32 y )
-{
-	if( GetCurrentTool() )
-	{
-		return GetCurrentTool()->MouseEnter( ViewportClient, Viewport, x, y );
-	}
-
-	return false;
-}
-
-bool FEdMode::MouseLeave( FEditorViewportClient* ViewportClient,FViewport* Viewport )
-{
-	if( GetCurrentTool() )
-	{
-		return GetCurrentTool()->MouseLeave( ViewportClient, Viewport );
-	}
-
-	return false;
-}
-
-bool FEdMode::MouseMove(FEditorViewportClient* ViewportClient,FViewport* Viewport,int32 x, int32 y)
-{
-	if( GetCurrentTool() )
-	{
-		return GetCurrentTool()->MouseMove( ViewportClient, Viewport, x, y );
-	}
-
-	return false;
-}
-
-bool FEdMode::ReceivedFocus(FEditorViewportClient* ViewportClient,FViewport* Viewport)
-{
-	if( GetCurrentTool() )
-	{
-		return GetCurrentTool()->ReceivedFocus( ViewportClient, Viewport );
-	}
-
-	return false;
-}
-
-bool FEdMode::LostFocus(FEditorViewportClient* ViewportClient,FViewport* Viewport)
-{
-	if( GetCurrentTool() )
-	{
-		return GetCurrentTool()->LostFocus( ViewportClient, Viewport );
-	}
-
-	return false;
-}
-
-/**
- * Called when the mouse is moved while a window input capture is in effect
- *
- * @param	InViewportClient	Level editor viewport client that captured the mouse input
- * @param	InViewport			Viewport that captured the mouse input
- * @param	InMouseX			New mouse cursor X coordinate
- * @param	InMouseY			New mouse cursor Y coordinate
- *
- * @return	true if input was handled
- */
-bool FEdMode::CapturedMouseMove( FEditorViewportClient* InViewportClient, FViewport* InViewport, int32 InMouseX, int32 InMouseY )
-{
-	if( GetCurrentTool() )
-	{
-		return GetCurrentTool()->CapturedMouseMove( InViewportClient, InViewport, InMouseX, InMouseY );
-	}
-
-	return false;
-}
-
-
-bool FEdMode::InputKey(FEditorViewportClient* ViewportClient, FViewport* Viewport, FKey Key, EInputEvent Event)
-{
-	if( GetCurrentTool() && GetCurrentTool()->InputKey( ViewportClient, Viewport, Key, Event ) )
-	{
-		return true;
-	}
-	else
-	{
-		// Pass input up to selected actors if not in a tool mode
-		TArray<AActor*> SelectedActors;
-		Owner->GetSelectedActors()->GetSelectedObjects<AActor>(SelectedActors);
-
-		for( TArray<AActor*>::TIterator it(SelectedActors); it; ++it )
-		{
-			// Tell the object we've had a key press
-			(*it)->EditorKeyPressed(Key, Event);
-		}
-	}
-
-	return 0;
-}
-
-bool FEdMode::InputAxis(FEditorViewportClient* InViewportClient, FViewport* Viewport, int32 ControllerId, FKey Key, float Delta, float DeltaTime)
-{
-	FModeTool* Tool = GetCurrentTool();
-	if (Tool)
-	{
-		return Tool->InputAxis(InViewportClient, Viewport, ControllerId, Key, Delta, DeltaTime);
-	}
-
-	return false;
-}
-
-bool FEdMode::InputDelta(FEditorViewportClient* InViewportClient, FViewport* InViewport, FVector& InDrag, FRotator& InRot, FVector& InScale)
-{	
-	if(UsesPropertyWidgets())
-	{
-		AActor* SelectedActor = GetFirstSelectedActorInstance();
-		if(SelectedActor != NULL && InViewportClient->GetCurrentWidgetAxis() != EAxisList::None)
-		{
-			GEditor->NoteActorMovement();
-
-			if (EditedPropertyName != TEXT(""))
-			{
-				FTransform LocalTM = FTransform::Identity;
-
-				if(bEditedPropertyIsTransform)
-				{
-					LocalTM = GetPropertyValueByName<FTransform>(SelectedActor, EditedPropertyName, EditedPropertyIndex);
-				}
-				else
-				{					
-					FVector LocalPos = GetPropertyValueByName<FVector>(SelectedActor, EditedPropertyName, EditedPropertyIndex);
-					LocalTM = FTransform(LocalPos);
-				}
-
-				// Get actor transform (actor to world)
-				FTransform ActorTM = SelectedActor->ActorToWorld();
-				// Calculate world transform
-				FTransform WorldTM = LocalTM * ActorTM;
-				// Calc delta specified by drag
-				//FTransform DeltaTM(InRot.Quaternion(), InDrag);
-				// Apply delta in world space
-				WorldTM.SetTranslation(WorldTM.GetTranslation() + InDrag);
-				WorldTM.SetRotation(InRot.Quaternion() * WorldTM.GetRotation());
-				// Convert new world transform back into local space
-				LocalTM = WorldTM.GetRelativeTransform(ActorTM);
-				// Apply delta scale
-				LocalTM.SetScale3D(LocalTM.GetScale3D() + InScale);
-
-				SelectedActor->PreEditChange(NULL);
-
-				if(bEditedPropertyIsTransform)
-				{
-					SetPropertyValueByName<FTransform>(SelectedActor, EditedPropertyName, EditedPropertyIndex, LocalTM);
-				}
-				else
-				{
-					SetPropertyValueByName<FVector>(SelectedActor, EditedPropertyName, EditedPropertyIndex, LocalTM.GetLocation());
-				}
-
-				SelectedActor->PostEditChange();
-
-				return true;
-			}
-		}
-	}
-
-	if( GetCurrentTool() )
-	{
-		return GetCurrentTool()->InputDelta(InViewportClient,InViewport,InDrag,InRot,InScale);
-	}
-
-	return 0;
-}
-
-/**
- * Lets each tool determine if it wants to use the editor widget or not.  If the tool doesn't want to use it, it will be
- * fed raw mouse delta information (not snapped or altered in any way).
- */
-
-bool FEdMode::UsesTransformWidget() const
-{
-	if( GetCurrentTool() )
-	{
-		return GetCurrentTool()->UseWidget();
-	}
-
-	return 1;
-}
-
-/**
- * Lets each mode selectively exclude certain widget types.
- */
-bool FEdMode::UsesTransformWidget(FWidget::EWidgetMode CheckMode) const
-{
-	if(UsesPropertyWidgets())
-	{
-		AActor* SelectedActor = GetFirstSelectedActorInstance();
-		if(SelectedActor != NULL)
-		{
-			// If editing a vector (not a transform)
-			if(EditedPropertyName != TEXT("") && !bEditedPropertyIsTransform)
-			{
-				return (CheckMode == FWidget::WM_Translate);
-			}
-		}
-	}
-
-	return true;
-}
-
-/**
- * Allows each mode/tool to determine a good location for the widget to be drawn at.
- */
-
-FVector FEdMode::GetWidgetLocation() const
-{
-	if(UsesPropertyWidgets())
-	{
-		AActor* SelectedActor = GetFirstSelectedActorInstance();
-		if(SelectedActor != NULL)
-		{
-			if(EditedPropertyName != TEXT(""))
-			{
-				FVector LocalPos = FVector::ZeroVector;
-
-				if(bEditedPropertyIsTransform)
-				{
-					FTransform LocalTM = GetPropertyValueByName<FTransform>(SelectedActor, EditedPropertyName, EditedPropertyIndex);
-					LocalPos = LocalTM.GetLocation();
-				}
-				else
-				{
-					LocalPos = GetPropertyValueByName<FVector>(SelectedActor, EditedPropertyName, EditedPropertyIndex);
-				}
-
-				FTransform ActorToWorld = SelectedActor->ActorToWorld();
-				FVector WorldPos = ActorToWorld.TransformPosition(LocalPos);
-				return WorldPos;
-			}
-		}
-	}
-
-	//UE_LOG(LogEditorModes, Log, TEXT("In FEdMode::GetWidgetLocation"));
-	return Owner->PivotLocation;
-}
-
-/**
- * Lets the mode determine if it wants to draw the widget or not.
- */
-
-bool FEdMode::ShouldDrawWidget() const
-{
-	return (Owner->GetSelectedActors()->GetTop<AActor>() != NULL);
-}
-
-/**
- * Allows each mode to customize the axis pieces of the widget they want drawn.
- *
- * @param	InwidgetMode	The current widget mode
- *
- * @return	A bitfield comprised of AXIS_ values
- */
-
-EAxisList::Type FEdMode::GetWidgetAxisToDraw( FWidget::EWidgetMode InWidgetMode ) const
-{
-	return EAxisList::All;
-}
-
-/**
- * Lets each mode/tool handle box selection in its own way.
- *
- * @param	InBox	The selection box to use, in worldspace coordinates.
- * @return		true if something was selected/deselected, false otherwise.
- */
-bool FEdMode::BoxSelect( FBox& InBox, bool InSelect )
-{
-	bool bResult = false;
-	if( GetCurrentTool() )
-	{
-		bResult = GetCurrentTool()->BoxSelect( InBox, InSelect );
-	}
-	return bResult;
-}
-
-/**
- * Lets each mode/tool handle frustum selection in its own way.
- *
- * @param	InFrustum	The selection frustum to use, in worldspace coordinates.
- * @return	true if something was selected/deselected, false otherwise.
- */
-bool FEdMode::FrustumSelect( const FConvexVolume& InFrustum, bool InSelect )
-{
-	bool bResult = false;
-	if( GetCurrentTool() )
-	{
-		bResult = GetCurrentTool()->FrustumSelect( InFrustum, InSelect );
-	}
-	return bResult;
-}
-
-void FEdMode::SelectNone()
-{
-	if( GetCurrentTool() )
-	{
-		GetCurrentTool()->SelectNone();
-	}
-}
-
-void FEdMode::Tick(FEditorViewportClient* ViewportClient,float DeltaTime)
-{
-	if( GetCurrentTool() )
-	{
-		GetCurrentTool()->Tick(ViewportClient,DeltaTime);
-	}
-}
-
-void FEdMode::ActorSelectionChangeNotify()
-{
-	EditedPropertyName = TEXT("");
-	EditedPropertyIndex = INDEX_NONE;
-	bEditedPropertyIsTransform = false;
-}
-
-bool FEdMode::HandleClick(FEditorViewportClient* InViewportClient, HHitProxy *HitProxy, const FViewportClick &Click)
-{
-	if(UsesPropertyWidgets() && HitProxy)
-	{
-		if( HitProxy->IsA(HPropertyWidgetProxy::StaticGetType()) )
-		{
-			HPropertyWidgetProxy* PropertyProxy = (HPropertyWidgetProxy*)HitProxy;
-			EditedPropertyName = PropertyProxy->PropertyName;
-			EditedPropertyIndex = PropertyProxy->PropertyIndex;
-			bEditedPropertyIsTransform = PropertyProxy->bPropertyIsTransform;
-			return true;
-		}
-		// Left clicking on an actor, stop editing a property
-		else if( HitProxy->IsA(HActor::StaticGetType()) )
-		{
-			EditedPropertyName = TEXT("");
-			EditedPropertyIndex = INDEX_NONE;
-			bEditedPropertyIsTransform = false;
-		}
-	}
-
-	return false;
-}
-
-void FEdMode::Enter()
-{
-	// Update components for selected actors, in case the mode we just exited
-	// was hijacking selection events selection and not updating components.
-	for ( FSelectionIterator It( *Owner->GetSelectedActors() ) ; It ; ++It )
-	{
-		AActor* SelectedActor = CastChecked<AActor>( *It );
-		SelectedActor->MarkComponentsRenderStateDirty();
-	}
-
-	bPendingDeletion = false;
-
-	FEditorDelegates::EditorModeEnter.Broadcast( this );
-	const bool bIsEnteringMode = true;
-	Owner->BroadcastEditorModeChanged( this, bIsEnteringMode );
-}
-
-void FEdMode::Exit()
-{
-	const bool bIsEnteringMode = false;
-	Owner->BroadcastEditorModeChanged( this, bIsEnteringMode );
-	FEditorDelegates::EditorModeExit.Broadcast( this );
-}
-
-void FEdMode::SetCurrentTool( EModeTools InID )
-{
-	CurrentTool = FindTool( InID );
-	check( CurrentTool );	// Tool not found!  This can't happen.
-
-	CurrentToolChanged();
-}
-
-void FEdMode::SetCurrentTool( FModeTool* InModeTool )
-{
-	CurrentTool = InModeTool;
-	check(CurrentTool);
-
-	CurrentToolChanged();
-}
-
-FModeTool* FEdMode::FindTool( EModeTools InID )
-{
-	for( int32 x = 0 ; x < Tools.Num() ; ++x )
-	{
-		if( Tools[x]->GetID() == InID )
-		{
-			return Tools[x];
-		}
-	}
-
-	UE_LOG(LogEditorModes, Fatal, TEXT("FEdMode::FindTool failed to find tool %d"), (int32)InID);
-	return NULL;
-}
-
-void FEdMode::Render(const FSceneView* View,FViewport* Viewport,FPrimitiveDrawInterface* PDI)
-{
-	if( GEditor->bShowBrushMarkerPolys )
-	{
-		// Draw translucent polygons on brushes and volumes
-
-		for( TActorIterator<ABrush> It(GetWorld()); It; ++ It )
-		{
-			ABrush* Brush = *It;
-
-			// Brush->Brush is checked to safe from brushes that were created without having their brush members attached.
-			if( Brush->Brush && (FActorEditorUtils::IsABuilderBrush(Brush) || Brush->IsVolumeBrush()) && Owner->GetSelectedActors()->IsSelected(Brush) )
-			{
-				// Build a mesh by basically drawing the triangles of each 
-				FDynamicMeshBuilder MeshBuilder;
-				int32 VertexOffset = 0;
-
-				for( int32 PolyIdx = 0 ; PolyIdx < Brush->Brush->Polys->Element.Num() ; ++PolyIdx )
-				{
-					const FPoly* Poly = &Brush->Brush->Polys->Element[PolyIdx];
-
-					if( Poly->Vertices.Num() > 2 )
-					{
-						const FVector Vertex0 = Poly->Vertices[0];
-						FVector Vertex1 = Poly->Vertices[1];
-
-						MeshBuilder.AddVertex(Vertex0, FVector2D::ZeroVector, FVector(1,0,0), FVector(0,1,0), FVector(0,0,1), FColor(255,255,255));
-						MeshBuilder.AddVertex(Vertex1, FVector2D::ZeroVector, FVector(1,0,0), FVector(0,1,0), FVector(0,0,1), FColor(255,255,255));
-
-						for( int32 VertexIdx = 2 ; VertexIdx < Poly->Vertices.Num() ; ++VertexIdx )
-						{
-							const FVector Vertex2 = Poly->Vertices[VertexIdx];
-							MeshBuilder.AddVertex(Vertex2, FVector2D::ZeroVector, FVector(1,0,0), FVector(0,1,0), FVector(0,0,1), FColor(255,255,255));
-							MeshBuilder.AddTriangle(VertexOffset,VertexOffset + VertexIdx,VertexOffset+VertexIdx-1);
-							Vertex1 = Vertex2;
-						}
-
-						// Increment the vertex offset so the next polygon uses the correct vertex indices.
-						VertexOffset += Poly->Vertices.Num();
-					}
-				}
-
-				// Allocate the material proxy and register it so it can be deleted properly once the rendering is done with it.
-				FDynamicColoredMaterialRenderProxy* MaterialProxy = new FDynamicColoredMaterialRenderProxy(GEngine->EditorBrushMaterial->GetRenderProxy(false),Brush->GetWireColor());
-				PDI->RegisterDynamicResource( MaterialProxy );
-
-				// Flush the mesh triangles.
-				MeshBuilder.Draw(PDI, Brush->ActorToWorld().ToMatrixWithScale(), MaterialProxy, SDPG_World, 0.f);
-			}
-		}
-	}
-
-	const bool bIsInGameView = !Viewport->GetClient() || Viewport->GetClient()->IsInGameView();
-	if (Owner->ShouldDrawBrushVertices() && !bIsInGameView)
-	{
-		UTexture2D* VertexTexture = GetVertexTexture();
-		const float TextureSizeX = VertexTexture->GetSizeX() * 0.170f;
-		const float TextureSizeY = VertexTexture->GetSizeY() * 0.170f;
-
-		for (FSelectionIterator It(*Owner->GetSelectedActors()); It; ++It)
-		{
-			AActor* SelectedActor = static_cast<AActor*>(*It);
-			checkSlow(SelectedActor->IsA(AActor::StaticClass()));
-
-			ABrush* Brush = Cast< ABrush >(SelectedActor);
-			if (Brush && Brush->Brush && !FActorEditorUtils::IsABuilderBrush(Brush))
-			{
-				for (int32 p = 0; p < Brush->Brush->Polys->Element.Num(); ++p)
-				{
-					FTransform BrushTransform = Brush->ActorToWorld();
-
-					FPoly* poly = &Brush->Brush->Polys->Element[p];
-					for (int32 VertexIndex = 0; VertexIndex < poly->Vertices.Num(); ++VertexIndex)
-					{
-						const FVector& PolyVertex = poly->Vertices[VertexIndex];
-						const FVector WorldLocation = BrushTransform.TransformPosition(PolyVertex);
-						
-						const float Scale = View->WorldToScreen( WorldLocation ).W * ( 4.0f / View->ViewRect.Width() / View->ViewMatrices.ProjMatrix.M[0][0] );
-
-						const FColor Color(Brush->GetWireColor());
-						PDI->SetHitProxy(new HBSPBrushVert(Brush, &poly->Vertices[VertexIndex]));
-
-						PDI->DrawSprite(WorldLocation, TextureSizeX * Scale, TextureSizeY * Scale, VertexTexture->Resource, Color, SDPG_World, 0.0f, 0.0f, 0.0f, 0.0f, SE_BLEND_Masked );
-
-						PDI->SetHitProxy(NULL);
-			
-					}
-				}
-			}
-		}
-	}
-
-	// Let the current mode tool render if it wants to
-	FModeTool* tool = GetCurrentTool();
-	if( tool )
-	{
-		tool->Render( View, Viewport, PDI );
-	}
-
-	AGroupActor::DrawBracketsForGroups(PDI, Viewport);
-
-	if(UsesPropertyWidgets())
-	{
-		bool bHitTesting = PDI->IsHitTesting();
-		AActor* SelectedActor = GetFirstSelectedActorInstance();
-		if (SelectedActor != NULL)
-		{
-			UClass* Class = SelectedActor->GetClass();
-			TArray<FPropertyWidgetInfo> WidgetInfos;
-			GetPropertyWidgetInfos(Class, SelectedActor, WidgetInfos);
-			FEditorScriptExecutionGuard ScriptGuard;
-			for(int32 i=0; i<WidgetInfos.Num(); i++)
-			{
-				FString WidgetName = WidgetInfos[i].PropertyName;
-				FName WidgetValidator = WidgetInfos[i].PropertyValidationName;
-				int32 WidgetIndex = WidgetInfos[i].PropertyIndex;
-				bool bIsTransform = WidgetInfos[i].bIsTransform;
-
-				bool bSelected = (WidgetName == EditedPropertyName) && (WidgetIndex == EditedPropertyIndex);
-				FColor WidgetColor = bSelected ? FColor(255,255,255) : FColor(128,128,255);
-
-				FVector LocalPos = FVector::ZeroVector;
-				if(bIsTransform)
-				{
-					FTransform LocalTM = GetPropertyValueByName<FTransform>(SelectedActor, WidgetName, WidgetIndex);
-					LocalPos = LocalTM.GetLocation();
-				}
-				else
-				{
-					LocalPos = GetPropertyValueByName<FVector>(SelectedActor, WidgetName, WidgetIndex);
-				}
-
-				FTransform ActorToWorld = SelectedActor->ActorToWorld();
-				FVector WorldPos = ActorToWorld.TransformPosition(LocalPos);
-
-				UFunction* ValidateFunc= NULL;
-				if(WidgetValidator != NAME_None && 
-					(ValidateFunc = SelectedActor->FindFunction(WidgetValidator)) != NULL)
-				{
-					FString ReturnText;
-					SelectedActor->ProcessEvent(ValidateFunc, &ReturnText);
-
-					//if we have a negative result, the widget color is red.
-					WidgetColor = ReturnText.IsEmpty() ? WidgetColor : FColor::Red;
-				}
-
-				FTranslationMatrix WidgetTM(WorldPos);
-
-				const float WidgetSize = 0.035f;
-				const float ZoomFactor = FMath::Min<float>(View->ViewMatrices.ProjMatrix.M[0][0], View->ViewMatrices.ProjMatrix.M[1][1]);
-				const float WidgetRadius = View->Project(WorldPos).W * (WidgetSize / ZoomFactor);
-
-				if(bHitTesting) PDI->SetHitProxy( new HPropertyWidgetProxy(WidgetName, WidgetIndex, bIsTransform) );
-				DrawWireDiamond(PDI, WidgetTM, WidgetRadius, WidgetColor, SDPG_Foreground );
-				if(bHitTesting) PDI->SetHitProxy( NULL );
-			}
-		}
-	}
-}
-
-void FEdMode::DrawHUD(FEditorViewportClient* ViewportClient,FViewport* Viewport,const FSceneView* View,FCanvas* Canvas)
-{
-	// Render the drag tool.
-	ViewportClient->RenderDragTool( View, Canvas );
-
-	// Let the current mode tool draw a HUD if it wants to
-	FModeTool* tool = GetCurrentTool();
-	if( tool )
-	{
-		tool->DrawHUD( ViewportClient, Viewport, View, Canvas );
-	}
-
-	if (ViewportClient->IsPerspective() && GetDefault<ULevelEditorViewportSettings>()->bHighlightWithBrackets)
-	{
-		DrawBrackets( ViewportClient, Viewport, View, Canvas );
-	}
-
-	// If this viewport doesn't show mode widgets or the mode itself doesn't want them, leave.
-	if( !(ViewportClient->EngineShowFlags.ModeWidgets) || !ShowModeWidgets() )
-	{
-		return;
-	}
-
-	// Clear Hit proxies
-	const bool bIsHitTesting = Canvas->IsHitTesting();
-	if ( !bIsHitTesting )
-	{
-		Canvas->SetHitProxy(NULL);
-	}
-
-	// Draw vertices for selected BSP brushes and static meshes if the large vertices show flag is set.
-	if ( !ViewportClient->bDrawVertices )
-	{
-		return;
-	}
-
-	const bool bLargeVertices		= View->Family->EngineShowFlags.LargeVertices;
-	const bool bShowBrushes			= View->Family->EngineShowFlags.Brushes;
-	const bool bShowBSP				= View->Family->EngineShowFlags.BSP;
-	const bool bShowBuilderBrush	= View->Family->EngineShowFlags.BuilderBrush != 0;
-
-	UTexture2D* VertexTexture = GetVertexTexture();
-	const float TextureSizeX		= VertexTexture->GetSizeX() * ( bLargeVertices ? 1.0f : 0.5f );
-	const float TextureSizeY		= VertexTexture->GetSizeY() * ( bLargeVertices ? 1.0f : 0.5f );
-
-	// Temporaries.
-	TArray<FVector> Vertices;
-
-	for ( FSelectionIterator It( *Owner->GetSelectedActors() ) ; It ; ++It )
-	{
-		AActor* SelectedActor = static_cast<AActor*>( *It );
-		checkSlow( SelectedActor->IsA(AActor::StaticClass()) );
-
-		if( bLargeVertices )
-		{
-			FCanvasItemTestbed::bTestState = !FCanvasItemTestbed::bTestState;
-
-			// Static mesh vertices
-			AStaticMeshActor* Actor = Cast<AStaticMeshActor>( SelectedActor );
-			if( Actor && Actor->GetStaticMeshComponent() && Actor->GetStaticMeshComponent()->StaticMesh
-				&& Actor->GetStaticMeshComponent()->StaticMesh->RenderData )
-			{
-				FTransform ActorToWorld = Actor->ActorToWorld();
-				Vertices.Empty();
-				const FPositionVertexBuffer& VertexBuffer = Actor->GetStaticMeshComponent()->StaticMesh->RenderData->LODResources[0].PositionVertexBuffer;
-				for( uint32 i = 0 ; i < VertexBuffer.GetNumVertices() ; i++ )
-				{
-					Vertices.AddUnique( ActorToWorld.TransformPosition( VertexBuffer.VertexPosition(i) ) );
-				}
-
-				FCanvasTileItem TileItem( FVector2D( 0.0f, 0.0f ), FVector2D( 0.0f, 0.0f ), FLinearColor::White );
-				TileItem.BlendMode = SE_BLEND_Translucent;
-				for( int32 VertexIndex = 0 ; VertexIndex < Vertices.Num() ; ++VertexIndex )
-				{				
-					const FVector& Vertex = Vertices[VertexIndex];
-					FVector2D PixelLocation;
-					if(View->ScreenToPixel(View->WorldToScreen(Vertex),PixelLocation))
-					{
-						const bool bOutside =
-							PixelLocation.X < 0.0f || PixelLocation.X > View->ViewRect.Width() ||
-							PixelLocation.Y < 0.0f || PixelLocation.Y > View->ViewRect.Height();
-						if( !bOutside )
-						{
-							const float X = PixelLocation.X - (TextureSizeX/2);
-							const float Y = PixelLocation.Y - (TextureSizeY/2);
-							if( bIsHitTesting ) 
-							{
-								Canvas->SetHitProxy( new HStaticMeshVert(Actor,Vertex) );
-							}
-							TileItem.Texture = VertexTexture->Resource;
-							
-							TileItem.Size = FVector2D( TextureSizeX, TextureSizeY );
-							Canvas->DrawItem( TileItem, FVector2D( X, Y ) );							
-							if( bIsHitTesting )
-							{
-								Canvas->SetHitProxy( NULL );
-							}
-						}
-					}
-				}
-			}
-		}
-	}
-
-	if(UsesPropertyWidgets())
-	{
-		AActor* SelectedActor = GetFirstSelectedActorInstance();
-		if (SelectedActor != NULL)
-		{
-			FEditorScriptExecutionGuard ScriptGuard;
-
-			const int32 HalfX = 0.5f * Viewport->GetSizeXY().X;
-			const int32 HalfY = 0.5f * Viewport->GetSizeXY().Y;
-
-			UClass* Class = SelectedActor->GetClass();		
-			TArray<FPropertyWidgetInfo> WidgetInfos;
-			GetPropertyWidgetInfos(Class, SelectedActor, WidgetInfos);
-			for(int32 i=0; i<WidgetInfos.Num(); i++)
-			{
-				FString WidgetName = WidgetInfos[i].PropertyName;
-				FName WidgetValidator = WidgetInfos[i].PropertyValidationName;
-				int32 WidgetIndex = WidgetInfos[i].PropertyIndex;
-				bool bIsTransform = WidgetInfos[i].bIsTransform;
-
-				FVector LocalPos = FVector::ZeroVector;
-				if(bIsTransform)
-				{
-					FTransform LocalTM = GetPropertyValueByName<FTransform>(SelectedActor, WidgetName, WidgetIndex);
-					LocalPos = LocalTM.GetLocation();
-				}
-				else
-				{
-					LocalPos = GetPropertyValueByName<FVector>(SelectedActor, WidgetName, WidgetIndex);
-				}
-
-				FTransform ActorToWorld = SelectedActor->ActorToWorld();
-				FVector WorldPos = ActorToWorld.TransformPosition(LocalPos);
-
-				UFunction* ValidateFunc = NULL;
-				FString FinalString;
-				if(WidgetValidator != NAME_None && 
-					(ValidateFunc = SelectedActor->FindFunction(WidgetValidator)) != NULL)
-				{
-					SelectedActor->ProcessEvent(ValidateFunc, &FinalString);
-				}
-
-				const FPlane Proj = View->Project( WorldPos );
-
-				//do some string fixing
-				const uint32 VectorIndex = WidgetInfos[i].PropertyIndex;
-				const FString WidgetDisplayName = WidgetInfos[i].DisplayName + ((VectorIndex != INDEX_NONE) ? FString::Printf(TEXT("[%d]"), VectorIndex) : TEXT(""));
-				FinalString = FinalString.IsEmpty() ? WidgetDisplayName : FinalString;
-
-				if(Proj.W > 0.f)
-				{
-					const int32 XPos = HalfX + ( HalfX * Proj.X );
-					const int32 YPos = HalfY + ( HalfY * (Proj.Y * -1.f) );
-					FCanvasTextItem TextItem( FVector2D( XPos + 5, YPos), FText::FromString( FinalString ), GEngine->GetSmallFont(), FLinearColor::White );
-					Canvas->DrawItem( TextItem );
-				}
-			}
-		}
-	}
-}
-
-// Draw brackets around all selected objects
-void FEdMode::DrawBrackets( FEditorViewportClient* ViewportClient, FViewport* Viewport, const FSceneView* View, FCanvas* Canvas )
-{
-	USelection& SelectedActors = *Owner->GetSelectedActors();
-	for( int32 CurSelectedActorIndex = 0; CurSelectedActorIndex < SelectedActors.Num(); ++CurSelectedActorIndex )
-	{
-		AActor* SelectedActor = Cast<AActor>( SelectedActors.GetSelectedObject(CurSelectedActorIndex ) );
-		if( SelectedActor != NULL )
-		{
-			// Draw a bracket for selected "paintable" static mesh actors
-			const bool bIsValidActor = ( Cast< AStaticMeshActor >( SelectedActor ) != NULL );
-
-			const FLinearColor SelectedActorBoxColor( 0.6f, 0.6f, 1.0f );
-			const bool bDrawBracket = bIsValidActor;
-			ViewportClient->DrawActorScreenSpaceBoundingBox( Canvas, View, Viewport, SelectedActor, SelectedActorBoxColor, bDrawBracket );
-		}
-	}
-}
-
-bool FEdMode::UsesToolkits() const
-{
-	return false;
-}
-
-bool FEdMode::StartTracking(FEditorViewportClient* InViewportClient, FViewport* InViewport)
-{
-	bool bResult = false;
-	if( GetCurrentTool() )
-	{
-		bResult = GetCurrentTool()->StartModify();
-	}
-	return bResult;
-}
-
-bool FEdMode::EndTracking(FEditorViewportClient* InViewportClient, FViewport* InViewport)
-{
-	bool bResult = false;
-	if( GetCurrentTool() )
-	{
-		bResult = GetCurrentTool()->EndModify();
-	}
-	return bResult;
-}
-
-FVector FEdMode::GetWidgetNormalFromCurrentAxis( void* InData )
-{
-	// Figure out the proper coordinate system.
-
-	FMatrix matrix = FMatrix::Identity;
-	if( Owner->GetCoordSystem() == COORD_Local )
-	{
-		GetCustomDrawingCoordinateSystem( matrix, InData );
-	}
-
-	// Get a base normal from the current axis.
-
-	FVector BaseNormal(1,0,0);		// Default to X axis
-	switch( CurrentWidgetAxis )
-	{
-		case EAxisList::Y:	BaseNormal = FVector(0,1,0);	break;
-		case EAxisList::Z:	BaseNormal = FVector(0,0,1);	break;
-		case EAxisList::XY:	BaseNormal = FVector(1,1,0);	break;
-		case EAxisList::XZ:	BaseNormal = FVector(1,0,1);	break;
-		case EAxisList::YZ:	BaseNormal = FVector(0,1,1);	break;
-		case EAxisList::XYZ:	BaseNormal = FVector(1,1,1);	break;
-	}
-
-	// Transform the base normal into the proper coordinate space.
-	return matrix.TransformPosition( BaseNormal );
-}
-
-AActor* FEdMode::GetFirstSelectedActorInstance() const
-{
-	return Owner->GetSelectedActors()->GetTop<AActor>();
-}
-
-bool FEdMode::CanCreateWidgetForStructure(const UStruct* InPropStruct)
-{
-	return InPropStruct && (InPropStruct->GetFName() == NAME_Vector || InPropStruct->GetFName() == NAME_Transform);
-}
-
-bool FEdMode::CanCreateWidgetForProperty(UProperty* InProp)
-{
-	UStructProperty* TestProperty = Cast<UStructProperty>(InProp);
-	if( !TestProperty )
-	{
-		UArrayProperty* ArrayProperty = Cast<UArrayProperty>(InProp);
-		if( ArrayProperty )
-		{
-			TestProperty = Cast<UStructProperty>(ArrayProperty->Inner);
-		}
-	}
-	return (TestProperty != NULL) && CanCreateWidgetForStructure(TestProperty->Struct);
-}
-
-bool FEdMode::ShouldCreateWidgetForProperty(UProperty* InProp)
-{
-	return CanCreateWidgetForProperty(InProp) && InProp->HasMetaData(MD_MakeEditWidget);
-}
-
-static bool IsTransformProperty(UProperty* InProp)
-{
-	UStructProperty* StructProp = Cast<UStructProperty>(InProp);
-	return (StructProp != NULL && StructProp->Struct->GetFName() == NAME_Transform);
-
-}
-
-void FEdMode::GetPropertyWidgetInfos(const UStruct* InStruct, const void* InContainer, TArray<FPropertyWidgetInfo>& OutInfos, FString PropertyNamePrefix, FString DisplayNamePrefix) const
-{
-	if(PropertyNamePrefix.Len() == 0)
-	{
-		OutInfos.Empty();
-	}
-
-	for (TFieldIterator<UProperty> PropertyIt(InStruct, EFieldIteratorFlags::IncludeSuper); PropertyIt; ++PropertyIt)
-	{
-		UProperty* CurrentProp = *PropertyIt;
-		check(CurrentProp);
-		FString DisplayName = CurrentProp->GetMetaData(TEXT("DisplayName"));
-		if (!PropertyNamePrefix.IsEmpty() && DisplayNamePrefix.IsEmpty()) // Display name is already invalid
-		{
-			DisplayName.Empty();
-		}
-		if (!DisplayName.IsEmpty()) //Display name cannot be only the prefix.
-		{
-			DisplayName = DisplayNamePrefix + DisplayName;
-		}
-
-		if(	ShouldCreateWidgetForProperty(CurrentProp) )
-		{
-			if( UArrayProperty* ArrayProp = Cast<UArrayProperty>(CurrentProp) )
-			{
-				check(InContainer != NULL);
-
-				FScriptArrayHelper_InContainer ArrayHelper(ArrayProp, InContainer);
-
-				// See how many widgets we need to make for the array property
-				uint32 ArrayDim = ArrayHelper.Num();
-				for( uint32 i = 0; i < ArrayDim; i++ )
-				{
-					//create a new widget info struct
-					FPropertyWidgetInfo WidgetInfo;
-
-					//fill it in with the struct name
-					WidgetInfo.PropertyName = PropertyNamePrefix + CurrentProp->GetFName().ToString();
-					WidgetInfo.DisplayName = DisplayName.IsEmpty() ? WidgetInfo.PropertyName : DisplayName;
-
-					//And see if we have any meta data that matches the MD_ValidateWidgetUsing name
-					WidgetInfo.PropertyValidationName = FName(*CurrentProp->GetMetaData(MD_ValidateWidgetUsing));
-
-					WidgetInfo.PropertyIndex = i;
-
-					// See if its a transform
-					WidgetInfo.bIsTransform = IsTransformProperty(ArrayProp->Inner);
-
-					//Add it to our out array
-					OutInfos.Add(WidgetInfo);
-				}
-			}
-			else
-			{
-
-				//create a new widget info struct
-				FPropertyWidgetInfo WidgetInfo;
-
-				//fill it in with the struct name
-				WidgetInfo.PropertyName = PropertyNamePrefix + CurrentProp->GetFName().ToString();
-				WidgetInfo.DisplayName = DisplayName.IsEmpty() ? WidgetInfo.PropertyName : DisplayName;
-
-				//And see if we have any meta data that matches the MD_ValidateWidgetUsing name
-				WidgetInfo.PropertyValidationName = FName(*CurrentProp->GetMetaData(MD_ValidateWidgetUsing));
-
-				// See if its a transform
-				WidgetInfo.bIsTransform = IsTransformProperty(CurrentProp);
-
-				//Add it to our out array
-				OutInfos.Add(WidgetInfo);
-			}
-
-		}
-		else
-		{
-			UStructProperty* StructProp = Cast<UStructProperty>(CurrentProp);
-			if(StructProp != NULL)
-			{
-				// Recursively traverse into structures, looking for additional vector properties to expose
-				GetPropertyWidgetInfos(StructProp->Struct
-					, StructProp->ContainerPtrToValuePtr<void>(InContainer)
-					, OutInfos
-					, PropertyNamePrefix + StructProp->GetFName().ToString() + TEXT(".")
-					, !DisplayName.IsEmpty() ? (DisplayName + TEXT(".")) : FString());
-			}
-			else
-			{
-				// Recursively traverse into arrays of structures, looking for additional vector properties to expose
-				UArrayProperty* ArrayProp = Cast<UArrayProperty>(CurrentProp);
-				if(ArrayProp != NULL)
-				{
-					StructProp = Cast<UStructProperty>(ArrayProp->Inner);
-					if(StructProp != NULL)
-					{
-						FScriptArrayHelper_InContainer ArrayHelper(ArrayProp, InContainer);
-						for(int32 ArrayIndex = 0; ArrayIndex < ArrayHelper.Num(); ++ArrayIndex)
-						{
-							if(ArrayHelper.IsValidIndex(ArrayIndex))
-							{
-								const FString ArrayPostfix = FString::Printf(TEXT("[%d]"), ArrayIndex) + TEXT(".");
-								GetPropertyWidgetInfos(StructProp->Struct
-									, ArrayHelper.GetRawPtr(ArrayIndex)
-									, OutInfos
-									, PropertyNamePrefix + ArrayProp->GetFName().ToString() + ArrayPostfix
-									, !DisplayName.IsEmpty() ? (DisplayName + ArrayPostfix) : FString());
-							}
-						}
-					}
-				}
-			}
-		}
-	}
-}
-
-bool FEdMode::IsSnapRotationEnabled()
-{
-	return GetDefault<ULevelEditorViewportSettings>()->RotGridEnabled;
-}
-
 /*------------------------------------------------------------------------------
 	Default.
 ------------------------------------------------------------------------------*/
@@ -1170,1077 +53,3 @@ FEdModeDefault::FEdModeDefault()
 	bDrawKillZ = true;
 }
 
-/*------------------------------------------------------------------------------
-	FEditorModeTools.
-
-	The master class that handles tracking of the current mode.
-------------------------------------------------------------------------------*/
-
-FEditorModeTools::FEditorModeTools()
-	:	PivotShown( 0 )
-	,	Snapping( 0 )
-	,	SnappedActor( 0 )
-	,	TranslateRotateXAxisAngle(0)
-	,	DefaultID(FBuiltinEditorModes::EM_Default)
-	,	WidgetMode( FWidget::WM_Translate )
-	,	OverrideWidgetMode( FWidget::WM_None )
-	,	bShowWidget( 1 )
-	,	bHideViewportUI(false)
-	,	CoordSystem(COORD_World)
-	,	bIsTracking(false)
-{
-	// Load the last used settings
-	LoadConfig();
-
-	// Register our callback for actor selection changes
-	USelection::SelectNoneEvent.AddRaw(this, &FEditorModeTools::OnEditorSelectNone);
-	USelection::SelectionChangedEvent.AddRaw(this, &FEditorModeTools::OnEditorSelectionChanged);
-	USelection::SelectObjectEvent.AddRaw(this, &FEditorModeTools::OnEditorSelectionChanged);
-}
-
-FEditorModeTools::~FEditorModeTools()
-{
-	// Should we call Exit on any modes that are still active, or is it too late?
-	USelection::SelectionChangedEvent.RemoveAll(this);
-	USelection::SelectNoneEvent.RemoveAll(this);
-	USelection::SelectObjectEvent.RemoveAll(this);
-}
-
-/**
- * Loads the state that was saved in the INI file
- */
-void FEditorModeTools::LoadConfig(void)
-{
-	GConfig->GetBool(TEXT("FEditorModeTools"),TEXT("ShowWidget"),bShowWidget,
-		GEditorUserSettingsIni);
-
-	const bool bGetRawValue = true;
-	int32 Bogus = (int32)GetCoordSystem(bGetRawValue);
-	GConfig->GetInt(TEXT("FEditorModeTools"),TEXT("CoordSystem"),Bogus,
-		GEditorUserSettingsIni);
-	SetCoordSystem((ECoordSystem)Bogus);
-
-	
-	GConfig->GetBool( TEXT("UnEdViewport"), TEXT("InterpEdPanInvert"), bInterpPanInverted, GEditorUserSettingsIni );
-	LoadWidgetSettings();
-}
-
-/**
- * Saves the current state to the INI file
- */
-void FEditorModeTools::SaveConfig(void)
-{
-	GConfig->SetBool(TEXT("FEditorModeTools"),TEXT("ShowWidget"),bShowWidget,
-		GEditorUserSettingsIni);
-
-	const bool bGetRawValue = true;
-	GConfig->SetInt(TEXT("FEditorModeTools"),TEXT("CoordSystem"),(int32)GetCoordSystem(bGetRawValue),
-		GEditorUserSettingsIni);
-
-	GConfig->SetBool( TEXT("UnEdViewport"), TEXT("InterpEdPanInvert"), bInterpPanInverted, GEditorUserSettingsIni );
-
-	SaveWidgetSettings();
-}
-
-TSharedPtr<class IToolkitHost> FEditorModeTools::GetToolkitHost() const
-{
-	TSharedPtr<class IToolkitHost> Result = ToolkitHost.Pin();
-	check(ToolkitHost.IsValid());
-	return Result;
-}
-
-void FEditorModeTools::SetToolkitHost(TSharedRef<class IToolkitHost> InHost)
-{
-	checkf(!ToolkitHost.IsValid(), TEXT("SetToolkitHost can only be called once"));
-	ToolkitHost = InHost;
-}
-
-class USelection* FEditorModeTools::GetSelectedActors() const
-{
-	return GEditor->GetSelectedActors();
-}
-
-class USelection* FEditorModeTools::GetSelectedObjects() const
-{
-	return GEditor->GetSelectedObjects();
-}
-
-void FEditorModeTools::OnEditorSelectionChanged(UObject* NewSelection)
-{
-	// If selecting an actor, move the pivot location.
-	AActor* Actor = Cast<AActor>(NewSelection);
-	if (Actor != NULL)
-	{
-		//@fixme - why isn't this using UObject::IsSelected()?
-		if ( GEditor->GetSelectedActors()->IsSelected( Actor ) )
-		{
-			SetPivotLocation( Actor->GetActorLocation(), false );
-
-			// If this actor wasn't part of the original selection set during pie/sie, clear it now
-			if ( GEditor->ActorsThatWereSelected.Num() > 0 )
-			{
-				AActor* EditorActor = EditorUtilities::GetEditorWorldCounterpartActor( Actor );
-				if ( !EditorActor || !GEditor->ActorsThatWereSelected.Contains(EditorActor) )
-				{
-					GEditor->ActorsThatWereSelected.Empty();
-				}
-			}
-		}
-		else if ( GEditor->ActorsThatWereSelected.Num() > 0 )
-		{
-			// Clear the selection set
-			GEditor->ActorsThatWereSelected.Empty();
-		}
-	}
-
-	for (const auto& Pair : FEditorModeRegistry::Get().GetFactoryMap())
-	{
-		Pair.Value->OnSelectionChanged(*this, NewSelection);
-	}
-}
-
-void FEditorModeTools::OnEditorSelectNone()
-{
-	GEditor->SelectNone( false, true );
-	GEditor->ActorsThatWereSelected.Empty();
-}
-
-/** 
- * Sets the pivot locations
- * 
- * @param Location 		The location to set
- * @param bIncGridBase	Whether or not to also set the GridBase
- */
-void FEditorModeTools::SetPivotLocation( const FVector& Location, const bool bIncGridBase )
-{
-	CachedLocation = PivotLocation = SnappedLocation = Location;
-	if ( bIncGridBase )
-	{
-		GridBase = Location;
-	}
-}
-
-ECoordSystem FEditorModeTools::GetCoordSystem(bool bGetRawValue)
-{
-	if (!bGetRawValue && GetWidgetMode() == FWidget::WM_Scale )
-	{
-		return COORD_Local;
-	}
-	else
-	{
-		return CoordSystem;
-	}
-}
-
-void FEditorModeTools::SetCoordSystem(ECoordSystem NewCoordSystem)
-{
-	CoordSystem = NewCoordSystem;
-}
-
-void FEditorModeTools::SetDefaultMode ( FEditorModeID InDefaultID )
-{
-	DefaultID = InDefaultID;
-}
-
-void FEditorModeTools::ActivateDefaultMode()
-{
-	ActivateMode( DefaultID );
-
-	check( IsModeActive( DefaultID ) );
-}
-
-void FEditorModeTools::DeactivateModeAtIndex(int32 InIndex)
-{
-	check( InIndex >= 0 && InIndex < Modes.Num() );
-
-	auto& Mode = Modes[InIndex];
-
-	Mode->Exit();
-	RecycledModes.Add( Mode->GetID(), Mode );
-	Modes.RemoveAt( InIndex );
-}
-
-/**
- * Deactivates an editor mode. 
- * 
- * @param InID		The ID of the editor mode to deactivate.
- */
-void FEditorModeTools::DeactivateMode( FEditorModeID InID )
-{
-	// Find the mode from the ID and exit it.
-	for( int32 Index = Modes.Num() - 1; Index >= 0; --Index )
-	{
-		auto& Mode = Modes[Index];
-		if( Mode->GetID() == InID )
-		{
-			DeactivateModeAtIndex(Index);
-			break;
-		}
-	}
-
-	if( Modes.Num() == 0 )
-	{
-		// Ensure the default mode is active if there are no active modes.
-		ActivateDefaultMode();
-	}
-}
-
-void FEditorModeTools::DeactivateAllModes()
-{
-	for( int32 Index = Modes.Num() - 1; Index >= 0; --Index )
-	{
-		DeactivateModeAtIndex(Index);
-	}
-}
-
-void FEditorModeTools::DestroyMode( FEditorModeID InID )
-{
-	// Find the mode from the ID and exit it.
-	for( int32 Index = Modes.Num() - 1; Index >= 0; --Index )
-	{
-		auto& Mode = Modes[Index];
-		if ( Mode->GetID() == InID )
-		{
-			// Deactivate and destroy
-			DeactivateModeAtIndex(Index);
-			break;
-		}
-	}
-
-	RecycledModes.Remove(InID);
-}
-
-/**
- * Activates an editor mode. Shuts down all other active modes which cannot run with the passed in mode.
- * 
- * @param InID		The ID of the editor mode to activate.
- * @param bToggle	true if the passed in editor mode should be toggled off if it is already active.
- */
-void FEditorModeTools::ActivateMode( FEditorModeID InID, bool bToggle )
-{
-	if (InID == FBuiltinEditorModes::EM_Default)
-	{
-		InID = DefaultID;
-	}
-
-	// Check to see if the mode is already active
-	if( IsModeActive(InID) )
-	{
-		// The mode is already active toggle it off if we should toggle off already active modes.
-		if( bToggle )
-		{
-			DeactivateMode( InID );
-		}
-		// Nothing more to do
-		return;
-	}
-
-	// Recycle a mode or factory a new one
-	TSharedPtr<FEdMode> Mode = RecycledModes.FindRef( InID );
-
-	if ( Mode.IsValid() )
-	{
-		RecycledModes.Remove( InID );
-	}
-	else
-	{
-		Mode = FEditorModeRegistry::Get().CreateMode( InID, *this );
-	}
-
-	if( !Mode.IsValid() )
-	{
-		UE_LOG(LogEditorModes, Log, TEXT("FEditorModeTools::ActivateMode : Couldn't find mode '%s'."), *InID.ToString() );
-		// Just return and leave the mode list unmodified
-		return;
-	}
-
-	// Remove anything that isn't compatible with this mode
-	for( int32 ModeIndex = Modes.Num() - 1; ModeIndex >= 0; --ModeIndex )
-	{
-		const bool bModesAreCompatible = Mode->IsCompatibleWith( Modes[ModeIndex]->GetID() ) || Modes[ModeIndex]->IsCompatibleWith( Mode->GetID() );
-		if ( !bModesAreCompatible )
-		{
-			DeactivateModeAtIndex(ModeIndex);
-		}
-	}
-
-	Modes.Add( Mode );
-
-	// Enter the new mode
-	Mode->Enter();
-	
-	// Update the editor UI
-	FEditorSupportDelegates::UpdateUI.Broadcast();
-}
-
-bool FEditorModeTools::EnsureNotInMode(FEditorModeID ModeID, const FText& ErrorMsg, bool bNotifyUser) const
-{
-	// We're in a 'safe' mode if we're not in the specified mode.
-	const bool bInASafeMode = !IsModeActive(ModeID);
-	if( !bInASafeMode && !ErrorMsg.IsEmpty() )
-	{
-		// Do we want to display this as a notification or a dialog to the user
-		if ( bNotifyUser )
-		{
-			FNotificationInfo Info( ErrorMsg );
-			FSlateNotificationManager::Get().AddNotification( Info );
-		}
-		else
-		{
-			FMessageDialog::Open( EAppMsgType::Ok, ErrorMsg );
-		}		
-	}
-	return bInASafeMode;
-}
-
-FEdMode* FEditorModeTools::FindMode( FEditorModeID InID )
-{
-	for( auto& Mode : Modes )
-	{
-		if( Mode->GetID() == InID )
-		{
-			return Mode.Get();
-		}
-	}
-
-	return NULL;
-}
-
-/**
- * Returns a coordinate system that should be applied on top of the worldspace system.
- */
-
-FMatrix FEditorModeTools::GetCustomDrawingCoordinateSystem()
-{
-	FMatrix Matrix = FMatrix::Identity;
-
-	switch (GetCoordSystem())
-	{
-		case COORD_Local:
-		{
-			// Let the current mode have a shot at setting the local coordinate system.
-			// If it doesn't want to, create it by looking at the currently selected actors list.
-
-			bool CustomCoordinateSystemProvided = false;
-			for (const auto& Mode : Modes)
-			{
-				if (Mode->GetCustomDrawingCoordinateSystem(Matrix, nullptr))
-				{
-					CustomCoordinateSystemProvided = true;
-					break;
-				}
-			}
-
-			if (!CustomCoordinateSystemProvided)
-			{
-				const int32 Num = GetSelectedActors()->CountSelections<AActor>();
-
-				// Coordinate system needs to come from the last actor selected
-				if (Num > 0)
-				{
-					Matrix = FRotationMatrix(GetSelectedActors()->GetBottom<AActor>()->GetActorRotation());
-				}
-			}
-
-			if (!Matrix.Equals(FMatrix::Identity))
-			{
-				Matrix.RemoveScaling();
-			}
-		}
-		break;
-
-		case COORD_World:
-			break;
-
-		default:
-			break;
-	}
-
-	return Matrix;
-}
-
-FMatrix FEditorModeTools::GetCustomInputCoordinateSystem()
-{
-	return GetCustomDrawingCoordinateSystem();
-}
-
-/** Gets the widget axis to be drawn */
-EAxisList::Type FEditorModeTools::GetWidgetAxisToDraw( FWidget::EWidgetMode InWidgetMode ) const
-{
-	EAxisList::Type OutAxis = EAxisList::All;
-	for( int Index = Modes.Num() - 1; Index >= 0 ; Index-- )
-	{
-		if ( Modes[Index]->ShouldDrawWidget() )
-		{
-			OutAxis = Modes[Index]->GetWidgetAxisToDraw( InWidgetMode );
-			break;
-		}
-	}
-
-	return OutAxis;
-}
-
-/** Mouse tracking interface.  Passes tracking messages to all active modes */
-bool FEditorModeTools::StartTracking(FEditorViewportClient* InViewportClient, FViewport* InViewport)
-{
-	bIsTracking = true;
-	bool bTransactionHandled = false;
-
-	CachedLocation = PivotLocation;	// Cache the pivot location
-
-	for ( const auto& Mode : Modes)
-	{
-		bTransactionHandled |= Mode->StartTracking(InViewportClient, InViewport);
-	}
-
-	return bTransactionHandled;
-}
-
-/** Mouse tracking interface.  Passes tracking messages to all active modes */
-bool FEditorModeTools::EndTracking(FEditorViewportClient* InViewportClient, FViewport* InViewport)
-{
-	bIsTracking = false;
-	bool bTransactionHandled = false;
-
-	for ( const auto& Mode : Modes)
-	{
-		bTransactionHandled |= Mode->EndTracking(InViewportClient, InViewportClient->Viewport);
-	}
-
-	CachedLocation = PivotLocation;	// Clear the pivot location
-	
-	return bTransactionHandled;
-}
-
-bool FEditorModeTools::AllowsViewportDragTool() const
-{
-	bool bCanUseDragTool = false;
-	for (const TSharedPtr<FEdMode>& Mode : Modes)
-	{
-		bCanUseDragTool |= Mode->AllowsViewportDragTool();
-	}
-	return bCanUseDragTool;
-}
-
-/** Notifies all active modes that a map change has occured */
-void FEditorModeTools::MapChangeNotify()
-{
-	for ( const auto& Mode : Modes)
-	{
-		Mode->MapChangeNotify();
-	}
-}
-
-
-/** Notifies all active modes to empty their selections */
-void FEditorModeTools::SelectNone()
-{
-	for ( const auto& Mode : Modes)
-	{
-		Mode->SelectNone();
-	}
-}
-
-/** Notifies all active modes of box selection attempts */
-bool FEditorModeTools::BoxSelect( FBox& InBox, bool InSelect )
-{
-	bool bHandled = false;
-	for( const auto& Mode : Modes)
-	{
-		bHandled |= Mode->BoxSelect( InBox, InSelect );
-	}
-	return bHandled;
-}
-
-/** Notifies all active modes of frustum selection attempts */
-bool FEditorModeTools::FrustumSelect( const FConvexVolume& InFrustum, bool InSelect )
-{
-	bool bHandled = false;
-	for( const auto& Mode : Modes)
-	{
-		bHandled |= Mode->FrustumSelect( InFrustum, InSelect );
-	}
-	return bHandled;
-}
-
-
-/** true if any active mode uses a transform widget */
-bool FEditorModeTools::UsesTransformWidget() const
-{
-	bool bUsesTransformWidget = false;
-	for( const auto& Mode : Modes)
-	{
-		bUsesTransformWidget |= Mode->UsesTransformWidget();
-	}
-
-	return bUsesTransformWidget;
-}
-
-/** true if any active mode uses the passed in transform widget */
-bool FEditorModeTools::UsesTransformWidget( FWidget::EWidgetMode CheckMode ) const
-{
-	bool bUsesTransformWidget = false;
-	for( const auto& Mode : Modes)
-	{
-		bUsesTransformWidget |= Mode->UsesTransformWidget(CheckMode);
-	}
-
-	return bUsesTransformWidget;
-}
-
-/** Sets the current widget axis */
-void FEditorModeTools::SetCurrentWidgetAxis( EAxisList::Type NewAxis )
-{
-	for( const auto& Mode : Modes)
-	{
-		Mode->SetCurrentWidgetAxis( NewAxis );
-	}
-}
-
-/** Notifies all active modes of mouse click messages. */
-bool FEditorModeTools::HandleClick(FEditorViewportClient* InViewportClient,  HHitProxy *HitProxy, const FViewportClick& Click )
-{
-	bool bHandled = false;
-	for( const auto& Mode : Modes)
-	{
-		bHandled |= Mode->HandleClick(InViewportClient, HitProxy, Click);
-	}
-
-	return bHandled;
-}
-
-/** true if the passed in brush actor should be drawn in wireframe */	
-bool FEditorModeTools::ShouldDrawBrushWireframe( AActor* InActor ) const
-{
-	bool bShouldDraw = false;
-	for( const auto& Mode : Modes)
-	{
-		bShouldDraw |= Mode->ShouldDrawBrushWireframe( InActor );
-	}
-
-	if( Modes.Num() == 0 )
-	{
-		// We can get into a state where there are no active modes at editor startup if the builder brush is created before the default mode is activated.
-		// Ensure we can see the builder brush when no modes are active.
-		bShouldDraw = true;
-	}
-	return bShouldDraw;
-}
-
-/** true if brush vertices should be drawn */
-bool FEditorModeTools::ShouldDrawBrushVertices() const
-{
-	// Currently only geometry mode being active prevents vertices from being drawn.
-	return !IsModeActive( FBuiltinEditorModes::EM_Geometry );
-}
-
-/** Ticks all active modes */
-void FEditorModeTools::Tick( FEditorViewportClient* ViewportClient, float DeltaTime )
-{
-	// Remove anything pending destruction
-	for( int32 Index = Modes.Num() - 1; Index >= 0; --Index)
-	{
-		if (Modes[Index]->IsPendingDeletion())
-		{
-			DeactivateModeAtIndex(Index);
-		}
-	}
-	
-	if (Modes.Num() == 0)
-	{
-		// Ensure the default mode is active if there are no active modes.
-		ActivateDefaultMode();
-	}
-
-	for( const auto& Mode : Modes)
-	{
-		Mode->Tick( ViewportClient, DeltaTime );
-	}
-}
-
-/** Notifies all active modes of any change in mouse movement */
-bool FEditorModeTools::InputDelta( FEditorViewportClient* InViewportClient,FViewport* InViewport,FVector& InDrag,FRotator& InRot,FVector& InScale )
-{
-	bool bHandled = false;
-	for( const auto& Mode : Modes)
-	{
-		bHandled |= Mode->InputDelta( InViewportClient, InViewport, InDrag, InRot, InScale );
-	}
-	return bHandled;
-}
-
-/** Notifies all active modes of captured mouse movement */	
-bool FEditorModeTools::CapturedMouseMove( FEditorViewportClient* InViewportClient, FViewport* InViewport, int32 InMouseX, int32 InMouseY )
-{
-	bool bHandled = false;
-	for( const auto& Mode : Modes)
-	{
-		bHandled |= Mode->CapturedMouseMove( InViewportClient, InViewport, InMouseX, InMouseY );
-	}
-	return bHandled;
-}
-
-/** Notifies all active modes of keyboard input */
-bool FEditorModeTools::InputKey(FEditorViewportClient* InViewportClient, FViewport* Viewport, FKey Key, EInputEvent Event)
-{
-	bool bHandled = false;
-	for (const auto& Mode : Modes)
-	{
-		bHandled |= Mode->InputKey( InViewportClient, Viewport, Key, Event );
-	}
-	return bHandled;
-}
-
-/** Notifies all active modes of axis movement */
-bool FEditorModeTools::InputAxis(FEditorViewportClient* InViewportClient, FViewport* Viewport, int32 ControllerId, FKey Key, float Delta, float DeltaTime)
-{
-	bool bHandled = false;
-	for( const auto& Mode : Modes)
-	{
-		bHandled |= Mode->InputAxis( InViewportClient, Viewport, ControllerId, Key, Delta, DeltaTime );
-	}
-	return bHandled;
-}
-
-bool FEditorModeTools::MouseEnter( FEditorViewportClient* InViewportClient, FViewport* Viewport, int32 X, int32 Y )
-{
-	bool bHandled = false;
-	for( const auto& Mode : Modes)
-	{
-		bHandled |= Mode->MouseEnter( InViewportClient, Viewport, X, Y );
-	}
-	return bHandled;
-}
-
-bool FEditorModeTools::MouseLeave( FEditorViewportClient* InViewportClient, FViewport* Viewport )
-{
-	bool bHandled = false;
-	for( const auto& Mode : Modes)
-	{
-		bHandled |= Mode->MouseLeave( InViewportClient, Viewport );
-	}
-	return bHandled;
-}
-
-/** Notifies all active modes that the mouse has moved */
-bool FEditorModeTools::MouseMove( FEditorViewportClient* InViewportClient, FViewport* Viewport, int32 X, int32 Y )
-{
-	bool bHandled = false;
-	for( const auto& Mode : Modes)
-	{
-		bHandled |= Mode->MouseMove( InViewportClient, Viewport, X, Y );
-	}
-	return bHandled;
-}
-
-bool FEditorModeTools::ReceivedFocus( FEditorViewportClient* InViewportClient, FViewport* Viewport )
-{
-	bool bHandled = false;
-	for( const auto& Mode : Modes)
-	{
-		bHandled |= Mode->ReceivedFocus( InViewportClient, Viewport );
-	}
-	return bHandled;
-}
-
-bool FEditorModeTools::LostFocus( FEditorViewportClient* InViewportClient, FViewport* Viewport )
-{
-	bool bHandled = false;
-	for( const auto& Mode : Modes)
-	{
-		bHandled |= Mode->LostFocus( InViewportClient, Viewport );
-	}
-	return bHandled;
-}
-
-/** Draws all active mode components */	
-void FEditorModeTools::DrawActiveModes( const FSceneView* InView, FPrimitiveDrawInterface* PDI )
-{
-	for( const auto& Mode : Modes)
-	{
-		Mode->Draw( InView, PDI );
-	}
-}
-
-/** Renders all active modes */
-void FEditorModeTools::Render( const FSceneView* InView, FViewport* Viewport, FPrimitiveDrawInterface* PDI )
-{
-	for( const auto& Mode : Modes)
-	{
-		Mode->Render( InView, Viewport, PDI );
-	}
-}
-
-/** Draws the HUD for all active modes */
-void FEditorModeTools::DrawHUD( FEditorViewportClient* InViewportClient,FViewport* Viewport, const FSceneView* View, FCanvas* Canvas )
-{
-	for( const auto& Mode : Modes)
-	{
-		Mode->DrawHUD( InViewportClient, Viewport, View, Canvas );
-	}
-}
-
-/** Calls PostUndo on all active components */
-void FEditorModeTools::PostUndo()
-{
-	for( const auto& Mode : Modes)
-	{
-		Mode->PostUndo();
-	}
-}
-
-/** true if we should allow widget move */
-bool FEditorModeTools::AllowWidgetMove() const
-{
-	bool bAllow = false;
-	for( const auto& Mode : Modes)
-	{
-		bAllow |= Mode->AllowWidgetMove();
-	}
-	return bAllow;
-}
-
-bool FEditorModeTools::DisallowMouseDeltaTracking() const
-{
-	bool bDisallow = false;
-	for( const auto& Mode : Modes)
-	{
-		bDisallow |= Mode->DisallowMouseDeltaTracking();
-	}
-	return bDisallow;
-}
-
-bool FEditorModeTools::GetCursor(EMouseCursor::Type& OutCursor) const
-{
-	bool bHandled = false;
-	for( const auto& Mode : Modes)
-	{
-		bHandled |= Mode->GetCursor(OutCursor);
-	}
-	return bHandled;
-}
-
-/**
- * Used to cycle widget modes
- */
-void FEditorModeTools::CycleWidgetMode (void)
-{
-	//make sure we're not currently tracking mouse movement.  If we are, changing modes could cause a crash due to referencing an axis/plane that is incompatible with the widget
-	for(int32 ViewportIndex = 0;ViewportIndex < GEditor->LevelViewportClients.Num();ViewportIndex++)
-	{
-		FEditorViewportClient* ViewportClient = GEditor->LevelViewportClients[ ViewportIndex ];
-		if (ViewportClient->IsTracking())
-		{
-			return;
-		}
-	}
-
-	//only cycle when the mode is requesting the drawing of a widget
-	if( GetShowWidget() )
-	{
-		const int32 CurrentWk = GetWidgetMode();
-		int32 Wk = CurrentWk;
-		do
-		{
-			Wk++;
-			if ((Wk == FWidget::WM_TranslateRotateZ) && (!GetDefault<ULevelEditorViewportSettings>()->bAllowTranslateRotateZWidget))
-			{
-				Wk++;
-			}
-			// Roll back to the start if we go past FWidget::WM_Scale
-			if( Wk >= FWidget::WM_Max)
-			{
-				Wk -= FWidget::WM_Max;
-			}
-		}
-		while (!UsesTransformWidget((FWidget::EWidgetMode)Wk) && Wk != CurrentWk);
-		SetWidgetMode( (FWidget::EWidgetMode)Wk );
-		FEditorSupportDelegates::RedrawAllViewports.Broadcast();
-	}
-}
-
-/**Save Widget Settings to Ini file*/
-void FEditorModeTools::SaveWidgetSettings(void)
-{
-	GEditor->SaveEditorUserSettings();
-}
-
-/**Load Widget Settings from Ini file*/
-void FEditorModeTools::LoadWidgetSettings(void)
-{
-}
-
-/**
- * Returns a good location to draw the widget at.
- */
-
-FVector FEditorModeTools::GetWidgetLocation() const
-{
-	for (int Index = Modes.Num() - 1; Index >= 0 ; Index--)
-	{
-		if ( Modes[Index]->UsesTransformWidget() )
-		{
-			 return Modes[Index]->GetWidgetLocation();
-		}
-	}
-	
-	return FVector(EForceInit::ForceInitToZero);
-}
-
-/**
- * Changes the current widget mode.
- */
-
-void FEditorModeTools::SetWidgetMode( FWidget::EWidgetMode InWidgetMode )
-{
-	WidgetMode = InWidgetMode;
-}
-
-/**
- * Allows you to temporarily override the widget mode.  Call this function again
- * with FWidget::WM_None to turn off the override.
- */
-
-void FEditorModeTools::SetWidgetModeOverride( FWidget::EWidgetMode InWidgetMode )
-{
-	OverrideWidgetMode = InWidgetMode;
-}
-
-/**
- * Retrieves the current widget mode, taking overrides into account.
- */
-
-FWidget::EWidgetMode FEditorModeTools::GetWidgetMode() const
-{
-	if( OverrideWidgetMode != FWidget::WM_None )
-	{
-		return OverrideWidgetMode;
-	}
-
-	return WidgetMode;
-}
-
-bool FEditorModeTools::GetShowFriendlyVariableNames() const
-{
-	return GetDefault<UEditorStyleSettings>()->bShowFriendlyNames;
-}
-
-/**
- * Sets a bookmark in the levelinfo file, allocating it if necessary.
- */
-
-void FEditorModeTools::SetBookmark( uint32 InIndex, FEditorViewportClient* InViewportClient )
-{
-	UWorld* World = InViewportClient->GetWorld();
-	if ( World )
-	{
-		AWorldSettings* WorldSettings = World->GetWorldSettings();
-
-		// Verify the index is valid for the bookmark
-		if ( WorldSettings && InIndex < AWorldSettings::MAX_BOOKMARK_NUMBER )
-		{
-			// If the index doesn't already have a bookmark in place, create a new one
-			if ( !WorldSettings->BookMarks[ InIndex ] )
-			{
-				WorldSettings->BookMarks[ InIndex ] = ConstructObject<UBookMark>( UBookMark::StaticClass(), WorldSettings );
-			}
-
-			UBookMark* CurBookMark = WorldSettings->BookMarks[ InIndex ];
-			check(CurBookMark);
-			check(InViewportClient);
-
-			// Use the rotation from the first perspective viewport can find.
-			FRotator Rotation(0,0,0);
-			if( !InViewportClient->IsOrtho() )
-			{
-				Rotation = InViewportClient->GetViewRotation();
-			}
-
-			CurBookMark->Location = InViewportClient->GetViewLocation();
-			CurBookMark->Rotation = Rotation;
-
-			// Keep a record of which levels were hidden so that we can restore these with the bookmark
-			CurBookMark->HiddenLevels.Empty();
-			for ( int32 LevelIndex = 0 ; LevelIndex < World->StreamingLevels.Num() ; ++LevelIndex )
-			{
-				ULevelStreaming* StreamingLevel = World->StreamingLevels[LevelIndex];
-				if ( StreamingLevel )
-				{
-					if( !StreamingLevel->bShouldBeVisibleInEditor )
-					{
-						CurBookMark->HiddenLevels.Add( StreamingLevel->GetFullName() );
-					}
-				}
-			}
-		}
-	}
-}
-
-/**
- * Checks to see if a bookmark exists at a given index
- */
-
-bool FEditorModeTools::CheckBookmark( uint32 InIndex, FEditorViewportClient* InViewportClient )
-{
-	UWorld* World = InViewportClient->GetWorld();
-	if ( World )
-	{
-		AWorldSettings* WorldSettings = World->GetWorldSettings();
-		if ( WorldSettings && InIndex < AWorldSettings::MAX_BOOKMARK_NUMBER && WorldSettings->BookMarks[ InIndex ] )
-		{
-			return ( WorldSettings->BookMarks[ InIndex ] ? true : false );
-		}
-	}
-
-	return false;
-}
-
-/**
- * Retrieves a bookmark from the list.
- */
-
-void FEditorModeTools::JumpToBookmark( uint32 InIndex, bool bShouldRestoreLevelVisibility, FEditorViewportClient* InViewportClient )
-{
-	UWorld* World = InViewportClient->GetWorld();
-	if ( World )
-	{
-		AWorldSettings* WorldSettings = World->GetWorldSettings();
-
-		// Can only jump to a pre-existing bookmark
-		if ( WorldSettings && InIndex < AWorldSettings::MAX_BOOKMARK_NUMBER && WorldSettings->BookMarks[ InIndex ] )
-		{
-			const UBookMark* CurBookMark = WorldSettings->BookMarks[ InIndex ];
-			check(CurBookMark);
-
-			// Set all level editing cameras to this bookmark
-			for( int32 v = 0 ; v < GEditor->LevelViewportClients.Num() ; v++ )
-			{
-				GEditor->LevelViewportClients[v]->SetViewLocation( CurBookMark->Location );
-				if( !GEditor->LevelViewportClients[v]->IsOrtho() )
-				{
-					GEditor->LevelViewportClients[v]->SetViewRotation( CurBookMark->Rotation );
-				}
-				GEditor->LevelViewportClients[v]->Invalidate();
-			}
-		}
-	}
-}
-
-
-/**
- * Clears a bookmark
- */
-void FEditorModeTools::ClearBookmark( uint32 InIndex, FEditorViewportClient* InViewportClient )
-{
-	UWorld* World = InViewportClient->GetWorld();
-	if( World )
-	{
-		AWorldSettings* pWorldSettings = World->GetWorldSettings();
-
-		// Verify the index is valid for the bookmark
-		if ( pWorldSettings && InIndex < AWorldSettings::MAX_BOOKMARK_NUMBER )
-		{
-			pWorldSettings->BookMarks[ InIndex ] = NULL;
-		}
-	}
-}
-
-/**
-* Clears all book marks
-*/
-void FEditorModeTools::ClearAllBookmarks( FEditorViewportClient* InViewportClient )
-{
-	for( int i = 0; i <  AWorldSettings::MAX_BOOKMARK_NUMBER; ++i )
-	{
-		ClearBookmark( i , InViewportClient );
-	}
-}
-
-/**
- * Serializes the components for all modes.
- */
-
-void FEditorModeTools::AddReferencedObjects( FReferenceCollector& Collector )
-{
-	for( int32 x = 0 ; x < Modes.Num() ; ++x )
-	{
-		Modes[x]->AddReferencedObjects( Collector );
-	}
-}
-
-/**
- * Returns a pointer to an active mode specified by the passed in ID
- * If the editor mode is not active, NULL is returned
- */
-FEdMode* FEditorModeTools::GetActiveMode( FEditorModeID InID )
-{
-	for( auto& Mode : Modes )
-	{
-		if( Mode->GetID() == InID )
-		{
-			return Mode.Get();
-		}
-	}
-	return nullptr;
-}
-
-/**
- * Returns a pointer to an active mode specified by the passed in ID
- * If the editor mode is not active, NULL is returned
- */
-const FEdMode* FEditorModeTools::GetActiveMode( FEditorModeID InID ) const
-{
-	for (const auto& Mode : Modes)
-	{
-		if (Mode->GetID() == InID)
-		{
-			return Mode.Get();
-		}
-	}
-
-	return nullptr;
-}
-
-/**
- * Returns the active tool of the passed in editor mode.
- * If the passed in editor mode is not active or the mode has no active tool, NULL is returned
- */
-const FModeTool* FEditorModeTools::GetActiveTool( FEditorModeID InID ) const
-{
-	const FEdMode* ActiveMode = GetActiveMode( InID );
-	const FModeTool* Tool = NULL;
-	if( ActiveMode )
-	{
-		Tool = ActiveMode->GetCurrentTool();
-	}
-	return Tool;
-}
-
-/** 
- * Returns true if the passed in editor mode is active 
- */
-bool FEditorModeTools::IsModeActive( FEditorModeID InID ) const
-{
-	return GetActiveMode( InID ) != NULL;
-}
-
-/** 
- * Returns true if the default editor mode is active 
- */
-bool FEditorModeTools::IsDefaultModeActive() const
-{
-	return IsModeActive(DefaultID);
-}
-
-/** 
- * Returns an array of all active modes
- */
-void FEditorModeTools::GetActiveModes( TArray<FEdMode*>& OutActiveModes )
-{
-	OutActiveModes.Empty();
-	// Copy into an array.  Do not let users modify the active list directly.
-	for( auto& Mode : Modes)
-	{
-		OutActiveModes.Add(Mode.Get());
-	}
-}
diff --git a/Engine/Source/Editor/UnrealEd/Public/EdMode.h b/Engine/Source/Editor/UnrealEd/Public/EdMode.h
new file mode 100644
--- /dev/null
+++ b/Engine/Source/Editor/UnrealEd/Public/EdMode.h
@@ -0,0 +1,316 @@
+// Copyright 1998-2014 Epic Games, Inc. All Rights Reserved.
+
+#pragma once
+
+/**
+ * Base class for all editor modes.
+ */
+class UNREALED_API FEdMode : public TSharedFromThis<FEdMode>, public FGCObject, public FEditorCommonDrawHelper
+{
+public:
+	/** Friends so it can access mode's internals on construction */
+	friend class FEditorModeRegistry;
+
+	FEdMode();
+	virtual ~FEdMode();
+
+	virtual void Initialize() {}
+
+	virtual bool MouseEnter( FEditorViewportClient* ViewportClient,FViewport* Viewport,int32 x, int32 y );
+
+	virtual bool MouseLeave( FEditorViewportClient* ViewportClient,FViewport* Viewport );
+
+	virtual bool MouseMove(FEditorViewportClient* ViewportClient,FViewport* Viewport,int32 x, int32 y);
+
+	virtual bool ReceivedFocus(FEditorViewportClient* ViewportClient,FViewport* Viewport);
+
+	virtual bool LostFocus(FEditorViewportClient* ViewportClient,FViewport* Viewport);
+
+	/**
+	 * Called when the mouse is moved while a window input capture is in effect
+	 *
+	 * @param	InViewportClient	Level editor viewport client that captured the mouse input
+	 * @param	InViewport			Viewport that captured the mouse input
+	 * @param	InMouseX			New mouse cursor X coordinate
+	 * @param	InMouseY			New mouse cursor Y coordinate
+	 *
+	 * @return	true if input was handled
+	 */
+	virtual bool CapturedMouseMove( FEditorViewportClient* InViewportClient, FViewport* InViewport, int32 InMouseX, int32 InMouseY );
+
+	virtual bool InputKey(FEditorViewportClient* ViewportClient,FViewport* Viewport,FKey Key,EInputEvent Event);
+	virtual bool InputAxis(FEditorViewportClient* InViewportClient,FViewport* Viewport,int32 ControllerId,FKey Key,float Delta,float DeltaTime);
+	virtual bool InputDelta(FEditorViewportClient* InViewportClient,FViewport* InViewport,FVector& InDrag,FRotator& InRot,FVector& InScale);
+	virtual bool StartTracking(FEditorViewportClient* InViewportClient, FViewport* InViewport);
+	virtual bool EndTracking(FEditorViewportClient* InViewportClient, FViewport* InViewport);
+	// Added for handling EDIT Command...
+	virtual EEditAction::Type GetActionEditDuplicate() { return EEditAction::Skip; }
+	virtual EEditAction::Type GetActionEditDelete() { return EEditAction::Skip; }
+	virtual EEditAction::Type GetActionEditCut() { return EEditAction::Skip; }
+	virtual EEditAction::Type GetActionEditCopy() { return EEditAction::Skip; }
+	virtual EEditAction::Type GetActionEditPaste() { return EEditAction::Skip; }
+	virtual bool ProcessEditDuplicate() { return false; }
+	virtual bool ProcessEditDelete() { return false; }
+	virtual bool ProcessEditCut() { return false; }
+	virtual bool ProcessEditCopy() { return false; }
+	virtual bool ProcessEditPaste() { return false; }
+
+	virtual void Tick(FEditorViewportClient* ViewportClient,float DeltaTime);
+
+	virtual bool IsCompatibleWith(FEditorModeID OtherModeID) const { return false; }
+	virtual void ActorMoveNotify() {}
+	virtual void ActorsDuplicatedNotify(TArray<AActor*>& PreDuplicateSelection, TArray<AActor*>& PostDuplicateSelection, bool bOffsetLocations) {}
+	virtual void ActorSelectionChangeNotify();
+	virtual void ActorPropChangeNotify() {}
+	virtual void MapChangeNotify() {}
+	virtual bool ShowModeWidgets() const { return 1; }
+
+	/** If the EdMode is handling InputDelta (i.e., returning true from it), this allows a mode to indicated whether or not the Widget should also move. */
+	virtual bool AllowWidgetMove() { return true; }
+	
+	/** If the Edmode is handling its own mouse deltas, it can disable the MouseDeltaTacker */
+	virtual bool DisallowMouseDeltaTracking() const { return false; }
+
+	/**
+	 * Get a cursor to override the default with, if any.
+	 * @return true if the cursor was overriden.
+	 */
+	virtual bool GetCursor(EMouseCursor::Type& OutCursor) const { return false; }
+
+	virtual bool ShouldDrawBrushWireframe( AActor* InActor ) const { return true; }
+	virtual bool GetCustomDrawingCoordinateSystem( FMatrix& InMatrix, void* InData ) { return 0; }
+	virtual bool GetCustomInputCoordinateSystem( FMatrix& InMatrix, void* InData ) { return 0; }
+
+	/** If Rotation Snap should be enabled for this mode*/ 
+	virtual bool IsSnapRotationEnabled();
+
+	/** If this mode should override the snap rotation
+	* @param	Rotation		The Rotation Override
+	*
+	* @return					True if you have overridden the value
+	*/
+	virtual bool SnapRotatorToGridOverride(FRotator& Rotation){ return false; };
+
+	/**
+	 * Allows each mode to customize the axis pieces of the widget they want drawn.
+	 *
+	 * @param	InwidgetMode	The current widget mode
+	 *
+	 * @return					A bitfield comprised of AXIS_* values
+	 */
+	virtual EAxisList::Type GetWidgetAxisToDraw( FWidget::EWidgetMode InWidgetMode ) const;
+
+	/**
+	 * Allows each mode/tool to determine a good location for the widget to be drawn at.
+	 */
+	virtual FVector GetWidgetLocation() const;
+
+	/**
+	 * Lets the mode determine if it wants to draw the widget or not.
+	 */
+	virtual bool ShouldDrawWidget() const;
+	virtual void UpdateInternalData() {}
+
+	virtual FVector GetWidgetNormalFromCurrentAxis( void* InData );
+
+	virtual void AddReferencedObjects( FReferenceCollector& Collector ) override {}
+
+	virtual void Enter();
+	virtual void Exit();
+	virtual UTexture2D* GetVertexTexture() { return GEngine->DefaultBSPVertexTexture; }
+	
+	/**
+	 * Lets each tool determine if it wants to use the editor widget or not.  If the tool doesn't want to use it,
+	 * it will be fed raw mouse delta information (not snapped or altered in any way).
+	 */
+	virtual bool UsesTransformWidget() const;
+
+	/**
+	 * Lets each mode selectively exclude certain widget types.
+	 */
+	virtual bool UsesTransformWidget(FWidget::EWidgetMode CheckMode) const;
+
+	virtual void PostUndo() {}
+
+	/**
+	 * Lets each mode/tool handle box selection in its own way.
+	 *
+	 * @param	InBox	The selection box to use, in worldspace coordinates.
+	 * @return		true if something was selected/deselected, false otherwise.
+	 */
+	bool BoxSelect( FBox& InBox, bool InSelect = true );
+
+	/**
+	 * Lets each mode/tool handle frustum selection in its own way.
+	 *
+	 * @param	InFrustum	The selection box to use, in worldspace coordinates.
+	 * @return	true if something was selected/deselected, false otherwise.
+	 */
+	bool FrustumSelect( const FConvexVolume& InFrustum, bool InSelect = true );
+
+	void SelectNone();
+	virtual void SelectionChanged() {}
+
+	virtual bool HandleClick(FEditorViewportClient* InViewportClient, HHitProxy *HitProxy, const FViewportClick &Click);
+
+	/** Handling SelectActor */
+	virtual bool Select( AActor* InActor, bool bInSelected ) { return 0; }
+
+	/** Check to see if an actor can be selected in this mode - no side effects */
+	virtual bool IsSelectionAllowed( AActor* InActor, bool bInSelection ) const { return true; }
+
+	/** @return True if this mode allows the viewport to use a drag tool */
+	virtual bool AllowsViewportDragTool() const { return true; }
+
+	/** Returns the editor mode identifier. */
+	FEditorModeID GetID() const { return Info.ID; }
+
+	/** Returns the editor mode information. */
+	const FEditorModeInfo& GetModeInfo() const { return Info; }
+
+	// Tools
+
+	void SetCurrentTool( EModeTools InID );
+	void SetCurrentTool( FModeTool* InModeTool );
+	FModeTool* FindTool( EModeTools InID );
+
+	const TArray<FModeTool*>& GetTools() const		{ return Tools; }
+
+	virtual void CurrentToolChanged() {}
+
+	/** Returns the current tool. */
+	//@{
+	FModeTool* GetCurrentTool()				{ return CurrentTool; }
+	const FModeTool* GetCurrentTool() const	{ return CurrentTool; }
+	//@}
+
+	/** @name Current widget axis. */
+	//@{
+	void SetCurrentWidgetAxis(EAxisList::Type InAxis)		{ CurrentWidgetAxis = InAxis; }
+	EAxisList::Type GetCurrentWidgetAxis() const			{ return CurrentWidgetAxis; }
+	//@}
+
+	/** @name Rendering */
+	//@{
+	/** Draws translucent polygons on brushes and volumes. */
+	virtual void Render(const FSceneView* View,FViewport* Viewport,FPrimitiveDrawInterface* PDI);
+	//void DrawGridSection(int32 ViewportLocX,int32 ViewportGridY,FVector* A,FVector* B,float* AX,float* BX,int32 Axis,int32 AlphaCase,FSceneView* View,FPrimitiveDrawInterface* PDI);
+
+	/** Overlays the editor hud (brushes, drag tools, static mesh vertices, etc*. */
+	virtual void DrawHUD(FEditorViewportClient* ViewportClient,FViewport* Viewport,const FSceneView* View,FCanvas* Canvas);
+	//@}
+
+	/**
+	 * Called when attempting to duplicate the selected actors by alt+dragging,
+	 * return true to prevent normal duplication.
+	 */
+	virtual bool HandleDragDuplicate() { return false; }
+
+	/**
+	 * Called when the mode wants to draw brackets around selected objects
+	 */
+	virtual void DrawBrackets( FEditorViewportClient* ViewportClient, FViewport* Viewport, const FSceneView* View, FCanvas* Canvas );
+
+	/** True if this mode uses a toolkit mode (eventually they all should) */
+	virtual bool UsesToolkits() const;
+
+	UWorld* GetWorld() const
+	{
+		return GEditor->GetEditorWorldContext().World();
+	}
+
+	// Property Widgets
+
+	/**
+	 * Lets each mode selectively enable widgets for editing properties tagged with 'Show 3D Widget' metadata.
+	 */
+	virtual bool UsesPropertyWidgets() const { return false; }
+
+	// @TODO: Find a better home for these?
+	/** Value of UPROPERTY can be edited with a widget in the editor (translation, rotation) */
+	static const FName MD_MakeEditWidget;
+	/** Specifies a function used for validation of the current value of a property.  The function returns a string that is empty if the value is valid, or contains an error description if the value is invalid */
+	static const FName MD_ValidateWidgetUsing;
+	/** Returns true if this structure can support creating a widget in the editor */
+	static bool CanCreateWidgetForStructure(const UStruct* InPropStruct);
+	/** Returns true if this property can support creating a widget in the editor */
+	static bool CanCreateWidgetForProperty(UProperty* InProp);
+	/** See if we should create a widget for the supplied property when selecting an actor instance */
+	static bool ShouldCreateWidgetForProperty(UProperty* InProp);
+
+public:
+
+	/** Request that this mode be deleted at the next convenient opportunity (FEditorModeTools::Tick) */
+	void RequestDeletion() { bPendingDeletion = true; }
+	
+	/** returns true if this mode is to be deleted at the next convenient opportunity (FEditorModeTools::Tick) */
+	bool IsPendingDeletion() const { return bPendingDeletion; }
+
+private:
+	/** true if this mode is pending removal from its owner */
+	bool bPendingDeletion;
+
+	/** Called whenever a mode type is unregistered */
+	void OnModeUnregistered(FEditorModeID ModeID);
+
+protected:
+	/** The current axis that is being dragged on the widget. */
+	EAxisList::Type CurrentWidgetAxis;
+
+	/** Optional array of tools for this mode. */
+	TArray<FModeTool*> Tools;
+
+	/** The tool that is currently active within this mode. */
+	FModeTool* CurrentTool;
+
+	/** Information pertaining to this mode. Assigned by FEditorModeRegistry. */
+	FEditorModeInfo Info;
+
+	/** Editor Mode Toolkit that is associated with this toolkit mode */
+	TSharedPtr<class FModeToolkit> Toolkit;
+
+	/** Pointer back to the mode tools that we are registered with */
+	FEditorModeTools* Owner;
+
+	// Property Widgets
+
+	/** Structure that holds info about our optional property widget */
+	struct FPropertyWidgetInfo
+	{
+		FString PropertyName;
+		int32 PropertyIndex;
+		FName PropertyValidationName;
+		FString DisplayName;
+		bool bIsTransform;
+
+		FPropertyWidgetInfo()
+			: PropertyIndex(INDEX_NONE)
+			, bIsTransform(false)
+		{
+		}
+	};
+
+	/**
+	 * Returns the first selected Actor, or NULL if there is no selection.
+	 */
+	AActor* GetFirstSelectedActorInstance() const;
+
+	/**
+	 * Gets an array of property widget info structures for the given struct/class type for the given container.
+	 *
+	 * @param InStruct The type of structure/class to access widget info structures for.
+	 * @param InContainer The container of the given type.
+	 * @param OutInfos An array of widget info structures (output).
+	 * @param PropertyNamePrefix Optional prefix to use to filter the output.
+	 */
+	void GetPropertyWidgetInfos(const UStruct* InStruct, const void* InContainer, TArray<FPropertyWidgetInfo>& OutInfos, FString PropertyNamePrefix = TEXT(""), FString DisplayNamePrefix = TEXT("")) const;
+
+	/** Name of the property currently being edited */
+	FString EditedPropertyName;
+
+	/** If the property being edited is an array property, this is the index of the element we're currently dealing with */
+	int32 EditedPropertyIndex;
+
+	/** Indicates  */
+	bool bEditedPropertyIsTransform;
+};
diff --git a/Engine/Source/Editor/UnrealEd/Public/EditorModeManager.h b/Engine/Source/Editor/UnrealEd/Public/EditorModeManager.h
new file mode 100644
--- /dev/null
+++ b/Engine/Source/Editor/UnrealEd/Public/EditorModeManager.h
@@ -0,0 +1,419 @@
+// Copyright 1998-2014 Epic Games, Inc. All Rights Reserved.
+
+#pragma once
+
+/**
+ * A helper class to store the state of the various editor modes.
+ */
+class UNREALED_API FEditorModeTools : public FGCObject
+{
+public:
+	FEditorModeTools();
+	virtual ~FEditorModeTools();
+
+	/**
+	 * Set the default editor mode for these tools
+	 */
+	void SetDefaultMode(FEditorModeID DefaultID);
+
+	/**
+	 * Activates the default mode defined by this class
+	 */
+	void ActivateDefaultMode();
+
+	/** 
+	 * Returns true if the default editor mode is active 
+	 */
+	bool IsDefaultModeActive() const;
+
+	/**
+	 * Activates an editor mode. Shuts down all other active modes which cannot run with the passed in mode.
+	 * 
+	 * @param InID		The ID of the editor mode to activate.
+	 * @param bToggle	true if the passed in editor mode should be toggled off if it is already active.
+	 */
+	void ActivateMode( FEditorModeID InID, bool bToggle = false );
+
+	/**
+	 * Deactivates an editor mode. 
+	 * 
+	 * @param InID		The ID of the editor mode to deactivate.
+	 */
+	void DeactivateMode(FEditorModeID InID);
+
+	/**
+	 * Deactivate the mode and entirely purge it from memory. Used when a mode type is unregistered
+	 */
+	void DestroyMode(FEditorModeID InID);
+
+protected:
+	
+	/** Deactivates the editor mode at the specified index */
+	void DeactivateModeAtIndex( int32 InIndex );
+		
+public:
+
+	/**
+	 * Deactivates all modes, note some modes can never be deactivated.
+	 */
+	void DeactivateAllModes();
+
+	/** 
+	 * Returns the editor mode specified by the passed in ID
+	 */
+	FEdMode* FindMode( FEditorModeID InID );
+
+	/**
+	 * Returns true if the current mode is not the specified ModeID.  Also optionally warns the user.
+	 *
+	 * @param	ModeID			The editor mode to query.
+	 * @param	ErrorMsg		If specified, inform the user the reason why this is a problem
+	 * @param	bNotifyUser		If true, display the error as a notification, instead of a dialog
+	 * @return					true if the current mode is not the specified mode.
+	 */
+	bool EnsureNotInMode(FEditorModeID ModeID, const FText& ErrorMsg = FText::GetEmpty(), bool bNotifyUser = false) const;
+
+	FMatrix GetCustomDrawingCoordinateSystem();
+	FMatrix GetCustomInputCoordinateSystem();
+	
+	/** 
+	 * Returns true if the passed in editor mode is active 
+	 */
+	bool IsModeActive( FEditorModeID InID ) const;
+
+	/**
+	 * Returns a pointer to an active mode specified by the passed in ID
+	 * If the editor mode is not active, NULL is returned
+	 */
+	FEdMode* GetActiveMode( FEditorModeID InID );
+	const FEdMode* GetActiveMode( FEditorModeID InID ) const;
+
+	template <typename SpecificModeType>
+	SpecificModeType* GetActiveModeTyped( FEditorModeID InID )
+	{
+		return (SpecificModeType*)GetActiveMode(InID);
+	}
+
+	template <typename SpecificModeType>
+	const SpecificModeType* GetActiveModeTyped( FEditorModeID InID ) const
+	{
+		return (SpecificModeType*)GetActiveMode(InID);
+	}
+
+	/**
+	 * Returns the active tool of the passed in editor mode.
+	 * If the passed in editor mode is not active or the mode has no active tool, NULL is returned
+	 */
+	const FModeTool* GetActiveTool( FEditorModeID InID ) const;
+
+	/** 
+	 * Returns an array of all active modes
+	 */
+	void GetActiveModes( TArray<FEdMode*>& OutActiveModes );
+
+	void SetShowWidget( bool InShowWidget )	{ bShowWidget = InShowWidget; }
+	bool GetShowWidget() const				{ return bShowWidget; }
+
+	void CycleWidgetMode (void);
+
+	/**Save Widget Settings to Ini file*/
+	void SaveWidgetSettings();
+	/**Load Widget Settings from Ini file*/
+	void LoadWidgetSettings();
+
+	/** Gets the widget axis to be drawn */
+	EAxisList::Type GetWidgetAxisToDraw( FWidget::EWidgetMode InWidgetMode ) const;
+
+	/** Mouse tracking interface.  Passes tracking messages to all active modes */
+	bool StartTracking(FEditorViewportClient* InViewportClient, FViewport* InViewport);
+	bool EndTracking(FEditorViewportClient* InViewportClient, FViewport* InViewport);
+	bool IsTracking() const { return bIsTracking; }
+
+	bool AllowsViewportDragTool() const;
+
+	/** Notifies all active modes that a map change has occured */
+	void MapChangeNotify();
+
+	/** Notifies all active modes to empty their selections */
+	void SelectNone();
+
+	/** Notifies all active modes of box selection attempts */
+	bool BoxSelect( FBox& InBox, bool InSelect );
+
+	/** Notifies all active modes of frustum selection attempts */
+	bool FrustumSelect( const FConvexVolume& InFrustum, bool InSelect );
+
+	/** true if any active mode uses a transform widget */
+	bool UsesTransformWidget() const;
+
+	/** true if any active mode uses the passed in transform widget */
+	bool UsesTransformWidget( FWidget::EWidgetMode CheckMode ) const;
+
+	/** Sets the current widget axis */
+	void SetCurrentWidgetAxis( EAxisList::Type NewAxis );
+
+	/** Notifies all active modes of mouse click messages. */
+	bool HandleClick(FEditorViewportClient* InViewportClient,  HHitProxy *HitProxy, const FViewportClick &Click );
+
+	/** true if the passed in brush actor should be drawn in wireframe */	
+	bool ShouldDrawBrushWireframe( AActor* InActor ) const;
+
+	/** true if brush vertices should be drawn */
+	bool ShouldDrawBrushVertices() const;
+
+	/** Ticks all active modes */
+	void Tick( FEditorViewportClient* ViewportClient, float DeltaTime );
+
+	/** Notifies all active modes of any change in mouse movement */
+	bool InputDelta( FEditorViewportClient* InViewportClient,FViewport* InViewport,FVector& InDrag,FRotator& InRot,FVector& InScale );
+
+	/** Notifies all active modes of captured mouse movement */	
+	bool CapturedMouseMove( FEditorViewportClient* InViewportClient, FViewport* InViewport, int32 InMouseX, int32 InMouseY );
+
+	/** Notifies all active modes of keyboard input */
+	bool InputKey( FEditorViewportClient* InViewportClient, FViewport* Viewport, FKey Key, EInputEvent Event);
+
+	/** Notifies all active modes of axis movement */
+	bool InputAxis( FEditorViewportClient* InViewportClient, FViewport* Viewport, int32 ControllerId, FKey Key, float Delta, float DeltaTime);
+
+	bool MouseEnter( FEditorViewportClient* InViewportClient, FViewport* Viewport, int32 X, int32 Y );
+	
+	bool MouseLeave( FEditorViewportClient* InViewportClient, FViewport* Viewport );
+
+	/** Notifies all active modes that the mouse has moved */
+	bool MouseMove( FEditorViewportClient* InViewportClient, FViewport* Viewport, int32 X, int32 Y );
+
+	/** Notifies all active modes that a viewport has received focus */
+	bool ReceivedFocus( FEditorViewportClient* InViewportClient, FViewport* Viewport );
+
+	/** Notifies all active modes that a viewport has lost focus */
+	bool LostFocus( FEditorViewportClient* InViewportClient, FViewport* Viewport );
+
+	/** Draws all active modes */	
+	void DrawActiveModes( const FSceneView* InView, FPrimitiveDrawInterface* PDI );
+
+	/** Renders all active modes */
+	void Render( const FSceneView* InView, FViewport* Viewport, FPrimitiveDrawInterface* PDI );
+
+	/** Draws the HUD for all active modes */
+	void DrawHUD( FEditorViewportClient* InViewportClient,FViewport* Viewport,const FSceneView* View,FCanvas* Canvas );
+
+	/** Calls PostUndo on all active components */
+	void PostUndo();
+
+	/** True if we should allow widget move */
+	bool AllowWidgetMove() const;
+
+	/** True if we should disallow mouse delta tracking. */
+	bool DisallowMouseDeltaTracking() const;
+
+	/** Get a cursor to override the default with, if any */
+	bool GetCursor(EMouseCursor::Type& OutCursor) const;
+
+	/**
+	 * Returns a good location to draw the widget at.
+	 */
+	FVector GetWidgetLocation() const;
+
+	/**
+	 * Changes the current widget mode.
+	 */
+	void SetWidgetMode( FWidget::EWidgetMode InWidgetMode );
+
+	/**
+	 * Allows you to temporarily override the widget mode.  Call this function again
+	 * with WM_None to turn off the override.
+	 */
+	void SetWidgetModeOverride( FWidget::EWidgetMode InWidgetMode );
+
+	/**
+	 * Retrieves the current widget mode, taking overrides into account.
+	 */
+	FWidget::EWidgetMode GetWidgetMode() const;
+
+	/**
+	 * Gets the current state of script editor usage of show friendly names
+	 * @ return - If true, replace variable names with sanatized ones
+	 */
+	bool GetShowFriendlyVariableNames() const;
+
+	/**
+	 * Sets Interp Editor inversion of panning on drag
+	 */
+	void SetInterpPanInvert(bool bInInterpPanInverted) { bInterpPanInverted = bInInterpPanInverted; }
+
+	/**
+	 * Gets the current state of the interp editor panning
+	 * @ return - If true, the interp editor left-right panning is inverted
+	 */
+	bool GetInterpPanInvert(void) const { return bInterpPanInverted; }
+
+	/**
+	 * Sets a bookmark in the levelinfo file, allocating it if necessary.
+	 * 
+	 * @param InIndex			Index of the bookmark to set.
+	 * @param InViewportClient	Level editor viewport client used to get pointer the world that the bookmark is for.
+	 */
+	void SetBookmark( uint32 InIndex, FEditorViewportClient* InViewportClient );
+
+	/**
+	 * Checks to see if a bookmark exists at a given index
+	 * 
+	 * @param InIndex			Index of the bookmark to set.
+	 * @param InViewportClient	Level editor viewport client used to get pointer the world to check for the bookmark in.
+	 */
+	bool CheckBookmark( uint32 InIndex, FEditorViewportClient* InViewportClient );
+
+	/**
+	 * Retrieves a bookmark from the list.
+	 * 
+	 * @param InIndex			Index of the bookmark to set.
+	 * @param					Set to true to restore the visibility of any streaming levels.
+	 * @param InViewportClient	Level editor viewport client used to get pointer the world that the bookmark is relevant to.
+	 */
+	void JumpToBookmark( uint32 InIndex, bool bShouldRestoreLevelVisibility, FEditorViewportClient* InViewportClient );
+
+	/**
+	 * Clears a bookmark from the list.
+	 * 
+	 * @param InIndex			Index of the bookmark to clear.
+	 * @param InViewportClient	Level editor viewport client used to get pointer the world that the bookmark is relevant to.
+	 */
+	void ClearBookmark( uint32 InIndex, FEditorViewportClient* InViewportClient );
+
+	/**
+	 * Clears all book marks
+	 * 
+	 * @param InViewportClient	Level editor viewport client used to get pointer the world to clear the bookmarks from.
+	 */
+	void ClearAllBookmarks( FEditorViewportClient* InViewportClient );
+
+	// FGCObject interface
+	virtual void AddReferencedObjects( FReferenceCollector& Collector ) override;
+	// End of FGCObject interface
+
+	/**
+	 * Loads the state that was saved in the INI file
+	 */
+	void LoadConfig(void);
+
+	/**
+	 * Saves the current state to the INI file
+	 */
+	void SaveConfig(void);
+
+	/** 
+	 * Sets the pivot locations
+	 * 
+	 * @param Location 		The location to set
+	 * @param bIncGridBase	Whether or not to also set the GridBase
+	 */
+	void SetPivotLocation( const FVector& Location, const bool bIncGridBase );
+
+	/**
+	 * Multicast delegate for OnModeEntered and OnModeExited callbacks.
+	 *
+	 * First parameter:  The editor mode that was changed
+	 * Second parameter:  True if entering the mode, or false if exiting the mode
+	 */
+	DECLARE_EVENT_TwoParams( FEditorModeTools, FEditorModeChangedEvent, FEdMode*, bool );
+	FEditorModeChangedEvent& OnEditorModeChanged() { return EditorModeChangedEvent; }
+
+	/** delegate type for triggering when widget mode changed */
+	DECLARE_EVENT_OneParam( FEditorModeTools, FWidgetModeChangedEvent, FWidget::EWidgetMode );
+	FWidgetModeChangedEvent& OnWidgetModeChanged() { return WidgetModeChangedEvent; }
+
+	/**	Broadcasts the WidgetModeChanged event */
+	void BroadcastWidgetModeChanged(FWidget::EWidgetMode InWidgetMode) { WidgetModeChangedEvent.Broadcast(InWidgetMode); }
+
+	/**	Broadcasts the EditorModeChanged event */
+	void BroadcastEditorModeChanged(FEdMode* Mode, bool IsEnteringMode) { EditorModeChangedEvent.Broadcast(Mode, IsEnteringMode); }
+
+	/**
+	 * Returns the current CoordSystem
+	 * 
+	 * @param bGetRawValue true when you want the actual value of CoordSystem, not the value modified by the state.
+	 */
+	ECoordSystem GetCoordSystem(bool bGetRawValue = false);
+
+	/** Sets the current CoordSystem */
+	void SetCoordSystem(ECoordSystem NewCoordSystem);
+
+	/** Sets the hide viewport UI state */
+	void SetHideViewportUI( bool bInHideViewportUI ) { bHideViewportUI = bInHideViewportUI; }
+
+	/** Is the viewport UI hidden? */
+	bool IsViewportUIHidden() const { return bHideViewportUI; }
+
+	bool PivotShown, Snapping, SnappedActor;
+	FVector CachedLocation, PivotLocation, SnappedLocation, GridBase;
+
+	/** The angle for the translate rotate widget */
+	float TranslateRotateXAxisAngle;
+
+	/** Draws in the top level corner of all FEditorViewportClient windows (can be used to relay info to the user). */
+	FString InfoString;
+
+	/** Sets the host for toolkits created via modes from this mode manager (can only be called once) */
+	void SetToolkitHost(TSharedRef<class IToolkitHost> Host);
+
+	/** Returns the host for toolkits created via modes from this mode manager */
+	TSharedPtr<class IToolkitHost> GetToolkitHost() const;
+
+	/**
+	 * Returns the set of selected actors.
+	 */
+	virtual class USelection* GetSelectedActors() const;
+
+	/**
+	 * @return the set of selected non-actor objects.
+	 */
+	virtual class USelection* GetSelectedObjects() const;
+protected:
+	/** 
+	 * Delegate handlers
+	 **/
+	void OnEditorSelectionChanged(UObject* NewSelection);
+	void OnEditorSelectNone();
+
+	/** The mode ID to use as a default */
+	FEditorModeID DefaultID;
+
+	/** A list of active editor modes. */
+	TArray< TSharedPtr<FEdMode> > Modes;
+
+	/** The host of the toolkits created by these modes */
+	TWeakPtr<IToolkitHost> ToolkitHost;
+
+	/** A list of previously active editor modes that we will potentially recycle */
+	TMap< FEditorModeID, TSharedPtr<FEdMode> > RecycledModes;
+
+	/** The mode that the editor viewport widget is in. */
+	FWidget::EWidgetMode WidgetMode;
+
+	/** If the widget mode is being overridden, this will be != WM_None. */
+	FWidget::EWidgetMode OverrideWidgetMode;
+
+	/** If 1, draw the widget and let the user interact with it. */
+	bool bShowWidget;
+
+	/** If true, panning in the Interp Editor is inverted. If false, panning is normal*/
+	bool bInterpPanInverted;
+
+	/** if true, the viewports will hide all UI overlays */
+	bool bHideViewportUI;
+
+private:
+
+	/** The coordinate system the widget is operating within. */
+	ECoordSystem CoordSystem;
+
+	/** Multicast delegate that is broadcast when a mode is entered or exited */
+	FEditorModeChangedEvent EditorModeChangedEvent;
+
+	/** Multicast delegate that is broadcast when a widget mode is changed */
+	FWidgetModeChangedEvent WidgetModeChangedEvent;
+
+	/** Flag set between calls to StartTracking() and EndTracking() */
+	bool bIsTracking;
+};
diff --git a/Engine/Source/Editor/UnrealEd/Public/EditorModes.h b/Engine/Source/Editor/UnrealEd/Public/EditorModes.h
--- a/Engine/Source/Editor/UnrealEd/Public/EditorModes.h
+++ b/Engine/Source/Editor/UnrealEd/Public/EditorModes.h
@@ -13,6 +13,8 @@ class HHitProxy;
 struct FViewportClick;
 class FModeTool;
 
+DECLARE_LOG_CATEGORY_EXTERN(LogEditorModes, Log, All);
+
 /** Outcomes when determining whether it's possible to perform an action on the edit modes*/
 namespace EEditAction
 {
@@ -100,319 +102,7 @@ public:
 	}
 };
 
-
-/**
- * Base class for all editor modes.
- */
-class UNREALED_API FEdMode : public TSharedFromThis<FEdMode>, public FGCObject, public FEditorCommonDrawHelper
-{
-public:
-	/** Friends so it can access mode's internals on construction */
-	friend class FEditorModeRegistry;
-
-	FEdMode();
-	virtual ~FEdMode();
-
-	virtual void Initialize() {}
-
-	virtual bool MouseEnter( FEditorViewportClient* ViewportClient,FViewport* Viewport,int32 x, int32 y );
-
-	virtual bool MouseLeave( FEditorViewportClient* ViewportClient,FViewport* Viewport );
-
-	virtual bool MouseMove(FEditorViewportClient* ViewportClient,FViewport* Viewport,int32 x, int32 y);
-
-	virtual bool ReceivedFocus(FEditorViewportClient* ViewportClient,FViewport* Viewport);
-
-	virtual bool LostFocus(FEditorViewportClient* ViewportClient,FViewport* Viewport);
-
-	/**
-	 * Called when the mouse is moved while a window input capture is in effect
-	 *
-	 * @param	InViewportClient	Level editor viewport client that captured the mouse input
-	 * @param	InViewport			Viewport that captured the mouse input
-	 * @param	InMouseX			New mouse cursor X coordinate
-	 * @param	InMouseY			New mouse cursor Y coordinate
-	 *
-	 * @return	true if input was handled
-	 */
-	virtual bool CapturedMouseMove( FEditorViewportClient* InViewportClient, FViewport* InViewport, int32 InMouseX, int32 InMouseY );
-
-	virtual bool InputKey(FEditorViewportClient* ViewportClient,FViewport* Viewport,FKey Key,EInputEvent Event);
-	virtual bool InputAxis(FEditorViewportClient* InViewportClient,FViewport* Viewport,int32 ControllerId,FKey Key,float Delta,float DeltaTime);
-	virtual bool InputDelta(FEditorViewportClient* InViewportClient,FViewport* InViewport,FVector& InDrag,FRotator& InRot,FVector& InScale);
-	virtual bool StartTracking(FEditorViewportClient* InViewportClient, FViewport* InViewport);
-	virtual bool EndTracking(FEditorViewportClient* InViewportClient, FViewport* InViewport);
-	// Added for handling EDIT Command...
-	virtual EEditAction::Type GetActionEditDuplicate() { return EEditAction::Skip; }
-	virtual EEditAction::Type GetActionEditDelete() { return EEditAction::Skip; }
-	virtual EEditAction::Type GetActionEditCut() { return EEditAction::Skip; }
-	virtual EEditAction::Type GetActionEditCopy() { return EEditAction::Skip; }
-	virtual EEditAction::Type GetActionEditPaste() { return EEditAction::Skip; }
-	virtual bool ProcessEditDuplicate() { return false; }
-	virtual bool ProcessEditDelete() { return false; }
-	virtual bool ProcessEditCut() { return false; }
-	virtual bool ProcessEditCopy() { return false; }
-	virtual bool ProcessEditPaste() { return false; }
-
-	virtual void Tick(FEditorViewportClient* ViewportClient,float DeltaTime);
-
-	virtual bool IsCompatibleWith(FEditorModeID OtherModeID) const { return false; }
-	virtual void ActorMoveNotify() {}
-	virtual void ActorsDuplicatedNotify(TArray<AActor*>& PreDuplicateSelection, TArray<AActor*>& PostDuplicateSelection, bool bOffsetLocations) {}
-	virtual void ActorSelectionChangeNotify();
-	virtual void ActorPropChangeNotify() {}
-	virtual void MapChangeNotify() {}
-	virtual bool ShowModeWidgets() const { return 1; }
-
-	/** If the EdMode is handling InputDelta (i.e., returning true from it), this allows a mode to indicated whether or not the Widget should also move. */
-	virtual bool AllowWidgetMove() { return true; }
-	
-	/** If the Edmode is handling its own mouse deltas, it can disable the MouseDeltaTacker */
-	virtual bool DisallowMouseDeltaTracking() const { return false; }
-
-	/**
-	 * Get a cursor to override the default with, if any.
-	 * @return true if the cursor was overriden.
-	 */
-	virtual bool GetCursor(EMouseCursor::Type& OutCursor) const { return false; }
-
-	virtual bool ShouldDrawBrushWireframe( AActor* InActor ) const { return true; }
-	virtual bool GetCustomDrawingCoordinateSystem( FMatrix& InMatrix, void* InData ) { return 0; }
-	virtual bool GetCustomInputCoordinateSystem( FMatrix& InMatrix, void* InData ) { return 0; }
-
-	/** If Rotation Snap should be enabled for this mode*/ 
-	virtual bool IsSnapRotationEnabled();
-
-	/** If this mode should override the snap rotation
-	* @param	Rotation		The Rotation Override
-	*
-	* @return					True if you have overridden the value
-	*/
-	virtual bool SnapRotatorToGridOverride(FRotator& Rotation){ return false; };
-
-	/**
-	 * Allows each mode to customize the axis pieces of the widget they want drawn.
-	 *
-	 * @param	InwidgetMode	The current widget mode
-	 *
-	 * @return					A bitfield comprised of AXIS_* values
-	 */
-	virtual EAxisList::Type GetWidgetAxisToDraw( FWidget::EWidgetMode InWidgetMode ) const;
-
-	/**
-	 * Allows each mode/tool to determine a good location for the widget to be drawn at.
-	 */
-	virtual FVector GetWidgetLocation() const;
-
-	/**
-	 * Lets the mode determine if it wants to draw the widget or not.
-	 */
-	virtual bool ShouldDrawWidget() const;
-	virtual void UpdateInternalData() {}
-
-	virtual FVector GetWidgetNormalFromCurrentAxis( void* InData );
-
-	virtual void AddReferencedObjects( FReferenceCollector& Collector ) override {}
-
-	virtual void Enter();
-	virtual void Exit();
-	virtual UTexture2D* GetVertexTexture() { return GEngine->DefaultBSPVertexTexture; }
-	
-	/**
-	 * Lets each tool determine if it wants to use the editor widget or not.  If the tool doesn't want to use it,
-	 * it will be fed raw mouse delta information (not snapped or altered in any way).
-	 */
-	virtual bool UsesTransformWidget() const;
-
-	/**
-	 * Lets each mode selectively exclude certain widget types.
-	 */
-	virtual bool UsesTransformWidget(FWidget::EWidgetMode CheckMode) const;
-
-	virtual void PostUndo() {}
-
-	/**
-	 * Lets each mode/tool handle box selection in its own way.
-	 *
-	 * @param	InBox	The selection box to use, in worldspace coordinates.
-	 * @return		true if something was selected/deselected, false otherwise.
-	 */
-	bool BoxSelect( FBox& InBox, bool InSelect = true );
-
-	/**
-	 * Lets each mode/tool handle frustum selection in its own way.
-	 *
-	 * @param	InFrustum	The selection box to use, in worldspace coordinates.
-	 * @return	true if something was selected/deselected, false otherwise.
-	 */
-	bool FrustumSelect( const FConvexVolume& InFrustum, bool InSelect = true );
-
-	void SelectNone();
-	virtual void SelectionChanged() {}
-
-	virtual bool HandleClick(FEditorViewportClient* InViewportClient, HHitProxy *HitProxy, const FViewportClick &Click);
-
-	/** Handling SelectActor */
-	virtual bool Select( AActor* InActor, bool bInSelected ) { return 0; }
-
-	/** Check to see if an actor can be selected in this mode - no side effects */
-	virtual bool IsSelectionAllowed( AActor* InActor, bool bInSelection ) const { return true; }
-
-	/** @return True if this mode allows the viewport to use a drag tool */
-	virtual bool AllowsViewportDragTool() const { return true; }
-
-	/** Returns the editor mode identifier. */
-	FEditorModeID GetID() const { return Info.ID; }
-
-	/** Returns the editor mode information. */
-	const FEditorModeInfo& GetModeInfo() const { return Info; }
-
-	// Tools
-
-	void SetCurrentTool( EModeTools InID );
-	void SetCurrentTool( FModeTool* InModeTool );
-	FModeTool* FindTool( EModeTools InID );
-
-	const TArray<FModeTool*>& GetTools() const		{ return Tools; }
-
-	virtual void CurrentToolChanged() {}
-
-	/** Returns the current tool. */
-	//@{
-	FModeTool* GetCurrentTool()				{ return CurrentTool; }
-	const FModeTool* GetCurrentTool() const	{ return CurrentTool; }
-	//@}
-
-	/** @name Current widget axis. */
-	//@{
-	void SetCurrentWidgetAxis(EAxisList::Type InAxis)		{ CurrentWidgetAxis = InAxis; }
-	EAxisList::Type GetCurrentWidgetAxis() const			{ return CurrentWidgetAxis; }
-	//@}
-
-	/** @name Rendering */
-	//@{
-	/** Draws translucent polygons on brushes and volumes. */
-	virtual void Render(const FSceneView* View,FViewport* Viewport,FPrimitiveDrawInterface* PDI);
-	//void DrawGridSection(int32 ViewportLocX,int32 ViewportGridY,FVector* A,FVector* B,float* AX,float* BX,int32 Axis,int32 AlphaCase,FSceneView* View,FPrimitiveDrawInterface* PDI);
-
-	/** Overlays the editor hud (brushes, drag tools, static mesh vertices, etc*. */
-	virtual void DrawHUD(FEditorViewportClient* ViewportClient,FViewport* Viewport,const FSceneView* View,FCanvas* Canvas);
-	//@}
-
-	/**
-	 * Called when attempting to duplicate the selected actors by alt+dragging,
-	 * return true to prevent normal duplication.
-	 */
-	virtual bool HandleDragDuplicate() { return false; }
-
-	/**
-	 * Called when the mode wants to draw brackets around selected objects
-	 */
-	virtual void DrawBrackets( FEditorViewportClient* ViewportClient, FViewport* Viewport, const FSceneView* View, FCanvas* Canvas );
-
-	/** True if this mode uses a toolkit mode (eventually they all should) */
-	virtual bool UsesToolkits() const;
-
-	UWorld* GetWorld() const
-	{
-		return GEditor->GetEditorWorldContext().World();
-	}
-
-	// Property Widgets
-
-	/**
-	 * Lets each mode selectively enable widgets for editing properties tagged with 'Show 3D Widget' metadata.
-	 */
-	virtual bool UsesPropertyWidgets() const { return false; }
-
-	// @TODO: Find a better home for these?
-	/** Value of UPROPERTY can be edited with a widget in the editor (translation, rotation) */
-	static const FName MD_MakeEditWidget;
-	/** Specifies a function used for validation of the current value of a property.  The function returns a string that is empty if the value is valid, or contains an error description if the value is invalid */
-	static const FName MD_ValidateWidgetUsing;
-	/** Returns true if this structure can support creating a widget in the editor */
-	static bool CanCreateWidgetForStructure(const UStruct* InPropStruct);
-	/** Returns true if this property can support creating a widget in the editor */
-	static bool CanCreateWidgetForProperty(UProperty* InProp);
-	/** See if we should create a widget for the supplied property when selecting an actor instance */
-	static bool ShouldCreateWidgetForProperty(UProperty* InProp);
-
-public:
-
-	/** Request that this mode be deleted at the next convenient opportunity (FEditorModeTools::Tick) */
-	void RequestDeletion() { bPendingDeletion = true; }
-	
-	/** returns true if this mode is to be deleted at the next convenient opportunity (FEditorModeTools::Tick) */
-	bool IsPendingDeletion() const { return bPendingDeletion; }
-
-private:
-	/** true if this mode is pending removal from its owner */
-	bool bPendingDeletion;
-
-	/** Called whenever a mode type is unregistered */
-	void OnModeUnregistered(FEditorModeID ModeID);
-
-protected:
-	/** The current axis that is being dragged on the widget. */
-	EAxisList::Type CurrentWidgetAxis;
-
-	/** Optional array of tools for this mode. */
-	TArray<FModeTool*> Tools;
-
-	/** The tool that is currently active within this mode. */
-	FModeTool* CurrentTool;
-
-	/** Information pertaining to this mode. Assigned by FEditorModeRegistry. */
-	FEditorModeInfo Info;
-
-	/** Editor Mode Toolkit that is associated with this toolkit mode */
-	TSharedPtr<class FModeToolkit> Toolkit;
-
-	/** Pointer back to the mode tools that we are registered with */
-	FEditorModeTools* Owner;
-
-	// Property Widgets
-
-	/** Structure that holds info about our optional property widget */
-	struct FPropertyWidgetInfo
-	{
-		FString PropertyName;
-		int32 PropertyIndex;
-		FName PropertyValidationName;
-		FString DisplayName;
-		bool bIsTransform;
-
-		FPropertyWidgetInfo()
-			: PropertyIndex(INDEX_NONE)
-			, bIsTransform(false)
-		{
-		}
-	};
-
-	/**
-	 * Returns the first selected Actor, or NULL if there is no selection.
-	 */
-	AActor* GetFirstSelectedActorInstance() const;
-
-	/**
-	 * Gets an array of property widget info structures for the given struct/class type for the given container.
-	 *
-	 * @param InStruct The type of structure/class to access widget info structures for.
-	 * @param InContainer The container of the given type.
-	 * @param OutInfos An array of widget info structures (output).
-	 * @param PropertyNamePrefix Optional prefix to use to filter the output.
-	 */
-	void GetPropertyWidgetInfos(const UStruct* InStruct, const void* InContainer, TArray<FPropertyWidgetInfo>& OutInfos, FString PropertyNamePrefix = TEXT(""), FString DisplayNamePrefix = TEXT("")) const;
-
-	/** Name of the property currently being edited */
-	FString EditedPropertyName;
-
-	/** If the property being edited is an array property, this is the index of the element we're currently dealing with */
-	int32 EditedPropertyIndex;
-
-	/** Indicates  */
-	bool bEditedPropertyIsTransform;
-};
+#include "EdMode.h"
 
 /*------------------------------------------------------------------------------
 	Default.
@@ -432,418 +122,5 @@ public:
 };
 
 
-/**
- * A helper class to store the state of the various editor modes.
- */
-class UNREALED_API FEditorModeTools : public FGCObject
-{
-public:
-	FEditorModeTools();
-	virtual ~FEditorModeTools();
-
-	/**
-	 * Set the default editor mode for these tools
-	 */
-	void SetDefaultMode(FEditorModeID DefaultID);
-
-	/**
-	 * Activates the default mode defined by this class
-	 */
-	void ActivateDefaultMode();
-
-	/** 
-	 * Returns true if the default editor mode is active 
-	 */
-	bool IsDefaultModeActive() const;
-
-	/**
-	 * Activates an editor mode. Shuts down all other active modes which cannot run with the passed in mode.
-	 * 
-	 * @param InID		The ID of the editor mode to activate.
-	 * @param bToggle	true if the passed in editor mode should be toggled off if it is already active.
-	 */
-	void ActivateMode( FEditorModeID InID, bool bToggle = false );
-
-	/**
-	 * Deactivates an editor mode. 
-	 * 
-	 * @param InID		The ID of the editor mode to deactivate.
-	 */
-	void DeactivateMode(FEditorModeID InID);
-
-	/**
-	 * Deactivate the mode and entirely purge it from memory. Used when a mode type is unregistered
-	 */
-	void DestroyMode(FEditorModeID InID);
-
-protected:
-	
-	/** Deactivates the editor mode at the specified index */
-	void DeactivateModeAtIndex( int32 InIndex );
-		
-public:
-
-	/**
-	 * Deactivates all modes, note some modes can never be deactivated.
-	 */
-	void DeactivateAllModes();
-
-	/** 
-	 * Returns the editor mode specified by the passed in ID
-	 */
-	FEdMode* FindMode( FEditorModeID InID );
-
-	/**
-	 * Returns true if the current mode is not the specified ModeID.  Also optionally warns the user.
-	 *
-	 * @param	ModeID			The editor mode to query.
-	 * @param	ErrorMsg		If specified, inform the user the reason why this is a problem
-	 * @param	bNotifyUser		If true, display the error as a notification, instead of a dialog
-	 * @return					true if the current mode is not the specified mode.
-	 */
-	bool EnsureNotInMode(FEditorModeID ModeID, const FText& ErrorMsg = FText::GetEmpty(), bool bNotifyUser = false) const;
-
-	FMatrix GetCustomDrawingCoordinateSystem();
-	FMatrix GetCustomInputCoordinateSystem();
-	
-	/** 
-	 * Returns true if the passed in editor mode is active 
-	 */
-	bool IsModeActive( FEditorModeID InID ) const;
-
-	/**
-	 * Returns a pointer to an active mode specified by the passed in ID
-	 * If the editor mode is not active, NULL is returned
-	 */
-	FEdMode* GetActiveMode( FEditorModeID InID );
-	const FEdMode* GetActiveMode( FEditorModeID InID ) const;
-
-	template <typename SpecificModeType>
-	SpecificModeType* GetActiveModeTyped( FEditorModeID InID )
-	{
-		return (SpecificModeType*)GetActiveMode(InID);
-	}
-
-	template <typename SpecificModeType>
-	const SpecificModeType* GetActiveModeTyped( FEditorModeID InID ) const
-	{
-		return (SpecificModeType*)GetActiveMode(InID);
-	}
-
-	/**
-	 * Returns the active tool of the passed in editor mode.
-	 * If the passed in editor mode is not active or the mode has no active tool, NULL is returned
-	 */
-	const FModeTool* GetActiveTool( FEditorModeID InID ) const;
-
-	/** 
-	 * Returns an array of all active modes
-	 */
-	void GetActiveModes( TArray<FEdMode*>& OutActiveModes );
-
-	void SetShowWidget( bool InShowWidget )	{ bShowWidget = InShowWidget; }
-	bool GetShowWidget() const				{ return bShowWidget; }
-
-	void CycleWidgetMode (void);
-
-	/**Save Widget Settings to Ini file*/
-	void SaveWidgetSettings();
-	/**Load Widget Settings from Ini file*/
-	void LoadWidgetSettings();
-
-	/** Gets the widget axis to be drawn */
-	EAxisList::Type GetWidgetAxisToDraw( FWidget::EWidgetMode InWidgetMode ) const;
-
-	/** Mouse tracking interface.  Passes tracking messages to all active modes */
-	bool StartTracking(FEditorViewportClient* InViewportClient, FViewport* InViewport);
-	bool EndTracking(FEditorViewportClient* InViewportClient, FViewport* InViewport);
-	bool IsTracking() const { return bIsTracking; }
-
-	bool AllowsViewportDragTool() const;
-
-	/** Notifies all active modes that a map change has occured */
-	void MapChangeNotify();
-
-	/** Notifies all active modes to empty their selections */
-	void SelectNone();
-
-	/** Notifies all active modes of box selection attempts */
-	bool BoxSelect( FBox& InBox, bool InSelect );
-
-	/** Notifies all active modes of frustum selection attempts */
-	bool FrustumSelect( const FConvexVolume& InFrustum, bool InSelect );
-
-	/** true if any active mode uses a transform widget */
-	bool UsesTransformWidget() const;
-
-	/** true if any active mode uses the passed in transform widget */
-	bool UsesTransformWidget( FWidget::EWidgetMode CheckMode ) const;
-
-	/** Sets the current widget axis */
-	void SetCurrentWidgetAxis( EAxisList::Type NewAxis );
-
-	/** Notifies all active modes of mouse click messages. */
-	bool HandleClick(FEditorViewportClient* InViewportClient,  HHitProxy *HitProxy, const FViewportClick &Click );
-
-	/** true if the passed in brush actor should be drawn in wireframe */	
-	bool ShouldDrawBrushWireframe( AActor* InActor ) const;
-
-	/** true if brush vertices should be drawn */
-	bool ShouldDrawBrushVertices() const;
-
-	/** Ticks all active modes */
-	void Tick( FEditorViewportClient* ViewportClient, float DeltaTime );
-
-	/** Notifies all active modes of any change in mouse movement */
-	bool InputDelta( FEditorViewportClient* InViewportClient,FViewport* InViewport,FVector& InDrag,FRotator& InRot,FVector& InScale );
-
-	/** Notifies all active modes of captured mouse movement */	
-	bool CapturedMouseMove( FEditorViewportClient* InViewportClient, FViewport* InViewport, int32 InMouseX, int32 InMouseY );
-
-	/** Notifies all active modes of keyboard input */
-	bool InputKey( FEditorViewportClient* InViewportClient, FViewport* Viewport, FKey Key, EInputEvent Event);
-
-	/** Notifies all active modes of axis movement */
-	bool InputAxis( FEditorViewportClient* InViewportClient, FViewport* Viewport, int32 ControllerId, FKey Key, float Delta, float DeltaTime);
-
-	bool MouseEnter( FEditorViewportClient* InViewportClient, FViewport* Viewport, int32 X, int32 Y );
-	
-	bool MouseLeave( FEditorViewportClient* InViewportClient, FViewport* Viewport );
-
-	/** Notifies all active modes that the mouse has moved */
-	bool MouseMove( FEditorViewportClient* InViewportClient, FViewport* Viewport, int32 X, int32 Y );
-
-	/** Notifies all active modes that a viewport has received focus */
-	bool ReceivedFocus( FEditorViewportClient* InViewportClient, FViewport* Viewport );
-
-	/** Notifies all active modes that a viewport has lost focus */
-	bool LostFocus( FEditorViewportClient* InViewportClient, FViewport* Viewport );
-
-	/** Draws all active modes */	
-	void DrawActiveModes( const FSceneView* InView, FPrimitiveDrawInterface* PDI );
-
-	/** Renders all active modes */
-	void Render( const FSceneView* InView, FViewport* Viewport, FPrimitiveDrawInterface* PDI );
-
-	/** Draws the HUD for all active modes */
-	void DrawHUD( FEditorViewportClient* InViewportClient,FViewport* Viewport,const FSceneView* View,FCanvas* Canvas );
-
-	/** Calls PostUndo on all active components */
-	void PostUndo();
-
-	/** True if we should allow widget move */
-	bool AllowWidgetMove() const;
-
-	/** True if we should disallow mouse delta tracking. */
-	bool DisallowMouseDeltaTracking() const;
-
-	/** Get a cursor to override the default with, if any */
-	bool GetCursor(EMouseCursor::Type& OutCursor) const;
-
-	/**
-	 * Returns a good location to draw the widget at.
-	 */
-	FVector GetWidgetLocation() const;
+#include "EditorModeManager.h"
 
-	/**
-	 * Changes the current widget mode.
-	 */
-	void SetWidgetMode( FWidget::EWidgetMode InWidgetMode );
-
-	/**
-	 * Allows you to temporarily override the widget mode.  Call this function again
-	 * with WM_None to turn off the override.
-	 */
-	void SetWidgetModeOverride( FWidget::EWidgetMode InWidgetMode );
-
-	/**
-	 * Retrieves the current widget mode, taking overrides into account.
-	 */
-	FWidget::EWidgetMode GetWidgetMode() const;
-
-	/**
-	 * Gets the current state of script editor usage of show friendly names
-	 * @ return - If true, replace variable names with sanatized ones
-	 */
-	bool GetShowFriendlyVariableNames() const;
-
-	/**
-	 * Sets Interp Editor inversion of panning on drag
-	 */
-	void SetInterpPanInvert(bool bInInterpPanInverted) { bInterpPanInverted = bInInterpPanInverted; }
-
-	/**
-	 * Gets the current state of the interp editor panning
-	 * @ return - If true, the interp editor left-right panning is inverted
-	 */
-	bool GetInterpPanInvert(void) const { return bInterpPanInverted; }
-
-	/**
-	 * Sets a bookmark in the levelinfo file, allocating it if necessary.
-	 * 
-	 * @param InIndex			Index of the bookmark to set.
-	 * @param InViewportClient	Level editor viewport client used to get pointer the world that the bookmark is for.
-	 */
-	void SetBookmark( uint32 InIndex, FEditorViewportClient* InViewportClient );
-
-	/**
-	 * Checks to see if a bookmark exists at a given index
-	 * 
-	 * @param InIndex			Index of the bookmark to set.
-	 * @param InViewportClient	Level editor viewport client used to get pointer the world to check for the bookmark in.
-	 */
-	bool CheckBookmark( uint32 InIndex, FEditorViewportClient* InViewportClient );
-
-	/**
-	 * Retrieves a bookmark from the list.
-	 * 
-	 * @param InIndex			Index of the bookmark to set.
-	 * @param					Set to true to restore the visibility of any streaming levels.
-	 * @param InViewportClient	Level editor viewport client used to get pointer the world that the bookmark is relevant to.
-	 */
-	void JumpToBookmark( uint32 InIndex, bool bShouldRestoreLevelVisibility, FEditorViewportClient* InViewportClient );
-
-	/**
-	 * Clears a bookmark from the list.
-	 * 
-	 * @param InIndex			Index of the bookmark to clear.
-	 * @param InViewportClient	Level editor viewport client used to get pointer the world that the bookmark is relevant to.
-	 */
-	void ClearBookmark( uint32 InIndex, FEditorViewportClient* InViewportClient );
-
-	/**
-	 * Clears all book marks
-	 * 
-	 * @param InViewportClient	Level editor viewport client used to get pointer the world to clear the bookmarks from.
-	 */
-	void ClearAllBookmarks( FEditorViewportClient* InViewportClient );
-
-	// FGCObject interface
-	virtual void AddReferencedObjects( FReferenceCollector& Collector ) override;
-	// End of FGCObject interface
-
-	/**
-	 * Loads the state that was saved in the INI file
-	 */
-	void LoadConfig(void);
-
-	/**
-	 * Saves the current state to the INI file
-	 */
-	void SaveConfig(void);
-
-	/** 
-	 * Sets the pivot locations
-	 * 
-	 * @param Location 		The location to set
-	 * @param bIncGridBase	Whether or not to also set the GridBase
-	 */
-	void SetPivotLocation( const FVector& Location, const bool bIncGridBase );
-
-	/**
-	 * Multicast delegate for OnModeEntered and OnModeExited callbacks.
-	 *
-	 * First parameter:  The editor mode that was changed
-	 * Second parameter:  True if entering the mode, or false if exiting the mode
-	 */
-	DECLARE_EVENT_TwoParams( FEditorModeTools, FEditorModeChangedEvent, FEdMode*, bool );
-	FEditorModeChangedEvent& OnEditorModeChanged() { return EditorModeChangedEvent; }
-
-	/** delegate type for triggering when widget mode changed */
-	DECLARE_EVENT_OneParam( FEditorModeTools, FWidgetModeChangedEvent, FWidget::EWidgetMode );
-	FWidgetModeChangedEvent& OnWidgetModeChanged() { return WidgetModeChangedEvent; }
-
-	/**	Broadcasts the WidgetModeChanged event */
-	void BroadcastWidgetModeChanged(FWidget::EWidgetMode InWidgetMode) { WidgetModeChangedEvent.Broadcast(InWidgetMode); }
-
-	/**	Broadcasts the EditorModeChanged event */
-	void BroadcastEditorModeChanged(FEdMode* Mode, bool IsEnteringMode) { EditorModeChangedEvent.Broadcast(Mode, IsEnteringMode); }
-
-	/**
-	 * Returns the current CoordSystem
-	 * 
-	 * @param bGetRawValue true when you want the actual value of CoordSystem, not the value modified by the state.
-	 */
-	ECoordSystem GetCoordSystem(bool bGetRawValue = false);
-
-	/** Sets the current CoordSystem */
-	void SetCoordSystem(ECoordSystem NewCoordSystem);
-
-	/** Sets the hide viewport UI state */
-	void SetHideViewportUI( bool bInHideViewportUI ) { bHideViewportUI = bInHideViewportUI; }
-
-	/** Is the viewport UI hidden? */
-	bool IsViewportUIHidden() const { return bHideViewportUI; }
-
-	bool PivotShown, Snapping, SnappedActor;
-	FVector CachedLocation, PivotLocation, SnappedLocation, GridBase;
-
-	/** The angle for the translate rotate widget */
-	float TranslateRotateXAxisAngle;
-
-	/** Draws in the top level corner of all FEditorViewportClient windows (can be used to relay info to the user). */
-	FString InfoString;
-
-	/** Sets the host for toolkits created via modes from this mode manager (can only be called once) */
-	void SetToolkitHost(TSharedRef<class IToolkitHost> Host);
-
-	/** Returns the host for toolkits created via modes from this mode manager */
-	TSharedPtr<class IToolkitHost> GetToolkitHost() const;
-
-	/**
-	 * Returns the set of selected actors.
-	 */
-	virtual class USelection* GetSelectedActors() const;
-
-	/**
-	 * @return the set of selected non-actor objects.
-	 */
-	virtual class USelection* GetSelectedObjects() const;
-protected:
-	/** 
-	 * Delegate handlers
-	 **/
-	void OnEditorSelectionChanged(UObject* NewSelection);
-	void OnEditorSelectNone();
-
-	/** The mode ID to use as a default */
-	FEditorModeID DefaultID;
-
-	/** A list of active editor modes. */
-	TArray< TSharedPtr<FEdMode> > Modes;
-
-	/** The host of the toolkits created by these modes */
-	TWeakPtr<IToolkitHost> ToolkitHost;
-
-	/** A list of previously active editor modes that we will potentially recycle */
-	TMap< FEditorModeID, TSharedPtr<FEdMode> > RecycledModes;
-
-	/** The mode that the editor viewport widget is in. */
-	FWidget::EWidgetMode WidgetMode;
-
-	/** If the widget mode is being overridden, this will be != WM_None. */
-	FWidget::EWidgetMode OverrideWidgetMode;
-
-	/** If 1, draw the widget and let the user interact with it. */
-	bool bShowWidget;
-
-	/** If true, panning in the Interp Editor is inverted. If false, panning is normal*/
-	bool bInterpPanInverted;
-
-	/** if true, the viewports will hide all UI overlays */
-	bool bHideViewportUI;
-
-private:
-
-	/** The coordinate system the widget is operating within. */
-	ECoordSystem CoordSystem;
-
-	/** Multicast delegate that is broadcast when a mode is entered or exited */
-	FEditorModeChangedEvent EditorModeChangedEvent;
-
-	/** Multicast delegate that is broadcast when a widget mode is changed */
-	FWidgetModeChangedEvent WidgetModeChangedEvent;
-
-	/** Flag set between calls to StartTracking() and EndTracking() */
-	bool bIsTracking;
-};
