Editor: Fix a crash when doing viewport console commands after opening and closing a standalone asset editor (stat delegates were not getting unregistered correctly)

diff --git a/Engine/Source/Editor/UnrealEd/Private/EditorViewportClient.cpp b/Engine/Source/Editor/UnrealEd/Private/EditorViewportClient.cpp
--- a/Engine/Source/Editor/UnrealEd/Private/EditorViewportClient.cpp
+++ b/Engine/Source/Editor/UnrealEd/Private/EditorViewportClient.cpp
@@ -321,6 +321,10 @@ FEditorViewportClient::~FEditorViewportClient()
 		GEditor->AllViewportClients[ViewportIndex]->ViewIndex = ViewportIndex;
 	}
 
+	FCoreDelegates::StatCheckEnabled.RemoveAll(this);
+	FCoreDelegates::StatEnabled.RemoveAll(this);
+	FCoreDelegates::StatDisabled.RemoveAll(this);
+	FCoreDelegates::StatDisableAll.RemoveAll(this);
 }
 
 void FEditorViewportClient::RedrawRequested(FViewport* InViewport)
diff --git a/Engine/Source/Editor/UnrealEd/Private/LevelEditorViewport.cpp b/Engine/Source/Editor/UnrealEd/Private/LevelEditorViewport.cpp
--- a/Engine/Source/Editor/UnrealEd/Private/LevelEditorViewport.cpp
+++ b/Engine/Source/Editor/UnrealEd/Private/LevelEditorViewport.cpp
@@ -1461,11 +1461,6 @@ FLevelEditorViewportClient::~FLevelEditorViewportClient()
 
 	GetMutableDefault<ULevelEditorViewportSettings>()->OnSettingChanged().RemoveAll(this);
 
-	FCoreDelegates::StatCheckEnabled.RemoveAll(this);
-	FCoreDelegates::StatEnabled.RemoveAll(this);
-	FCoreDelegates::StatDisabled.RemoveAll(this);
-	FCoreDelegates::StatDisableAll.RemoveAll(this);
-
 	GEditor->LevelViewportClients.Remove(this);
 
 	RemoveReferenceToWorldContext(GEditor->GetEditorWorldContext());
