Gather text from data tables implemented.

diff --git a/Engine/Source/Editor/UnrealEd/Private/Commandlets/GatherTextFromAssetsCommandlet.cpp b/Engine/Source/Editor/UnrealEd/Private/Commandlets/GatherTextFromAssetsCommandlet.cpp
--- a/Engine/Source/Editor/UnrealEd/Private/Commandlets/GatherTextFromAssetsCommandlet.cpp
+++ b/Engine/Source/Editor/UnrealEd/Private/Commandlets/GatherTextFromAssetsCommandlet.cpp
@@ -351,29 +351,25 @@ UGatherTextFromAssetsCommandlet::UGatherTextFromAssetsCommandlet(const FObjectIn
 
 namespace
 {
-	class FGatherTextFromObject
+	template<typename TConcreteChildClass>
+	class FGatherTextFromObject_Base
 	{
 	public:
 		static inline void Execute(UGatherTextFromAssetsCommandlet* const Commandlet, UObject* const Object, const UPackage* const ObjectPackage, const bool ShouldFixBroken)
 		{
 			if( !Object->HasAnyFlags( RF_Transient | RF_PendingKill ) )
 			{
-				FGatherTextFromObject(Commandlet, Object, ObjectPackage, ShouldFixBroken);
+				TConcreteChildClass(Commandlet, Object, ObjectPackage, ShouldFixBroken).Execute_Internal();
 			}
 		}
 
-	private:
-		FGatherTextFromObject(UGatherTextFromAssetsCommandlet* const InCommandlet, UObject* const InObject, const UPackage* const InObjectPackage, const bool InShouldFixBroken)
+	protected:
+		FGatherTextFromObject_Base(UGatherTextFromAssetsCommandlet* const InCommandlet, UObject* const InObject, const UPackage* const InObjectPackage, const bool InShouldFixBroken)
 			: Commandlet(InCommandlet)
 			, Object(InObject)
 			, ObjectPackage(InObjectPackage)
 			, ShouldFixBroken(InShouldFixBroken)
 		{
-			// Iterate over all fields of the object's class.
-			for (TFieldIterator<UProperty>PropIt(Object->GetClass(), EFieldIteratorFlags::IncludeSuper, EFieldIteratorFlags::ExcludeDeprecated, EFieldIteratorFlags::IncludeInterfaces); PropIt; ++PropIt)
-			{
-				ProcessProperty( *PropIt, PropIt->ContainerPtrToValuePtr<void>(Object) );
-			}
 		}
 
 		void ProcessProperty(UProperty* const Property, void* const ValueAddress)
@@ -457,12 +453,53 @@ namespace
 			}
 		}
 
-	private:
+	protected:
 		UGatherTextFromAssetsCommandlet* const Commandlet;
 		UObject* const Object;
 		const UPackage* const ObjectPackage;
 		const bool ShouldFixBroken;
 	};
+
+	class FGatherTextFromObject : public FGatherTextFromObject_Base<FGatherTextFromObject>
+	{
+	public:
+		FGatherTextFromObject(UGatherTextFromAssetsCommandlet* const InCommandlet, UObject* const InObject, const UPackage* const InObjectPackage, const bool InShouldFixBroken)
+			: FGatherTextFromObject_Base(InCommandlet, InObject, InObjectPackage, InShouldFixBroken)
+		{
+		}
+
+		void Execute_Internal()
+		{
+			// Iterate over all fields of the object's class.
+			for (TFieldIterator<UProperty>PropIt(Object->GetClass(), EFieldIteratorFlags::IncludeSuper, EFieldIteratorFlags::ExcludeDeprecated, EFieldIteratorFlags::IncludeInterfaces); PropIt; ++PropIt)
+			{
+				ProcessProperty( *PropIt, PropIt->ContainerPtrToValuePtr<void>(Object) );
+			}
+		}
+	};
+
+	class FGatherTextFromDataTable : public FGatherTextFromObject_Base<FGatherTextFromDataTable>
+	{
+	public:
+		FGatherTextFromDataTable(UGatherTextFromAssetsCommandlet* const InCommandlet, UObject* const InDataTable, const UPackage* const InObjectPackage, const bool InShouldFixBroken)
+			: FGatherTextFromObject_Base(InCommandlet, InDataTable, InObjectPackage, InShouldFixBroken)
+		{
+		}
+
+		void Execute_Internal()
+		{
+			UDataTable* const DataTable = Cast<UDataTable>(Object);
+
+			// Iterate over all properties of the struct's class.
+			for (TFieldIterator<UProperty>PropIt(DataTable->RowStruct, EFieldIteratorFlags::IncludeSuper, EFieldIteratorFlags::ExcludeDeprecated, EFieldIteratorFlags::IncludeInterfaces); PropIt; ++PropIt)
+			{
+				for (const auto& Pair : DataTable->RowMap)
+				{
+					ProcessProperty( *PropIt, PropIt->ContainerPtrToValuePtr<void>(Pair.Value) );
+				}
+			}
+		}
+	};
 }
 
 void UGatherTextFromAssetsCommandlet::ProcessPackages( const TArray< UPackage* >& PackagesToProcess )
@@ -489,10 +526,15 @@ void UGatherTextFromAssetsCommandlet::ProcessPackages( const TArray< UPackage* >
 					UE_LOG(LogGatherTextFromAssetsCommandlet, Warning, TEXT("%s - Invalid generated class!"), *Blueprint->GetFullName());
 				}
 			}
+			else if ( Object->IsA( UDataTable::StaticClass() ) )
+			{
+				UDataTable* DataTable = Cast<UDataTable>(Object);
+				FGatherTextFromDataTable::Execute(this, DataTable, Package, bFixBroken);
+			}
 			else if( Object->IsA( UDialogueWave::StaticClass() ) )
 			{
-				UDialogueWave* DialogueWave = Cast<UDialogueWave>( Object );
-				ProcessDialogueWave( DialogueWave );
+				UDialogueWave* DialogueWave = Cast<UDialogueWave>(Object);
+				ProcessDialogueWave(DialogueWave);
 			}
 
 			FGatherTextFromObject::Execute(this, Object, Package, bFixBroken);
