Fixing compile error

diff --git a/Engine/Source/Editor/UnrealEd/Public/Editor.h b/Engine/Source/Editor/UnrealEd/Public/Editor.h
--- a/Engine/Source/Editor/UnrealEd/Public/Editor.h
+++ b/Engine/Source/Editor/UnrealEd/Public/Editor.h
@@ -587,6 +587,17 @@ namespace EditorUtilities
 	 */
 	UNREALED_API AActor* GetSimWorldCounterpartActor( AActor* Actor );
 
+	/**
+	 * Guiven an actor in the editor world, and SourceComponent from Simulation or PIE world
+	 * find the matching component in the Editor World
+	 *
+	 * @param	SourceComponent	SouceCompoent in SIM world
+	 * @param	TargetActor		TargetActor in editor world
+	 *
+	 * @return	the sound editor component or NULL if we couldn't find
+	 */
+	UNREALED_API UActorComponent* FindMatchingComponentInstance( UActorComponent* SourceComponent, AActor* TargetActor );
+
 	/** Options for CopyActorProperties */
 	namespace ECopyOptions
 	{
