[UE-5068] First pass, adding command-line options to invoke the Blueprint diff tool (from the command-line... can be set up with p4v).

diff --git a/Engine/Source/Editor/UnrealEd/Private/Editor.cpp b/Engine/Source/Editor/UnrealEd/Private/Editor.cpp
--- a/Engine/Source/Editor/UnrealEd/Private/Editor.cpp
+++ b/Engine/Source/Editor/UnrealEd/Private/Editor.cpp
@@ -10,6 +10,7 @@
 #include "EditorSupportDelegates.h"
 #include "Factories.h"
 #include "BSPOps.h"
+#include "EditorCommandLineUtils.h"
 
 // needed for the RemotePropagator
 #include "SoundDefinitions.h"
@@ -713,22 +714,7 @@ void UEditorEngine::Init(IEngineLoop* InEngineLoop)
 	// Purge garbage.
 	Cleanse( false, 0, NSLOCTEXT("UnrealEd", "Startup", "Startup") );
 
-	// If specified, Lightmass has to be launched manually with -debug (e.g. through a debugger).
-	// This creates a job with a hard-coded GUID, and allows Lightmass to be executed multiple times (even stand-alone).
-	if ( FParse::Param(FCommandLine::Get(), TEXT("LIGHTMASSDEBUG")) )
-	{
-		extern bool GLightmassDebugMode;
-		GLightmassDebugMode = true;
-		UE_LOG(LogInit, Log, TEXT("Running Engine with Lightmass Debug Mode ENABLED"));
-	}
-
-	// If specified, all participating Lightmass agents will report back detailed stats to the log.
-	if ( FParse::Param(FCommandLine::Get(), TEXT("LIGHTMASSSTATS")) )
-	{
-		extern bool GLightmassStatsMode;
-		GLightmassStatsMode = true;
-		UE_LOG(LogInit, Log, TEXT("Running Engine with Lightmass Stats Mode ENABLED"));
-	}
+	FEditorCommandLineUtils::ProcessEditorCommands(FCommandLine::Get());
 
 	// for IsInitialized()
 	bIsInitialized = true;
diff --git a/Engine/Source/Editor/UnrealEd/Private/EditorCommandLineUtils.cpp b/Engine/Source/Editor/UnrealEd/Private/EditorCommandLineUtils.cpp
new file mode 100644
--- /dev/null
+++ b/Engine/Source/Editor/UnrealEd/Private/EditorCommandLineUtils.cpp
@@ -0,0 +1,472 @@
+// Copyright 1998-2014 Epic Games, Inc. All Rights Reserved.
+
+#include "UnrealEd.h"
+#include "EditorCommandLineUtils.h"
+#include "Commandlets/Commandlet.h" // for ParseCommandLine()
+#include "Paths.h"
+#include "App.h"					// for IsGameNameEmpty()
+#include "Dialogs.h"				// for OpenMsgDlgInt_NonModal()
+#include "TickableEditorObject.h"
+#include "Editor/MainFrame/Public/Interfaces/IMainFrameModule.h"
+#include "AssetRegistryModule.h"
+
+#define LOCTEXT_NAMESPACE "EditorCommandLineUtils"
+
+/*******************************************************************************
+ * EditorCommandLineUtilsImpl Declaration
+ ******************************************************************************/
+
+/**  */
+namespace EditorCommandLineUtilsImpl
+{
+	static const TCHAR* DebugLightmassCommandSwitch = TEXT("LIGHTMASSDEBUG");
+	static const TCHAR* LightmassStatsCommandSwitch = TEXT("LIGHTMASSSTATS");
+
+	static const TCHAR* DiffCommandSwitch  = TEXT("diff");
+	static const FText  DiffCommandHelpTxt = LOCTEXT("DiffCommandeHelpText", "\
+Usage: \n\
+    -diff [options] left right                                                 \n\
+    -diff [options] left right base                                            \n\
+\n\
+Options: \n\
+    -echo               Prints back the command arguments and then exits.      \n\
+    -help, -h, -?       Display this message and then exits.                   \n");
+
+	/**  */
+	static bool ParseCommandArgs(const TCHAR* FullEditorCmdLine, const TCHAR* CmdSwitch, FString& CmdArgsOut);
+
+	/**  */
+	static void RaiseEditorMessageBox(const FText& Title, const FText& BodyText, const bool bExitOnClose);
+
+	/**  */
+	static void ForceCloseEditor();
+	
+	/**  */
+	static void RunAssetDiffCommand(TSharedPtr<SWindow> MainEditorWindow, bool bIsRunningProjBrowser, FString CommandArgs);
+
+	/**  */
+	static int32 ParseRevisionFromFilename(const FString& Filename, FString& CLeanFilenameOut);
+
+	/** */
+	static UObject* ExtractAssetFromPackage(UPackage* Package, const FString& AssetName);
+
+	/**  */
+//	static int32 GetAssetPackagePaths(const FString& InAssetName, TArray<FString>& PackagePathsOut, bool bIncludeDevFolder);
+}
+
+/*******************************************************************************
+ * FCommandLineErrorReporter
+ ******************************************************************************/
+ 
+/**  */
+struct FCommandLineErrorReporter
+{
+	FCommandLineErrorReporter(const FString& Command, const FString& CommandArgs)
+		: CommandSwitch(FText::FromString(Command))
+		, FullCommand(FText::FromString("-" + Command + " " + CommandArgs))
+		, bHasBlockingError(false)
+	{}
+
+	/**  */
+	void ReportFatalError(const FText& Title, const FText& ErrorMsg);
+
+	/**  */
+	void ReportError(const FText& Title, const FText& ErrorMsg, bool bIsFatal);
+
+	/**  */
+	bool HasBlockingError() const;
+
+private:
+	FText CommandSwitch;
+	FText FullCommand;
+	bool  bHasBlockingError;
+};
+
+//------------------------------------------------------------------------------
+void FCommandLineErrorReporter::ReportFatalError(const FText& Title, const FText& ErrorMsg)
+{
+	ReportError(Title, ErrorMsg, /*bIsFatal =*/true);
+}
+
+
+//------------------------------------------------------------------------------
+void FCommandLineErrorReporter::ReportError(const FText& Title, const FText& ErrorMsg, bool bIsFatal)
+{
+	if (bHasBlockingError)
+	{
+		return;
+	}
+
+	FText FullErrorMsg = FText::Format(LOCTEXT("CommandLineError", "Erroneous editor command: {0}\n\n{1}\n\nRun '-{2} -h' for more help."),
+		FullCommand, ErrorMsg, CommandSwitch);
+
+	bHasBlockingError = bIsFatal;
+	EditorCommandLineUtilsImpl::RaiseEditorMessageBox(Title, FullErrorMsg, /*bShutdownOnOk =*/bIsFatal);
+}
+
+//------------------------------------------------------------------------------
+bool FCommandLineErrorReporter::HasBlockingError() const
+{
+	return bHasBlockingError;
+}
+
+/*******************************************************************************
+ * FFauxStandaloneToolManager Implementation
+ ******************************************************************************/
+
+/**  
+ * Helps keep up the facade for tools can launch "stand-alone"... Hides the main 
+ * editor window, and monitors for when all visible windows are closed (so it
+ * can kill the editor process).
+ */
+class FFauxStandaloneToolManager : FTickableEditorObject
+{
+public:
+	FFauxStandaloneToolManager(TSharedPtr<SWindow> MainEditorWindow);
+
+	// FTickableEditorObject interface
+	virtual TStatId GetStatId() const override;
+	virtual bool IsTickable() const override { return true; }
+	virtual void Tick(float DeltaTime) override;	
+	// End FTickableEditorObject interface
+
+private:
+	TWeakPtr<SWindow> MainEditorWindow;
+};
+
+//------------------------------------------------------------------------------
+FFauxStandaloneToolManager::FFauxStandaloneToolManager(TSharedPtr<SWindow> InMainEditorWindow)
+	: MainEditorWindow(InMainEditorWindow)
+{
+	// present the illusion that this is a stand-alone editor by hiding the 
+	// root level editor window
+	InMainEditorWindow->HideWindow();
+}
+
+//------------------------------------------------------------------------------
+TStatId FFauxStandaloneToolManager::GetStatId() const
+{
+	RETURN_QUICK_DECLARE_CYCLE_STAT(FFauxStandaloneToolManager, STATGROUP_Tickables);
+}
+
+//------------------------------------------------------------------------------
+void FFauxStandaloneToolManager::Tick(float DeltaTime)
+{
+	if (MainEditorWindow.IsValid())
+	{	
+		FSlateApplication& WindowManager = FSlateApplication::Get();
+		TArray< TSharedRef<SWindow> > ActiveWindows = WindowManager.GetInteractiveTopLevelWindows();
+
+		bool bVisibleWindowFound = false;
+		for (TSharedRef<SWindow> Window : ActiveWindows)
+		{
+			if (Window->IsVisible())
+			{
+				bVisibleWindowFound = true;
+				break;
+			}
+		}
+
+		if (!bVisibleWindowFound)
+		{
+			EditorCommandLineUtilsImpl::ForceCloseEditor();
+		}
+	}
+	else
+	{
+		EditorCommandLineUtilsImpl::ForceCloseEditor();
+	}
+}
+
+/*******************************************************************************
+ * EditorCommandLineUtilsImpl Implementation
+ ******************************************************************************/
+
+//------------------------------------------------------------------------------
+static bool EditorCommandLineUtilsImpl::ParseCommandArgs(const TCHAR* FullEditorCmdLine, const TCHAR* CmdSwitch, FString& CmdArgsOut)
+{
+	if (FParse::Param(FullEditorCmdLine, CmdSwitch))
+	{
+		FString CmdPrefix;
+		FString(FullEditorCmdLine).Split(FString("-") + CmdSwitch, &CmdPrefix, &CmdArgsOut);
+		return true;
+	}
+	return false;
+}
+
+//------------------------------------------------------------------------------
+static void EditorCommandLineUtilsImpl::RaiseEditorMessageBox(const FText& Title, const FText& BodyText, const bool bExitOnClose)
+{
+	FOnMsgDlgResult OnDialogClosed;
+	if (bExitOnClose)
+	{
+		auto OnDialogClosedLambda = [](const TSharedRef<SWindow>&, EAppReturnType::Type)
+		{
+			ForceCloseEditor();
+		};
+		OnDialogClosed = FOnMsgDlgResult::CreateStatic(OnDialogClosedLambda);
+	}
+
+	OpenMsgDlgInt_NonModal(EAppMsgType::Ok, BodyText, Title, OnDialogClosed)->ShowWindow();
+}
+
+//------------------------------------------------------------------------------
+static void EditorCommandLineUtilsImpl::ForceCloseEditor()
+{
+	IMainFrameModule& MainFrameModule = IMainFrameModule::Get();
+	if (MainFrameModule.IsWindowInitialized())
+	{
+		MainFrameModule.RequestCloseEditor();
+	}
+}
+
+//------------------------------------------------------------------------------
+static void EditorCommandLineUtilsImpl::RunAssetDiffCommand(TSharedPtr<SWindow> MainEditorWindow, bool bIsRunningProjBrowser, FString CommandArgs)
+{
+	// if the editor is running the project browser, then the user has to first 
+	// select a project (and then the editor will re-launch with this command).
+	if (bIsRunningProjBrowser) 
+	{
+		return;
+	}
+
+	// static so it exists past this function, but doesn't get instantiated  
+	// until this function is called
+	static FFauxStandaloneToolManager FauxStandaloneToolManager(MainEditorWindow);
+
+	TMap<FString, FString> Params;
+	TArray<FString> Tokens;
+	TArray<FString> Switches;
+	UCommandlet::ParseCommandLine(*CommandArgs, Tokens, Switches, Params);
+
+	if (Switches.Contains("h") ||
+		Switches.Contains("?") ||
+		Switches.Contains("help"))
+	{
+		RaiseEditorMessageBox(LOCTEXT("DiffCommandHelp", "UAsset Diff Command-Line Help"), DiffCommandHelpTxt, /*bExitOnClose =*/true);
+		return;
+	}
+
+	if (Switches.Contains("echo"))
+	{
+		RaiseEditorMessageBox(LOCTEXT("DiffCommandHelp", "Passed Command Arguments"), FText::FromString(CommandArgs), /*bExitOnClose =*/true);
+		return;
+	}
+
+	FCommandLineErrorReporter ErrorReporter(DiffCommandSwitch, CommandArgs);
+	if (Tokens.Num() < 2)
+	{
+		ErrorReporter.ReportFatalError(LOCTEXT("TooFewParamsTitle", "Too Few Parameters"),
+			LOCTEXT("TooFewParamsError", "At least two files are needed."));
+	}
+
+	UObject* LeftAsset = nullptr, *RightAsset = nullptr, *BaseAsset = nullptr;
+	const int32 MaxFileCount = 3;
+	UObject** FileAssets[MaxFileCount] = { &LeftAsset, &RightAsset, &BaseAsset };
+	int32 ParsedFileCount = 0;
+	int32 AssetRevisionNums[MaxFileCount] = { -1, -1, -1 };
+
+	for (int32 FileIndex = 0; FileIndex < Tokens.Num() && ParsedFileCount < MaxFileCount; ++FileIndex)
+	{
+		FString& FilePath = Tokens[FileIndex];
+		if (!FPaths::FileExists(FilePath))
+		{
+			ErrorReporter.ReportFatalError(LOCTEXT("BadFilePathTitle", "Bad File Path"),
+				FText::Format(LOCTEXT("BadFilePathError", "'{0}' is an invalid file."), FText::FromString(FilePath)));
+			continue;
+		}
+
+		FString CleanFilename;
+		int32 RevisionNum = ParseRevisionFromFilename(FilePath, CleanFilename);
+		FString LocalFilePath = FPaths::Combine(*FPaths::DiffDir(), *(TEXT("Temp-Rev-") + FString::FromInt(RevisionNum) + "-" + CleanFilename));
+		// UE4 cannot open files with a number in them
+		IFileManager::Get().Copy(*LocalFilePath, *FilePath);
+
+		if (UPackage* TempAssetPkg = LoadPackage(/*Outer =*/nullptr, *LocalFilePath, LOAD_None))
+		{
+			if (UObject* FileAsset = ExtractAssetFromPackage(TempAssetPkg, CleanFilename))
+			{
+				*FileAssets[ParsedFileCount] = FileAsset;
+				AssetRevisionNums[ParsedFileCount] = RevisionNum;
+				++ParsedFileCount;
+			}
+			else
+			{
+				ErrorReporter.ReportFatalError(LOCTEXT("AssetNotFoundTitle", "Asset Not Found"),
+					FText::Format( LOCTEXT("AssetNotFoundError", "Failed to find the asset '{0}', inside the package file: '{1}'."), 
+					FText::FromString(FPaths::GetBaseFilename(CleanFilename)), FText::FromString(FilePath) ));
+			}
+		}
+		else
+		{
+			ErrorReporter.ReportFatalError(LOCTEXT("LoadFailedTitle", "Unable to Open File"),
+				FText::Format(LOCTEXT("LoadFailedError", "'{0}' failed to load."), FText::FromString(FilePath)));
+		}
+	}
+
+	if (!ErrorReporter.HasBlockingError())
+	{
+		FRevisionInfo LeftRevisionInfo = FRevisionInfo::InvalidRevision();
+		LeftRevisionInfo.Revision  = AssetRevisionNums[0];
+		FRevisionInfo RightRevisionInfo = FRevisionInfo::InvalidRevision();
+		RightRevisionInfo.Revision = AssetRevisionNums[1];
+
+		FAssetToolsModule& AssetToolsModule = FModuleManager::GetModuleChecked<FAssetToolsModule>("AssetTools");
+		AssetToolsModule.Get().DiffAssets(LeftAsset, RightAsset, LeftRevisionInfo, RightRevisionInfo);
+	}
+}
+
+//------------------------------------------------------------------------------
+static int32 EditorCommandLineUtilsImpl::ParseRevisionFromFilename(const FString& Filename, FString& CleanedFilenameOut)
+{
+	int32 RevisionNum = -1;
+
+	CleanedFilenameOut = FPaths::GetBaseFilename(Filename);
+
+	FString BaseFileName, RevisionStr;
+	CleanedFilenameOut.Split(TEXT("#"), &BaseFileName, &RevisionStr);
+
+	CleanedFilenameOut = BaseFileName + FPackageName::GetAssetPackageExtension();
+
+	if (RevisionStr.IsNumeric())
+	{
+		RevisionNum = FCString::Atoi(*RevisionStr);
+	}
+	return RevisionNum;
+}
+
+//------------------------------------------------------------------------------
+static UObject* EditorCommandLineUtilsImpl::ExtractAssetFromPackage(UPackage* Package, const FString& AssetName)
+{
+	return FindObject<UObject>(Package, *FPaths::GetBaseFilename(AssetName));
+}
+
+//------------------------------------------------------------------------------
+// static int32 EditorCommandLineUtilsImpl::GetAssetPackagePaths(const FString& InAssetName, TArray<FString>& PackagePathsOut, bool bIncludeDevFolder)
+// {
+// 	PackagePathsOut.Empty();
+// 	if (InAssetName.IsEmpty())
+// 	{
+// 		return 0;
+// 	}
+// 
+// 	FString AssetPath = InAssetName;
+// 	FPaths::NormalizeFilename(AssetPath);
+// 	FString AssetName = FPaths::GetCleanFilename(AssetPath);
+// 	const bool bIsPlainAssetName = (AssetName == AssetPath);
+// 
+// 	const FString& AssetExt = FPackageName::GetAssetPackageExtension();
+// 	const bool bHasExtension = AssetName.EndsWith(AssetExt);
+// 	FString BaseAssetName = FPaths::GetBaseFilename(AssetName);
+// 
+// 	FString AssetWildcard = FString::Printf(TEXT("*%s%s"), *AssetName, bHasExtension ? TEXT("") : *(TEXT("*") + AssetExt));
+// 	const uint8 PackageFilter = NORMALIZE_ExcludeMapPackages | (bIncludeDevFolder ? 0x00 : NORMALIZE_ExcludeDeveloperPackages);
+// 	TArray<FString> RelativeFilePaths;
+// 	NormalizePackageNames(TArray<FString>(), RelativeFilePaths, AssetName, PackageFilter);
+// 
+// 	bool bHasExactMatch = false;
+// 	for (FString& RelativePath : RelativeFilePaths)
+// 	{
+// 		FString FullFilePath = FPaths::ConvertRelativePathToFull(RelativePath);
+// 		FPaths::NormalizeFilename(FullFilePath);
+// 
+// 		bool bIsExactMatch = false;
+// 		if (BaseAssetName == FPaths::GetBaseFilename(FullFilePath))
+// 		{
+// 			bIsExactMatch = bIsPlainAssetName ? true : FullFilePath.Contains(AssetPath);
+// 		}
+// 
+// 		if (!bIsExactMatch && bHasExactMatch)
+// 		{
+// 			continue;
+// 		}
+// 
+// 		FString PackagePath;
+// 		if (FPackageName::TryConvertFilenameToLongPackageName(RelativePath, PackagePath))
+// 		{
+// 			if (bIsExactMatch && !bHasExactMatch)
+// 			{
+// 				PackagePathsOut.Empty();
+// 				bHasExactMatch = true;
+// 			}
+// 			PackagePathsOut.Add(PackagePath);
+// 		}
+// 	}
+// 
+// 	return PackagePathsOut.Num();
+// }
+
+/*******************************************************************************
+ * FEditorCommandLineUtils Definition
+ ******************************************************************************/
+
+//------------------------------------------------------------------------------
+bool FEditorCommandLineUtils::ParseGameProjectPath(const TCHAR* CmdLine, FString& ProjPathOut, FString& GameNameOut)
+{
+	using namespace EditorCommandLineUtilsImpl; // for ParseCommandArgs(), etc.
+
+	FString DiffArgs;
+	if (ParseCommandArgs(CmdLine, DiffCommandSwitch, DiffArgs))
+	{
+		TArray<FString> Tokens, Switches;
+		UCommandlet::ParseCommandLine(*DiffArgs, Tokens, Switches);
+
+		if (Tokens.Num() > 0)
+		{
+			FString DiffFilePath = Tokens[0];
+			FPaths::NormalizeFilename(DiffFilePath);
+
+			const TCHAR* const ContentDir = TEXT("/Content/");
+
+			FString ContentPath, FileSubPath;
+			if (DiffFilePath.Split(ContentDir, &ContentPath, &FileSubPath))
+			{
+				GameNameOut = FPaths::GetCleanFilename(ContentPath);
+				ProjPathOut = FPaths::Combine(*FPaths::RootDir(), *GameNameOut, *FString(GameNameOut + TEXT(".") + FProjectDescriptor::GetExtension()));
+				return FPaths::FileExists(ProjPathOut);
+			}
+		}
+	}
+	return false;
+}
+
+//------------------------------------------------------------------------------
+void FEditorCommandLineUtils::ProcessEditorCommands(const TCHAR* EditorCmdLine)
+{
+	using namespace EditorCommandLineUtilsImpl; // for DiffCommandSwitch, etc.
+
+	// If specified, Lightmass has to be launched manually with -debug (e.g. through a debugger).
+	// This creates a job with a hard-coded GUID, and allows Lightmass to be executed multiple times (even stand-alone).
+	if (FParse::Param(EditorCmdLine, DebugLightmassCommandSwitch))
+	{
+		extern bool GLightmassDebugMode;
+		GLightmassDebugMode = true;
+		UE_LOG(LogInit, Log, TEXT("Running Engine with Lightmass Debug Mode ENABLED"));
+	}
+
+	// If specified, all participating Lightmass agents will report back detailed stats to the log.
+	if (FParse::Param(EditorCmdLine, LightmassStatsCommandSwitch))
+	{
+		extern bool GLightmassStatsMode;
+		GLightmassStatsMode = true;
+		UE_LOG(LogInit, Log, TEXT("Running Engine with Lightmass Stats Mode ENABLED"));
+	}
+
+	FString DiffArgs;
+	if (ParseCommandArgs(EditorCmdLine, DiffCommandSwitch, DiffArgs))
+	{		
+		IMainFrameModule& MainFrameModule = IMainFrameModule::Get();
+		const bool bIsMainFramInitialized = MainFrameModule.IsWindowInitialized();
+
+		if (bIsMainFramInitialized)
+		{
+			RunAssetDiffCommand(MainFrameModule.GetParentWindow(), /*bIsNewProjectWindow =*/FApp::IsGameNameEmpty(), DiffArgs);
+		}
+		else
+		{
+			MainFrameModule.OnMainFrameCreationFinished().AddStatic(&RunAssetDiffCommand, DiffArgs);
+		}
+	}
+}
+
+#undef LOCTEXT_NAMESPACE
diff --git a/Engine/Source/Editor/UnrealEd/Public/EditorCommandLineUtils.h b/Engine/Source/Editor/UnrealEd/Public/EditorCommandLineUtils.h
new file mode 100644
--- /dev/null
+++ b/Engine/Source/Editor/UnrealEd/Public/EditorCommandLineUtils.h
@@ -0,0 +1,32 @@
+// Copyright 1998-2014 Epic Games, Inc. All Rights Reserved.
+
+#pragma once
+
+/** 
+ * 
+ */
+struct FEditorCommandLineUtils
+{
+	/**
+	 * Parses commands from the supplied command-line and invokes any tools/
+	 * settings that we're specified.
+	 *
+	 * @param  CommandLine	The string to parse for command-line switches/params
+	 */
+	static void ProcessEditorCommands(const TCHAR* CommandLine);
+
+	/**
+	 * Generally, users supply the project path/name as the second argument of 
+	 * the editor command-line, but in some cases (like for the -diff command), 
+	 * we can infer the project from elements elsewhere in the command-line...
+	 * This function serves as a fallback for determining the launching project
+	 * path, and does special processing based off of commands that are found in
+	 * the command-line.
+	 *
+	 * @param  EditorCmdLine	The string to parse a project path from.
+	 * @param  ProjPathOut		[out] Will be filled out with a full .uproject file path (if one was successfully parsed out) from the command-line.
+	 * @param  GameNameOut		[out] Will be filled out with a game name (if a project path was successfully parsed out) from the command-line.
+	 * @return True if the command-line contained enough info to pull a .uproject file from.
+	 */
+	UNREALED_API static bool ParseGameProjectPath(const TCHAR* EditorCmdLine, FString& ProjPathOut, FString& GameNameOut);
+};
\ No newline at end of file
diff --git a/Engine/Source/Runtime/Launch/Private/LaunchEngineLoop.cpp b/Engine/Source/Runtime/Launch/Private/LaunchEngineLoop.cpp
--- a/Engine/Source/Runtime/Launch/Private/LaunchEngineLoop.cpp
+++ b/Engine/Source/Runtime/Launch/Private/LaunchEngineLoop.cpp
@@ -20,6 +20,7 @@
 	#include "AutomationController.h"
 	#include "ProfilerClient.h"
 	#include "RemoteConfigIni.h"
+	#include "EditorCommandLineUtils.h"
 
 	#if PLATFORM_WINDOWS
 		#include "AllowWindowsPlatformTypes.h"
@@ -215,6 +216,13 @@ bool ParseGameProjectFromCommandLine(const TCHAR* InCmdLine, FString& OutProject
 			return true;
 		}
 	}
+
+#if WITH_EDITOR
+	if (FEditorCommandLineUtils::ParseGameProjectPath(InCmdLine, OutProjectFilePath, OutGameName))
+	{
+		return true;
+	}
+#endif
 	return false;
 }
 
