PR-601: Add UNREALED_API to UFbxExporter.

diff --git a/Engine/Source/Editor/UnrealEd/Private/FbxExporter.h b/Engine/Source/Editor/UnrealEd/Private/FbxExporter.h
--- a/Engine/Source/Editor/UnrealEd/Private/FbxExporter.h
+++ b/Engine/Source/Editor/UnrealEd/Private/FbxExporter.h
@@ -47,13 +47,13 @@ namespace UnFbx
 /**
  * Main FBX Exporter class.
  */
-class FFbxExporter  : public MatineeExporter
+class UNREALED_API FFbxExporter  : public MatineeExporter
 {
 public:
 	/**
 	 * Returns the exporter singleton. It will be created on the first request.
 	 */
-	UNREALED_API static FFbxExporter* GetInstance();
+	static FFbxExporter* GetInstance();
 	static void DeleteInstance();
 	~FFbxExporter();
 	
@@ -157,7 +157,7 @@ public:
 	/**
 	 * Exports the list of UAnimSequences as a single animation based on the settings in the TrackKeys
 	 */
-	UNREALED_API void ExportAnimSequencesAsSingle( USkeletalMesh* SkelMesh, const ASkeletalMeshActor* SkelMeshActor, const FString& ExportName, const TArray<UAnimSequence*>& AnimSeqList, const TArray<FAnimControlTrackKey>& TrackKeys );
+	void ExportAnimSequencesAsSingle( USkeletalMesh* SkelMesh, const ASkeletalMeshActor* SkelMeshActor, const FString& ExportName, const TArray<UAnimSequence*>& AnimSeqList, const TArray<FAnimControlTrackKey>& TrackKeys );
 
 private:
 	FFbxExporter();
