Add analytics for importing FBX files

- vendor
- app name
- app version

Also added UE4 information when export FBX

diff --git a/Engine/Source/Editor/UnrealEd/Private/Fbx/FbxMainExport.cpp b/Engine/Source/Editor/UnrealEd/Private/Fbx/FbxMainExport.cpp
--- a/Engine/Source/Editor/UnrealEd/Private/Fbx/FbxMainExport.cpp
+++ b/Engine/Source/Editor/UnrealEd/Private/Fbx/FbxMainExport.cpp
@@ -122,7 +122,12 @@ void FFbxExporter::CreateDocument()
 	FbxDocumentInfo* SceneInfo = FbxDocumentInfo::Create(SdkManager,"SceneInfo");
 	SceneInfo->mTitle = "Unreal FBX Exporter";
 	SceneInfo->mSubject = "Export FBX meshes from Unreal";
-
+	SceneInfo->Original_ApplicationVendor.Set( "Epic Games" );
+	SceneInfo->Original_ApplicationName.Set( "Unreal Engine" );
+	SceneInfo->Original_ApplicationVersion.Set( TCHAR_TO_ANSI(*GEngineVersion.ToString()) );
+	SceneInfo->LastSaved_ApplicationVendor.Set( "Epic Games" );
+	SceneInfo->LastSaved_ApplicationName.Set( "Unreal Engine" );
+	SceneInfo->LastSaved_ApplicationVersion.Set( TCHAR_TO_ANSI(*GEngineVersion.ToString()) );
 
 	Scene->SetSceneInfo(SceneInfo);
 	
diff --git a/Engine/Source/Editor/UnrealEd/Private/Fbx/FbxMainImport.cpp b/Engine/Source/Editor/UnrealEd/Private/Fbx/FbxMainImport.cpp
--- a/Engine/Source/Editor/UnrealEd/Private/Fbx/FbxMainImport.cpp
+++ b/Engine/Source/Editor/UnrealEd/Private/Fbx/FbxMainImport.cpp
@@ -43,6 +43,8 @@
 #include "FbxOptionWindow.h"
 #include "FbxErrors.h"
 #include "MainFrame.h"
+#include "EngineAnalytics.h"
+#include "Runtime/Analytics/Analytics/Public/Interfaces/IAnalyticsProvider.h"
 
 DEFINE_LOG_CATEGORY(LogFbx);
 
@@ -762,6 +764,32 @@ bool FFbxImporter::ImportFromFile(const FString& Filename, const FString& Type)
 						UnrealZUp.ConvertScene(Scene);
 					}
 				}
+
+				// do analytics on getting Fbx data
+				FbxDocumentInfo* DocInfo = Scene->GetSceneInfo();
+				if (DocInfo)
+				{
+					TArray<FAnalyticsEventAttribute> Attribs;
+
+					FString OriginalVendor(ANSI_TO_TCHAR(DocInfo->Original_ApplicationVendor.Get().Buffer()));
+					FString OriginalAppName(ANSI_TO_TCHAR(DocInfo->Original_ApplicationName.Get().Buffer()));
+					FString OriginalAppVersion(ANSI_TO_TCHAR(DocInfo->Original_ApplicationVersion.Get().Buffer()));
+
+					FString LastSavedVendor(ANSI_TO_TCHAR(DocInfo->LastSaved_ApplicationVendor.Get().Buffer()));
+					FString LastSavedAppName(ANSI_TO_TCHAR(DocInfo->LastSaved_ApplicationName.Get().Buffer()));
+					FString LastSavedAppVersion(ANSI_TO_TCHAR(DocInfo->LastSaved_ApplicationVersion.Get().Buffer()));
+
+					Attribs.Add(FAnalyticsEventAttribute(TEXT("Original Application Vendor"), OriginalVendor));
+					Attribs.Add(FAnalyticsEventAttribute(TEXT("Original Application Name"), OriginalAppName));
+					Attribs.Add(FAnalyticsEventAttribute(TEXT("Original Application Version"), OriginalAppVersion));
+
+					Attribs.Add(FAnalyticsEventAttribute(TEXT("LastSaved Application Vendor"), LastSavedVendor));
+					Attribs.Add(FAnalyticsEventAttribute(TEXT("LastSaved Application Name"), LastSavedAppName));
+					Attribs.Add(FAnalyticsEventAttribute(TEXT("LastSaved Application Version"), LastSavedAppVersion));
+
+					FString EventString = FString::Printf(TEXT("Editor.Usage.FBX.Import"));
+					FEngineAnalytics::GetProvider().RecordEvent(EventString, Attribs);
+				}
 			}
 
 			// Convert the scene's units to what is used in this program, if needed.
