Missed header include (for non-unity builds).

diff --git a/Engine/Source/Editor/UnrealEd/Private/EditorCommandLineUtils.cpp b/Engine/Source/Editor/UnrealEd/Private/EditorCommandLineUtils.cpp
--- a/Engine/Source/Editor/UnrealEd/Private/EditorCommandLineUtils.cpp
+++ b/Engine/Source/Editor/UnrealEd/Private/EditorCommandLineUtils.cpp
@@ -9,6 +9,7 @@
 #include "TickableEditorObject.h"
 #include "Editor/MainFrame/Public/Interfaces/IMainFrameModule.h"
 #include "AssetRegistryModule.h"
+#include "IAssetTypeActions.h"		// for FRevisionInfo
 
 #define LOCTEXT_NAMESPACE "EditorCommandLineUtils"
 
