Fixed ShaderCompilerManager from not being updated in cook commandlet.

diff --git a/Engine/Source/Editor/UnrealEd/Private/Commandlets/CookCommandlet.cpp b/Engine/Source/Editor/UnrealEd/Private/Commandlets/CookCommandlet.cpp
--- a/Engine/Source/Editor/UnrealEd/Private/Commandlets/CookCommandlet.cpp
+++ b/Engine/Source/Editor/UnrealEd/Private/Commandlets/CookCommandlet.cpp
@@ -23,6 +23,7 @@
 #include "ChunkManifestGenerator.h"
 #include "PhysicsPublic.h"
 #include "CookerSettings.h"
+#include "ShaderCompiler.h"
 
 DEFINE_LOG_CATEGORY_STATIC(LogCookCommandlet, Log, All);
 
@@ -1162,6 +1163,8 @@ bool UCookCommandlet::NewCook( const TArray<ITargetPlatform*>& Platforms, TArray
 		}
 
 
+		GShaderCompilingManager->ProcessAsyncResults(true, false);
+	
 		if (NonMapPackageCountSinceLastGC > 0)
 		{
 			// We should GC if we have packages to collect and we've been idle for some time.
