Enable more warnings and fix errors/problems caught by them.

- PaperSprite

- EditorBuildPromotionTests

- AttributeDetails

- RepairLocData

- K2Node

diff --git a/Engine/Plugins/2D/Paper2D/Source/Paper2D/Private/PaperSprite.cpp b/Engine/Plugins/2D/Paper2D/Source/Paper2D/Private/PaperSprite.cpp
--- a/Engine/Plugins/2D/Paper2D/Source/Paper2D/Private/PaperSprite.cpp
+++ b/Engine/Plugins/2D/Paper2D/Source/Paper2D/Private/PaperSprite.cpp
@@ -297,7 +297,7 @@ void UPaperSprite::PostEditChangeProperty(FPropertyChangedEvent& PropertyChanged
 	if (CollisionGeometry.GeometryType == ESpritePolygonMode::Diced)
 	{
 		// Disallow dicing on collision geometry for now
-		CollisionGeometry.GeometryType == ESpritePolygonMode::SourceBoundingBox;
+		CollisionGeometry.GeometryType = ESpritePolygonMode::SourceBoundingBox;
 	}
 	RenderGeometry.PixelsPerSubdivisionX = FMath::Max(RenderGeometry.PixelsPerSubdivisionX, 4);
 	RenderGeometry.PixelsPerSubdivisionY = FMath::Max(RenderGeometry.PixelsPerSubdivisionY, 4);
diff --git a/Engine/Source/Editor/BlueprintGraph/Classes/K2Node.h b/Engine/Source/Editor/BlueprintGraph/Classes/K2Node.h
--- a/Engine/Source/Editor/BlueprintGraph/Classes/K2Node.h
+++ b/Engine/Source/Editor/BlueprintGraph/Classes/K2Node.h
@@ -334,7 +334,7 @@ class UK2Node : public UEdGraphNode
 	 */
 	virtual FBlueprintNodeSignature GetSignature() const { return FBlueprintNodeSignature(GetClass()); }
 
-	BLUEPRINTGRAPH_API enum EBaseNodeRefreshPriority
+	enum BLUEPRINTGRAPH_API EBaseNodeRefreshPriority
 	{
 		Low_UsesDependentWildcard = 100,
 		Normal = 200,
diff --git a/Engine/Source/Editor/GameplayAbilitiesEditor/Private/AttributeDetails.cpp b/Engine/Source/Editor/GameplayAbilitiesEditor/Private/AttributeDetails.cpp
--- a/Engine/Source/Editor/GameplayAbilitiesEditor/Private/AttributeDetails.cpp
+++ b/Engine/Source/Editor/GameplayAbilitiesEditor/Private/AttributeDetails.cpp
@@ -346,11 +346,13 @@ void FScalableFloatDetails::CustomizeHeader( TSharedRef<class IPropertyHandle> S
 	uint32 NumChildren = 0;
 	StructPropertyHandle->GetNumChildren(NumChildren);
 
+	/* Debug?
 	for (uint32 i = 0; i < NumChildren; i++)
 	{
 		FString PropName = StructPropertyHandle->GetChildHandle(i)->GetPropertyDisplayName();
 		PropName;
 	}
+	*/
 
 	ValueProperty = StructPropertyHandle->GetChildHandle(GET_MEMBER_NAME_CHECKED(FScalableFloat,Value));
 	CurveTableHandleProperty = StructPropertyHandle->GetChildHandle(GET_MEMBER_NAME_CHECKED(FScalableFloat,Curve));
diff --git a/Engine/Source/Editor/UnrealEd/Private/Commandlets/RepairLocalizationDataCommandlet.cpp b/Engine/Source/Editor/UnrealEd/Private/Commandlets/RepairLocalizationDataCommandlet.cpp
--- a/Engine/Source/Editor/UnrealEd/Private/Commandlets/RepairLocalizationDataCommandlet.cpp
+++ b/Engine/Source/Editor/UnrealEd/Private/Commandlets/RepairLocalizationDataCommandlet.cpp
@@ -16,7 +16,7 @@ void RepairManifestAndArchives(TSharedRef<FInternationalizationManifest> Manifes
 	if (Manifest->GetFormatVersion() < FInternationalizationManifest::EFormatVersion::EscapeFixes)
 	{
 		TManifestEntryBySourceTextContainer::TConstIterator Iterator = Manifest->GetEntriesBySourceTextIterator();
-		for(Iterator; Iterator; ++Iterator)
+		for(; Iterator; ++Iterator)
 		{
 			const FString& Source = Iterator.Key();
 			const TSharedRef<FManifestEntry>& ManifestEntry = Iterator.Value();
@@ -39,7 +39,7 @@ void RepairManifestAndArchives(TSharedRef<FInternationalizationManifest> Manifes
 		{
 			TArchiveEntryContainer::TConstIterator Iterator = Archive->GetEntryIterator();
 
-			for (Iterator; Iterator; ++Iterator)
+			for (; Iterator; ++Iterator)
 			{
 				const FString& Source = Iterator.Key();
 				const TSharedRef<FArchiveEntry>& ArchiveEntry = Iterator.Value();
@@ -62,7 +62,7 @@ void RepairManifestAndArchives(TSharedRef<FInternationalizationManifest> Manifes
 	const FRegexPattern Pattern(TEXT(".* - line \\d+"));
 
 	TManifestEntryBySourceTextContainer::TConstIterator Iterator = Manifest->GetEntriesBySourceTextIterator();
-	for(Iterator; Iterator; ++Iterator)
+	for(; Iterator; ++Iterator)
 	{
 		const FString& Source = Iterator.Key();
 		const TSharedRef<FManifestEntry>& ManifestEntry = Iterator.Value();
diff --git a/Engine/Source/Editor/UnrealEd/Private/Tests/EditorBuildPromotionTests.cpp b/Engine/Source/Editor/UnrealEd/Private/Tests/EditorBuildPromotionTests.cpp
--- a/Engine/Source/Editor/UnrealEd/Private/Tests/EditorBuildPromotionTests.cpp
+++ b/Engine/Source/Editor/UnrealEd/Private/Tests/EditorBuildPromotionTests.cpp
@@ -2389,18 +2389,18 @@ namespace BuildPromotionTestHelper
 			IMPORT_ASSET_WITH_FACTORY(UFbxFactory, UStaticMesh, ReimportStaticMesh, ReimportMesh, );
 
 			//Blend Shape Mesh
-			IMPORT_ASSET_WITH_FACTORY(UFbxFactory, USkeletalMesh, BlendShapeMesh, BlendShape, FactoryInst->ImportUI->MeshTypeToImport == FBXIT_SkeletalMesh;);
+			IMPORT_ASSET_WITH_FACTORY(UFbxFactory, USkeletalMesh, BlendShapeMesh, BlendShape, FactoryInst->ImportUI->MeshTypeToImport = FBXIT_SkeletalMesh;);
 
 			//Morph Mesh
-			IMPORT_ASSET_WITH_FACTORY(UFbxFactory, USkeletalMesh, MorphMesh, MorphAndMorphAnim, FactoryInst->ImportUI->MeshTypeToImport == FBXIT_SkeletalMesh;);
+			IMPORT_ASSET_WITH_FACTORY(UFbxFactory, USkeletalMesh, MorphMesh, MorphAndMorphAnim, FactoryInst->ImportUI->MeshTypeToImport = FBXIT_SkeletalMesh;);
 
 			//Skeletal Mesh
-			IMPORT_ASSET_WITH_FACTORY(UFbxFactory, USkeletalMesh, SkeletalMesh, SkeletalMesh_Test, FactoryInst->ImportUI->MeshTypeToImport == FBXIT_SkeletalMesh;);
+			IMPORT_ASSET_WITH_FACTORY(UFbxFactory, USkeletalMesh, SkeletalMesh, SkeletalMesh_Test, FactoryInst->ImportUI->MeshTypeToImport = FBXIT_SkeletalMesh;);
 
 			if (SkeletalMesh_Test)
 			{
 				//Animation
-				IMPORT_ASSET_WITH_FACTORY(UFbxFactory, UAnimSequence, Animation, AnimationTest, FactoryInst->ImportUI->MeshTypeToImport == FBXIT_Animation; FactoryInst->ImportUI->Skeleton = SkeletalMesh_Test->Skeleton;);
+				IMPORT_ASSET_WITH_FACTORY(UFbxFactory, UAnimSequence, Animation, AnimationTest, FactoryInst->ImportUI->MeshTypeToImport = FBXIT_Animation; FactoryInst->ImportUI->Skeleton = SkeletalMesh_Test->Skeleton;);
 			}
 			else
 			{
diff --git a/Engine/Source/Programs/UnrealBuildTool/Linux/LinuxToolChain.cs b/Engine/Source/Programs/UnrealBuildTool/Linux/LinuxToolChain.cs
--- a/Engine/Source/Programs/UnrealBuildTool/Linux/LinuxToolChain.cs
+++ b/Engine/Source/Programs/UnrealBuildTool/Linux/LinuxToolChain.cs
@@ -294,14 +294,6 @@ namespace UnrealBuildTool
 			Result += " -fno-math-errno";               // do not assume that math ops have side effects
             Result += " -fno-rtti";                     // no run-time type info
 
-            // these needs to be removed ASAP
-            Result += " -Wno-delete-non-virtual-dtor";
-            Result += " -Wno-reorder";
-            Result += " -Wno-logical-op-parentheses";
-            Result += " -Wno-ignored-attributes";
-            Result += " -Wno-overloaded-virtual";
-            Result += " -Wno-unused-value";
-
             if (String.IsNullOrEmpty(ClangPath))
             {
                 // GCC only option
diff --git a/Engine/Source/Runtime/Core/Private/Linux/LinuxPlatformFeedbackContextPrivate.h b/Engine/Source/Runtime/Core/Private/Linux/LinuxPlatformFeedbackContextPrivate.h
--- a/Engine/Source/Runtime/Core/Private/Linux/LinuxPlatformFeedbackContextPrivate.h
+++ b/Engine/Source/Runtime/Core/Private/Linux/LinuxPlatformFeedbackContextPrivate.h
@@ -73,10 +73,9 @@ public:
 			GLog->Serialize( V, Verbosity, Category );
 		}
 	}
-	VARARG_BODY( bool, YesNof, const TCHAR*, VARARG_NONE )
+	/** Ask the user a binary question, returning their answer */
+	virtual bool YesNof( const FText& Question )
 	{
-		TCHAR TempStr[4096];
-		GET_VARARGS( TempStr, ARRAY_COUNT(TempStr), ARRAY_COUNT(TempStr)-1, Fmt, Fmt );
 		if( ( GIsClient || GIsEditor ) && ( ( GIsSilent != true ) && ( FApp::IsUnattended() != true ) ) )
 		{
 			//return( ::MessageBox( NULL, TempStr, *NSLOCTEXT("Core", "Question", "Question").ToString(), MB_YESNO|MB_TASKMODAL ) == IDYES);
