Fixing non-unity build

diff --git a/Engine/Source/Editor/UnrealEd/Private/EditorCommandLineUtils.cpp b/Engine/Source/Editor/UnrealEd/Private/EditorCommandLineUtils.cpp
--- a/Engine/Source/Editor/UnrealEd/Private/EditorCommandLineUtils.cpp
+++ b/Engine/Source/Editor/UnrealEd/Private/EditorCommandLineUtils.cpp
@@ -10,6 +10,7 @@
 #include "Editor/MainFrame/Public/Interfaces/IMainFrameModule.h"
 #include "AssetRegistryModule.h"
 #include "IAssetTypeActions.h"		// for FRevisionInfo
+#include "AssetToolsModule.h"
 
 #define LOCTEXT_NAMESPACE "EditorCommandLineUtils"
 
