UE-5457 When using a UserDefinedEnum in a BlueprintMacroLibrary, Enum is not loaded

diff --git a/Engine/Source/Editor/UnrealEd/Private/Kismet2/BlueprintEditorUtils.cpp b/Engine/Source/Editor/UnrealEd/Private/Kismet2/BlueprintEditorUtils.cpp
--- a/Engine/Source/Editor/UnrealEd/Private/Kismet2/BlueprintEditorUtils.cpp
+++ b/Engine/Source/Editor/UnrealEd/Private/Kismet2/BlueprintEditorUtils.cpp
@@ -826,6 +826,16 @@ struct FRegenerationHelper
 								ProcessHierarchy(ParentFunction, Dependencies);
 							}
 						}
+
+						// load Enums
+						for (auto Pin : Node->Pins)
+						{
+							auto SubCategoryObject = Pin ? Pin->PinType.PinSubCategoryObject.Get() : NULL;
+							if (SubCategoryObject && SubCategoryObject->IsA<UEnum>())
+							{
+								ForcedLoad(SubCategoryObject);
+							}
+						}
 					}
 				}
 			}
