Fixed slow task dialogs not being ticked correctly for the very first time they are opened

diff --git a/Engine/Source/Editor/UnrealEd/Private/FeedbackContextEditor.cpp b/Engine/Source/Editor/UnrealEd/Private/FeedbackContextEditor.cpp
--- a/Engine/Source/Editor/UnrealEd/Private/FeedbackContextEditor.cpp
+++ b/Engine/Source/Editor/UnrealEd/Private/FeedbackContextEditor.cpp
@@ -36,6 +36,9 @@ public:
 	/** Construct this widget */
 	void Construct( const FArguments& InArgs )
 	{
+		// Ensure this gets ticked at least once
+		LastTickTime = 0;
+
 		OnCancelClickedDelegate = InArgs._OnCancelClickedDelegate;
 		ScopeStack = InArgs._ScopeStack;
 
@@ -123,9 +126,20 @@ public:
 				SNew(SBox).WidthOverride(FixedWidth) [ VerticalBox ]
 			]
 		);
+
+		// Make sure all our bars are set up
+		UpdateDynamicProgressBars();
 	}
 
 	virtual void Tick(const FGeometry& AllottedGeometry, const double InCurrentTime, const float InDeltaTime) override
+	{
+		UpdateDynamicProgressBars();
+	}
+
+private:
+
+	/** Updates the dynamic progress bars for this widget */
+	void UpdateDynamicProgressBars()
 	{
 		static const double VisibleScopeThreshold = 0.5;
 
@@ -164,8 +178,6 @@ public:
 		}
 	}
 
-private:
-
 	/** Create a progress bar for the specified index */
 	void CreateSecondaryBar(int32 Index) 
 	{
@@ -288,6 +300,10 @@ private:
 
 	/** Array mapping progress bar index -> scope stack index. Updated every tick. */
 	TArray<int32> DynamicProgressIndices;
+
+public:
+	/** Publicly accessible tick time to throttle slate ticking */
+	static double LastTickTime;
 };
 
 /** Static integer definitions required on some builds where the linker needs access to these */
@@ -296,14 +312,14 @@ const int32 SSlowTaskWidget::FixedWidth;
 const int32 SSlowTaskWidget::FixedPaddingH;;
 const int32 SSlowTaskWidget::MainBarHeight;
 const int32 SSlowTaskWidget::SecondaryBarHeight;
+double SSlowTaskWidget::LastTickTime = 0;
 
 static void TickSlate()
 {
-	static double Seconds = FPlatformTime::Seconds();
 	static double MinFrameTime = 0.05;		// Only update at 20fps so as not to slow down the actual task
-	if (FPlatformTime::Seconds() - Seconds > MinFrameTime)
+	if (FPlatformTime::Seconds() - SSlowTaskWidget::LastTickTime > MinFrameTime)
 	{
-		Seconds = FPlatformTime::Seconds();
+		SSlowTaskWidget::LastTickTime = FPlatformTime::Seconds();
 
 		// Tick Slate application
 		FSlateApplication::Get().Tick();
