BP: Reinstancer does additional initialization for AnimInstance

diff --git a/Engine/Source/Editor/UnrealEd/Private/Kismet2/KismetReinstanceUtilities.cpp b/Engine/Source/Editor/UnrealEd/Private/Kismet2/KismetReinstanceUtilities.cpp
--- a/Engine/Source/Editor/UnrealEd/Private/Kismet2/KismetReinstanceUtilities.cpp
+++ b/Engine/Source/Editor/UnrealEd/Private/Kismet2/KismetReinstanceUtilities.cpp
@@ -181,10 +181,10 @@ void FBlueprintCompileReinstancer::ReinstanceObjects(bool bAlwaysReinstance)
 			auto BPClassA = Cast<const UBlueprintGeneratedClass>(DuplicatedClass);
 			auto BPClassB = Cast<const UBlueprintGeneratedClass>(ClassToReinstance);
 			auto BP = Cast<const UBlueprint>(ClassToReinstance->ClassGeneratedBy);
-			const bool bBPSubClass = (BP && BP->GetClass() != UBlueprint::StaticClass());
+
 			static const FBoolConfigValueHelper ChangeDefaultValueWithoutReinstancing(TEXT("Kismet"), TEXT("bChangeDefaultValueWithoutReinstancing"), GEngineIni);
 			const bool bTheSameDefaultValues = BP && ClassToReinstanceDefaultValuesCRC && (BP->CrcPreviousCompiledCDO == ClassToReinstanceDefaultValuesCRC);
-			const bool bTheSame = !bBPSubClass && (ChangeDefaultValueWithoutReinstancing || bTheSameDefaultValues) && BPClassA && BPClassB && FStructUtils::TheSameLayout(BPClassA, BPClassB, true);
+			const bool bTheSame = (ChangeDefaultValueWithoutReinstancing || bTheSameDefaultValues) && BPClassA && BPClassB && FStructUtils::TheSameLayout(BPClassA, BPClassB, true);
 			if (bTheSame)
 			{
 				UE_LOG(LogBlueprint, Log, TEXT("BlueprintCompileReinstancer: class '%s' is replaced without reinstancing (the optimized way)."), *GetPathNameSafe(ClassToReinstance));
@@ -193,6 +193,7 @@ void FBlueprintCompileReinstancer::ReinstanceObjects(bool bAlwaysReinstance)
 				GetObjectsOfClass(DuplicatedClass, ObjectsToReplace, false);
 
 				const bool bIsActor = ClassToReinstance->IsChildOf<AActor>();
+				const bool bIsAnimInstance = ClassToReinstance->IsChildOf<UAnimInstance>();
 				for (auto Obj : ObjectsToReplace)
 				{
 					if (!Obj->IsTemplate() && !Obj->IsPendingKill())
@@ -204,6 +205,16 @@ void FBlueprintCompileReinstancer::ReinstanceObjects(bool bAlwaysReinstance)
 							Actor->ReregisterAllComponents();
 							Actor->RerunConstructionScripts();
 						}
+
+						if (bIsAnimInstance)
+						{
+							// Initialising the anim instance isn't enough to correctly set up the skeletal mesh again in a
+							// paused world, need to initialise the skeletal mesh component that contains the anim instance.
+							if (USkeletalMeshComponent* SkelComponent = Cast<USkeletalMeshComponent>(Obj->GetOuter()))
+							{
+								SkelComponent->InitAnim(true);
+							}
+						}
 					}
 				}
 
