#ue4 Blocking save on compile during COL

diff --git a/Engine/Source/Editor/UnrealEd/Private/Kismet2/Kismet2.cpp b/Engine/Source/Editor/UnrealEd/Private/Kismet2/Kismet2.cpp
--- a/Engine/Source/Editor/UnrealEd/Private/Kismet2/Kismet2.cpp
+++ b/Engine/Source/Editor/UnrealEd/Private/Kismet2/Kismet2.cpp
@@ -492,7 +492,7 @@ void FKismetEditorUtilities::CompileBlueprint(UBlueprint* BlueprintObj, bool bIs
 		BlueprintPackage->SetDirtyFlag(bStartedWithUnsavedChanges);
 
 		UBlueprintEditorSettings* Settings = GetMutableDefault<UBlueprintEditorSettings>();
-		const bool bShouldSaveOnCompile = ( Settings->SaveOnCompile == SoC_Always ) || ( ( Settings->SaveOnCompile == SoC_SuccessOnly ) && ( BlueprintObj->Status == BS_UpToDate ) );
+		const bool bShouldSaveOnCompile = !bIsRegeneratingOnLoad && ( Settings->SaveOnCompile == SoC_Always ) || ( ( Settings->SaveOnCompile == SoC_SuccessOnly ) && ( BlueprintObj->Status == BS_UpToDate ) );
 
 		// Only try to save on compile if we aren't running a commandlet (i.e. cooking a blueprint shouldn't try to save it)
 		if ( !IsRunningCommandlet() && bShouldSaveOnCompile )
