Clean up physics scenes if they are created by the UWorld during cooking.

diff --git a/Engine/Source/Editor/UnrealEd/Private/CookOnTheFlyServer.cpp b/Engine/Source/Editor/UnrealEd/Private/CookOnTheFlyServer.cpp
--- a/Engine/Source/Editor/UnrealEd/Private/CookOnTheFlyServer.cpp
+++ b/Engine/Source/Editor/UnrealEd/Private/CookOnTheFlyServer.cpp
@@ -263,7 +263,7 @@ public:
 	}
 
 	FCachedPackageFilename( FCachedPackageFilename &&In )
-	{
+{
 		PackageFilename = MoveTemp(In.PackageFilename);
 		StandardFilename = MoveTemp(In.StandardFilename);
 		StandardFileFName = In.StandardFileFName;
@@ -289,9 +289,9 @@ static const FCachedPackageFilename &Cache(const FName& PackageName)
 	FString StandardFilename;
 	FName StandardFileFName = NAME_None;
 	if (FPackageName::DoesPackageExist(PackageName.ToString(), NULL, &Filename))
-	{
+{
 		StandardFilename = PackageFilename = FPaths::ConvertRelativePathToFull(Filename);
-		
+
 
 		FPaths::MakeStandardFilename(StandardFilename);
 		StandardFileFName = FName(*StandardFilename);
@@ -380,7 +380,7 @@ struct FRecompileRequest
  * @param Severity of the message
  */
 void LogCookerMessage( const FString& MessageText, EMessageSeverity::Type Severity)
-{
+	{
 	FMessageLog MessageLog("CookResults");
 
 	TSharedRef<FTokenizedMessage> Message = FTokenizedMessage::Create(Severity);
@@ -1655,6 +1655,7 @@ bool UCookOnTheFlyServer::SaveCookedPackage( UPackage* Package, uint32 SaveFlags
 					Package->PackageFlags &= ~PKG_FilterEditorOnly;
 				}
 
+				bool bDidInitializeWorld = false;
 				if (World)
 				{
 					World->PersistentLevel->OwningWorld = World;
@@ -1662,6 +1663,7 @@ bool UCookOnTheFlyServer::SaveCookedPackage( UPackage* Package, uint32 SaveFlags
 					{
 						// we need to initialize the world - at least need physics scene since BP construction script runs during cooking, otherwise trace won't work
 						World->InitWorld(UWorld::InitializationValues().RequiresHitProxies(false).ShouldSimulatePhysics(false).EnableTraceCollision(false).CreateNavigation(false).AllowAudioPlayback(false).CreatePhysicsScene(true));
+						bDidInitializeWorld = true;
 					}
 				}
 
@@ -1678,6 +1680,17 @@ bool UCookOnTheFlyServer::SaveCookedPackage( UPackage* Package, uint32 SaveFlags
 					bSavedCorrectly &= GEditor->SavePackage( Package, World, Flags, *PlatFilename, GError, NULL, bSwap, false, SaveFlags, Target, FDateTime::MinValue(), false );
 				}
 
+				// if we initialized the world we are responsible for cleaning it up.
+				if (World && World->bIsWorldInitialized && bDidInitializeWorld)
+				{
+					// Make sure we clean up the physics scene here. If we leave too many scenes in memory, undefined behavior occurs when locking a scene for read/write.
+					World->SetPhysicsScene(nullptr);
+					if ( GPhysCommandHandler )
+					{
+						GPhysCommandHandler->Flush();
+					}
+				}
+
 				
 				bOutWasUpToDate = false;
 			}
