[UE-5487/UE-4078] Making sure the package dirty flag isn't set when we unload a blueprint for automation testing.

diff --git a/Engine/Source/Editor/UnrealEd/Private/Kismet2/Kismet2.cpp b/Engine/Source/Editor/UnrealEd/Private/Kismet2/Kismet2.cpp
--- a/Engine/Source/Editor/UnrealEd/Private/Kismet2/Kismet2.cpp
+++ b/Engine/Source/Editor/UnrealEd/Private/Kismet2/Kismet2.cpp
@@ -495,14 +495,14 @@ void FKismetEditorUtilities::CompileBlueprint(UBlueprint* BlueprintObj, bool bIs
 		const bool bShouldSaveOnCompile = !bIsRegeneratingOnLoad && ( Settings->SaveOnCompile == SoC_Always ) || ( ( Settings->SaveOnCompile == SoC_SuccessOnly ) && ( BlueprintObj->Status == BS_UpToDate ) );
 
 		// Only try to save on compile if we aren't running a commandlet (i.e. cooking a blueprint shouldn't try to save it)
-		if ( !IsRunningCommandlet() && bShouldSaveOnCompile )
+		if ( !IsRunningCommandlet() && bShouldSaveOnCompile && !GIsAutomationTesting )
 		{
 			bool const bIsLevelPackage = (UWorld::FindWorldInPackage(BlueprintPackage) != nullptr);
 			// we don't want to save the entire level (especially if this 
 			// compile was already kicked off as a result of a level save, as it
 			// could cause a recursive save)... let the "SaveOnCompile" setting 
 			// only save blueprint assets
-			if (!bIsLevelPackage && !IsRunningCommandlet())
+			if (!bIsLevelPackage)
 			{
 				TArray<UPackage*> PackagesToSave;
 				PackagesToSave.Add(BlueprintPackage);
diff --git a/Engine/Source/Editor/UnrealEd/Private/Tests/BlueprintAutomationTests.cpp b/Engine/Source/Editor/UnrealEd/Private/Tests/BlueprintAutomationTests.cpp
--- a/Engine/Source/Editor/UnrealEd/Private/Tests/BlueprintAutomationTests.cpp
+++ b/Engine/Source/Editor/UnrealEd/Private/Tests/BlueprintAutomationTests.cpp
@@ -337,6 +337,11 @@ public:
 			FName UnloadedName = MakeUniqueObjectName(NewPackage, UBlueprint::StaticClass(), BlueprintObj->GetFName());
 			BlueprintObj->Rename(*UnloadedName.ToString(), NewPackage, REN_DontCreateRedirectors|REN_DoNotDirty);
 
+			// Rename() will mark the OldPackage dirty (since it is removing the
+			// blueprint from it), we don't want this to affect the dirty flag 
+			// (for if/when we load it again)
+			OldPackage->SetDirtyFlag(/*bIsDirty =*/false);
+
 			// make sure the blueprint is properly trashed so we can rerun tests on it
 			BlueprintObj->SetFlags(RF_Transient);
 			BlueprintObj->ClearFlags(RF_Standalone | RF_RootSet | RF_Transactional);
