fix non unity build

diff --git a/Engine/Source/Editor/UnrealEd/Private/CookOnTheFlyServer.cpp b/Engine/Source/Editor/UnrealEd/Private/CookOnTheFlyServer.cpp
--- a/Engine/Source/Editor/UnrealEd/Private/CookOnTheFlyServer.cpp
+++ b/Engine/Source/Editor/UnrealEd/Private/CookOnTheFlyServer.cpp
@@ -20,6 +20,7 @@
 #include "AssetRegistryModule.h"
 #include "UnrealEdMessages.h"
 #include "GameDelegates.h"
+#include "PhysicsPublic.h"
 
 // cook by the book requirements
 #include "Commandlets/ChunkManifestGenerator.h"
