Removed unnecessary transactions when dropping actors into the world.

diff --git a/Engine/Source/Editor/UnrealEd/Private/LevelEditorViewport.cpp b/Engine/Source/Editor/UnrealEd/Private/LevelEditorViewport.cpp
--- a/Engine/Source/Editor/UnrealEd/Private/LevelEditorViewport.cpp
+++ b/Engine/Source/Editor/UnrealEd/Private/LevelEditorViewport.cpp
@@ -784,10 +784,35 @@ bool FLevelEditorViewportClient::AttemptApplyObjAsMaterialToSurface( UObject* Ob
 }
 
 
+static bool AreAllDroppedObjectsBrushBuilders(const TArray<UObject*>& DroppedObjects)
+{
+	for (UObject* DroppedObject : DroppedObjects)
+	{
+		if (!DroppedObject->IsA(UBrushBuilder::StaticClass()))
+		{
+			return false;
+		}
+	}
 
-bool FLevelEditorViewportClient::DropObjectsOnBackground( FViewportCursorLocation& Cursor, const TArray<UObject*>& DroppedObjects, EObjectFlags ObjectFlags, TArray<AActor*>& OutNewActors, bool bSelectActors, UActorFactory* FactoryToUse )
+	return true;
+}
+
+bool FLevelEditorViewportClient::DropObjectsOnBackground(FViewportCursorLocation& Cursor, const TArray<UObject*>& DroppedObjects, EObjectFlags ObjectFlags, TArray<AActor*>& OutNewActors, bool bCreateDropPreview, bool bSelectActors, UActorFactory* FactoryToUse)
 {
-	bool bResult = DroppedObjects.Num() > 0;
+	if (DroppedObjects.Num() == 0)
+	{
+		return false;
+	}
+
+	bool bSuccess = false;
+
+	const bool bTransacted = !bCreateDropPreview && !AreAllDroppedObjectsBrushBuilders(DroppedObjects);
+
+	// Create a transaction if not a preview drop
+	if (bTransacted)
+	{
+		GEditor->BeginTransaction(NSLOCTEXT("UnrealEd", "CreateActors", "Create Actors"));
+	}
 
 	for ( int32 DroppedObjectsIdx = 0; DroppedObjectsIdx < DroppedObjects.Num(); ++DroppedObjectsIdx )
 	{
@@ -800,20 +825,20 @@ bool FLevelEditorViewportClient::DropObjectsOnBackground( FViewportCursorLocatio
 		if ( NewActors.Num() > 0 )
 		{
 			OutNewActors.Append(NewActors);
-		}
-		else
-		{
-			bResult = false;
+			bSuccess = true;
 		}
 	}
 
-	return bResult;
+	if (bTransacted)
+	{
+		GEditor->EndTransaction();
+	}
+
+	return bSuccess;
 }
 
-bool FLevelEditorViewportClient::DropObjectsOnActor(FViewportCursorLocation& Cursor, const TArray<UObject*>& DroppedObjects, AActor* DroppedUponActor, int32 DroppedUponSlot, EObjectFlags ObjectFlags, TArray<AActor*>& OutNewActors, bool bSelectActors, UActorFactory* FactoryToUse)
+bool FLevelEditorViewportClient::DropObjectsOnActor(FViewportCursorLocation& Cursor, const TArray<UObject*>& DroppedObjects, AActor* DroppedUponActor, int32 DroppedUponSlot, EObjectFlags ObjectFlags, TArray<AActor*>& OutNewActors, bool bCreateDropPreview, bool bSelectActors, UActorFactory* FactoryToUse)
 {
-	bool bResult = false;
-
 	if ( !DroppedUponActor || DroppedObjects.Num() == 0 )
 	{
 		return false;
@@ -821,8 +846,14 @@ bool FLevelEditorViewportClient::DropObjectsOnActor(FViewportCursorLocation& Cur
 
 	bool bSuccess = false;
 
-	// Create a transaction if we have more than 1 item. This causes them all to get "un-created" together
-	FScopedTransaction Transaction( NSLOCTEXT("UnrealEd", "CreateActors", "Create Actors") );
+	const bool bTransacted = !bCreateDropPreview && !AreAllDroppedObjectsBrushBuilders(DroppedObjects);
+
+	// Create a transaction if not a preview drop
+	if (bTransacted)
+	{
+		GEditor->BeginTransaction(NSLOCTEXT("UnrealEd", "CreateActors", "Create Actors"));
+	}
+
 	for ( auto DroppedObject : DroppedObjects )
 	{
 		const bool bAppliedToActor = ( FactoryToUse == NULL ) ? AttemptApplyObjToActor( DroppedObject, DroppedUponActor, DroppedUponSlot ) : false;
@@ -844,13 +875,16 @@ bool FLevelEditorViewportClient::DropObjectsOnActor(FViewportCursorLocation& Cur
 		}
 	}
 
+	if (bTransacted)
+	{
+		GEditor->EndTransaction();
+	}
+
 	return bSuccess;
 }
 
-bool FLevelEditorViewportClient::DropObjectsOnBSPSurface( FSceneView* View, FViewportCursorLocation& Cursor, const TArray<UObject*>& DroppedObjects, HModel* TargetProxy, EObjectFlags ObjectFlags, TArray<AActor*>& OutNewActors, bool bSelectActors, UActorFactory* FactoryToUse )
+bool FLevelEditorViewportClient::DropObjectsOnBSPSurface(FSceneView* View, FViewportCursorLocation& Cursor, const TArray<UObject*>& DroppedObjects, HModel* TargetProxy, EObjectFlags ObjectFlags, TArray<AActor*>& OutNewActors, bool bCreateDropPreview, bool bSelectActors, UActorFactory* FactoryToUse)
 {
-	bool bResult = false;
-
 	if (DroppedObjects.Num() == 0)
 	{
 		return false;
@@ -864,8 +898,14 @@ bool FLevelEditorViewportClient::DropObjectsOnBSPSurface( FSceneView* View, FVie
 
 	bool bSuccess = false;
 
-	// Create a transaction if we have more than 1 item. This causes them all to get "un-created" together
-	FScopedTransaction Transaction(NSLOCTEXT("UnrealEd", "CreateActors", "Create Actors"));
+	const bool bTransacted = !bCreateDropPreview && !AreAllDroppedObjectsBrushBuilders(DroppedObjects);
+
+	// Create a transaction if not a preview drop
+	if (bTransacted)
+	{
+		GEditor->BeginTransaction(NSLOCTEXT("UnrealEd", "CreateActors", "Create Actors"));
+	}
+
 	for (auto DroppedObject : DroppedObjects)
 	{
 		// Attempt to create actors from the dropped object
@@ -878,6 +918,11 @@ bool FLevelEditorViewportClient::DropObjectsOnBSPSurface( FSceneView* View, FVie
 		}
 	}
 
+	if (bTransacted)
+	{
+		GEditor->EndTransaction();
+	}
+
 	return bSuccess;
 }
 
@@ -890,7 +935,7 @@ bool FLevelEditorViewportClient::DropObjectsOnBSPSurface( FSceneView* View, FVie
  *
  * @return	true if the drop operation was successfully handled; false otherwise
  */
-bool FLevelEditorViewportClient::DropObjectsOnWidget( FSceneView* View, FViewportCursorLocation& Cursor, const TArray<UObject*>& DroppedObjects )
+bool FLevelEditorViewportClient::DropObjectsOnWidget(FSceneView* View, FViewportCursorLocation& Cursor, const TArray<UObject*>& DroppedObjects, bool bCreateDropPreview)
 {
 	bool bResult = false;
 
@@ -920,7 +965,8 @@ bool FLevelEditorViewportClient::DropObjectsOnWidget( FSceneView* View, FViewpor
 	// Try this again, but without the widgets this time!
 	TArray< AActor* > TemporaryActors;
 	const FIntPoint& CursorPos = Cursor.GetCursorPos();
-	bResult = DropObjectsAtCoordinates(CursorPos.X, CursorPos.Y, DroppedObjects, TemporaryActors);
+	const bool bOnlyDropOnTarget = false;
+	bResult = DropObjectsAtCoordinates(CursorPos.X, CursorPos.Y, DroppedObjects, TemporaryActors, bOnlyDropOnTarget, bCreateDropPreview);
 
 	// Restore the original flags
 	EngineShowFlags.ModeWidgets = bOldModeWidgets1;
@@ -1149,7 +1195,7 @@ bool FLevelEditorViewportClient::DropObjectsAtCoordinates(int32 MouseX, int32 Mo
 		EObjectFlags ObjectFlags = bCreateDropPreview ? RF_Transient : RF_Transactional;
 		if ( HitProxy == nullptr )
 		{
-			bResult = DropObjectsOnBackground( Cursor, DroppedObjects, ObjectFlags, OutNewActors, SelectActors, FactoryToUse );
+			bResult = DropObjectsOnBackground(Cursor, DroppedObjects, ObjectFlags, OutNewActors, bCreateDropPreview, SelectActors, FactoryToUse);
 		}
 		else if (HitProxy->IsA(HActor::StaticGetType()) || HitProxy->IsA(HBSPBrushVert::StaticGetType()))
 		{
@@ -1189,7 +1235,7 @@ bool FLevelEditorViewportClient::DropObjectsAtCoordinates(int32 MouseX, int32 Mo
 					FactoryToUse != NULL ||
 					!AttemptApplyObjToActor(DroppedObjects[0], TargetActor, TargetMaterialSlot, true) )
 				{
-					bResult = DropObjectsOnActor( Cursor, DroppedObjects, TargetActor, TargetMaterialSlot, ObjectFlags, OutNewActors, SelectActors, FactoryToUse );
+					bResult = DropObjectsOnActor(Cursor, DroppedObjects, TargetActor, TargetMaterialSlot, ObjectFlags, OutNewActors, bCreateDropPreview, SelectActors, FactoryToUse);
 				}
 				else
 				{
@@ -1198,7 +1244,7 @@ bool FLevelEditorViewportClient::DropObjectsAtCoordinates(int32 MouseX, int32 Mo
 						TargetActor = static_cast<AActor*>(*It);
 						if( TargetActor )
 						{
-							DropObjectsOnActor( Cursor, DroppedObjects, TargetActor, TargetMaterialSlot, ObjectFlags, OutNewActors, SelectActors, FactoryToUse );
+							DropObjectsOnActor(Cursor, DroppedObjects, TargetActor, TargetMaterialSlot, ObjectFlags, OutNewActors, bCreateDropPreview, SelectActors, FactoryToUse);
 							bResult = true;
 						}
 					}
@@ -1208,7 +1254,7 @@ bool FLevelEditorViewportClient::DropObjectsAtCoordinates(int32 MouseX, int32 Mo
 		else if (HitProxy->IsA(HModel::StaticGetType()))
 		{
 			// BSP surface
-			bResult = DropObjectsOnBSPSurface(View, Cursor, DroppedObjects, static_cast<HModel*>(HitProxy), ObjectFlags, OutNewActors, SelectActors, FactoryToUse );
+			bResult = DropObjectsOnBSPSurface(View, Cursor, DroppedObjects, static_cast<HModel*>(HitProxy), ObjectFlags, OutNewActors, bCreateDropPreview, SelectActors, FactoryToUse);
 		}
 		else if( HitProxy->IsA( HWidgetAxis::StaticGetType() ) )
 		{
diff --git a/Engine/Source/Editor/UnrealEd/Public/LevelEditorViewport.h b/Engine/Source/Editor/UnrealEd/Public/LevelEditorViewport.h
--- a/Engine/Source/Editor/UnrealEd/Public/LevelEditorViewport.h
+++ b/Engine/Source/Editor/UnrealEd/Public/LevelEditorViewport.h
@@ -555,12 +555,13 @@ private:
 	 * @param	DroppedObjects		Array of objects dropped into the viewport
 	 * @param	ObjectFlags			The object flags to place on the actors that this function spawns.
 	 * @param	OutNewActors		The list of actors created while dropping
+	 * @param	bCreateDropPreview	If true, the actor being dropped is a preview actor (defaults: false)
 	 * @param	bSelectActors		If true, select the newly dropped actors (defaults: true)
 	 * @param	FactoryToUse		The preferred actor factory to use (optional)
 	 *
 	 * @return	true if the drop operation was successfully handled; false otherwise
 	 */
-	bool DropObjectsOnBackground( struct FViewportCursorLocation& Cursor, const TArray<UObject*>& DroppedObjects, EObjectFlags ObjectFlags, TArray<AActor*>& OutNewActors, bool bSelectActors = true, class UActorFactory* FactoryToUse = NULL );
+	bool DropObjectsOnBackground( struct FViewportCursorLocation& Cursor, const TArray<UObject*>& DroppedObjects, EObjectFlags ObjectFlags, TArray<AActor*>& OutNewActors, bool bCreateDropPreview = false, bool bSelectActors = true, class UActorFactory* FactoryToUse = NULL );
 
 	/**
 	* Called when an asset is dropped upon an existing actor.
@@ -571,12 +572,13 @@ private:
 	* @param    DroppedUponSlot     The material slot/submesh that was identified as the drop location.  If unknown use -1.
 	* @param	ObjectFlags			The object flags to place on the actors that this function spawns.
 	* @param	OutNewActors		The list of actors created while dropping
+	* @param	bCreateDropPreview	If true, the actor being dropped is a preview actor (defaults: false)
 	* @param	bSelectActors		If true, select the newly dropped actors (defaults: true)
 	* @param	FactoryToUse		The preferred actor factory to use (optional)
 	*
 	* @return	true if the drop operation was successfully handled; false otherwise
 	*/
-	bool DropObjectsOnActor(struct FViewportCursorLocation& Cursor, const TArray<UObject*>& DroppedObjects, AActor* DroppedUponActor, int32 DroppedUponSlot, EObjectFlags ObjectFlags, TArray<AActor*>& OutNewActors, bool bSelectActors = true, class UActorFactory* FactoryToUse = NULL);
+	bool DropObjectsOnActor(struct FViewportCursorLocation& Cursor, const TArray<UObject*>& DroppedObjects, AActor* DroppedUponActor, int32 DroppedUponSlot, EObjectFlags ObjectFlags, TArray<AActor*>& OutNewActors, bool bCreateDropPreview = false, bool bSelectActors = true, class UActorFactory* FactoryToUse = NULL);
 
 	/**
 	 * Called when an asset is dropped upon a BSP surface.
@@ -587,12 +589,13 @@ private:
 	 * @param	TargetProxy			Hit proxy representing the dropped upon model
 	 * @param	ObjectFlags			The object flags to place on the actors that this function spawns.
 	 * @param	OutNewActors		The list of actors created while dropping
+	 * @param	bCreateDropPreview	If true, the actor being dropped is a preview actor (defaults: false)
 	 * @param	bSelectActors		If true, select the newly dropped actors (defaults: true)
 	 * @param	FactoryToUse		The preferred actor factory to use (optional)
 	 *
 	 * @return	true if the drop operation was successfully handled; false otherwise
 	 */
-	bool DropObjectsOnBSPSurface( FSceneView* View, struct FViewportCursorLocation& Cursor, const TArray<UObject*>& DroppedObjects, HModel* TargetProxy, EObjectFlags ObjectFlags, TArray<AActor*>& OutNewActors, bool bSelectActors, UActorFactory* FactoryToUse );
+	bool DropObjectsOnBSPSurface(FSceneView* View, struct FViewportCursorLocation& Cursor, const TArray<UObject*>& DroppedObjects, HModel* TargetProxy, EObjectFlags ObjectFlags, TArray<AActor*>& OutNewActors, bool bCreateDropPreview = false, bool bSelectActors = true, UActorFactory* FactoryToUse = NULL);
 
 	/**
 	 * Called when an asset is dropped upon a manipulation widget.
@@ -600,10 +603,11 @@ private:
 	 * @param	View				The SceneView for the dropped-in viewport
 	 * @param	Cursor				Mouse cursor location
 	 * @param	DroppedObjects		Array of objects dropped into the viewport
+	 * @param	bCreateDropPreview	If true, the actor being dropped is a preview actor (defaults: false)
 	 *
 	 * @return	true if the drop operation was successfully handled; false otherwise
 	 */
-	bool DropObjectsOnWidget( FSceneView* View, struct FViewportCursorLocation& Cursor, const TArray<UObject*>& DroppedObjects);
+	bool DropObjectsOnWidget(FSceneView* View, struct FViewportCursorLocation& Cursor, const TArray<UObject*>& DroppedObjects, bool bCreateDropPreview = false);
 
 	
 
