Added cancel delegate to Open/Save Asset dialog, called when the user cancels the operation without specifying an asset.  Used this to track whether the window is open and thus prevent multiple instances at once.

diff --git a/Engine/Source/Editor/ContentBrowser/Private/ContentBrowserSingleton.cpp b/Engine/Source/Editor/ContentBrowser/Private/ContentBrowserSingleton.cpp
--- a/Engine/Source/Editor/ContentBrowser/Private/ContentBrowserSingleton.cpp
+++ b/Engine/Source/Editor/ContentBrowser/Private/ContentBrowserSingleton.cpp
@@ -81,11 +81,14 @@ TSharedRef<class SWidget> FContentBrowserSingleton::CreateCollectionPicker(const
 		.CollectionPickerConfig(CollectionPickerConfig);
 }
 
-void FContentBrowserSingleton::CreateOpenAssetDialog(const FOpenAssetDialogConfig& InConfig, const FOnAssetsChosenForOpen& InOnAssetsChosenForOpen)
+void FContentBrowserSingleton::CreateOpenAssetDialog(const FOpenAssetDialogConfig& InConfig,
+													 const FOnAssetsChosenForOpen& InOnAssetsChosenForOpen,
+													 const FOnAssetDialogCancelled& InOnAssetDialogCancelled)
 {
 	const bool bModal = false;
 	TSharedRef<SAssetDialog> AssetDialog = SNew(SAssetDialog, InConfig);
 	AssetDialog->SetOnAssetsChosenForOpen(InOnAssetsChosenForOpen);
+	AssetDialog->SetOnAssetDialogCancelled(InOnAssetDialogCancelled);
 	SharedCreateAssetDialogWindow(AssetDialog, InConfig, bModal);
 }
 
@@ -112,11 +115,14 @@ TArray<FAssetData> FContentBrowserSingleton::CreateModalOpenAssetDialog(const FO
 	return ModalWindowResults.SavedResults;
 }
 
-void FContentBrowserSingleton::CreateSaveAssetDialog(const FSaveAssetDialogConfig& InConfig, const FOnObjectPathChosenForSave& InOnObjectPathChosenForSave)
+void FContentBrowserSingleton::CreateSaveAssetDialog(const FSaveAssetDialogConfig& InConfig,
+													 const FOnObjectPathChosenForSave& InOnObjectPathChosenForSave,
+													 const FOnAssetDialogCancelled& InOnAssetDialogCancelled)
 {
 	const bool bModal = false;
 	TSharedRef<SAssetDialog> AssetDialog = SNew(SAssetDialog, InConfig);
 	AssetDialog->SetOnObjectPathChosenForSave(InOnObjectPathChosenForSave);
+	AssetDialog->SetOnAssetDialogCancelled(InOnAssetDialogCancelled);
 	SharedCreateAssetDialogWindow(AssetDialog, InConfig, bModal);
 }
 
diff --git a/Engine/Source/Editor/ContentBrowser/Private/ContentBrowserSingleton.h b/Engine/Source/Editor/ContentBrowser/Private/ContentBrowserSingleton.h
--- a/Engine/Source/Editor/ContentBrowser/Private/ContentBrowserSingleton.h
+++ b/Engine/Source/Editor/ContentBrowser/Private/ContentBrowserSingleton.h
@@ -21,9 +21,9 @@ public:
 	virtual TSharedRef<class SWidget> CreateAssetPicker(const FAssetPickerConfig& AssetPickerConfig) override;
 	virtual TSharedRef<class SWidget> CreatePathPicker(const FPathPickerConfig& PathPickerConfig) override;
 	virtual TSharedRef<class SWidget> CreateCollectionPicker(const FCollectionPickerConfig& CollectionPickerConfig) override;
-	virtual void CreateOpenAssetDialog(const FOpenAssetDialogConfig& OpenAssetConfig, const FOnAssetsChosenForOpen& OnAssetsChosenForOpen) override;
+	virtual void CreateOpenAssetDialog(const FOpenAssetDialogConfig& OpenAssetConfig, const FOnAssetsChosenForOpen& OnAssetsChosenForOpen, const FOnAssetDialogCancelled& OnAssetDialogCancelled) override;
 	virtual TArray<FAssetData> CreateModalOpenAssetDialog(const FOpenAssetDialogConfig& InConfig) override;
-	virtual void CreateSaveAssetDialog(const FSaveAssetDialogConfig& SaveAssetConfig, const FOnObjectPathChosenForSave& OnAssetNameChosenForSave) override;
+	virtual void CreateSaveAssetDialog(const FSaveAssetDialogConfig& SaveAssetConfig, const FOnObjectPathChosenForSave& OnAssetNameChosenForSave, const FOnAssetDialogCancelled& OnAssetDialogCancelled) override;
 	virtual FString CreateModalSaveAssetDialog(const FSaveAssetDialogConfig& SaveAssetConfig) override;
 	virtual bool HasPrimaryContentBrowser() const override;
 	virtual void FocusPrimaryContentBrowser(bool bFocusSearch) override;
diff --git a/Engine/Source/Editor/ContentBrowser/Private/SAssetDialog.cpp b/Engine/Source/Editor/ContentBrowser/Private/SAssetDialog.cpp
--- a/Engine/Source/Editor/ContentBrowser/Private/SAssetDialog.cpp
+++ b/Engine/Source/Editor/ContentBrowser/Private/SAssetDialog.cpp
@@ -10,9 +10,18 @@ SAssetDialog::SAssetDialog()
 	, ExistingAssetPolicy(ESaveAssetDialogExistingAssetPolicy::Disallow)
 	, bLastInputValidityCheckSuccessful(false)
 	, bPendingFocusNextFrame(true)
+	, bValidAssetsChosen(false)
 {
 }
 
+SAssetDialog::~SAssetDialog()
+{
+	if (!bValidAssetsChosen)
+	{
+		OnAssetDialogCancelled.ExecuteIfBound();
+	}
+}
+
 BEGIN_SLATE_FUNCTION_BUILD_OPTIMIZATION
 void SAssetDialog::Construct(const FArguments& InArgs, const FSharedAssetDialogConfig& InConfig)
 {
@@ -228,6 +237,11 @@ void SAssetDialog::SetOnObjectPathChosenForSave(const FOnObjectPathChosenForSave
 	OnObjectPathChosenForSave = InOnObjectPathChosenForSave;
 }
 
+void SAssetDialog::SetOnAssetDialogCancelled(const FOnAssetDialogCancelled& InOnAssetDialogCancelled)
+{
+	OnAssetDialogCancelled = InOnAssetDialogCancelled;
+}
+
 void SAssetDialog::FocusNameBox()
 {
 	if ( NameEditableText.IsValid() )
@@ -434,6 +448,7 @@ void SAssetDialog::ChooseAssetsForOpen(const TArray<FAssetData>& SelectedAssets)
 	{
 		if (SelectedAssets.Num() > 0)
 		{
+			bValidAssetsChosen = true;
 			OnAssetsChosenForOpen.ExecuteIfBound(SelectedAssets);
 			CloseDialog();
 		}
@@ -469,6 +484,7 @@ void SAssetDialog::CommitObjectPathForSave()
 
 			if ( bProceedWithSave )
 			{
+				bValidAssetsChosen = true;
 				OnObjectPathChosenForSave.ExecuteIfBound(ObjectPath);
 				CloseDialog();
 			}
diff --git a/Engine/Source/Editor/ContentBrowser/Private/SAssetDialog.h b/Engine/Source/Editor/ContentBrowser/Private/SAssetDialog.h
--- a/Engine/Source/Editor/ContentBrowser/Private/SAssetDialog.h
+++ b/Engine/Source/Editor/ContentBrowser/Private/SAssetDialog.h
@@ -10,6 +10,7 @@ public:
 	SLATE_END_ARGS()
 
 	SAssetDialog();
+	virtual ~SAssetDialog();
 
 	virtual void Construct(const FArguments& InArgs, const FSharedAssetDialogConfig& InConfig);
 
@@ -22,6 +23,9 @@ public:
 	/** Sets the delegate handler for when a save operation is committed */
 	void SetOnObjectPathChosenForSave(const FOnObjectPathChosenForSave& InOnObjectPathChosenForSave);
 
+	/** Sets the delegate handler for when the dialog is closed or cancelled */
+	void SetOnAssetDialogCancelled(const FOnAssetDialogCancelled& InOnAssetDialogCancelled);
+
 private:
 
 	/** Moves keyboard focus to the name box if this is a save dialog */
@@ -100,6 +104,9 @@ private:
 	/** Fired when an object path was chosen for save. Only fired in save dialogs. */
 	FOnObjectPathChosenForSave OnObjectPathChosenForSave;
 
+	/** Fired when the asset dialog is cancelled or closed */
+	FOnAssetDialogCancelled OnAssetDialogCancelled;
+
 	/** The assets that are currently selected in the asset picker */
 	TArray<FAssetData> CurrentlySelectedAssets;
 
@@ -123,4 +130,7 @@ private:
 
 	/** Used to focus the name box after opening the dialog */
 	bool bPendingFocusNextFrame;
+
+	/** Used to specify that valid assets were chosen */
+	bool bValidAssetsChosen;
 };
\ No newline at end of file
diff --git a/Engine/Source/Editor/ContentBrowser/Public/ContentBrowserDelegates.h b/Engine/Source/Editor/ContentBrowser/Public/ContentBrowserDelegates.h
--- a/Engine/Source/Editor/ContentBrowser/Public/ContentBrowserDelegates.h
+++ b/Engine/Source/Editor/ContentBrowser/Public/ContentBrowserDelegates.h
@@ -81,6 +81,9 @@ DECLARE_DELEGATE_RetVal_OneParam( TSharedRef<SToolTip>, FOnGetCustomAssetToolTip
 /** Called when an asset item visualizes its tooltip */
 DECLARE_DELEGATE_RetVal_TwoParams(bool, FOnVisualizeAssetToolTip, const TSharedPtr<SWidget>& /*ToolTipContent*/, class FAssetData& /*AssetData*/);
 
+/** Called from the Asset Dialog when a non-modal dialog is closed or cancelled */
+DECLARE_DELEGATE(FOnAssetDialogCancelled);
+
 /** Called from the Asset Dialog when assets are chosen in non-modal Open dialogs */
 DECLARE_DELEGATE_OneParam(FOnAssetsChosenForOpen, const TArray<FAssetData>& /*SelectedAssets*/);
 
diff --git a/Engine/Source/Editor/ContentBrowser/Public/IContentBrowserSingleton.h b/Engine/Source/Editor/ContentBrowser/Public/IContentBrowserSingleton.h
--- a/Engine/Source/Editor/ContentBrowser/Public/IContentBrowserSingleton.h
+++ b/Engine/Source/Editor/ContentBrowser/Public/IContentBrowserSingleton.h
@@ -307,8 +307,9 @@ public:
 	 *
 	 * @param OpenAssetConfig				A struct containing details about how the open asset dialog should behave
 	 * @param OnAssetsChosenForOpen			A delegate that is fired when assets are chosen and the open button is pressed
+	 * @param OnAssetDialogCancelled		A delegate that is fired when the asset dialog is closed or cancelled
 	 */
-	virtual void CreateOpenAssetDialog(const FOpenAssetDialogConfig& OpenAssetConfig, const FOnAssetsChosenForOpen& OnAssetsChosenForOpen) = 0;
+	virtual void CreateOpenAssetDialog(const FOpenAssetDialogConfig& OpenAssetConfig, const FOnAssetsChosenForOpen& OnAssetsChosenForOpen, const FOnAssetDialogCancelled& OnAssetDialogCancelled) = 0;
 
 	/**
 	 * Opens the Open Asset dialog in a modal window
@@ -323,8 +324,9 @@ public:
 	 *
 	 * @param SaveAssetConfig				A struct containing details about how the save asset dialog should behave
 	 * @param OnAssetNameChosenForSave		A delegate that is fired when an object path is chosen and the save button is pressed
+	 * @param OnAssetDialogCancelled		A delegate that is fired when the asset dialog is closed or cancelled
 	 */
-	virtual void CreateSaveAssetDialog(const FSaveAssetDialogConfig& SaveAssetConfig, const FOnObjectPathChosenForSave& OnAssetNameChosenForSave) = 0;
+	virtual void CreateSaveAssetDialog(const FSaveAssetDialogConfig& SaveAssetConfig, const FOnObjectPathChosenForSave& OnAssetNameChosenForSave, const FOnAssetDialogCancelled& OnAssetDialogCancelled) = 0;
 
 	/**
 	 * Opens the Save Asset dialog in a modal window
diff --git a/Engine/Source/Editor/Levels/Private/LevelCollectionViewModel.cpp b/Engine/Source/Editor/Levels/Private/LevelCollectionViewModel.cpp
--- a/Engine/Source/Editor/Levels/Private/LevelCollectionViewModel.cpp
+++ b/Engine/Source/Editor/Levels/Private/LevelCollectionViewModel.cpp
@@ -1405,9 +1405,14 @@ void FLevelCollectionViewModel::AddExistingLevel(bool bRemoveInvalidSelectedLeve
 {
 	if ( UEditorEngine::IsUsingWorldAssets() )
 	{
-		FEditorFileUtils::FOnLevelsChosen LevelsChosenDelegate = FEditorFileUtils::FOnLevelsChosen::CreateSP(this, &FLevelCollectionViewModel::HandleAddExistingLevelSelected, bRemoveInvalidSelectedLevelsAfter);
-		const bool bAllowMultipleSelection = true;
-		FEditorFileUtils::OpenLevelPickingDialog(LevelsChosenDelegate, bAllowMultipleSelection);
+		if (!bAssetDialogOpen)
+		{
+			bAssetDialogOpen = true;
+			FEditorFileUtils::FOnLevelsChosen LevelsChosenDelegate = FEditorFileUtils::FOnLevelsChosen::CreateSP(this, &FLevelCollectionViewModel::HandleAddExistingLevelSelected, bRemoveInvalidSelectedLevelsAfter);
+			FEditorFileUtils::FOnLevelPickingCancelled DialogCancelledDelegate = FEditorFileUtils::FOnLevelPickingCancelled::CreateSP(this, &FLevelCollectionViewModel::HandleAddExistingLevelCancelled);
+			const bool bAllowMultipleSelection = true;
+			FEditorFileUtils::OpenLevelPickingDialog(LevelsChosenDelegate, DialogCancelledDelegate, bAllowMultipleSelection);
+		}
 	}
 	else
 	{
@@ -1491,6 +1496,7 @@ void FLevelCollectionViewModel::AddExistingLevel(bool bRemoveInvalidSelectedLeve
 
 void FLevelCollectionViewModel::HandleAddExistingLevelSelected(const TArray<FAssetData>& SelectedAssets, bool bRemoveInvalidSelectedLevelsAfter)
 {
+	bAssetDialogOpen = false;
 	TArray<FString> PackageNames;
 	for (const auto& AssetData : SelectedAssets)
 	{
@@ -1505,6 +1511,11 @@ void FLevelCollectionViewModel::HandleAddExistingLevelSelected(const TArray<FAss
 	}
 }
 
+void FLevelCollectionViewModel::HandleAddExistingLevelCancelled()
+{
+	bAssetDialogOpen = false;
+}
+
 void FLevelCollectionViewModel::AddSelectedActorsToNewLevel_Executed()
 {
 	EditorLevelUtils::CreateNewLevel( CurrentWorld.Get(), true, AddedLevelStreamingClass );
diff --git a/Engine/Source/Editor/Levels/Private/LevelCollectionViewModel.h b/Engine/Source/Editor/Levels/Private/LevelCollectionViewModel.h
--- a/Engine/Source/Editor/Levels/Private/LevelCollectionViewModel.h
+++ b/Engine/Source/Editor/Levels/Private/LevelCollectionViewModel.h
@@ -376,6 +376,9 @@ private:
 	/** Handler for when a level is selected after invoking AddExistingLevel */
 	void HandleAddExistingLevelSelected(const TArray<FAssetData>& SelectedAssets, bool bRemoveInvalidSelectedLevelsAfter);
 
+	/** Handler for when the dialog is cancelled after invoking AddExistingLevel */
+	void HandleAddExistingLevelCancelled();
+
 	/** Add Selected Actors to New Level; prompts for level save location */
 	void AddSelectedActorsToNewLevel_Executed();
 
@@ -544,6 +547,9 @@ private:
 	
 	/** Flag for whether the selection of levels or actors has changed */
 	bool bSelectionHasChanged;
+
+	/** Boolean indicating whether the asset dialog is currently open */
+	bool bAssetDialogOpen;
 };
 
 
diff --git a/Engine/Source/Editor/UnrealEd/Private/FileHelpers.cpp b/Engine/Source/Editor/UnrealEd/Private/FileHelpers.cpp
--- a/Engine/Source/Editor/UnrealEd/Private/FileHelpers.cpp
+++ b/Engine/Source/Editor/UnrealEd/Private/FileHelpers.cpp
@@ -1543,7 +1543,7 @@ bool FEditorFileUtils::PromptToCheckoutLevels(bool bCheckDirty, ULevel* Specific
 	return FEditorFileUtils::PromptToCheckoutLevels( bCheckDirty, LevelsToCheckOut );	
 }
 
-void FEditorFileUtils::OpenLevelPickingDialog(const FOnLevelsChosen& OnLevelsChosen, bool bAllowMultipleSelection)
+void FEditorFileUtils::OpenLevelPickingDialog(const FOnLevelsChosen& OnLevelsChosen, const FOnLevelPickingCancelled& OnLevelPickingCancelled, bool bAllowMultipleSelection)
 {
 	struct FLocal
 	{
@@ -1569,6 +1569,11 @@ void FEditorFileUtils::OpenLevelPickingDialog(const FOnLevelsChosen& OnLevelsCho
 				OnLevelsChosen.ExecuteIfBound(SelectedLevels);
 			}
 		}
+
+		static void OnDialogCancelled(FOnLevelPickingCancelled OnLevelPickingCancelled)
+		{
+			OnLevelPickingCancelled.ExecuteIfBound();
+		}
 	};
 
 	// Determine the starting path. Try to use the most recently used directory
@@ -1589,7 +1594,9 @@ void FEditorFileUtils::OpenLevelPickingDialog(const FOnLevelsChosen& OnLevelsCho
 	OpenAssetDialogConfig.bAllowMultipleSelection = bAllowMultipleSelection;
 
 	FContentBrowserModule& ContentBrowserModule = FModuleManager::LoadModuleChecked<FContentBrowserModule>("ContentBrowser");
-	ContentBrowserModule.Get().CreateOpenAssetDialog(OpenAssetDialogConfig, FOnAssetsChosenForOpen::CreateStatic(&FLocal::OnLevelsSelected, OnLevelsChosen));
+	ContentBrowserModule.Get().CreateOpenAssetDialog(OpenAssetDialogConfig,
+													 FOnAssetsChosenForOpen::CreateStatic(&FLocal::OnLevelsSelected, OnLevelsChosen),
+													 FOnAssetDialogCancelled::CreateStatic(&FLocal::OnDialogCancelled, OnLevelPickingCancelled));
 }
 
 bool FEditorFileUtils::IsValidMapFilename(const FString& MapFilename, FText& OutErrorMessage)
@@ -1741,10 +1748,14 @@ void FEditorFileUtils::LoadMap()
 
 	if (UEditorEngine::IsUsingWorldAssets())
 	{
+		static bool bIsDialogOpen = false;
+
 		struct FLocal
 		{
 			static void HandleLevelsChosen(const TArray<FAssetData>& SelectedAssets)
 			{
+				bIsDialogOpen = false;
+
 				if ( SelectedAssets.Num() > 0 )
 				{
 					const FAssetData& AssetData = SelectedAssets[0];
@@ -1767,10 +1778,21 @@ void FEditorFileUtils::LoadMap()
 					FEditorFileUtils::LoadMap(FileToOpen, bLoadAsTemplate, bShowProgress);
 				}
 			}
+
+			static void HandleDialogCancelled()
+			{
+				bIsDialogOpen = false;
+			}
 		};
 		
-		const bool bAllowMultipleSelection = false;
-		OpenLevelPickingDialog(FOnLevelsChosen::CreateStatic(&FLocal::HandleLevelsChosen), bAllowMultipleSelection);
+		if (!bIsDialogOpen)
+		{
+			bIsDialogOpen = true;
+			const bool bAllowMultipleSelection = false;
+			OpenLevelPickingDialog(FOnLevelsChosen::CreateStatic(&FLocal::HandleLevelsChosen),
+								   FOnLevelPickingCancelled::CreateStatic(&FLocal::HandleDialogCancelled),
+								   bAllowMultipleSelection);
+		}
 	}
 	else
 	{
diff --git a/Engine/Source/Editor/WorldBrowser/Private/StreamingLevels/StreamingLevelCollectionModel.cpp b/Engine/Source/Editor/WorldBrowser/Private/StreamingLevels/StreamingLevelCollectionModel.cpp
--- a/Engine/Source/Editor/WorldBrowser/Private/StreamingLevels/StreamingLevelCollectionModel.cpp
+++ b/Engine/Source/Editor/WorldBrowser/Private/StreamingLevels/StreamingLevelCollectionModel.cpp
@@ -21,6 +21,7 @@
 FStreamingLevelCollectionModel::FStreamingLevelCollectionModel(UEditorEngine* InEditor)
 	: FLevelCollectionModel(InEditor)
 	, AddedLevelStreamingClass(ULevelStreamingKismet::StaticClass())
+	, bAssetDialogOpen(false)
 {
 	const TSubclassOf<ULevelStreaming> DefaultLevelStreamingClass = GetDefault<ULevelEditorMiscSettings>()->DefaultLevelStreamingClass;
 	if ( DefaultLevelStreamingClass )
@@ -385,9 +386,14 @@ void FStreamingLevelCollectionModel::AddExistingLevel(bool bRemoveInvalidSelecte
 {
 	if (UEditorEngine::IsUsingWorldAssets())
 	{
-		FEditorFileUtils::FOnLevelsChosen LevelsChosenDelegate = FEditorFileUtils::FOnLevelsChosen::CreateSP(this, &FStreamingLevelCollectionModel::HandleAddExistingLevelSelected, bRemoveInvalidSelectedLevelsAfter);
-		const bool bAllowMultipleSelection = true;
-		FEditorFileUtils::OpenLevelPickingDialog(LevelsChosenDelegate, bAllowMultipleSelection);
+		if (!bAssetDialogOpen)
+		{
+			bAssetDialogOpen = true;
+			FEditorFileUtils::FOnLevelsChosen LevelsChosenDelegate = FEditorFileUtils::FOnLevelsChosen::CreateSP(this, &FStreamingLevelCollectionModel::HandleAddExistingLevelSelected, bRemoveInvalidSelectedLevelsAfter);
+			FEditorFileUtils::FOnLevelPickingCancelled LevelPickingCancelledDelegate = FEditorFileUtils::FOnLevelPickingCancelled::CreateSP(this, &FStreamingLevelCollectionModel::HandleAddExistingLevelCancelled);
+			const bool bAllowMultipleSelection = true;
+			FEditorFileUtils::OpenLevelPickingDialog(LevelsChosenDelegate, LevelPickingCancelledDelegate, bAllowMultipleSelection);
+		}
 	}
 	else
 	{
@@ -478,6 +484,8 @@ void FStreamingLevelCollectionModel::AddExistingLevel(bool bRemoveInvalidSelecte
 
 void FStreamingLevelCollectionModel::HandleAddExistingLevelSelected(const TArray<FAssetData>& SelectedAssets, bool bRemoveInvalidSelectedLevelsAfter)
 {
+	bAssetDialogOpen = false;
+
 	TArray<FString> PackageNames;
 	for (const auto& AssetData : SelectedAssets)
 	{
@@ -499,6 +507,11 @@ void FStreamingLevelCollectionModel::HandleAddExistingLevelSelected(const TArray
 	}
 }
 
+void FStreamingLevelCollectionModel::HandleAddExistingLevelCancelled()
+{
+	bAssetDialogOpen = false;
+}
+
 void FStreamingLevelCollectionModel::AddSelectedActorsToNewLevel_Executed()
 {
 	EditorLevelUtils::CreateNewLevel( CurrentWorld.Get(), true, AddedLevelStreamingClass );
diff --git a/Engine/Source/Editor/WorldBrowser/Private/StreamingLevels/StreamingLevelCollectionModel.h b/Engine/Source/Editor/WorldBrowser/Private/StreamingLevels/StreamingLevelCollectionModel.h
--- a/Engine/Source/Editor/WorldBrowser/Private/StreamingLevels/StreamingLevelCollectionModel.h
+++ b/Engine/Source/Editor/WorldBrowser/Private/StreamingLevels/StreamingLevelCollectionModel.h
@@ -81,6 +81,9 @@ private:
 	/** Handler for when a level is selected after invoking AddExistingLevel */
 	void HandleAddExistingLevelSelected(const TArray<FAssetData>& SelectedAssets, bool bRemoveInvalidSelectedLevelsAfter);
 
+	/** Handler for when the level picker dialog is cancelled */
+	void HandleAddExistingLevelCancelled();
+
 	/** Add Selected Actors to New Level; prompts for level save location */
 	void AddSelectedActorsToNewLevel_Executed();
 
@@ -114,6 +117,9 @@ private:
 
 	/** The current class to set new or added levels streaming method to. */
 	UClass*				AddedLevelStreamingClass;
+
+	/** Boolean indicating whether the asset dialog is currently open */
+	bool				bAssetDialogOpen;
 };
 
 
