Added missing file (fixes CL 2364646)

diff --git a/Engine/Source/Editor/UnrealEd/Public/FileHelpers.h b/Engine/Source/Editor/UnrealEd/Public/FileHelpers.h
--- a/Engine/Source/Editor/UnrealEd/Public/FileHelpers.h
+++ b/Engine/Source/Editor/UnrealEd/Public/FileHelpers.h
@@ -48,13 +48,16 @@ public:
 	// Loading
 
 	DECLARE_DELEGATE_OneParam(FOnLevelsChosen, const TArray<FAssetData>& /*SelectedLevels*/);
+	DECLARE_DELEGATE(FOnLevelPickingCancelled);
+
 	/**
 	 * Opens a non-modal dialog to allow the user to choose a level
 	 *
 	 * @param	OnLevelsChosen				Delegate executed when one more more levels have been selected
+	 * @param	OnLevelPickingDialogClosed	Delegate executed when the level picking dialog is closed
 	 * @param	bAllowMultipleSelection		If true, more than one level can be chosen
 	 */
-	static UNREALED_API void OpenLevelPickingDialog(const FOnLevelsChosen& OnLevelsChosen, bool bAllowMultipleSelection);
+	static UNREALED_API void OpenLevelPickingDialog(const FOnLevelsChosen& OnLevelsChosen, const FOnLevelPickingCancelled& OnLevelPickingCancelled, bool bAllowMultipleSelection);
 
 	/**
 	 * Returns true if the specified map filename is valid for loading or saving.
