UE-5559 LIVE: Variables of type UserDefinedStruct that are not compiled when a project is closed will not be visible when the project is reopened

diff --git a/Engine/Source/Editor/UnrealEd/Private/Kismet2/BlueprintEditorUtils.cpp b/Engine/Source/Editor/UnrealEd/Private/Kismet2/BlueprintEditorUtils.cpp
--- a/Engine/Source/Editor/UnrealEd/Private/Kismet2/BlueprintEditorUtils.cpp
+++ b/Engine/Source/Editor/UnrealEd/Private/Kismet2/BlueprintEditorUtils.cpp
@@ -762,6 +762,19 @@ struct FRegenerationHelper
 		TSet<UStruct*> Dependencies;
 		ProcessHierarchy(Blueprint->ParentClass, Dependencies);
 
+		for (const auto& NewVar : Blueprint->NewVariables)
+		{
+			if (UObject* TypeObject = NewVar.VarType.PinSubCategoryObject.Get())
+			{
+				ForcedLoad(TypeObject);
+			}
+
+			if (UClass* TypeClass = NewVar.VarType.PinSubCategoryMemberReference.MemberParentClass)
+			{
+				ForcedLoad(TypeClass);
+			}
+		}
+
 		TSet<UBlueprint*> MacroSources;
 		TArray<UEdGraph*> Graphs;
 		Blueprint->GetAllGraphs(Graphs);
