Fixed input binding adding in PackFactory.

diff --git a/Engine/Source/Editor/UnrealEd/Private/Factories/PackFactory.cpp b/Engine/Source/Editor/UnrealEd/Private/Factories/PackFactory.cpp
--- a/Engine/Source/Editor/UnrealEd/Private/Factories/PackFactory.cpp
+++ b/Engine/Source/Editor/UnrealEd/Private/Factories/PackFactory.cpp
@@ -209,13 +209,13 @@ namespace PackFactoryHelper
 
 					for (const FInputActionKeyMapping& ActionKeyMapping : ActionMappingsToAdd)
 					{
-						InputSettingsCDO->AddActionMapping(ActionKeyMapping);
+						InputSettingsCDO->ActionMappings.Add(ActionKeyMapping);
 					}
 					for (const FInputAxisKeyMapping& AxisKeyMapping : AxisMappingsToAdd)
 					{
-						InputSettingsCDO->AddAxisMapping(AxisKeyMapping);
+						InputSettingsCDO->AxisMappings.Add(AxisKeyMapping);
 					}
-
+					
 					InputSettingsCDO->SaveKeyMappings();
 					InputSettingsCDO->UpdateDefaultConfigFile();
 				}
