Back out changelist 2360073. It is breaking some cooks.

diff --git a/Engine/Source/Editor/UnrealEd/Private/Commandlets/CookCommandlet.cpp b/Engine/Source/Editor/UnrealEd/Private/Commandlets/CookCommandlet.cpp
--- a/Engine/Source/Editor/UnrealEd/Private/Commandlets/CookCommandlet.cpp
+++ b/Engine/Source/Editor/UnrealEd/Private/Commandlets/CookCommandlet.cpp
@@ -562,13 +562,13 @@ int32 UCookCommandlet::Main(const FString& CmdLineParams)
 		GenerateAssetRegistry(Platforms);
 
 		// new cook is better 
-		if ( Switches.Contains(TEXT("OLDCOOK")))
+		if ( Switches.Contains(TEXT("NEWCOOK")))
 		{
-			Cook(Platforms, FilesInPath);
+			NewCook(Platforms, FilesInPath );
 		}
 		else
 		{
-			NewCook(Platforms, FilesInPath );
+			Cook(Platforms, FilesInPath);
 		}
 	}
 	
