- Fixed unity error (Log was conflicting with DumpBlueprintsInfoCommandlet)

diff --git a/Engine/Source/Editor/UnrealEd/Private/Commandlets/GenerateBlueprintAPICommandlet.cpp b/Engine/Source/Editor/UnrealEd/Private/Commandlets/GenerateBlueprintAPICommandlet.cpp
--- a/Engine/Source/Editor/UnrealEd/Private/Commandlets/GenerateBlueprintAPICommandlet.cpp
+++ b/Engine/Source/Editor/UnrealEd/Private/Commandlets/GenerateBlueprintAPICommandlet.cpp
@@ -25,7 +25,7 @@
 #include "EditorCategoryUtils.h"
 #include "Engine/LevelScriptActor.h"
 
-DEFINE_LOG_CATEGORY_STATIC(LogBlueprintInfoDump, Log, All);
+DEFINE_LOG_CATEGORY_STATIC(LogBlueprintAPIGenerate, Log, All);
 
 /*******************************************************************************
  * Static Helpers
@@ -257,7 +257,7 @@ GenerateBlueprintAPIUtils::CommandletOptions::CommandletOptions(TArray<FString>
 
 			if (BlueprintClass == nullptr)
 			{
-				UE_LOG(LogBlueprintInfoDump, Error, TEXT("Unrecognized blueprint class '%s', defaulting to 'Actor'"), *ClassName);
+				UE_LOG(LogBlueprintAPIGenerate, Error, TEXT("Unrecognized blueprint class '%s', defaulting to 'Actor'"), *ClassName);
 				BlueprintClass = AActor::StaticClass();
 			}
 		}
@@ -272,7 +272,7 @@ GenerateBlueprintAPIUtils::CommandletOptions::CommandletOptions(TArray<FString>
 			{
 				if (ClassName.Compare("all", ESearchCase::IgnoreCase))
 				{
-					UE_LOG(LogBlueprintInfoDump, Error, TEXT("Unrecognized palette filter '%s', defaulting to unfiltered"), *ClassName);
+					UE_LOG(LogBlueprintAPIGenerate, Error, TEXT("Unrecognized palette filter '%s', defaulting to unfiltered"), *ClassName);
 					NewDumpFlags &= ~(GenerateBlueprintAPIUtils::BPDUMP_FilteredPalette);
 				}
 			}
@@ -303,7 +303,7 @@ GenerateBlueprintAPIUtils::CommandletOptions::CommandletOptions(TArray<FString>
 		}
 		else
 		{
-			UE_LOG(LogBlueprintInfoDump, Warning, TEXT("Unrecognized command switch '%s', use -help for a listing of all accepted params"), *Switch);
+			UE_LOG(LogBlueprintAPIGenerate, Warning, TEXT("Unrecognized command switch '%s', use -help for a listing of all accepted params"), *Switch);
 		}
 	}
 
@@ -333,7 +333,7 @@ static UWorld* GenerateBlueprintAPIUtils::GetWorld()
 		{
 			if (GUnrealEd == nullptr)
 			{
-				UE_LOG(LogBlueprintInfoDump, Error, TEXT("Cannot create a temp map to test within, without a valid editor world"));
+				UE_LOG(LogBlueprintAPIGenerate, Error, TEXT("Cannot create a temp map to test within, without a valid editor world"));
 			}
 			else
 			{
@@ -377,7 +377,7 @@ static UBlueprint* GenerateBlueprintAPIUtils::MakeTempBlueprint(UClass* ParentCl
 			UWorld* World = GetWorld();
 			if (World == nullptr)
 			{
-				UE_LOG(LogBlueprintInfoDump, Error, TEXT("Cannot make a proper level blueprint without a valid editor level for its outer."));
+				UE_LOG(LogBlueprintAPIGenerate, Error, TEXT("Cannot make a proper level blueprint without a valid editor level for its outer."));
 			}
 			else
 			{
@@ -565,7 +565,7 @@ static void GenerateBlueprintAPIUtils::GetComponentProperties(UBlueprint* Bluepr
 static void GenerateBlueprintAPIUtils::DumpInfoForClass(uint32 Indent, UClass* BlueprintClass, FArchive* FileOutWriter)
 {
 	FString const ClassName = BlueprintClass->GetName();
-	UE_LOG(LogBlueprintInfoDump, Display, TEXT("%sDumping BP class: '%s'..."), *BuildIndentString(Indent, true), *ClassName);
+	UE_LOG(LogBlueprintAPIGenerate, Display, TEXT("%sDumping BP class: '%s'..."), *BuildIndentString(Indent, true), *ClassName);
 
 	FString const ClassEntryIndent = BuildIndentString(Indent);
 	FString BeginClassEntry = FString::Printf(TEXT("%s\"%s\" : {"), *ClassEntryIndent, *ClassName);
@@ -685,7 +685,7 @@ static void GenerateBlueprintAPIUtils::DumpPalette(uint32 Indent, UBlueprint* Bl
 	BeginPaletteEntry += "\" : {\n";
 
 	FString const NestedIndent = BuildIndentString(Indent + 1);
-	UE_LOG(LogBlueprintInfoDump, Display, TEXT("%sDumping palette: %s"), *BuildIndentString(Indent, true), *FilterClassName);
+	UE_LOG(LogBlueprintAPIGenerate, Display, TEXT("%sDumping palette: %s"), *BuildIndentString(Indent, true), *FilterClassName);
 
 	bool bIsAnimBlueprint = (Cast<UAnimBlueprint>(Blueprint) != nullptr);
 	// animation blueprints don't have a palette
@@ -1038,15 +1038,15 @@ int32 UGenerateBlueprintAPICommandlet::Main(FString const& Params)
 	UClass* const BlueprintClass = CommandOptions.BlueprintClass;
 	if (CommandOptions.DumpFlags & GenerateBlueprintAPIUtils::BPDUMP_LogHelp)
 	{
-		UE_LOG(LogBlueprintInfoDump, Display, TEXT("%s"), *GenerateBlueprintAPIUtils::HelpString);
+		UE_LOG(LogBlueprintAPIGenerate, Display, TEXT("%s"), *GenerateBlueprintAPIUtils::HelpString);
 	}
 	else if (BlueprintClass != nullptr)
 	{
-		UE_LOG(LogBlueprintInfoDump, Display, TEXT("Dumping Blueprint info..."));
+		UE_LOG(LogBlueprintAPIGenerate, Display, TEXT("Dumping Blueprint info..."));
 		// make sure the class that the user specified is a blueprintable type
 		if (IsInvalidBlueprintClass(BlueprintClass))
 		{
-			UE_LOG(LogBlueprintInfoDump, Error, TEXT("Cannot create a blueprint from class: '%s'"), *BlueprintClass->GetName());
+			UE_LOG(LogBlueprintAPIGenerate, Error, TEXT("Cannot create a blueprint from class: '%s'"), *BlueprintClass->GetName());
 			if (FileOut != nullptr)
 			{
 				FString const InvalidClassEntry = GenerateBlueprintAPIUtils::BuildIndentString(1) + "\"INVALID_BLUEPRINT_CLASS\" : \"" + BlueprintClass->GetName() + "\"";
@@ -1061,7 +1061,7 @@ int32 UGenerateBlueprintAPICommandlet::Main(FString const& Params)
 	// if the user didn't specify a class, then we take that to mean dump ALL the classes!
 	else
 	{
-		UE_LOG(LogBlueprintInfoDump, Display, TEXT("Dumping Blueprint info..."));
+		UE_LOG(LogBlueprintAPIGenerate, Display, TEXT("Dumping Blueprint info..."));
 		for (TObjectIterator<UClass> ClassIt; ClassIt; ++ClassIt)
 		{
 			UClass* Class = *ClassIt;
