SCurveEditor - Fix FCurveViewModel constructor initialization.

diff --git a/Engine/Source/Editor/UnrealEd/Public/SCurveEditor.h b/Engine/Source/Editor/UnrealEd/Public/SCurveEditor.h
--- a/Engine/Source/Editor/UnrealEd/Public/SCurveEditor.h
+++ b/Engine/Source/Editor/UnrealEd/Public/SCurveEditor.h
@@ -90,8 +90,8 @@ public:
 	{
 		CurveInfo = InCurveInfo;
 		Color = InColor;
+		bIsVisible = true;
 		bIsLocked = bInIsLocked;
-		bIsLocked = false;
 	}
 };
 
