Added process async results to the cook on the fly commandlet.

diff --git a/Engine/Source/Editor/UnrealEd/Private/Commandlets/CookCommandlet.cpp b/Engine/Source/Editor/UnrealEd/Private/Commandlets/CookCommandlet.cpp
--- a/Engine/Source/Editor/UnrealEd/Private/Commandlets/CookCommandlet.cpp
+++ b/Engine/Source/Editor/UnrealEd/Private/Commandlets/CookCommandlet.cpp
@@ -174,6 +174,10 @@ bool UCookCommandlet::CookOnTheFly( FGuid InstanceId, int32 Timeout, bool bForce
 				}
 			}
 
+
+			// Shaders need to be updated
+			GShaderCompilingManager->ProcessAsyncResults(true, false);
+
 			ProcessDeferredCommands();
 
 			// handle server timeout
