Editor: Ensure the asset type color strip in the content browser is always at least 3 units tall to preserve utility at small thumbnail sizes

diff --git a/Engine/Source/Editor/UnrealEd/Private/AssetThumbnail.cpp b/Engine/Source/Editor/UnrealEd/Private/AssetThumbnail.cpp
--- a/Engine/Source/Editor/UnrealEd/Private/AssetThumbnail.cpp
+++ b/Engine/Source/Editor/UnrealEd/Private/AssetThumbnail.cpp
@@ -344,8 +344,8 @@ private:
 
 	FMargin GetAssetColorStripPadding() const
 	{
-		// The strip is 2.5% the height of the thumbnail
-		const float Height = FMath::CeilToFloat(WidthLastFrame*0.025f);
+		// The strip is 2.5% the height of the thumbnail, but at least 3 units tall
+		const float Height = FMath::Max(FMath::CeilToFloat(WidthLastFrame*0.025f), 3.0f);
 		return FMargin(0,Height,0,0);
 	}
 
