Properly map multiples of the same keybinding

diff --git a/Engine/Source/Editor/UnrealEd/Private/Factories/PackFactory.cpp b/Engine/Source/Editor/UnrealEd/Private/Factories/PackFactory.cpp
--- a/Engine/Source/Editor/UnrealEd/Private/Factories/PackFactory.cpp
+++ b/Engine/Source/Editor/UnrealEd/Private/Factories/PackFactory.cpp
@@ -145,12 +145,12 @@ namespace PackFactoryHelper
 		UInputSettings* InputSettingsCDO = GetMutableDefault<UInputSettings>();
 		bool bCheckedOut = false;
 
-		TArray<FInputActionKeyMapping> ActionMappingsToAdd;
-		TArray<FInputAxisKeyMapping> AxisMappingsToAdd;
-
 		FConfigSection* InputSettingsSection = PackConfig.Find("InputSettings");
 		if (InputSettingsSection)
 		{
+			TArray<FInputActionKeyMapping> ActionMappingsToAdd;
+			TArray<FInputAxisKeyMapping> AxisMappingsToAdd;
+
 			for (auto SettingPair : *InputSettingsSection)
 			{
 				struct FMatchMappingByName
@@ -193,32 +193,32 @@ namespace PackFactoryHelper
 						AxisMappingsToAdd.Add(AxisKeyMapping);
 					}
 				}
+			}
 
-				if (ActionMappingsToAdd.Num() > 0 || AxisMappingsToAdd.Num() > 0)
+			if (ActionMappingsToAdd.Num() > 0 || AxisMappingsToAdd.Num() > 0)
+			{
+				if (ISourceControlModule::Get().IsEnabled())
 				{
-					if (ISourceControlModule::Get().IsEnabled())
-					{
-						FText ErrorMessage;
-
-						const FString InputSettingsFilename = FPaths::ConvertRelativePathToFull(InputSettingsCDO->GetDefaultConfigFilename());
-						if (!SourceControlHelpers::CheckoutOrMarkForAdd(InputSettingsFilename, FText::FromString(InputSettingsFilename), NULL, ErrorMessage))
-						{
-							UE_LOG(LogPackFactory, Error, TEXT("%s"), *ErrorMessage.ToString());
-						}
-					}
+					FText ErrorMessage;
 
-					for (const FInputActionKeyMapping& ActionKeyMapping : ActionMappingsToAdd)
+					const FString InputSettingsFilename = FPaths::ConvertRelativePathToFull(InputSettingsCDO->GetDefaultConfigFilename());
+					if (!SourceControlHelpers::CheckoutOrMarkForAdd(InputSettingsFilename, FText::FromString(InputSettingsFilename), NULL, ErrorMessage))
 					{
-						InputSettingsCDO->ActionMappings.Add(ActionKeyMapping);
-					}
-					for (const FInputAxisKeyMapping& AxisKeyMapping : AxisMappingsToAdd)
-					{
-						InputSettingsCDO->AxisMappings.Add(AxisKeyMapping);
+						UE_LOG(LogPackFactory, Error, TEXT("%s"), *ErrorMessage.ToString());
 					}
-					
-					InputSettingsCDO->SaveKeyMappings();
-					InputSettingsCDO->UpdateDefaultConfigFile();
 				}
+
+				for (const FInputActionKeyMapping& ActionKeyMapping : ActionMappingsToAdd)
+				{
+					InputSettingsCDO->AddActionMapping(ActionKeyMapping);
+				}
+				for (const FInputAxisKeyMapping& AxisKeyMapping : AxisMappingsToAdd)
+				{
+					InputSettingsCDO->AddxisMapping(AxisKeyMapping);
+				}
+					
+				InputSettingsCDO->SaveKeyMappings();
+				InputSettingsCDO->UpdateDefaultConfigFile();
 			}
 		}
 
