Fix editor compile error from typo.

diff --git a/Engine/Source/Editor/UnrealEd/Private/Factories/PackFactory.cpp b/Engine/Source/Editor/UnrealEd/Private/Factories/PackFactory.cpp
--- a/Engine/Source/Editor/UnrealEd/Private/Factories/PackFactory.cpp
+++ b/Engine/Source/Editor/UnrealEd/Private/Factories/PackFactory.cpp
@@ -214,7 +214,7 @@ namespace PackFactoryHelper
 				}
 				for (const FInputAxisKeyMapping& AxisKeyMapping : AxisMappingsToAdd)
 				{
-					InputSettingsCDO->AddxisMapping(AxisKeyMapping);
+					InputSettingsCDO->AddAxisMapping(AxisKeyMapping);
 				}
 					
 				InputSettingsCDO->SaveKeyMappings();
