Made cook in editor / cook on the side wait until packages is requested again / launch on is restarted before fixing up modified packages.  Stops editor from bogging down while modifying cooked packages.

diff --git a/Engine/Source/Editor/UnrealEd/Private/CookOnTheFlyServer.cpp b/Engine/Source/Editor/UnrealEd/Private/CookOnTheFlyServer.cpp
--- a/Engine/Source/Editor/UnrealEd/Private/CookOnTheFlyServer.cpp
+++ b/Engine/Source/Editor/UnrealEd/Private/CookOnTheFlyServer.cpp
@@ -627,16 +627,15 @@ void UCookOnTheFlyServer::GetDependencies( const TSet<UPackage*>& Packages, TSet
 	// Iterate through the object list
 	for( FObjectIterator It; It; ++It )
 	{
-		UPackage* Package = It->GetOutermost();
-		if ( Packages.Find( Package ) )
+		if (*It)
 		{
-			RootSet.Add(*It);
-			Found.Add(*It);
+			UPackage* Package = It->GetOutermost();
+			if ( Packages.Find( Package ) )
+			{
+				RootSet.Add(*It);
+				Found.Add(*It);
+			}
 		}
-
-		/*if (It->IsIn(Package))
-		{
-		}*/
 	}
 
 	TArray<UObject*> Exclude;
@@ -1382,11 +1381,15 @@ void UCookOnTheFlyServer::MarkPackageDirtyForCooker( UPackage *Package )
 		if ( CurrentCookMode == ECookMode::CookByTheBookFromTheEditor )
 		{
 			TArray<FName> CookedPlatforms;
+			// if we have already cooked this package and we have made changes then recook ;)
 			if ( CookedPackages.GetCookedPlatforms(PackageFFileName, CookedPlatforms) )
 			{
-				// if this package was previously cooked and we are doing a cook by the book 
-				// we need to recook this package before finishing cook by the book
-				CookRequests.EnqueueUnique( FFilePlatformRequest(PackageFFileName, CookedPlatforms) );
+				if ( IsCookByTheBookRunning() )
+				{
+					// if this package was previously cooked and we are doing a cook by the book 
+					// we need to recook this package before finishing cook by the book
+					CookRequests.EnqueueUnique( FFilePlatformRequest(PackageFFileName, CookedPlatforms) );	
+				}
 			}
 		}
 		CookedPackages.RemoveAll( PackageFFileName );
