[UE-5068] Extending the new command-line asset diff command to handle Blueprint merging.

diff --git a/Engine/Source/Developer/AssetTools/Private/AssetTypeActions/AssetTypeActions_Blueprint.cpp b/Engine/Source/Developer/AssetTools/Private/AssetTypeActions/AssetTypeActions_Blueprint.cpp
--- a/Engine/Source/Developer/AssetTools/Private/AssetTypeActions/AssetTypeActions_Blueprint.cpp
+++ b/Engine/Source/Developer/AssetTools/Private/AssetTypeActions/AssetTypeActions_Blueprint.cpp
@@ -101,6 +101,21 @@ void FAssetTypeActions_Blueprint::Merge(UObject* InObject)
 	}
 }
 
+void FAssetTypeActions_Blueprint::Merge(UObject* BaseAsset, UObject* RemoteAsset, UObject* LocalAsset)
+{
+	UBlueprint* AsBlueprint = CastChecked<UBlueprint>(LocalAsset);
+	check(LocalAsset->GetClass() == BaseAsset->GetClass());
+	check(LocalAsset->GetClass() == RemoteAsset->GetClass());
+	
+	if (FAssetEditorManager::Get().OpenEditorForAsset(AsBlueprint))
+	{
+		FBlueprintEditorModule& BlueprintEditorModule = FModuleManager::LoadModuleChecked<FBlueprintEditorModule>("Kismet");
+
+		FBlueprintEditor* BlueprintEditor = static_cast<FBlueprintEditor*>(FAssetEditorManager::Get().FindEditorForAsset(AsBlueprint, /*bFocusIfOpen =*/false));
+		BlueprintEditor->CreateMergeToolTab(Cast<UBlueprint>(BaseAsset), Cast<UBlueprint>(RemoteAsset));
+	}
+}
+
 bool FAssetTypeActions_Blueprint::CanCreateNewDerivedBlueprint() const
 {
 	return true;
diff --git a/Engine/Source/Developer/AssetTools/Private/AssetTypeActions/AssetTypeActions_Blueprint.h b/Engine/Source/Developer/AssetTools/Private/AssetTypeActions/AssetTypeActions_Blueprint.h
--- a/Engine/Source/Developer/AssetTools/Private/AssetTypeActions/AssetTypeActions_Blueprint.h
+++ b/Engine/Source/Developer/AssetTools/Private/AssetTypeActions/AssetTypeActions_Blueprint.h
@@ -16,6 +16,7 @@ public:
 	virtual void OpenAssetEditor(const TArray<UObject*>& InObjects, TSharedPtr<class IToolkitHost> EditWithinLevelEditor = TSharedPtr<IToolkitHost>()) override;
 	virtual bool CanMerge() const override;
 	virtual void Merge(UObject* InObject) override;
+	virtual void Merge(UObject* BaseAsset, UObject* RemoteAsset, UObject* LocalAsset);
 	virtual uint32 GetCategories() override { return EAssetTypeCategories::Blueprint | EAssetTypeCategories::Basic; }
 	virtual void PerformAssetDiff(UObject* Asset1, UObject* Asset2, const struct FRevisionInfo& OldRevision, const struct FRevisionInfo& NewRevision) const override;
 	virtual class UThumbnailInfo* GetThumbnailInfo(UObject* Asset) const override;
diff --git a/Engine/Source/Developer/AssetTools/Public/AssetTypeActions_Base.h b/Engine/Source/Developer/AssetTools/Public/AssetTypeActions_Base.h
--- a/Engine/Source/Developer/AssetTools/Public/AssetTypeActions_Base.h
+++ b/Engine/Source/Developer/AssetTools/Public/AssetTypeActions_Base.h
@@ -56,6 +56,11 @@ public:
 		check(false); // no generic merge operation exists yet, did you override CanMerge but not Merge?
 	}
 
+	virtual void Merge(UObject* BaseAsset, UObject* RemoteAsset, UObject* LocalAsset) override
+	{
+		check(false); // no generic merge operation exists yet, did you override CanMerge but not Merge?
+	}
+
 	virtual bool ShouldForceWorldCentric() override
 	{
 		return false;
diff --git a/Engine/Source/Developer/AssetTools/Public/IAssetTypeActions.h b/Engine/Source/Developer/AssetTools/Public/IAssetTypeActions.h
--- a/Engine/Source/Developer/AssetTools/Public/IAssetTypeActions.h
+++ b/Engine/Source/Developer/AssetTools/Public/IAssetTypeActions.h
@@ -80,9 +80,12 @@ public:
 	/** Returns true if this class can be merged (either manually or automatically) */
 	virtual bool CanMerge() const = 0;
 
-	/** Begins a merge operation for InObjects */
+	/** Begins a merge operation for InObject (automatically determines remote/base versions needed to resolve) */
 	virtual void Merge( UObject* InObject ) = 0;
 
+	/** Begins a merge between the specified assets */
+	virtual void Merge(UObject* BaseAsset, UObject* RemoteAsset, UObject* LocalAsset) = 0;
+
 	/** Returns the categories that this asset type. The return value is one or more flags from EAssetTypeCategories.  */
 	virtual uint32 GetCategories() = 0;
 
diff --git a/Engine/Source/Developer/Merge/Private/Merge.cpp b/Engine/Source/Developer/Merge/Private/Merge.cpp
--- a/Engine/Source/Developer/Merge/Private/Merge.cpp
+++ b/Engine/Source/Developer/Merge/Private/Merge.cpp
@@ -115,12 +115,65 @@ static UObject* LoadBaseRev(const FString& PackageName, const FString& AssetName
 	return LoadRevision(AssetName, *Revision);
 }
 
+static TSharedPtr<SWidget> GenerateMergeTabContents(TSharedRef<FBlueprintEditor> Editor,
+                                                    const UBlueprint*    BaseBlueprint,
+                                                    const FRevisionInfo& BaseRevInfo,
+													const UBlueprint*    RemoteBlueprint,
+													const FRevisionInfo& RemoteRevInfo,
+													const UBlueprint&    LocalBlueprint)
+{
+	TSharedPtr<SWidget> TabContent;
+
+	if (RemoteBlueprint && BaseBlueprint)
+	{
+		LocalBlueprint.bDuplicatingReadOnly = true;
+		FBlueprintMergeData Data(Editor
+			, static_cast<const UBlueprint*>(StaticDuplicateObject(&LocalBlueprint, GetTransientPackage(), TEXT("None")))
+			, BaseBlueprint
+			, RemoteRevInfo
+			, RemoteBlueprint
+			, BaseRevInfo);
+		LocalBlueprint.bDuplicatingReadOnly = false;
+
+		TabContent = SNew(SBlueprintMerge, Data);
+	}
+	else
+	{
+		FText MissingFiles;
+		if (!RemoteBlueprint && !BaseBlueprint)
+		{
+			MissingFiles = LOCTEXT("MergeBothRevisionsFailed", "common base, nor conflicting revision");
+		}
+		else if (!RemoteBlueprint)
+		{
+			MissingFiles = LOCTEXT("MergeConflictingRevisionsFailed", "conflicting revision");
+		}
+		else
+		{
+			MissingFiles = LOCTEXT("MergeBaseRevisionsFailed", "common base");
+		}
+
+		DisplayErrorMessage(
+			FText::Format(
+				LOCTEXT("MergeRevisionLoadFailed", "Aborted Merge of {0} because we could not load {1}")
+				, FText::FromString(LocalBlueprint.GetName())
+				, MissingFiles
+			)
+		);
+
+		TabContent = SNew(SHorizontalBox);
+	}
+
+	return TabContent;
+}
+
 class FMerge : public IMerge
 {
 	/** IModuleInterface implementation */
 	virtual void StartupModule() override;
 	virtual void ShutdownModule() override;
 	virtual TSharedRef<SDockTab> GenerateMergeWidget(const UBlueprint& Object, TSharedRef<FBlueprintEditor> Editor) override;
+	virtual TSharedRef<SDockTab> GenerateMergeWidget(const UBlueprint* BaseAsset, const UBlueprint* RemoteAsset, const UBlueprint* LocalAsset, TSharedRef<FBlueprintEditor> Editor) override;
 	virtual bool PendingMerge(const UBlueprint& BlueprintObj) const override;
 
 	// Simplest to only allow one merge operation at a time, we could easily make this a map of Blueprint=>MergeTab
@@ -190,45 +243,7 @@ TSharedRef<SDockTab> FMerge::GenerateMergeWidget(const UBlueprint& Object, TShar
 		FRevisionInfo BaseRevInfo = FRevisionInfo::InvalidRevision();
 		const UBlueprint* BaseBlueprint = Cast< UBlueprint >(LoadBaseRev(PackageName, AssetName, SourceControlStateRef, BaseRevInfo));
 
-		if (RemoteBlueprint && BaseBlueprint)
-		{
-			Object.bDuplicatingReadOnly = true;
-			FBlueprintMergeData Data(Editor 
-									, static_cast<const UBlueprint*>(StaticDuplicateObject(&Object, GetTransientPackage(), TEXT("None")))
-									, BaseBlueprint
-									, CurrentRevInfo
-									, RemoteBlueprint
-									, BaseRevInfo);
-			Object.bDuplicatingReadOnly = false;
-
-			Contents = SNew(SBlueprintMerge, Data);
-		}
-		else
-		{
-			FText MissingFiles;
-			if (!RemoteBlueprint && !BaseBlueprint)
-			{
-				MissingFiles = LOCTEXT("MergeBothRevisionsFailed", "common base, nor conflicting revision");
-			}
-			else if (!RemoteBlueprint)
-			{
-				MissingFiles = LOCTEXT("MergeConflictingRevisionsFailed", "conflicting revision");
-			}
-			else
-			{
-				MissingFiles = LOCTEXT("MergeBaseRevisionsFailed", "common base");
-			}
-
-			DisplayErrorMessage(
-				FText::Format(
-					LOCTEXT("MergeRevisionLoadFailed", "Aborted Merge of {0} because we could not load {1}")
-					, FText::FromString(Object.GetName())
-					, MissingFiles
-				)
-			);
-
-			Contents = SNew(SHorizontalBox);
-		}
+		Contents = GenerateMergeTabContents(Editor, BaseBlueprint, BaseRevInfo, RemoteBlueprint, CurrentRevInfo, Object);
 	}
 
 	TSharedRef<SDockTab> Tab =  FGlobalTabmanager::Get()->InvokeTab(MergeToolTabId);
@@ -238,6 +253,26 @@ TSharedRef<SDockTab> FMerge::GenerateMergeWidget(const UBlueprint& Object, TShar
 
 }
 
+TSharedRef<SDockTab> FMerge::GenerateMergeWidget(const UBlueprint* BaseBlueprint, const UBlueprint* RemoteBlueprint, const UBlueprint* LocalBlueprint, TSharedRef<FBlueprintEditor> Editor)
+{
+	if (ActiveTab.IsValid())
+	{
+		TSharedPtr<SDockTab> ActiveTabPtr = ActiveTab.Pin();
+		// just bring the tab to the foreground:
+		TSharedRef<SDockTab> CurrentTab = FGlobalTabmanager::Get()->InvokeTab(MergeToolTabId);
+		check(CurrentTab == ActiveTabPtr);
+		return ActiveTabPtr.ToSharedRef();
+	}
+
+	// @TODO: pipe revision info through
+	TSharedPtr<SWidget> TabContents = GenerateMergeTabContents(Editor, BaseBlueprint, FRevisionInfo::InvalidRevision(), RemoteBlueprint, FRevisionInfo::InvalidRevision(), *LocalBlueprint);
+
+	TSharedRef<SDockTab> Tab = FGlobalTabmanager::Get()->InvokeTab(MergeToolTabId);
+	Tab->SetContent(TabContents.ToSharedRef());
+	ActiveTab = Tab;
+	return Tab;
+}
+
 bool FMerge::PendingMerge(const UBlueprint& BlueprintObj) const
 {
 	bool bIsMergeEnabled = GetDefault<UEditorExperimentalSettings>()->bEnableBlueprintMergeTool;
diff --git a/Engine/Source/Developer/Merge/Public/Merge.h b/Engine/Source/Developer/Merge/Public/Merge.h
--- a/Engine/Source/Developer/Merge/Public/Merge.h
+++ b/Engine/Source/Developer/Merge/Public/Merge.h
@@ -42,6 +42,7 @@ public:
 	 * @return The merge view widget
 	 */
 	virtual TSharedRef<class SDockTab> GenerateMergeWidget( const UBlueprint& Object, TSharedRef< class FBlueprintEditor> Editor ) = 0;
+	virtual TSharedRef<class SDockTab> GenerateMergeWidget( const UBlueprint* BaseBlueprint, const UBlueprint* RemoteBlueprint, const UBlueprint* LocalBlueprint, TSharedRef<class FBlueprintEditor> Editor ) = 0;
 
 	/** 
 	 * @return whether the blueprint is in a conflicted state, and can therefore be merged.
diff --git a/Engine/Source/Editor/Kismet/Private/BlueprintEditor.cpp b/Engine/Source/Editor/Kismet/Private/BlueprintEditor.cpp
--- a/Engine/Source/Editor/Kismet/Private/BlueprintEditor.cpp
+++ b/Engine/Source/Editor/Kismet/Private/BlueprintEditor.cpp
@@ -5048,6 +5048,11 @@ void FBlueprintEditor::CreateMergeToolTab()
 	MergeTool = IMerge::Get().GenerateMergeWidget(*GetBlueprintObj(), SharedThis(this));
 }
 
+void FBlueprintEditor::CreateMergeToolTab(const UBlueprint* BaseBlueprint, const UBlueprint* RemoteBlueprint)
+{
+	MergeTool = IMerge::Get().GenerateMergeWidget(BaseBlueprint, RemoteBlueprint, GetBlueprintObj(), SharedThis(this));
+}
+
 void FBlueprintEditor::CloseMergeTool()
 {
 	auto MergeToolPtr = MergeTool.Pin();
diff --git a/Engine/Source/Editor/Kismet/Public/BlueprintEditor.h b/Engine/Source/Editor/Kismet/Public/BlueprintEditor.h
--- a/Engine/Source/Editor/Kismet/Public/BlueprintEditor.h
+++ b/Engine/Source/Editor/Kismet/Public/BlueprintEditor.h
@@ -468,6 +468,7 @@ public:
 
 	/** Forces the merge tool to be shown */
 	void CreateMergeToolTab();
+	void CreateMergeToolTab(const UBlueprint* BaseBlueprint, const UBlueprint* RemoteBlueprint);
 	
 	/** Closes the merge tool, rather than simply hiding it */
 	void CloseMergeTool();
diff --git a/Engine/Source/Editor/UnrealEd/Private/EditorCommandLineUtils.cpp b/Engine/Source/Editor/UnrealEd/Private/EditorCommandLineUtils.cpp
--- a/Engine/Source/Editor/UnrealEd/Private/EditorCommandLineUtils.cpp
+++ b/Engine/Source/Editor/UnrealEd/Private/EditorCommandLineUtils.cpp
@@ -11,9 +11,14 @@
 #include "AssetRegistryModule.h"
 #include "IAssetTypeActions.h"		// for FRevisionInfo
 #include "AssetToolsModule.h"
+#include "AssetEditorManager.h"
+#include "Editor.h"					// for OnShutdownPostPackagesSaved
 
 #define LOCTEXT_NAMESPACE "EditorCommandLineUtils"
 
+// Forward Declarations
+struct FMergeAsset;
+
 /*******************************************************************************
  * EditorCommandLineUtilsImpl Declaration
  ******************************************************************************/
@@ -28,7 +33,7 @@ namespace EditorCommandLineUtilsImpl
 	static const FText  DiffCommandHelpTxt = LOCTEXT("DiffCommandeHelpText", "\
 Usage: \n\
     -diff [options] left right                                                 \n\
-    -diff [options] left right base                                            \n\
+    -diff [options] remote local base result                                   \n\
 \n\
 Options: \n\
     -echo               Prints back the command arguments and then exits.      \n\
@@ -37,6 +42,9 @@ Options: \n\
 	/**  */
 	static bool ParseCommandArgs(const TCHAR* FullEditorCmdLine, const TCHAR* CmdSwitch, FString& CmdArgsOut);
 
+	/** */
+	static FString FindProjectFile(const FString& AssetFilePath);
+
 	/**  */
 	static void RaiseEditorMessageBox(const FText& Title, const FText& BodyText, const bool bExitOnClose);
 
@@ -47,13 +55,10 @@ Options: \n\
 	static void RunAssetDiffCommand(TSharedPtr<SWindow> MainEditorWindow, bool bIsRunningProjBrowser, FString CommandArgs);
 
 	/**  */
-	static int32 ParseRevisionFromFilename(const FString& Filename, FString& CLeanFilenameOut);
-
-	/** */
-	static UObject* ExtractAssetFromPackage(UPackage* Package, const FString& AssetName);
+	static void RunAssetMerge(FMergeAsset const& Base, FMergeAsset const& Remote, FMergeAsset const& Local, FMergeAsset const& Result);
 
 	/**  */
-//	static int32 GetAssetPackagePaths(const FString& InAssetName, TArray<FString>& PackagePathsOut, bool bIncludeDevFolder);
+	static UObject* ExtractAssetFromPackage(UPackage* Package);
 }
 
 /*******************************************************************************
@@ -132,6 +137,9 @@ public:
 	virtual void Tick(float DeltaTime) override;	
 	// End FTickableEditorObject interface
 
+	/**  */
+	void Disable();
+
 private:
 	TWeakPtr<SWindow> MainEditorWindow;
 };
@@ -180,6 +188,157 @@ void FFauxStandaloneToolManager::Tick(float DeltaTime)
 	}
 }
 
+//------------------------------------------------------------------------------
+void FFauxStandaloneToolManager::Disable()
+{
+	if (MainEditorWindow.IsValid())
+	{
+		MainEditorWindow.Pin()->ShowWindow();
+	}
+}
+
+/*******************************************************************************
+ * FMergeAsset
+ ******************************************************************************/
+ 
+struct FMergeAsset
+{
+public:
+	FMergeAsset(TCHAR* DstFileName);
+
+	/**  */
+	bool SetSourceFile(const FString& SrcFilePathIn, FCommandLineErrorReporter& ErrorReporter);
+
+	/**  */
+	bool Load(FCommandLineErrorReporter& ErrorReporter);
+
+	/**  */
+	UClass* GetClass() const;
+
+	/**  */
+	UObject* GetAssetObj() const;
+
+	/**  */
+	FRevisionInfo GetRevisionInfo() const;
+
+	/**  */
+	const FString& GetSourceFilePath() const;
+
+	/**  */
+	const FString& GetAssetFilePath() const;
+
+private:
+	UPackage* Package;
+	UObject*  AssetObj;
+	FString   DestFilePath;
+	FString   SrcFilePath;
+};
+
+
+//------------------------------------------------------------------------------
+FMergeAsset::FMergeAsset(TCHAR* DstFileName)
+	: Package(nullptr)
+	, AssetObj(nullptr)
+	, DestFilePath(FPaths::Combine(*FPaths::DiffDir(), DstFileName))
+{
+	const FString& AssetExt = FPackageName::GetAssetPackageExtension();
+	if (!DestFilePath.EndsWith(AssetExt))
+	{
+		DestFilePath += AssetExt;
+	}
+}
+
+//------------------------------------------------------------------------------
+bool FMergeAsset::SetSourceFile(const FString& SrcFilePathIn, FCommandLineErrorReporter& ErrorReporter)
+{
+	SrcFilePath.Empty();
+	if (!FPaths::FileExists(SrcFilePathIn))
+	{
+		ErrorReporter.ReportFatalError(LOCTEXT("BadFilePathTitle", "Bad File Path"),
+			FText::Format(LOCTEXT("BadFilePathError", "'{0}' is an invalid file."), FText::FromString(SrcFilePathIn)));
+	}
+	else
+	{
+		SrcFilePath = SrcFilePathIn;
+	}
+	return !SrcFilePath.IsEmpty();
+}
+
+//------------------------------------------------------------------------------
+bool FMergeAsset::Load(FCommandLineErrorReporter& ErrorReporter)
+{
+	if (SrcFilePath.IsEmpty())
+	{
+		// was SetSourceFile() called prior to this?
+		return false;
+	}
+
+	// UE4 cannot open files with certain special characters in them (like 
+	// the # symbol), so we make a copy of the file with a more UE digestible 
+	// path (since this may be a perforce temp file)
+	if (IFileManager::Get().Copy(*DestFilePath, *SrcFilePath) != COPY_OK)
+	{
+		ErrorReporter.ReportFatalError(LOCTEXT("LoadFailedTitle", "Unable to Copy File"),
+			FText::Format(LOCTEXT("LoadFailedError", "Failed to make a local copy of the asset file: '{0}'."), FText::FromString(SrcFilePath)));
+	}
+	else if (UPackage* AssetPkg = LoadPackage(/*Outer =*/nullptr, *DestFilePath, LOAD_None))
+	{
+		if (UObject* ExtractedAsset = EditorCommandLineUtilsImpl::ExtractAssetFromPackage(AssetPkg))
+		{
+			Package  = AssetPkg;
+			AssetObj = ExtractedAsset;
+		}
+		else
+		{
+			ErrorReporter.ReportFatalError(LOCTEXT("AssetNotFoundTitle", "Asset Not Found"),
+				FText::Format(LOCTEXT("AssetNotFoundError", "Failed to find the asset object inside the package file: '{0}'."), FText::FromString(SrcFilePath)));
+		}
+	}
+
+	return (AssetObj != nullptr);
+}
+
+//------------------------------------------------------------------------------
+UClass* FMergeAsset::GetClass() const
+{
+	return (AssetObj != nullptr) ? AssetObj->GetClass() : nullptr;
+}
+
+//------------------------------------------------------------------------------
+UObject* FMergeAsset::GetAssetObj() const
+{
+	return AssetObj;
+}
+
+//------------------------------------------------------------------------------
+FRevisionInfo FMergeAsset::GetRevisionInfo() const
+{
+	FString SrcFileName = FPaths::GetBaseFilename(SrcFilePath);
+
+	FRevisionInfo RevisionInfoOut = FRevisionInfo::InvalidRevision();
+
+	FString BaseFileName, RevisionStr;
+	if (SrcFileName.Split(TEXT("#"), &BaseFileName, &RevisionStr) && RevisionStr.IsNumeric())
+	{
+		// @TODO: if connected to source-control, extract changelist and date info
+		RevisionInfoOut.Revision = FCString::Atoi(*RevisionStr);
+	}
+
+	return RevisionInfoOut;
+}
+
+//------------------------------------------------------------------------------
+const FString& FMergeAsset::GetSourceFilePath() const
+{
+	return SrcFilePath;
+}
+
+//------------------------------------------------------------------------------
+const FString& FMergeAsset::GetAssetFilePath() const
+{
+	return DestFilePath;
+}
+
 /*******************************************************************************
  * EditorCommandLineUtilsImpl Implementation
  ******************************************************************************/
@@ -196,6 +355,58 @@ static bool EditorCommandLineUtilsImpl::ParseCommandArgs(const TCHAR* FullEditor
 	return false;
 }
 
+//------------------------------------------------------------------------------
+static FString EditorCommandLineUtilsImpl::FindProjectFile(const FString& AssetFilePathIn)
+{
+	FString FoundProjectPath;
+
+	FString AssetFilePath = AssetFilePathIn;
+	FPaths::NormalizeFilename(AssetFilePath);	
+
+	const TCHAR* const ContentDirName = TEXT("/Content/");
+
+	FString ProjectDir, AssetSubPath;
+	if (AssetFilePath.Split(ContentDirName, &ProjectDir, &AssetSubPath))
+	{
+		const FString UProjExt = TEXT(".") + FProjectDescriptor::GetExtension();
+		const FString ProjectWildcardPath = FPaths::Combine(*ProjectDir, *FString(TEXT("*") + UProjExt));
+
+		TArray<FString> FoundFiles;
+		IFileManager::Get().FindFiles(FoundFiles, *ProjectWildcardPath, /*Files =*/true, /*Directories =*/false);
+
+		if (FoundFiles.Num() > 0)
+		{
+			FoundProjectPath = FPaths::Combine(*ProjectDir, *FoundFiles[0]);
+
+			const FString DirName = FPaths::GetBaseFilename(ProjectDir);
+			for (FString FileName : FoundFiles)
+			{
+				// favor project files that match the directory name
+				if (FPaths::GetBaseFilename(FileName) == DirName)
+				{
+					FoundProjectPath = FPaths::Combine(*ProjectDir, *FileName);
+					break;
+				}
+			}
+		}
+		else 
+		{
+			// guess at what the project path would be (in case this is a  
+			// perforce temp file, and its path mimics the real asset file's 
+			// directory structure)
+			FString GameName = FPaths::GetCleanFilename(ProjectDir);
+			FoundProjectPath = FPaths::Combine(*FPaths::RootDir(), *GameName, *FString(GameName + UProjExt));
+			// make sure what we're guessing at exists...
+			if (!FPaths::FileExists(FoundProjectPath))
+			{
+				FoundProjectPath.Empty();
+			}
+		}
+	}
+
+	return FoundProjectPath;
+}
+
 //------------------------------------------------------------------------------
 static void EditorCommandLineUtilsImpl::RaiseEditorMessageBox(const FText& Title, const FText& BodyText, const bool bExitOnClose)
 {
@@ -229,6 +440,7 @@ static void EditorCommandLineUtilsImpl::RunAssetDiffCommand(TSharedPtr<SWindow>
 	// select a project (and then the editor will re-launch with this command).
 	if (bIsRunningProjBrowser) 
 	{
+		// @TODO: can we run without loading a project?
 		return;
 	}
 
@@ -245,158 +457,189 @@ static void EditorCommandLineUtilsImpl::RunAssetDiffCommand(TSharedPtr<SWindow>
 		Switches.Contains("?") ||
 		Switches.Contains("help"))
 	{
-		RaiseEditorMessageBox(LOCTEXT("DiffCommandHelp", "UAsset Diff Command-Line Help"), DiffCommandHelpTxt, /*bExitOnClose =*/true);
+		RaiseEditorMessageBox(LOCTEXT("DiffCommandHelp", "Diff/Merge Command-Line Help"), DiffCommandHelpTxt, /*bExitOnClose =*/true);
 		return;
 	}
 
 	if (Switches.Contains("echo"))
 	{
-		RaiseEditorMessageBox(LOCTEXT("DiffCommandHelp", "Passed Command Arguments"), FText::FromString(CommandArgs), /*bExitOnClose =*/true);
+		RaiseEditorMessageBox(LOCTEXT("DiffCommandHelp", "Passed Command Arguments"), 
+			FText::FromString(CommandArgs), /*bExitOnClose =*/true);
 		return;
 	}
 
+	const int32 FilesNeededForDiff  = 2;
+	const int32 FilesNeededForMerge = 4;
+	const int32 MaxFilesNeeded = FilesNeededForMerge;
+
+	FMergeAsset MergeAssets[MaxFilesNeeded] = {
+		FMergeAsset(TEXT("MergeTool-Left")),
+		FMergeAsset(TEXT("MergeTool-Right")),
+		FMergeAsset(TEXT("MergeTool-Base")),
+		FMergeAsset(TEXT("MergeTool-Merge")),
+	};
+	FMergeAsset& LeftAsset   = MergeAssets[0];
+	FMergeAsset& ThierAsset  = LeftAsset;
+	FMergeAsset& RightAsset  = MergeAssets[1];
+	FMergeAsset& OurAsset    = RightAsset;
+	FMergeAsset& BaseAsset   = MergeAssets[2];
+	FMergeAsset& MergeResult = MergeAssets[3];
+
+	//--------------------------------------
+	// Parse file paths from command-line
+	//--------------------------------------
+
 	FCommandLineErrorReporter ErrorReporter(DiffCommandSwitch, CommandArgs);
-	if (Tokens.Num() < 2)
-	{
-		ErrorReporter.ReportFatalError(LOCTEXT("TooFewParamsTitle", "Too Few Parameters"),
-			LOCTEXT("TooFewParamsError", "At least two files are needed."));
-	}
 
-	UObject* LeftAsset = nullptr, *RightAsset = nullptr, *BaseAsset = nullptr;
-	const int32 MaxFileCount = 3;
-	UObject** FileAssets[MaxFileCount] = { &LeftAsset, &RightAsset, &BaseAsset };
 	int32 ParsedFileCount = 0;
-	int32 AssetRevisionNums[MaxFileCount] = { -1, -1, -1 };
-
-	for (int32 FileIndex = 0; FileIndex < Tokens.Num() && ParsedFileCount < MaxFileCount; ++FileIndex)
+	for (int32 FileIndex = 0; FileIndex < Tokens.Num() && ParsedFileCount < MaxFilesNeeded; ++FileIndex)
 	{
 		FString& FilePath = Tokens[FileIndex];
-		if (!FPaths::FileExists(FilePath))
+
+		FMergeAsset& MergeAsset = MergeAssets[ParsedFileCount];
+		if (MergeAsset.SetSourceFile(FilePath, ErrorReporter))
 		{
-			ErrorReporter.ReportFatalError(LOCTEXT("BadFilePathTitle", "Bad File Path"),
-				FText::Format(LOCTEXT("BadFilePathError", "'{0}' is an invalid file."), FText::FromString(FilePath)));
-			continue;
+			++ParsedFileCount;
 		}
+	}
 
-		FString CleanFilename;
-		int32 RevisionNum = ParseRevisionFromFilename(FilePath, CleanFilename);
-		FString LocalFilePath = FPaths::Combine(*FPaths::DiffDir(), *(TEXT("Temp-Rev-") + FString::FromInt(RevisionNum) + "-" + CleanFilename));
-		// UE4 cannot open files with a number in them
-		IFileManager::Get().Copy(*LocalFilePath, *FilePath);
+	//--------------------------------------
+	// Verify file count
+	//--------------------------------------
 
-		if (UPackage* TempAssetPkg = LoadPackage(/*Outer =*/nullptr, *LocalFilePath, LOAD_None))
+	const bool bWantsMerge = (ParsedFileCount > FilesNeededForDiff);
+	if (ParsedFileCount < FilesNeededForDiff)
+	{
+		ErrorReporter.ReportFatalError(LOCTEXT("TooFewParamsTitle", "Too Few Parameters"),
+			LOCTEXT("TooFewParamsError", "At least two files are needed (for a diff)."));
+	}
+	else if (bWantsMerge && (ParsedFileCount < FilesNeededForMerge))
+	{
+		ErrorReporter.ReportFatalError(LOCTEXT("TooFewParamsTitle", "Too Few Parameters"),
+			LOCTEXT("TooFewMergeParamsError", "To merge, at least two files are needed."));
+	}
+	else if (Tokens.Num() > FilesNeededForMerge)
+	{
+		ErrorReporter.ReportFatalError(LOCTEXT("TooManyParamsTitle", "Too Many Parameters"),
+			FText::Format( LOCTEXT("TooManyParamsError", "There were too many command arguments supplied. The maximum files needed are {0} (for merging)"), FText::AsNumber(FilesNeededForMerge) ));
+	}
+
+	//--------------------------------------
+	// Load diff/merge asset files
+	//--------------------------------------
+
+	bool bLoadSuccess = true;
+	if (bWantsMerge)
+	{
+		bLoadSuccess &= ThierAsset.Load(ErrorReporter);
+		bLoadSuccess &= OurAsset.Load(ErrorReporter);
+		bLoadSuccess &= BaseAsset.Load(ErrorReporter);
+	}
+	else
+	{
+		bLoadSuccess &= LeftAsset.Load(ErrorReporter);
+		bLoadSuccess &= RightAsset.Load(ErrorReporter);
+	}
+
+	//--------------------------------------
+	// Verify asset types
+	//--------------------------------------
+
+	IAssetTools& AssetTools = FModuleManager::GetModuleChecked<FAssetToolsModule>("AssetTools").Get();
+	if (bLoadSuccess)
+	{
+		if (LeftAsset.GetClass() != RightAsset.GetClass())
+		{
+			ErrorReporter.ReportFatalError(LOCTEXT("TypeMismatchTitle", "Asset Type Mismatch"),
+				LOCTEXT("TypeMismatchError", "Cannot compare files of different asset types."));
+		}
+		else if (bWantsMerge)
 		{
-			if (UObject* FileAsset = ExtractAssetFromPackage(TempAssetPkg, CleanFilename))
+			UClass* AssetClass = OurAsset.GetClass();
+			TWeakPtr<IAssetTypeActions> AssetActions = AssetTools.GetAssetTypeActionsForClass(AssetClass);
+
+			if (AssetClass != BaseAsset.GetClass())
 			{
-				*FileAssets[ParsedFileCount] = FileAsset;
-				AssetRevisionNums[ParsedFileCount] = RevisionNum;
-				++ParsedFileCount;
+				ErrorReporter.ReportFatalError(LOCTEXT("TypeMismatchTitle", "Asset Type Mismatch"),
+					LOCTEXT("MergeTypeMismatchError", "Cannot merge files of different asset types."));
 			}
-			else
+			else if(!AssetActions.IsValid() || !AssetActions.Pin()->CanMerge())
 			{
-				ErrorReporter.ReportFatalError(LOCTEXT("AssetNotFoundTitle", "Asset Not Found"),
-					FText::Format( LOCTEXT("AssetNotFoundError", "Failed to find the asset '{0}', inside the package file: '{1}'."), 
-					FText::FromString(FPaths::GetBaseFilename(CleanFilename)), FText::FromString(FilePath) ));
+				ErrorReporter.ReportFatalError(LOCTEXT("CannotMergeTitle", "Cannot Merge"),
+					FText::Format(LOCTEXT("CannotMergeError", "{0} asset files can not be merged."), FText::FromName(AssetClass->GetFName())));
 			}
 		}
-		else
-		{
-			ErrorReporter.ReportFatalError(LOCTEXT("LoadFailedTitle", "Unable to Open File"),
-				FText::Format(LOCTEXT("LoadFailedError", "'{0}' failed to load."), FText::FromString(FilePath)));
-		}
 	}
 
-	if (!ErrorReporter.HasBlockingError())
-	{
-		FRevisionInfo LeftRevisionInfo = FRevisionInfo::InvalidRevision();
-		LeftRevisionInfo.Revision  = AssetRevisionNums[0];
-		FRevisionInfo RightRevisionInfo = FRevisionInfo::InvalidRevision();
-		RightRevisionInfo.Revision = AssetRevisionNums[1];
+	//--------------------------------------
+	// Preform diff/merge
+	//--------------------------------------
 
-		FAssetToolsModule& AssetToolsModule = FModuleManager::GetModuleChecked<FAssetToolsModule>("AssetTools");
-		AssetToolsModule.Get().DiffAssets(LeftAsset, RightAsset, LeftRevisionInfo, RightRevisionInfo);
+	if (bLoadSuccess && !ErrorReporter.HasBlockingError())
+	{
+		if (bWantsMerge)
+		{
+			// unlike with diff'ing, for merging we rely on asset editors for
+			// merging, and those windows get childed to the main window (so it
+			// needs to be visible)
+			//
+			// @TODO: get it so asset editor windows can be shown standalone
+			FauxStandaloneToolManager.Disable();
+
+			RunAssetMerge(BaseAsset, ThierAsset, OurAsset, MergeResult);
+		}
+		else
+		{
+			AssetTools.DiffAssets(LeftAsset.GetAssetObj(), RightAsset.GetAssetObj(), LeftAsset.GetRevisionInfo(), RightAsset.GetRevisionInfo());
+		}
 	}
 }
 
 //------------------------------------------------------------------------------
-static int32 EditorCommandLineUtilsImpl::ParseRevisionFromFilename(const FString& Filename, FString& CleanedFilenameOut)
+static void EditorCommandLineUtilsImpl::RunAssetMerge(FMergeAsset const& Base, FMergeAsset const& Remote, FMergeAsset const& Local, FMergeAsset const& Result)
 {
-	int32 RevisionNum = -1;
+	IAssetTools& AssetTools = FModuleManager::GetModuleChecked<FAssetToolsModule>("AssetTools").Get();
 
-	CleanedFilenameOut = FPaths::GetBaseFilename(Filename);
+	UClass* AssetClass = Local.GetClass();
+	check(AssetClass != nullptr);
+	TWeakPtr<IAssetTypeActions> AssetActions = AssetTools.GetAssetTypeActionsForClass(AssetClass);
+	check(AssetActions.IsValid());
 
-	FString BaseFileName, RevisionStr;
-	CleanedFilenameOut.Split(TEXT("#"), &BaseFileName, &RevisionStr);
-
-	CleanedFilenameOut = BaseFileName + FPackageName::GetAssetPackageExtension();
+	// have to mount the save directory so that the BP-editor can save
+	// the merged asset packages
+	FPackageName::RegisterMountPoint(TEXT("/Temp/"), FPaths::GameSavedDir());
+	// bring up the merge tool...
+	AssetActions.Pin()->Merge(Base.GetAssetObj(), Remote.GetAssetObj(), Local.GetAssetObj());
 
-	if (RevisionStr.IsNumeric())
+	auto CopyFilesLambda = [](FString SrcFilePath, FString DstFilePath)
 	{
-		RevisionNum = FCString::Atoi(*RevisionStr);
-	}
-	return RevisionNum;
+		IFileManager::Get().Copy(*DstFilePath, *SrcFilePath);
+	};
+
+	const FString& ModifiedFilePath = Local.GetAssetFilePath();
+	const FString& ResultFilePath   = (!Result.GetSourceFilePath().IsEmpty()) ? Result.GetSourceFilePath() : Local.GetSourceFilePath();
+	// have to copy the file into the expected result file when we're done
+	FEditorDelegates::OnShutdownPostPackagesSaved.AddStatic(CopyFilesLambda, ModifiedFilePath, ResultFilePath);
 }
 
 //------------------------------------------------------------------------------
-static UObject* EditorCommandLineUtilsImpl::ExtractAssetFromPackage(UPackage* Package, const FString& AssetName)
+static UObject* EditorCommandLineUtilsImpl::ExtractAssetFromPackage(UPackage* Package)
 {
-	return FindObject<UObject>(Package, *FPaths::GetBaseFilename(AssetName));
-}
+	TArray<UObject*> ObjectsWithOuter;
+	GetObjectsWithOuter(Package, ObjectsWithOuter, /*bIncludeNestedObjects =*/false);
 
-//------------------------------------------------------------------------------
-// static int32 EditorCommandLineUtilsImpl::GetAssetPackagePaths(const FString& InAssetName, TArray<FString>& PackagePathsOut, bool bIncludeDevFolder)
-// {
-// 	PackagePathsOut.Empty();
-// 	if (InAssetName.IsEmpty())
-// 	{
-// 		return 0;
-// 	}
-// 
-// 	FString AssetPath = InAssetName;
-// 	FPaths::NormalizeFilename(AssetPath);
-// 	FString AssetName = FPaths::GetCleanFilename(AssetPath);
-// 	const bool bIsPlainAssetName = (AssetName == AssetPath);
-// 
-// 	const FString& AssetExt = FPackageName::GetAssetPackageExtension();
-// 	const bool bHasExtension = AssetName.EndsWith(AssetExt);
-// 	FString BaseAssetName = FPaths::GetBaseFilename(AssetName);
-// 
-// 	FString AssetWildcard = FString::Printf(TEXT("*%s%s"), *AssetName, bHasExtension ? TEXT("") : *(TEXT("*") + AssetExt));
-// 	const uint8 PackageFilter = NORMALIZE_ExcludeMapPackages | (bIncludeDevFolder ? 0x00 : NORMALIZE_ExcludeDeveloperPackages);
-// 	TArray<FString> RelativeFilePaths;
-// 	NormalizePackageNames(TArray<FString>(), RelativeFilePaths, AssetName, PackageFilter);
-// 
-// 	bool bHasExactMatch = false;
-// 	for (FString& RelativePath : RelativeFilePaths)
-// 	{
-// 		FString FullFilePath = FPaths::ConvertRelativePathToFull(RelativePath);
-// 		FPaths::NormalizeFilename(FullFilePath);
-// 
-// 		bool bIsExactMatch = false;
-// 		if (BaseAssetName == FPaths::GetBaseFilename(FullFilePath))
-// 		{
-// 			bIsExactMatch = bIsPlainAssetName ? true : FullFilePath.Contains(AssetPath);
-// 		}
-// 
-// 		if (!bIsExactMatch && bHasExactMatch)
-// 		{
-// 			continue;
-// 		}
-// 
-// 		FString PackagePath;
-// 		if (FPackageName::TryConvertFilenameToLongPackageName(RelativePath, PackagePath))
-// 		{
-// 			if (bIsExactMatch && !bHasExactMatch)
-// 			{
-// 				PackagePathsOut.Empty();
-// 				bHasExactMatch = true;
-// 			}
-// 			PackagePathsOut.Add(PackagePath);
-// 		}
-// 	}
-// 
-// 	return PackagePathsOut.Num();
-// }
+	UObject* FoundAsset = nullptr;
+	for (UObject* PackageObj : ObjectsWithOuter)
+	{
+		if (PackageObj->IsAsset())
+		{
+			FoundAsset = PackageObj;
+			break;
+		}
+	}
+
+	return FoundAsset;
+}
 
 /*******************************************************************************
  * FEditorCommandLineUtils Definition
@@ -413,23 +656,24 @@ bool FEditorCommandLineUtils::ParseGameProjectPath(const TCHAR* CmdLine, FString
 		TArray<FString> Tokens, Switches;
 		UCommandlet::ParseCommandLine(*DiffArgs, Tokens, Switches);
 
-		if (Tokens.Num() > 0)
+		for (FString FilePath : Tokens)
 		{
-			FString DiffFilePath = Tokens[0];
-			FPaths::NormalizeFilename(DiffFilePath);
-
-			const TCHAR* const ContentDir = TEXT("/Content/");
+			FPaths::NormalizeFilename(FilePath);
+			ProjPathOut = FindProjectFile(FilePath);
 
-			FString ContentPath, FileSubPath;
-			if (DiffFilePath.Split(ContentDir, &ContentPath, &FileSubPath))
+			if (!ProjPathOut.IsEmpty())
 			{
-				GameNameOut = FPaths::GetCleanFilename(ContentPath);
-				ProjPathOut = FPaths::Combine(*FPaths::RootDir(), *GameNameOut, *FString(GameNameOut + TEXT(".") + FProjectDescriptor::GetExtension()));
-				return FPaths::FileExists(ProjPathOut);
+				GameNameOut = FPaths::GetBaseFilename(ProjPathOut);
+				// favor project files that are in the same directory tree as 
+				// the supplied file
+				if ( FilePath.StartsWith(FPaths::GetPath(ProjPathOut)) )
+				{
+					break;
+				}
 			}
 		}
 	}
-	return false;
+	return FPaths::FileExists(ProjPathOut);
 }
 
 //------------------------------------------------------------------------------
@@ -456,7 +700,7 @@ void FEditorCommandLineUtils::ProcessEditorCommands(const TCHAR* EditorCmdLine)
 
 	FString DiffArgs;
 	if (ParseCommandArgs(EditorCmdLine, DiffCommandSwitch, DiffArgs))
-	{		
+	{
 		IMainFrameModule& MainFrameModule = IMainFrameModule::Get();
 		const bool bIsMainFramInitialized = MainFrameModule.IsWindowInitialized();
 
