Added comments to editor experimental settings to make it more clear.

diff --git a/Engine/Source/Editor/UnrealEd/Classes/Settings/EditorExperimentalSettings.h b/Engine/Source/Editor/UnrealEd/Classes/Settings/EditorExperimentalSettings.h
--- a/Engine/Source/Editor/UnrealEd/Classes/Settings/EditorExperimentalSettings.h
+++ b/Engine/Source/Editor/UnrealEd/Classes/Settings/EditorExperimentalSettings.h
@@ -66,11 +66,11 @@ public:
 	bool bContextMenuChunkAssignments;
 
 	/** Disable cook in the editor */
-	UPROPERTY(EditAnywhere, config, Category = Cooking, meta = (DisplayName = "Disable Cook In The Editor feature (cooks from launch on will be run in a separate process if disabled)"))
+	UPROPERTY(EditAnywhere, config, Category = Cooking, meta = (DisplayName = "Disable Cook In The Editor feature, requires editor restart (cooks from launch on will be run in a separate process if disabled)"))
 	bool bDisableCookInEditor;
 
 	/** Enable cook on the side */
-	UPROPERTY(EditAnywhere, config, Category = Cooking, meta = (DisplayName = "Cook On The Side (Run a cook on the fly server in the background of the editor)"))
+	UPROPERTY(EditAnywhere, config, Category = Cooking, meta = (DisplayName = "Cook On The Side, requires editor restart (Run a cook on the fly server in the background of the editor)"))
 	bool bCookOnTheSide;
 
 	/** Enable -iterate for launch on */
