Fixing a mono compiler error.

diff --git a/Engine/Source/Editor/UnrealEd/Private/EditorCommandLineUtils.cpp b/Engine/Source/Editor/UnrealEd/Private/EditorCommandLineUtils.cpp
--- a/Engine/Source/Editor/UnrealEd/Private/EditorCommandLineUtils.cpp
+++ b/Engine/Source/Editor/UnrealEd/Private/EditorCommandLineUtils.cpp
@@ -204,7 +204,7 @@ void FFauxStandaloneToolManager::Disable()
 struct FMergeAsset
 {
 public:
-	FMergeAsset(TCHAR* DstFileName);
+	FMergeAsset(const TCHAR* DstFileName);
 
 	/**  */
 	bool SetSourceFile(const FString& SrcFilePathIn, FCommandLineErrorReporter& ErrorReporter);
@@ -236,7 +236,7 @@ private:
 
 
 //------------------------------------------------------------------------------
-FMergeAsset::FMergeAsset(TCHAR* DstFileName)
+FMergeAsset::FMergeAsset(const TCHAR* DstFileName)
 	: Package(nullptr)
 	, AssetObj(nullptr)
 	, DestFilePath(FPaths::Combine(*FPaths::DiffDir(), DstFileName))
