Each viewport type maintains its own transform details, so that perspective/ortho viewports remember their settings when moving from one to another.

diff --git a/Engine/Source/Developer/AssetTools/Private/AssetTypeActions/AssetTypeActions_CameraAnim.cpp b/Engine/Source/Developer/AssetTools/Private/AssetTypeActions/AssetTypeActions_CameraAnim.cpp
--- a/Engine/Source/Developer/AssetTools/Private/AssetTypeActions/AssetTypeActions_CameraAnim.cpp
+++ b/Engine/Source/Developer/AssetTools/Private/AssetTypeActions/AssetTypeActions_CameraAnim.cpp
@@ -41,8 +41,8 @@ void FAssetTypeActions_CameraAnim::CreateCameraActorForCameraAnim(UCameraAnim* I
 		FLevelEditorViewportClient* ViewportClient = GEditor->LevelViewportClients[ViewportIndex];
 		if(ViewportClient != NULL && ViewportClient->ViewportType == LVT_Perspective)
 		{
-			ViewportCamLocation = ViewportClient->ViewTransform.GetLocation();
-			ViewportCamRotation = ViewportClient->ViewTransform.GetRotation();
+			ViewportCamLocation = ViewportClient->ViewTransformPerspective.GetLocation();
+			ViewportCamRotation = ViewportClient->ViewTransformPerspective.GetRotation();
 			break;
 		}
 	}
diff --git a/Engine/Source/Editor/LevelEditor/Private/SLevelEditor.cpp b/Engine/Source/Editor/LevelEditor/Private/SLevelEditor.cpp
--- a/Engine/Source/Editor/LevelEditor/Private/SLevelEditor.cpp
+++ b/Engine/Source/Editor/LevelEditor/Private/SLevelEditor.cpp
@@ -880,15 +880,28 @@ void SLevelEditor::RestoreViewportTabInfo(TSharedRef<FLevelViewportTabContent> V
 		for(const auto& Viewport : *Viewports)
 		{
 			FLevelEditorViewportClient& LevelViewportClient = Viewport->GetLevelViewportClient();
-			const FString Key = FString::Printf(TEXT("%s[%d]"), *LayoutId, static_cast<int32>(LevelViewportClient.ViewportType));
-			const FLevelViewportInfo* const TransientEditorView = TransientEditorViews.Find(Key);
-			if(TransientEditorView)
+			bool bInitializedOrthoViewport = false;
+			for (int32 ViewportType = 0; ViewportType < LVT_MAX; ViewportType++)
 			{
-				LevelViewportClient.SetInitialViewTransform(
-					TransientEditorView->CamPosition,
-					TransientEditorView->CamRotation,
-					TransientEditorView->CamOrthoZoom
-					);
+				if (ViewportType == LVT_Perspective || !bInitializedOrthoViewport)
+				{
+					const FString Key = FString::Printf(TEXT("%s[%d]"), *LayoutId, ViewportType);
+					const FLevelViewportInfo* const TransientEditorView = TransientEditorViews.Find(Key);
+					if (TransientEditorView)
+					{
+						LevelViewportClient.SetInitialViewTransform(
+							static_cast<ELevelViewportType>(ViewportType),
+							TransientEditorView->CamPosition,
+							TransientEditorView->CamRotation,
+							TransientEditorView->CamOrthoZoom
+							);
+
+						if (ViewportType != LVT_Perspective)
+						{
+							bInitializedOrthoViewport = true;
+						}
+					}
+				}
 			}
 		}
 	}
diff --git a/Engine/Source/Editor/LevelEditor/Private/SLevelViewport.cpp b/Engine/Source/Editor/LevelEditor/Private/SLevelViewport.cpp
--- a/Engine/Source/Editor/LevelEditor/Private/SLevelViewport.cpp
+++ b/Engine/Source/Editor/LevelEditor/Private/SLevelViewport.cpp
@@ -949,10 +949,23 @@ void SLevelViewport::OnMapChanged( UWorld* World, EMapChangeType::Type MapChange
 	
 			ResetNewLevelViewFlags();
 
-			LevelViewportClient->SetInitialViewTransform(
-				World->EditorViews[LevelViewportClient->ViewportType].CamPosition,
-				World->EditorViews[LevelViewportClient->ViewportType].CamRotation,
-				World->EditorViews[LevelViewportClient->ViewportType].CamOrthoZoom );
+			bool bInitializedOrthoViewport = false;
+			for (int32 ViewportType = 0; ViewportType < LVT_MAX; ViewportType++)
+			{
+				if (ViewportType == LVT_Perspective || !bInitializedOrthoViewport)
+				{
+					LevelViewportClient->SetInitialViewTransform(
+						static_cast<ELevelViewportType>(ViewportType),
+						World->EditorViews[ViewportType].CamPosition,
+						World->EditorViews[ViewportType].CamRotation,
+						World->EditorViews[ViewportType].CamOrthoZoom);
+
+					if (ViewportType != LVT_Perspective)
+					{
+						bInitializedOrthoViewport = true;
+					}
+				}
+			}
 		}
 		else if( MapChangeType == EMapChangeType::SaveMap )
 		{
@@ -1936,7 +1949,7 @@ FLevelEditorViewportInstanceSettings SLevelViewport::LoadLegacyConfigFromIni(con
 	{
 		int32 ViewportTypeAsInt = ViewportInstanceSettings.ViewportType;
 		GConfig->GetInt(*IniSection, *(ConfigKey + TEXT(".Type")), ViewportTypeAsInt, GEditorUserSettingsIni);
-		ViewportInstanceSettings.ViewportType = (ViewportTypeAsInt == -1) ? LVT_None : static_cast<ELevelViewportType>(ViewportTypeAsInt); // LVT_None used to be -1
+		ViewportInstanceSettings.ViewportType = (ViewportTypeAsInt == -1 || ViewportTypeAsInt == 255) ? LVT_None : static_cast<ELevelViewportType>(ViewportTypeAsInt); // LVT_None used to be -1 or 255
 
 		if(ViewportInstanceSettings.ViewportType == LVT_None)
 		{
diff --git a/Engine/Source/Editor/Persona/Private/SAnimationEditorViewport.cpp b/Engine/Source/Editor/Persona/Private/SAnimationEditorViewport.cpp
--- a/Engine/Source/Editor/Persona/Private/SAnimationEditorViewport.cpp
+++ b/Engine/Source/Editor/Persona/Private/SAnimationEditorViewport.cpp
@@ -1384,7 +1384,7 @@ void SAnimationEditorViewportTabBody::SaveData(class SAnimationEditorViewportTab
 		SharedData.ViewLocation = OldAnimViewportClient.Get().GetViewLocation();
 		SharedData.ViewRotation = OldAnimViewportClient.Get().GetViewRotation();
 		SharedData.LookAtLocation = OldAnimViewportClient.Get().GetLookAtLocation();
-		SharedData.OrthoZoom = OldAnimViewportClient.Get().ViewTransform.GetOrthoZoom();
+		SharedData.OrthoZoom = OldAnimViewportClient.Get().GetOrthoZoom();
 		SharedData.bCameraLock = OldAnimViewportClient.Get().IsCameraLocked();
 		SharedData.bCameraFollow = OldAnimViewportClient.Get().IsSetCameraFollowChecked();
 		SharedData.bShowBound = OldAnimViewportClient.Get().IsSetShowBoundsChecked();
@@ -1407,7 +1407,7 @@ void SAnimationEditorViewportTabBody::RestoreData()
 		AnimViewportClient.Get().SetViewRotation( SharedData.ViewRotation );
 		AnimViewportClient.Get().SetShowBounds(SharedData.bShowBound);
 		AnimViewportClient.Get().SetLocalAxesMode((ELocalAxesMode::Type)SharedData.LocalAxesMode);
-		AnimViewportClient.Get().ViewTransform.SetOrthoZoom(SharedData.OrthoZoom);
+		AnimViewportClient.Get().SetOrthoZoom(SharedData.OrthoZoom);
 
 		OnSetPlaybackSpeed(SharedData.PlaybackSpeedMode);
 
diff --git a/Engine/Source/Editor/StaticMeshEditor/Private/StaticMeshEditorViewportClient.cpp b/Engine/Source/Editor/StaticMeshEditor/Private/StaticMeshEditorViewportClient.cpp
--- a/Engine/Source/Editor/StaticMeshEditor/Private/StaticMeshEditorViewportClient.cpp
+++ b/Engine/Source/Editor/StaticMeshEditor/Private/StaticMeshEditorViewportClient.cpp
@@ -1093,7 +1093,7 @@ void FStaticMeshEditorViewportClient::PerspectiveCameraMoved()
 	const bool bWasOrbit = bUsingOrbitCamera;
 	ToggleOrbitCamera(true);
 
-	const FVector OrbitPoint = ViewTransform.GetLookAt();
+	const FVector OrbitPoint = GetLookAtLocation();
 	const FVector OrbitZoom = GetViewLocation() - OrbitPoint;
 	StaticMesh->EditorCameraPosition = FAssetEditorOrbitCameraPosition(
 		OrbitPoint,
diff --git a/Engine/Source/Editor/UnrealEd/Private/EditorViewportClient.cpp b/Engine/Source/Editor/UnrealEd/Private/EditorViewportClient.cpp
--- a/Engine/Source/Editor/UnrealEd/Private/EditorViewportClient.cpp
+++ b/Engine/Source/Editor/UnrealEd/Private/EditorViewportClient.cpp
@@ -197,7 +197,6 @@ FEditorViewportClient::FEditorViewportClient(FEditorModeTools& InModeTools, FPre
 	, ImmersiveDelegate()
 	, VisibilityDelegate()
 	, Viewport(NULL)
-	, ViewTransform()
 	, ViewportType(LVT_Perspective)
 	, ViewState()
 	, EngineShowFlags(ESFIM_Editor)
@@ -363,23 +362,29 @@ void FEditorViewportClient::SetViewLocationForOrbiting(const FVector& LookAtPoin
 	SetLookAtLocation( LookAtPoint );
 }
 
-void FEditorViewportClient::SetInitialViewTransform( const FVector& ViewLocation, const FRotator& ViewRotation, float InOrthoZoom )
+void FEditorViewportClient::SetInitialViewTransform(ELevelViewportType ViewportType, const FVector& ViewLocation, const FRotator& ViewRotation, float InOrthoZoom )
 {
-	SetViewLocation( ViewLocation );
-	SetViewRotation( ViewRotation );
+	check(ViewportType < LVT_MAX);
+
+	FViewportCameraTransform& ViewTransform = (ViewportType == LVT_Perspective) ? ViewTransformPerspective : ViewTransformOrthographic;
+
+	ViewTransform.SetLocation(ViewLocation);
+	ViewTransform.SetRotation(ViewRotation);
 
 	// Make a look at location in front of the camera
-	const FQuat CameraOrientation = FQuat::MakeFromEuler( GetViewRotation().Euler() );
+	const FQuat CameraOrientation = FQuat::MakeFromEuler(ViewRotation.Euler());
 	FVector Direction = CameraOrientation.RotateVector( FVector(1,0,0) );
 
-	SetLookAtLocation( ViewLocation + Direction * OrbitConstants::IntialLookAtDistance );
-	SetOrthoZoom( InOrthoZoom );
+	ViewTransform.SetLookAt(ViewLocation + Direction * OrbitConstants::IntialLookAtDistance);
+	ViewTransform.SetOrthoZoom(InOrthoZoom);
 }
 
 void FEditorViewportClient::ToggleOrbitCamera( bool bEnableOrbitCamera )
 {
 	if( bUsingOrbitCamera != bEnableOrbitCamera )
 	{
+		FViewportCameraTransform& ViewTransform = GetViewTransform();
+
 		bUsingOrbitCamera = bEnableOrbitCamera;
 
 		// Convert orbit view to regular view
@@ -389,36 +394,38 @@ void FEditorViewportClient::ToggleOrbitCamera( bool bEnableOrbitCamera )
 		if( !bUsingOrbitCamera )
 		{
 			// Ensure that the view location and rotation is up to date to ensure smooth transition in and out of orbit mode
-			SetViewRotation( OrbitMatrix.Rotator() );
+			ViewTransform.SetRotation(OrbitMatrix.Rotator());
 		}
 		else
 		{
-			bool bUpsideDown = (GetViewRotation().Pitch < -90.0f || GetViewRotation().Pitch > 90.0f || !FMath::IsNearlyZero( GetViewRotation().Roll, KINDA_SMALL_NUMBER ) );
+			FRotator ViewRotation = ViewTransform.GetRotation();
+
+			bool bUpsideDown = (ViewRotation.Pitch < -90.0f || ViewRotation.Pitch > 90.0f || !FMath::IsNearlyZero(ViewRotation.Roll, KINDA_SMALL_NUMBER));
 
 			// if the camera is upside down compute the rotation differently to preserve pitch
 			// otherwise the view will pop to right side up when transferring to orbit controls
 			if( bUpsideDown )
 			{
-				FMatrix OrbitViewMatrix = FTranslationMatrix(-GetViewLocation());
-				OrbitViewMatrix *= FInverseRotationMatrix(GetViewRotation());
+				FMatrix OrbitViewMatrix = FTranslationMatrix(-ViewTransform.GetLocation());
+				OrbitViewMatrix *= FInverseRotationMatrix(ViewRotation);
 				OrbitViewMatrix *= FRotationMatrix( FRotator(0,90.f,0) );
 
-				FMatrix RotMat = FTranslationMatrix( -ViewTransform.GetLookAt() ) * OrbitViewMatrix;
+				FMatrix RotMat = FTranslationMatrix(-ViewTransform.GetLookAt()) * OrbitViewMatrix;
 				FMatrix RotMatInv = RotMat.InverseFast();
 				FRotator RollVec = RotMatInv.Rotator();
 				FMatrix YawMat = RotMatInv * FInverseRotationMatrix( FRotator(0, 0, -RollVec.Roll));
 				FMatrix YawMatInv = YawMat.InverseFast();
 				FRotator YawVec = YawMat.Rotator();
 				FRotator rotYawInv = YawMatInv.Rotator();
-				SetViewRotation( FRotator(-RollVec.Roll,YawVec.Yaw,0) );
+				ViewTransform.SetRotation(FRotator(-RollVec.Roll, YawVec.Yaw, 0));
 			}
 			else
 			{
-				SetViewRotation( OrbitMatrix.Rotator() );
+				ViewTransform.SetRotation(OrbitMatrix.Rotator());
 			}
 		}
 		
-		SetViewLocation( OrbitMatrix.GetOrigin() );
+		ViewTransform.SetLocation(OrbitMatrix.GetOrigin());
 	}
 }
 
@@ -438,6 +445,8 @@ void FEditorViewportClient::FocusViewportOnBox( const FBox& BoundingBox, bool bI
 	ToggleOrbitCamera(bEnable);
 
 	{
+		FViewportCameraTransform& ViewTransform = GetViewTransform();
+
 		if(!IsOrtho())
 		{
 		   /** 
@@ -461,15 +470,15 @@ void FEditorViewportClient::FocusViewportOnBox( const FBox& BoundingBox, bool bI
 			const float DistanceFromSphere = Radius / FMath::Tan( HalfFOVRadians );
 			FVector CameraOffsetVector = ViewTransform.GetRotation().Vector() * -DistanceFromSphere;
 
-			ViewTransform.SetLookAt(Position);		
-			ViewTransform.TransitionToLocation( Position + CameraOffsetVector, bInstant );
+			ViewTransform.SetLookAt(Position);
+			ViewTransform.TransitionToLocation(Position + CameraOffsetVector, bInstant);
 
 		}
 		else
 		{
 			// For ortho viewports just set the camera position to the center of the bounding volume.
 			//SetViewLocation( Position );
-			ViewTransform.TransitionToLocation( Position, bInstant );
+			ViewTransform.TransitionToLocation(Position, bInstant);
 
 			if( !(Viewport->KeyState(EKeys::LeftControl) || Viewport->KeyState(EKeys::RightControl)) )
 			{
@@ -484,7 +493,7 @@ void FEditorViewportClient::FocusViewportOnBox( const FBox& BoundingBox, bool bI
 
 				NewOrthoZoom = Zoom * (Viewport->GetSizeXY().X*15.0f);
 				NewOrthoZoom = FMath::Clamp<float>( NewOrthoZoom, MIN_ORTHOZOOM, MAX_ORTHOZOOM );
-				ViewTransform.SetOrthoZoom( NewOrthoZoom );
+				ViewTransform.SetOrthoZoom(NewOrthoZoom);
 			}
 		}
 	}
@@ -501,8 +510,11 @@ FSceneView* FEditorViewportClient::CalcSceneView(FSceneViewFamily* ViewFamily)
 {
 	FSceneViewInitOptions ViewInitOptions;
 
-	const FVector& ViewLocation = GetViewLocation();
-	const FRotator& ViewRotation = GetViewRotation();
+	FViewportCameraTransform& ViewTransform = GetViewTransform();
+	const ELevelViewportType EffectiveViewportType = GetViewportType();
+
+	const FVector& ViewLocation = ViewTransform.GetLocation();
+	const FRotator& ViewRotation = ViewTransform.GetRotation();
 
 	const FIntPoint ViewportSizeXY = Viewport->GetSizeXY();
 
@@ -512,10 +524,9 @@ FSceneView* FEditorViewportClient::CalcSceneView(FSceneViewFamily* ViewFamily)
 	// no matter how we are drawn (forced or otherwise), reset our time here
 	TimeForForceRedraw = 0.0;
 
-	const ELevelViewportType EffectiveViewportType = GetViewportType();
 	const bool bConstrainAspectRatio = bUseControllingActorViewInfo && ControllingActorViewInfo.bConstrainAspectRatio;
 
-	ViewInitOptions.ViewMatrix = FTranslationMatrix(-GetViewLocation());	
+	ViewInitOptions.ViewMatrix = FTranslationMatrix(-ViewLocation);
 	if (EffectiveViewportType == LVT_Perspective)
 	{
 		if (bUsingOrbitCamera)
@@ -750,6 +761,7 @@ void FEditorViewportClient::Tick(float DeltaTime)
 {
 	ConditionalCheckHoveredHitProxy();
 
+	FViewportCameraTransform& ViewTransform = GetViewTransform();
 	const bool bIsAnimating = ViewTransform.UpdateTransition();
 
 	if ( bIsTracking )
@@ -1672,6 +1684,8 @@ void FEditorViewportClient::InputAxisForOrbit(FViewport* InViewport, const FVect
 
 	Drag.X = DragDelta.X;
 
+	FViewportCameraTransform& ViewTransform = GetViewTransform();
+
 	if ( IsOrbitRotationMode( InViewport ) )
 	{
 		SetViewRotation( GetViewRotation() + FRotator( Rot.Pitch, -Rot.Yaw, Rot.Roll ) );
@@ -2180,7 +2194,8 @@ void FEditorViewportClient::DrawAxes(FViewport* InViewport, FCanvas* Canvas, con
 	FMatrix ViewTM = FMatrix::Identity;
 	if ( bUsingOrbitCamera)
 	{
-		ViewTM = FRotationMatrix( ViewTransform.ComputeOrbitMatrix().InverseFast().Rotator() );
+		FViewportCameraTransform& ViewTransform = GetViewTransform();
+		ViewTM = FRotationMatrix(ViewTransform.ComputeOrbitMatrix().InverseFast().Rotator());
 	}
 	else
 	{
diff --git a/Engine/Source/Editor/UnrealEd/Private/LevelEditorViewport.cpp b/Engine/Source/Editor/UnrealEd/Private/LevelEditorViewport.cpp
--- a/Engine/Source/Editor/UnrealEd/Private/LevelEditorViewport.cpp
+++ b/Engine/Source/Editor/UnrealEd/Private/LevelEditorViewport.cpp
@@ -1624,22 +1624,25 @@ void FLevelEditorViewportClient::PerspectiveCameraMoved()
  */
 void FLevelEditorViewportClient::ResetCamera()
 {
-	if(IsPerspective())
-	{
-		SetViewLocation( EditorViewportDefs::DefaultPerspectiveViewLocation );
-		SetViewRotation( EditorViewportDefs::DefaultPerspectiveViewRotation );
+	// Initialize perspective view transform
+	ViewTransformPerspective.SetLocation(EditorViewportDefs::DefaultPerspectiveViewLocation);
+	ViewTransformPerspective.SetRotation(EditorViewportDefs::DefaultPerspectiveViewRotation);
+	ViewTransformPerspective.SetLookAt(FVector::ZeroVector);
 
-		bool bRecalculateView = true;
-		SetLookAtLocation( FVector::ZeroVector, bRecalculateView );
-	}
-	else
-	{
-		SetViewLocation( FVector::ZeroVector );
-		SetViewRotation( FRotator::ZeroRotator );
-	}
-	ViewFOV = FOVAngle;
+	FMatrix OrbitMatrix = ViewTransformPerspective.ComputeOrbitMatrix();
+	OrbitMatrix = OrbitMatrix.InverseFast();
+
+	ViewTransformPerspective.SetRotation(OrbitMatrix.Rotator());
+	ViewTransformPerspective.SetLocation(OrbitMatrix.GetOrigin());
+
+	ViewTransformPerspective.SetOrthoZoom(DEFAULT_ORTHOZOOM);
 
-	ViewTransform.SetOrthoZoom( DEFAULT_ORTHOZOOM );
+	// Initialize orthographic view transform
+	ViewTransformOrthographic.SetLocation(FVector::ZeroVector);
+	ViewTransformOrthographic.SetRotation(FRotator::ZeroRotator);
+	ViewTransformOrthographic.SetOrthoZoom(DEFAULT_ORTHOZOOM);
+
+	ViewFOV = FOVAngle;
 
 	// If interp mode is active, tell it about the camera movement.
 	FEdMode* Mode = ModeTools->GetActiveMode(FBuiltinEditorModes::EM_InterpEdit);
diff --git a/Engine/Source/Editor/UnrealEd/Public/EditorViewportClient.h b/Engine/Source/Editor/UnrealEd/Public/EditorViewportClient.h
--- a/Engine/Source/Editor/UnrealEd/Public/EditorViewportClient.h
+++ b/Engine/Source/Editor/UnrealEd/Public/EditorViewportClient.h
@@ -332,16 +332,29 @@ public:
 	/** Callback for toggling the realtime preview flag. */
 	void SetRealtimePreview();
 
+	/** Gets ViewportCameraTransform object for the current viewport type */
+	FViewportCameraTransform& GetViewTransform()
+	{
+		return IsPerspective() ? ViewTransformPerspective : ViewTransformOrthographic;
+	}
+
+	const FViewportCameraTransform& GetViewTransform() const
+	{
+		return IsPerspective() ? ViewTransformPerspective : ViewTransformOrthographic;
+	}
+
 	/** Sets the location of the viewport's camera */
 	void SetViewLocation( const FVector& NewLocation )
 	{
-		ViewTransform.SetLocation( NewLocation );
+		FViewportCameraTransform& ViewTransform = GetViewTransform();
+		ViewTransform.SetLocation(NewLocation);
 	}
 
 	/** Sets the location of the viewport's camera */
 	void SetViewRotation( const FRotator& NewRotation )
 	{
-		ViewTransform.SetRotation( NewRotation );
+		FViewportCameraTransform& ViewTransform = GetViewTransform();
+		ViewTransform.SetRotation(NewRotation);
 	}
 
 	/** 
@@ -352,47 +365,55 @@ public:
 	 */
 	void SetLookAtLocation( const FVector& LookAt, bool bRecalculateView = false )
 	{
-		ViewTransform.SetLookAt( LookAt );
+		FViewportCameraTransform& ViewTransform = GetViewTransform();
+
+		ViewTransform.SetLookAt(LookAt);
 
 		if( bRecalculateView )
 		{
 			FMatrix OrbitMatrix = ViewTransform.ComputeOrbitMatrix();
 			OrbitMatrix = OrbitMatrix.InverseFast();
 
-			ViewTransform.SetRotation( OrbitMatrix.Rotator() );
-			ViewTransform.SetLocation( OrbitMatrix.GetOrigin() );
+			ViewTransform.SetRotation(OrbitMatrix.Rotator());
+			ViewTransform.SetLocation(OrbitMatrix.GetOrigin());
 		}
 	}
 
 	/** Sets ortho zoom amount */
 	void SetOrthoZoom( float InOrthoZoom ) 
 	{
+		FViewportCameraTransform& ViewTransform = GetViewTransform();
+
 		// A zero ortho zoom is not supported and causes NaN/div0 errors
 		check(InOrthoZoom != 0);
-		ViewTransform.SetOrthoZoom( InOrthoZoom );
+		ViewTransform.SetOrthoZoom(InOrthoZoom);
 	}
 
 	/** @return the current viewport camera location */
 	const FVector& GetViewLocation() const
 	{
+		const FViewportCameraTransform& ViewTransform = GetViewTransform();
 		return ViewTransform.GetLocation();
 	}
 
 	/** @return the current viewport camera rotation */
 	const FRotator& GetViewRotation() const
 	{
+		const FViewportCameraTransform& ViewTransform = GetViewTransform();
 		return ViewTransform.GetRotation();
 	}
 
 	/** @return the current look at location */
 	const FVector& GetLookAtLocation() const
 	{
+		const FViewportCameraTransform& ViewTransform = GetViewTransform();
 		return ViewTransform.GetLookAt();
 	}
 
 	/** @return the current ortho zoom amount */
 	float GetOrthoZoom() const
 	{
+		const FViewportCameraTransform& ViewTransform = GetViewTransform();
 		return ViewTransform.GetOrthoZoom();
 	}
 
@@ -409,7 +430,7 @@ public:
 		SetViewRotation( Rotation );
 	}
 
-	void SetInitialViewTransform( const FVector& ViewLocation, const FRotator& ViewRotation, float InOrthoZoom );
+	void SetInitialViewTransform(ELevelViewportType ViewportType, const FVector& ViewLocation, const FRotator& ViewRotation, float InOrthoZoom );
 
 	void TakeHighResScreenShot();
 
@@ -1198,8 +1219,11 @@ public:
 
 	FViewport*				Viewport;
 
-	/** Viewport camera transform data */
-	FViewportCameraTransform		ViewTransform;
+	/** Viewport camera transform data for perspective viewports */
+	FViewportCameraTransform		ViewTransformPerspective;
+
+	/** Viewport camera transform data for orthographic viewports */
+	FViewportCameraTransform		ViewTransformOrthographic;
 
 	/** The viewport type. */
 	ELevelViewportType		ViewportType;
