Added calls to PostEditMove when dragging an actor into the world.
#github Github pull request #621

diff --git a/Engine/Source/Editor/UnrealEd/Private/LevelEditorViewport.cpp b/Engine/Source/Editor/UnrealEd/Private/LevelEditorViewport.cpp
--- a/Engine/Source/Editor/UnrealEd/Private/LevelEditorViewport.cpp
+++ b/Engine/Source/Editor/UnrealEd/Private/LevelEditorViewport.cpp
@@ -259,6 +259,7 @@ static TArray<AActor*> AttemptDropObjAsActors( ULevel* InLevel, UObject* ObjToUs
 			if ( PlacedActor != NULL )
 			{
 				PlacedActors.Add(PlacedActor);
+				PlacedActor->PostEditMove(true);
 			}
 		}
 	}
@@ -1070,6 +1071,7 @@ bool FLevelEditorViewportClient::UpdateDropPreviewActors(int32 MouseX, int32 Mou
 
 		Actor->SetActorTransform(ActorTransform);
 		Actor->SetIsTemporarilyHiddenInEditor(false);
+		Actor->PostEditMove(false);
 	}
 
 	return true;
