[UE-5740] The finish/cancel buttons now close out the editor when running the merge tool from command-line.

diff --git a/Engine/Source/Developer/AssetTools/Private/AssetTypeActions/AssetTypeActions_Blueprint.cpp b/Engine/Source/Developer/AssetTools/Private/AssetTypeActions/AssetTypeActions_Blueprint.cpp
--- a/Engine/Source/Developer/AssetTools/Private/AssetTypeActions/AssetTypeActions_Blueprint.cpp
+++ b/Engine/Source/Developer/AssetTools/Private/AssetTypeActions/AssetTypeActions_Blueprint.cpp
@@ -101,7 +101,7 @@ void FAssetTypeActions_Blueprint::Merge(UObject* InObject)
 	}
 }
 
-void FAssetTypeActions_Blueprint::Merge(UObject* BaseAsset, UObject* RemoteAsset, UObject* LocalAsset)
+void FAssetTypeActions_Blueprint::Merge(UObject* BaseAsset, UObject* RemoteAsset, UObject* LocalAsset, const FOnMergeResolved& ResolutionCallback)
 {
 	UBlueprint* AsBlueprint = CastChecked<UBlueprint>(LocalAsset);
 	check(LocalAsset->GetClass() == BaseAsset->GetClass());
@@ -112,7 +112,7 @@ void FAssetTypeActions_Blueprint::Merge(UObject* BaseAsset, UObject* RemoteAsset
 		FBlueprintEditorModule& BlueprintEditorModule = FModuleManager::LoadModuleChecked<FBlueprintEditorModule>("Kismet");
 
 		FBlueprintEditor* BlueprintEditor = static_cast<FBlueprintEditor*>(FAssetEditorManager::Get().FindEditorForAsset(AsBlueprint, /*bFocusIfOpen =*/false));
-		BlueprintEditor->CreateMergeToolTab(Cast<UBlueprint>(BaseAsset), Cast<UBlueprint>(RemoteAsset));
+		BlueprintEditor->CreateMergeToolTab(Cast<UBlueprint>(BaseAsset), Cast<UBlueprint>(RemoteAsset), ResolutionCallback);
 	}
 }
 
diff --git a/Engine/Source/Developer/AssetTools/Private/AssetTypeActions/AssetTypeActions_Blueprint.h b/Engine/Source/Developer/AssetTools/Private/AssetTypeActions/AssetTypeActions_Blueprint.h
--- a/Engine/Source/Developer/AssetTools/Private/AssetTypeActions/AssetTypeActions_Blueprint.h
+++ b/Engine/Source/Developer/AssetTools/Private/AssetTypeActions/AssetTypeActions_Blueprint.h
@@ -16,7 +16,7 @@ public:
 	virtual void OpenAssetEditor(const TArray<UObject*>& InObjects, TSharedPtr<class IToolkitHost> EditWithinLevelEditor = TSharedPtr<IToolkitHost>()) override;
 	virtual bool CanMerge() const override;
 	virtual void Merge(UObject* InObject) override;
-	virtual void Merge(UObject* BaseAsset, UObject* RemoteAsset, UObject* LocalAsset);
+	virtual void Merge(UObject* BaseAsset, UObject* RemoteAsset, UObject* LocalAsset, const FOnMergeResolved& ResolutionCallback);
 	virtual uint32 GetCategories() override { return EAssetTypeCategories::Blueprint | EAssetTypeCategories::Basic; }
 	virtual void PerformAssetDiff(UObject* Asset1, UObject* Asset2, const struct FRevisionInfo& OldRevision, const struct FRevisionInfo& NewRevision) const override;
 	virtual class UThumbnailInfo* GetThumbnailInfo(UObject* Asset) const override;
diff --git a/Engine/Source/Developer/AssetTools/Public/AssetTypeActions_Base.h b/Engine/Source/Developer/AssetTools/Public/AssetTypeActions_Base.h
--- a/Engine/Source/Developer/AssetTools/Public/AssetTypeActions_Base.h
+++ b/Engine/Source/Developer/AssetTools/Public/AssetTypeActions_Base.h
@@ -56,7 +56,7 @@ public:
 		check(false); // no generic merge operation exists yet, did you override CanMerge but not Merge?
 	}
 
-	virtual void Merge(UObject* BaseAsset, UObject* RemoteAsset, UObject* LocalAsset) override
+	virtual void Merge(UObject* BaseAsset, UObject* RemoteAsset, UObject* LocalAsset, const FOnMergeResolved& ResolutionCallback) override
 	{
 		check(false); // no generic merge operation exists yet, did you override CanMerge but not Merge?
 	}
diff --git a/Engine/Source/Developer/AssetTools/Public/IAssetTypeActions.h b/Engine/Source/Developer/AssetTools/Public/IAssetTypeActions.h
--- a/Engine/Source/Developer/AssetTools/Public/IAssetTypeActions.h
+++ b/Engine/Source/Developer/AssetTools/Public/IAssetTypeActions.h
@@ -2,6 +2,8 @@
 
 #pragma once
 
+#include "Developer/Merge/Public/Merge.h" // for FOnMergeResolved
+
 namespace EAssetTypeCategories
 {
 	enum Type
@@ -84,7 +86,7 @@ public:
 	virtual void Merge( UObject* InObject ) = 0;
 
 	/** Begins a merge between the specified assets */
-	virtual void Merge(UObject* BaseAsset, UObject* RemoteAsset, UObject* LocalAsset) = 0;
+	virtual void Merge(UObject* BaseAsset, UObject* RemoteAsset, UObject* LocalAsset, const FOnMergeResolved& ResolutionCallback) = 0;
 
 	/** Returns the categories that this asset type. The return value is one or more flags from EAssetTypeCategories.  */
 	virtual uint32 GetCategories() = 0;
diff --git a/Engine/Source/Developer/Merge/Private/Merge.cpp b/Engine/Source/Developer/Merge/Private/Merge.cpp
--- a/Engine/Source/Developer/Merge/Private/Merge.cpp
+++ b/Engine/Source/Developer/Merge/Private/Merge.cpp
@@ -116,11 +116,12 @@ static UObject* LoadBaseRev(const FString& PackageName, const FString& AssetName
 }
 
 static TSharedPtr<SWidget> GenerateMergeTabContents(TSharedRef<FBlueprintEditor> Editor,
-                                                    const UBlueprint*    BaseBlueprint,
-                                                    const FRevisionInfo& BaseRevInfo,
-													const UBlueprint*    RemoteBlueprint,
-													const FRevisionInfo& RemoteRevInfo,
-													const UBlueprint&    LocalBlueprint)
+                                                    const UBlueprint*       BaseBlueprint,
+                                                    const FRevisionInfo&    BaseRevInfo,
+													const UBlueprint*       RemoteBlueprint,
+													const FRevisionInfo&    RemoteRevInfo,
+													const UBlueprint&       LocalBlueprint,
+													const FOnMergeResolved& MereResolutionCallback)
 {
 	TSharedPtr<SWidget> TabContent;
 
@@ -135,7 +136,8 @@ static TSharedPtr<SWidget> GenerateMergeTabContents(TSharedRef<FBlueprintEditor>
 			, BaseRevInfo);
 		LocalBlueprint.bDuplicatingReadOnly = false;
 
-		TabContent = SNew(SBlueprintMerge, Data);
+		TabContent = SNew(SBlueprintMerge, Data)
+			.OnMergeResolved(MereResolutionCallback);
 	}
 	else
 	{
@@ -173,8 +175,9 @@ class FMerge : public IMerge
 	virtual void StartupModule() override;
 	virtual void ShutdownModule() override;
 	virtual TSharedRef<SDockTab> GenerateMergeWidget(const UBlueprint& Object, TSharedRef<FBlueprintEditor> Editor) override;
-	virtual TSharedRef<SDockTab> GenerateMergeWidget(const UBlueprint* BaseAsset, const UBlueprint* RemoteAsset, const UBlueprint* LocalAsset, TSharedRef<FBlueprintEditor> Editor) override;
+	virtual TSharedRef<SDockTab> GenerateMergeWidget(const UBlueprint* BaseAsset, const UBlueprint* RemoteAsset, const UBlueprint* LocalAsset, const FOnMergeResolved& MergeResolutionCallback, TSharedRef<FBlueprintEditor> Editor) override;
 	virtual bool PendingMerge(const UBlueprint& BlueprintObj) const override;
+	//virtual FOnMergeResolved& OnMergeResolved() const override;
 
 	// Simplest to only allow one merge operation at a time, we could easily make this a map of Blueprint=>MergeTab
 	// but doing so will complicate the tab management
@@ -243,7 +246,7 @@ TSharedRef<SDockTab> FMerge::GenerateMergeWidget(const UBlueprint& Object, TShar
 		FRevisionInfo BaseRevInfo = FRevisionInfo::InvalidRevision();
 		const UBlueprint* BaseBlueprint = Cast< UBlueprint >(LoadBaseRev(PackageName, AssetName, SourceControlStateRef, BaseRevInfo));
 
-		Contents = GenerateMergeTabContents(Editor, BaseBlueprint, BaseRevInfo, RemoteBlueprint, CurrentRevInfo, Object);
+		Contents = GenerateMergeTabContents(Editor, BaseBlueprint, BaseRevInfo, RemoteBlueprint, CurrentRevInfo, Object, FOnMergeResolved());
 	}
 
 	TSharedRef<SDockTab> Tab =  FGlobalTabmanager::Get()->InvokeTab(MergeToolTabId);
@@ -253,7 +256,7 @@ TSharedRef<SDockTab> FMerge::GenerateMergeWidget(const UBlueprint& Object, TShar
 
 }
 
-TSharedRef<SDockTab> FMerge::GenerateMergeWidget(const UBlueprint* BaseBlueprint, const UBlueprint* RemoteBlueprint, const UBlueprint* LocalBlueprint, TSharedRef<FBlueprintEditor> Editor)
+TSharedRef<SDockTab> FMerge::GenerateMergeWidget(const UBlueprint* BaseBlueprint, const UBlueprint* RemoteBlueprint, const UBlueprint* LocalBlueprint, const FOnMergeResolved& MergeResolutionCallback, TSharedRef<FBlueprintEditor> Editor)
 {
 	if (ActiveTab.IsValid())
 	{
@@ -265,7 +268,7 @@ TSharedRef<SDockTab> FMerge::GenerateMergeWidget(const UBlueprint* BaseBlueprint
 	}
 
 	// @TODO: pipe revision info through
-	TSharedPtr<SWidget> TabContents = GenerateMergeTabContents(Editor, BaseBlueprint, FRevisionInfo::InvalidRevision(), RemoteBlueprint, FRevisionInfo::InvalidRevision(), *LocalBlueprint);
+	TSharedPtr<SWidget> TabContents = GenerateMergeTabContents(Editor, BaseBlueprint, FRevisionInfo::InvalidRevision(), RemoteBlueprint, FRevisionInfo::InvalidRevision(), *LocalBlueprint, MergeResolutionCallback);
 
 	TSharedRef<SDockTab> Tab = FGlobalTabmanager::Get()->InvokeTab(MergeToolTabId);
 	Tab->SetContent(TabContents.ToSharedRef());
diff --git a/Engine/Source/Developer/Merge/Private/SBlueprintMerge.cpp b/Engine/Source/Developer/Merge/Private/SBlueprintMerge.cpp
--- a/Engine/Source/Developer/Merge/Private/SBlueprintMerge.cpp
+++ b/Engine/Source/Developer/Merge/Private/SBlueprintMerge.cpp
@@ -41,6 +41,7 @@ void SBlueprintMerge::Construct(const FArguments InArgs, const FBlueprintMergeDa
 	check( InData.OwningEditor.Pin().IsValid() );
 
 	Data = InData;
+	OnMergeResolved = InArgs._OnMergeResolved;
 
 	BackupSubDir = FPaths::GameSavedDir() / TEXT("Backup") / TEXT("Resolve_Backup[") + FDateTime::Now().ToString(TEXT("%Y-%m-%d-%H-%M-%S")) + TEXT("]");
 
@@ -201,12 +202,30 @@ void SBlueprintMerge::OnFinishMerge()
 		FSlateNotificationManager::Get().AddNotification( ErrorNotification );
 	}
 
-	Data.OwningEditor.Pin()->CloseMergeTool();
+	// should come before CloseMergeTool(), because CloseMergeTool() makes its
+	// own call to OnMergeResolved (with an "Unknown" state).
+	OnMergeResolved.ExecuteIfBound(Package, EMergeResult::Completed);
+
+	// if we're using the merge command-line, it might close everything once
+	// a resolution is found (so the editor may be invalid now)
+	if (Data.OwningEditor.IsValid())
+	{
+		Data.OwningEditor.Pin()->CloseMergeTool();
+	}
 }
 
 void SBlueprintMerge::OnCancelClicked()
 {
-	Data.OwningEditor.Pin()->CloseMergeTool();
+	// should come before CloseMergeTool(), because CloseMergeTool() makes its
+	// own call to OnMergeResolved (with an "Unknown" state).
+	OnMergeResolved.ExecuteIfBound(GetTargetBlueprint()->GetOutermost(), EMergeResult::Cancelled);
+
+	// if we're using the merge command-line, it might close everything once
+	// a resolution is found (so the editor may be invalid now)
+	if (Data.OwningEditor.IsValid())
+	{
+		Data.OwningEditor.Pin()->CloseMergeTool();
+	}
 }
 
 void SBlueprintMerge::OnModeChanged(FName NewMode)
diff --git a/Engine/Source/Developer/Merge/Private/SBlueprintMerge.h b/Engine/Source/Developer/Merge/Private/SBlueprintMerge.h
--- a/Engine/Source/Developer/Merge/Private/SBlueprintMerge.h
+++ b/Engine/Source/Developer/Merge/Private/SBlueprintMerge.h
@@ -12,6 +12,7 @@ public:
 
 	SLATE_BEGIN_ARGS(SBlueprintMerge)
 	{}
+		SLATE_EVENT(FOnMergeResolved, OnMergeResolved)
 	SLATE_END_ARGS()
 
 	void Construct(const FArguments InArgs, const FBlueprintMergeData& InData);
@@ -57,6 +58,8 @@ private:
 	FMergeControl GraphControl;
 	FMergeControl TreeControl;
 	FMergeControl DetailsControl;
+	
+	FOnMergeResolved OnMergeResolved;
 
 	class IDiffControl* CurrentDiffControl;
 	class IMergeControl* CurrentMergeControl;
diff --git a/Engine/Source/Developer/Merge/Public/Merge.h b/Engine/Source/Developer/Merge/Public/Merge.h
--- a/Engine/Source/Developer/Merge/Public/Merge.h
+++ b/Engine/Source/Developer/Merge/Public/Merge.h
@@ -7,6 +7,17 @@
 class UBlueprint;
 class FBlueprintEditor;
 
+namespace EMergeResult
+{
+	enum Type
+	{
+		Unknown,
+		Completed,
+		Cancelled,
+	};
+}
+DECLARE_DELEGATE_TwoParams(FOnMergeResolved, UPackage*, EMergeResult::Type);
+
 /**
  * The public interface to this module
  */
@@ -42,12 +53,11 @@ public:
 	 * @return The merge view widget
 	 */
 	virtual TSharedRef<class SDockTab> GenerateMergeWidget( const UBlueprint& Object, TSharedRef< class FBlueprintEditor> Editor ) = 0;
-	virtual TSharedRef<class SDockTab> GenerateMergeWidget( const UBlueprint* BaseBlueprint, const UBlueprint* RemoteBlueprint, const UBlueprint* LocalBlueprint, TSharedRef<class FBlueprintEditor> Editor ) = 0;
+	virtual TSharedRef<class SDockTab> GenerateMergeWidget(const UBlueprint* BaseBlueprint, const UBlueprint* RemoteBlueprint, const UBlueprint* LocalBlueprint, const FOnMergeResolved& MergeResolutionCallback, TSharedRef<class FBlueprintEditor> Editor) = 0;
 
 	/** 
 	 * @return whether the blueprint is in a conflicted state, and can therefore be merged.
 	 */
 	virtual bool PendingMerge(const UBlueprint& BlueprintObj) const = 0;
-
 };
 
diff --git a/Engine/Source/Editor/Kismet/Private/BlueprintEditor.cpp b/Engine/Source/Editor/Kismet/Private/BlueprintEditor.cpp
--- a/Engine/Source/Editor/Kismet/Private/BlueprintEditor.cpp
+++ b/Engine/Source/Editor/Kismet/Private/BlueprintEditor.cpp
@@ -5048,9 +5048,10 @@ void FBlueprintEditor::CreateMergeToolTab()
 	MergeTool = IMerge::Get().GenerateMergeWidget(*GetBlueprintObj(), SharedThis(this));
 }
 
-void FBlueprintEditor::CreateMergeToolTab(const UBlueprint* BaseBlueprint, const UBlueprint* RemoteBlueprint)
+void FBlueprintEditor::CreateMergeToolTab(const UBlueprint* BaseBlueprint, const UBlueprint* RemoteBlueprint, const FOnMergeResolved& ResolutionCallback)
 {
-	MergeTool = IMerge::Get().GenerateMergeWidget(BaseBlueprint, RemoteBlueprint, GetBlueprintObj(), SharedThis(this));
+	OnMergeResolved = ResolutionCallback;
+	MergeTool = IMerge::Get().GenerateMergeWidget(BaseBlueprint, RemoteBlueprint, GetBlueprintObj(), ResolutionCallback, SharedThis(this));
 }
 
 void FBlueprintEditor::CloseMergeTool()
@@ -5058,6 +5059,12 @@ void FBlueprintEditor::CloseMergeTool()
 	auto MergeToolPtr = MergeTool.Pin();
 	if( MergeToolPtr.IsValid() )
 	{
+		UBlueprint* Blueprint = GetBlueprintObj();
+		UPackage* BpPackage = (Blueprint == nullptr) ? nullptr : Blueprint->GetOutermost();
+		// @TODO: right now crashes the editor on closing of the BP editor
+		//OnMergeResolved.ExecuteIfBound(BpPackage, EMergeResult::Unknown);
+		OnMergeResolved.Unbind();
+
 		MergeToolPtr->RequestCloseTab();
 	}
 }
diff --git a/Engine/Source/Editor/Kismet/Public/BlueprintEditor.h b/Engine/Source/Editor/Kismet/Public/BlueprintEditor.h
--- a/Engine/Source/Editor/Kismet/Public/BlueprintEditor.h
+++ b/Engine/Source/Editor/Kismet/Public/BlueprintEditor.h
@@ -10,6 +10,7 @@
 #include "GraphEditor.h"
 #include "EditorUndoClient.h"
 #include "Engine/UserDefinedEnum.h"
+#include "Developer/Merge/Public/Merge.h" // for FOnMergeResolved
 
 class USCS_Node;
 
@@ -468,7 +469,7 @@ public:
 
 	/** Forces the merge tool to be shown */
 	void CreateMergeToolTab();
-	void CreateMergeToolTab(const UBlueprint* BaseBlueprint, const UBlueprint* RemoteBlueprint);
+	void CreateMergeToolTab(const UBlueprint* BaseBlueprint, const UBlueprint* RemoteBlueprint, const FOnMergeResolved& ResolutionCallback);
 	
 	/** Closes the merge tool, rather than simply hiding it */
 	void CloseMergeTool();
@@ -959,6 +960,8 @@ protected:
 
 	/** Merge tool - WeakPtr because it's owned by the GlobalTabManager */
 	TWeakPtr<class SDockTab> MergeTool;
+	/** Merge tool - Delegate to call when the merge tool is closed. */
+	FOnMergeResolved OnMergeResolved;
 
 	/** Reference to owner of the current popup */
 	TWeakPtr<class SWindow> NameEntryPopupWindow;
diff --git a/Engine/Source/Editor/UnrealEd/Private/EditorCommandLineUtils.cpp b/Engine/Source/Editor/UnrealEd/Private/EditorCommandLineUtils.cpp
--- a/Engine/Source/Editor/UnrealEd/Private/EditorCommandLineUtils.cpp
+++ b/Engine/Source/Editor/UnrealEd/Private/EditorCommandLineUtils.cpp
@@ -598,28 +598,87 @@ static void EditorCommandLineUtilsImpl::RunAssetDiffCommand(TSharedPtr<SWindow>
 //------------------------------------------------------------------------------
 static void EditorCommandLineUtilsImpl::RunAssetMerge(FMergeAsset const& Base, FMergeAsset const& Remote, FMergeAsset const& Local, FMergeAsset const& Result)
 {
-	IAssetTools& AssetTools = FModuleManager::GetModuleChecked<FAssetToolsModule>("AssetTools").Get();
+	class FMergeResolutionHandler : public TSharedFromThis<FMergeResolutionHandler>
+	{
+	public:
+		FMergeResolutionHandler(UPackage* MergingPkgIn, const FString& DstFilePathIn)
+			: MergingPackage(MergingPkgIn)
+			, Resolution(EMergeResult::Unknown)
+			, DstFilePath(DstFilePathIn)
+		{
+			// force the user to save the result file (so we know if they "accepted" the merge)
+			MergingPkgIn->SetDirtyFlag(true);
+		}
 
-	UClass* AssetClass = Local.GetClass();
-	check(AssetClass != nullptr);
-	TWeakPtr<IAssetTypeActions> AssetActions = AssetTools.GetAssetTypeActionsForClass(AssetClass);
-	check(AssetActions.IsValid());
+		/** Records the user's selected resolution, and closes the editor. */
+		void HandleMergeResolution(UPackage* MergedPackageIn, EMergeResult::Type ResolutionIn)
+		{
+			if (MergedPackageIn == MergingPackage)
+			{
+				if (ResolutionIn == EMergeResult::Cancelled)
+				{
+					// they don't want to save any changes, so clear the flag
+					MergingPackage->SetDirtyFlag(false);
+				}
+				
+				if (Resolution == EMergeResult::Unknown)
+				{
+					Resolution = ResolutionIn;
+					EditorCommandLineUtilsImpl::ForceCloseEditor();
+				}
+			}
+		}
+
+		/** Copies the modified file if the user saved changes (and didn't cancel). */
+		void HandleEditorClose()
+		{
+			if ((Resolution != EMergeResult::Cancelled) && !MergingPackage->IsDirty())
+			{
+				FString SrcFilePath = MergingPackage->FileName.ToString();
+				IFileManager::Get().Copy(*DstFilePath, *SrcFilePath);
+			}
+		}
+
+	private:
+		UPackage* MergingPackage;
+		EMergeResult::Type Resolution;
+		FString DstFilePath;
+	};
+	
+	UPackage* MergeResultPkg = Local.GetAssetObj()->GetOutermost();
+	const FString& ResultFilePath = (!Result.GetSourceFilePath().IsEmpty()) ? Result.GetSourceFilePath() : Local.GetSourceFilePath();
+	TSharedRef<FMergeResolutionHandler> MergeHandler = MakeShareable(new FMergeResolutionHandler(MergeResultPkg, ResultFilePath));
+
+	// we use a lambda delegate to route the call into MergeHandler (we require 
+	// this intermediate to hold onto a MergeHandler ref, so it doesn't get 
+	// prematurely destroyed at the end of this function)
+	auto HandleMergeResolution = [MergeHandler](UPackage* MergedPackageIn, EMergeResult::Type ResolutionIn)
+	{
+		MergeHandler->HandleMergeResolution(MergedPackageIn, ResolutionIn);
+	};
+	const FOnMergeResolved MergeResolutionDelegate = FOnMergeResolved::CreateLambda(HandleMergeResolution);
 
 	// have to mount the save directory so that the BP-editor can save
 	// the merged asset packages
 	FPackageName::RegisterMountPoint(TEXT("/Temp/"), FPaths::GameSavedDir());
+	
+	IAssetTools& AssetTools = FModuleManager::GetModuleChecked<FAssetToolsModule>("AssetTools").Get();
+	UClass* AssetClass = Local.GetClass();
+	check(AssetClass != nullptr);
+	TWeakPtr<IAssetTypeActions> AssetActions = AssetTools.GetAssetTypeActionsForClass(AssetClass);
+	check(AssetActions.IsValid());
 	// bring up the merge tool...
-	AssetActions.Pin()->Merge(Base.GetAssetObj(), Remote.GetAssetObj(), Local.GetAssetObj());
+	AssetActions.Pin()->Merge(Base.GetAssetObj(), Remote.GetAssetObj(), Local.GetAssetObj(), MergeResolutionDelegate);
 
-	auto CopyFilesLambda = [](FString SrcFilePath, FString DstFilePath)
+	// we use a lambda delegate to route the call into MergeHandler (we require 
+	// this intermediate to hold onto a MergeHandler ref, so it doesn't get 
+	// prematurely destroyed at the end of this function)
+	auto HandleEditorClose = [](TSharedRef<FMergeResolutionHandler> MergeHandler)
 	{
-		IFileManager::Get().Copy(*DstFilePath, *SrcFilePath);
-	};
-
-	const FString& ModifiedFilePath = Local.GetAssetFilePath();
-	const FString& ResultFilePath   = (!Result.GetSourceFilePath().IsEmpty()) ? Result.GetSourceFilePath() : Local.GetSourceFilePath();
+		MergeHandler->HandleEditorClose();
+	};	
 	// have to copy the file into the expected result file when we're done
-	FEditorDelegates::OnShutdownPostPackagesSaved.AddStatic(CopyFilesLambda, ModifiedFilePath, ResultFilePath);
+	FEditorDelegates::OnShutdownPostPackagesSaved.AddStatic(HandleEditorClose, MergeHandler);
 }
 
 //------------------------------------------------------------------------------
