Altered Skeletal Mesh import path to auto attempt rigid-mesh import if no bone hierarchy was found.

diff --git a/Engine/Source/Editor/UnrealEd/Private/Fbx/FbxMainImport.cpp b/Engine/Source/Editor/UnrealEd/Private/Fbx/FbxMainImport.cpp
--- a/Engine/Source/Editor/UnrealEd/Private/Fbx/FbxMainImport.cpp
+++ b/Engine/Source/Editor/UnrealEd/Private/Fbx/FbxMainImport.cpp
@@ -1495,8 +1495,8 @@ void FFbxImporter::FillFbxSkelMeshArrayInScene(FbxNode* Node, TArray< TArray<Fbx
 	SkeletonArray.Empty();
 	// b) find rigid mesh
 	
-	// for rigid meshes, we don't convert to bone
-	if (ImportOptions->bImportRigidMesh)
+	// If we are attempting to import a skeletal mesh but we have no hierarchy attempt to find a rigid mesh.
+	if (outSkelMeshArray.Num() == 0)
 	{
 		RecursiveFindRigidMesh(Node, outSkelMeshArray, SkeletonArray, ExpandLOD);
 	}
