UE-5687 Crash occurs when opening a blueprint in the BlackJack sample game.

BP regeneration is less prone to cyclid dependency errors

diff --git a/Engine/Source/Editor/UnrealEd/Private/Kismet2/BlueprintEditorUtils.cpp b/Engine/Source/Editor/UnrealEd/Private/Kismet2/BlueprintEditorUtils.cpp
--- a/Engine/Source/Editor/UnrealEd/Private/Kismet2/BlueprintEditorUtils.cpp
+++ b/Engine/Source/Editor/UnrealEd/Private/Kismet2/BlueprintEditorUtils.cpp
@@ -761,17 +761,25 @@ struct FRegenerationHelper
 		const UEdGraphSchema_K2* Schema = GetDefault<UEdGraphSchema_K2>();
 		TSet<UStruct*> Dependencies;
 		ProcessHierarchy(Blueprint->ParentClass, Dependencies);
-
+		
 		for (const auto& NewVar : Blueprint->NewVariables)
 		{
 			if (UObject* TypeObject = NewVar.VarType.PinSubCategoryObject.Get())
 			{
-				ForcedLoad(TypeObject);
+				auto Linker = TypeObject->GetLinker();
+				if (Linker && TypeObject->HasAnyFlags(RF_NeedLoad))
+				{
+					Linker->Preload(TypeObject);
+				}
 			}
 
 			if (UClass* TypeClass = NewVar.VarType.PinSubCategoryMemberReference.MemberParentClass)
 			{
-				ForcedLoad(TypeClass);
+				auto Linker = TypeClass->GetLinker();
+				if (Linker && TypeClass->HasAnyFlags(RF_NeedLoad))
+				{
+					Linker->Preload(TypeClass);
+				}
 			}
 		}
 
