UE-5498 Leaving a variable name blank in a Blueprint Structure appears to append...

diff --git a/Engine/Source/Editor/UnrealEd/Private/Kismet2/StructureEditorUtils.cpp b/Engine/Source/Editor/UnrealEd/Private/Kismet2/StructureEditorUtils.cpp
--- a/Engine/Source/Editor/UnrealEd/Private/Kismet2/StructureEditorUtils.cpp
+++ b/Engine/Source/Editor/UnrealEd/Private/Kismet2/StructureEditorUtils.cpp
@@ -305,8 +305,10 @@ bool FStructureEditorUtils::RenameVariable(UUserDefinedStruct* Struct, FGuid Var
 	if (Struct)
 	{
 		auto VarDesc = GetVarDescByGuid(Struct, VarGuid);
-		if (VarDesc && FName(*NewDisplayNameStr).IsValidXName(INVALID_OBJECTNAME_CHARACTERS) &&
-			IsUniqueVariableDisplayName(Struct, NewDisplayNameStr))
+		if (VarDesc 
+			&& !NewDisplayNameStr.IsEmpty()
+			&& FName(*NewDisplayNameStr).IsValidXName(INVALID_OBJECTNAME_CHARACTERS) 
+			&& IsUniqueVariableDisplayName(Struct, NewDisplayNameStr))
 		{
 			const FScopedTransaction Transaction(LOCTEXT("RenameVariable", "Rename Variable"));
 			ModifyStructData(Struct);
