DataTable can be created without importing as a new asset.

diff --git a/Engine/Source/Editor/UnrealEd/Classes/Factories/DataTableFactory.h b/Engine/Source/Editor/UnrealEd/Classes/Factories/DataTableFactory.h
new file mode 100644
--- /dev/null
+++ b/Engine/Source/Editor/UnrealEd/Classes/Factories/DataTableFactory.h
@@ -0,0 +1,22 @@
+// Copyright 1998-2014 Epic Games, Inc. All Rights Reserved.
+
+#pragma once
+
+#include "DataTableFactory.generated.h"
+
+UCLASS(hidecategories=Object, MinimalAPI)
+class UDataTableFactory : public UFactory
+{
+	GENERATED_UCLASS_BODY()
+
+	UPROPERTY()
+	class UScriptStruct* Struct;
+
+	// Begin UFactory Interface
+	virtual bool ConfigureProperties() override;
+	virtual UObject* FactoryCreateNew(UClass* Class, UObject* InParent, FName Name, EObjectFlags Flags, UObject* Context, FFeedbackContext* Warn) override;
+	// Begin UFactory Interface
+};
+
+
+
diff --git a/Engine/Source/Editor/UnrealEd/Private/DataTableEditorUtils.cpp b/Engine/Source/Editor/UnrealEd/Private/DataTableEditorUtils.cpp
--- a/Engine/Source/Editor/UnrealEd/Private/DataTableEditorUtils.cpp
+++ b/Engine/Source/Editor/UnrealEd/Private/DataTableEditorUtils.cpp
@@ -91,3 +91,26 @@ void FDataTableEditorUtils::BroadcastPostChange(UDataTable* DataTable, EDataTabl
 	}
 	FDataTableEditorManager::Get().PostChange(DataTable, Info);
 }
+
+TArray<UScriptStruct*> FDataTableEditorUtils::GetPossibleStructs()
+{
+	TArray< UScriptStruct* > RowStructs;
+	UScriptStruct* TableRowStruct = FindObjectChecked<UScriptStruct>(ANY_PACKAGE, TEXT("TableRowBase"));
+	if (TableRowStruct != NULL)
+	{
+		// Make combo of table rowstruct options
+		for (TObjectIterator<UScriptStruct> It; It; ++It)
+		{
+			UScriptStruct* Struct = *It;
+			// If a child of the table row struct base, but not itself
+			const bool bBasedOnTableRowBase = Struct->IsChildOf(TableRowStruct) && (Struct != TableRowStruct);
+			const bool bUDStruct = Struct->IsA<UUserDefinedStruct>();
+			const bool bValidStruct = (Struct->GetOutermost() != GetTransientPackage());
+			if ((bBasedOnTableRowBase || bUDStruct) && bValidStruct)
+			{
+				RowStructs.Add(Struct);
+			}
+		}
+	}
+	return RowStructs;
+}
\ No newline at end of file
diff --git a/Engine/Source/Editor/UnrealEd/Private/Factories/CSVImportFactory.cpp b/Engine/Source/Editor/UnrealEd/Private/Factories/CSVImportFactory.cpp
--- a/Engine/Source/Editor/UnrealEd/Private/Factories/CSVImportFactory.cpp
+++ b/Engine/Source/Editor/UnrealEd/Private/Factories/CSVImportFactory.cpp
@@ -103,23 +103,7 @@ public:
 		ImportTypes.Add( MakeShareable(new ECSVImportType(ECSV_CurveVector)) );
 
 		// Find table row struct info
-		UScriptStruct* TableRowStruct = FindObjectChecked<UScriptStruct>(ANY_PACKAGE, TEXT("TableRowBase"));
-		if(TableRowStruct != NULL)
-		{
-			// Make combo of table rowstruct options
-			for (TObjectIterator<UScriptStruct> It; It; ++It)
-			{
-				UScriptStruct* Struct = *It;
-				// If a child of the table row struct base, but not itself
-				const bool bBasedOnTableRowBase = Struct->IsChildOf(TableRowStruct) && (Struct != TableRowStruct);
-				const bool bUDStruct = Struct->IsA<UUserDefinedStruct>();
-				const bool bValidStruct = (Struct->GetOutermost() != GetTransientPackage());
-				if ((bBasedOnTableRowBase || bUDStruct) && bValidStruct)
-				{
-					RowStructs.Add(Struct);
-				}
-			}
-		}
+		RowStructs = FDataTableEditorUtils::GetPossibleStructs();
 
 		// Create widget
 		this->ChildSlot
diff --git a/Engine/Source/Editor/UnrealEd/Private/Factories/EditorFactories.cpp b/Engine/Source/Editor/UnrealEd/Private/Factories/EditorFactories.cpp
--- a/Engine/Source/Editor/UnrealEd/Private/Factories/EditorFactories.cpp
+++ b/Engine/Source/Editor/UnrealEd/Private/Factories/EditorFactories.cpp
@@ -100,6 +100,7 @@
 #include "Engine/SubsurfaceProfile.h"
 #include "Camera/CameraAnim.h"
 #include "GameFramework/TouchInterface.h"
+#include "Engine/DataTable.h"
 
 DEFINE_LOG_CATEGORY_STATIC(LogEditorFactories, Log, All);
 
@@ -7129,5 +7130,149 @@ UObject* UCameraAnimFactory::FactoryCreateNew(UClass* Class,UObject* InParent,FN
 	return NewCamAnim;
 }
 
+/*------------------------------------------------------------------------------
+UDataTableFactory implementation.
+------------------------------------------------------------------------------*/
+UDataTableFactory::UDataTableFactory(const FObjectInitializer& ObjectInitializer)
+	: Super(ObjectInitializer)
+{
+	SupportedClass = UDataTable::StaticClass();
+	bCreateNew = true;
+	bEditAfterNew = true;
+}
+
+bool UDataTableFactory::ConfigureProperties()
+{
+	class FDataTableFactoryUI : public TSharedFromThis < FDataTableFactoryUI >
+	{
+		TSharedPtr<SWindow> PickerWindow;
+		TSharedPtr<SComboBox<UScriptStruct*>> RowStructCombo;
+		TSharedPtr<SButton> OkButton;
+		UScriptStruct* ResultStruct;
+	public:
+		FDataTableFactoryUI() : ResultStruct(NULL) {}
+
+		TSharedRef<SWidget> MakeRowStructItemWidget(class UScriptStruct* Struct) const
+		{
+			return SNew(STextBlock).Text(Struct ? Struct->GetDisplayNameText().ToString() : FString());
+		}
+
+		FString GetSelectedRowOptionText() const
+		{
+			UScriptStruct* Struct = RowStructCombo.IsValid() ? RowStructCombo->GetSelectedItem() : NULL;
+			return Struct ? Struct->GetDisplayNameText().ToString() : FString();
+		}
+
+		FReply OnCreate()
+		{
+			ResultStruct = RowStructCombo.IsValid() ? RowStructCombo->GetSelectedItem() : NULL;
+			if (PickerWindow.IsValid())
+			{
+				PickerWindow->RequestDestroyWindow();
+			}
+			return FReply::Handled();
+		}
+
+		FReply OnCancel()
+		{
+			ResultStruct = NULL;
+			if (PickerWindow.IsValid())
+			{
+				PickerWindow->RequestDestroyWindow();
+			}
+			return FReply::Handled();
+		}
+
+		bool IsAnyRowSelected() const
+		{
+			return  RowStructCombo.IsValid() && RowStructCombo->GetSelectedItem();
+		}
+
+		UScriptStruct* OpenStructSelector()
+		{
+			ResultStruct = NULL;
+			auto RowStructs = FDataTableEditorUtils::GetPossibleStructs();
+
+			RowStructCombo = SNew(SComboBox<UScriptStruct*>)
+			.OptionsSource(&RowStructs)
+			.OnGenerateWidget(this, &FDataTableFactoryUI::MakeRowStructItemWidget)
+			[
+				SNew(STextBlock)
+				.Text(this, &FDataTableFactoryUI::GetSelectedRowOptionText)
+			];
+
+			PickerWindow = SNew(SWindow)
+			.Title(LOCTEXT("DataTableFactoryOptions", "Pick Structure"))
+			.ClientSize(FVector2D(350, 100))
+			.SupportsMinimize(false).SupportsMaximize(false)
+			[
+				SNew(SBorder)
+				.BorderImage(FEditorStyle::GetBrush("Menu.Background"))
+				.Padding(10)
+				[
+					SNew(SVerticalBox)
+					+ SVerticalBox::Slot()
+					.AutoHeight()
+					[
+						RowStructCombo.ToSharedRef()
+					]
+					+ SVerticalBox::Slot()
+					.HAlign(HAlign_Right)
+					.AutoHeight()
+					[
+						SNew(SHorizontalBox)
+						+ SHorizontalBox::Slot()
+						.AutoWidth()
+						[
+							SAssignNew(OkButton, SButton)
+							.Text(LOCTEXT("OK", "OK"))
+							.OnClicked(this, &FDataTableFactoryUI::OnCreate)
+						]
+						+ SHorizontalBox::Slot()
+						.AutoWidth()
+						[
+							SNew(SButton)
+							.Text(LOCTEXT("Cancel", "Cancel"))
+							.OnClicked(this, &FDataTableFactoryUI::OnCancel)
+						]
+					]
+				]
+			];
+
+			OkButton->SetEnabled(
+				TAttribute<bool>::Create(TAttribute<bool>::FGetter::CreateSP(this, &FDataTableFactoryUI::IsAnyRowSelected)));
+
+			GEditor->EditorAddModalWindow(PickerWindow.ToSharedRef());
+
+			PickerWindow.Reset();
+			RowStructCombo.Reset();
+
+			return ResultStruct;
+		}
+	};
+
+
+	TSharedRef<FDataTableFactoryUI> StructSelector = MakeShareable(new FDataTableFactoryUI());
+	Struct = StructSelector->OpenStructSelector();
+
+	return Struct != NULL;
+}
+
+UObject* UDataTableFactory::FactoryCreateNew(UClass* Class, UObject* InParent, FName Name, EObjectFlags Flags, UObject* Context, FFeedbackContext* Warn)
+{
+	UDataTable* DataTable = NULL;
+	if (Struct && ensure(UDataTable::StaticClass() == Class))
+	{
+		ensure(0 != (RF_Public & Flags));
+		DataTable = NewNamedObject<UDataTable>(InParent, Name, Flags);
+		if (DataTable)
+		{
+			DataTable->RowStruct = Struct;
+		}
+	}
+	return DataTable;
+}
+
+
 #undef LOCTEXT_NAMESPACE
 
diff --git a/Engine/Source/Editor/UnrealEd/Public/DataTableEditorUtils.h b/Engine/Source/Editor/UnrealEd/Public/DataTableEditorUtils.h
--- a/Engine/Source/Editor/UnrealEd/Public/DataTableEditorUtils.h
+++ b/Engine/Source/Editor/UnrealEd/Public/DataTableEditorUtils.h
@@ -31,4 +31,6 @@ struct UNREALED_API FDataTableEditorUtils
 
 	static void BroadcastPreChange(UDataTable* DataTable, EDataTableChangeInfo Info);
 	static void BroadcastPostChange(UDataTable* DataTable, EDataTableChangeInfo Info);
+
+	static TArray<UScriptStruct*> GetPossibleStructs();
 };
\ No newline at end of file
