Fix missing or misplaced header guards.

diff --git a/Engine/Plugins/2D/Paper2D/Source/Paper2DEditor/Classes/PaperFlipbookActorFactory.h b/Engine/Plugins/2D/Paper2D/Source/Paper2DEditor/Classes/PaperFlipbookActorFactory.h
--- a/Engine/Plugins/2D/Paper2D/Source/Paper2DEditor/Classes/PaperFlipbookActorFactory.h
+++ b/Engine/Plugins/2D/Paper2D/Source/Paper2DEditor/Classes/PaperFlipbookActorFactory.h
@@ -1,5 +1,7 @@
 // Copyright 1998-2014 Epic Games, Inc. All Rights Reserved.
 
+#pragma once
+
 #include "PaperFlipbookActorFactory.generated.h"
 
 UCLASS()
diff --git a/Engine/Plugins/2D/Paper2D/Source/Paper2DEditor/Classes/PaperFlipbookThumbnailRenderer.h b/Engine/Plugins/2D/Paper2D/Source/Paper2DEditor/Classes/PaperFlipbookThumbnailRenderer.h
--- a/Engine/Plugins/2D/Paper2D/Source/Paper2DEditor/Classes/PaperFlipbookThumbnailRenderer.h
+++ b/Engine/Plugins/2D/Paper2D/Source/Paper2DEditor/Classes/PaperFlipbookThumbnailRenderer.h
@@ -1,5 +1,7 @@
 // Copyright 1998-2014 Epic Games, Inc. All Rights Reserved.
 
+#pragma once
+
 #include "PaperFlipbookThumbnailRenderer.generated.h"
 
 UCLASS()
diff --git a/Engine/Plugins/2D/Paper2D/Source/Paper2DEditor/Classes/PaperSpriteActorFactory.h b/Engine/Plugins/2D/Paper2D/Source/Paper2DEditor/Classes/PaperSpriteActorFactory.h
--- a/Engine/Plugins/2D/Paper2D/Source/Paper2DEditor/Classes/PaperSpriteActorFactory.h
+++ b/Engine/Plugins/2D/Paper2D/Source/Paper2DEditor/Classes/PaperSpriteActorFactory.h
@@ -1,5 +1,7 @@
 // Copyright 1998-2014 Epic Games, Inc. All Rights Reserved.
 
+#pragma once
+
 #include "PaperSpriteActorFactory.generated.h"
 
 UCLASS()
diff --git a/Engine/Plugins/2D/Paper2D/Source/Paper2DEditor/Classes/PaperSpriteFactory.h b/Engine/Plugins/2D/Paper2D/Source/Paper2DEditor/Classes/PaperSpriteFactory.h
--- a/Engine/Plugins/2D/Paper2D/Source/Paper2DEditor/Classes/PaperSpriteFactory.h
+++ b/Engine/Plugins/2D/Paper2D/Source/Paper2DEditor/Classes/PaperSpriteFactory.h
@@ -1,5 +1,7 @@
 // Copyright 1998-2014 Epic Games, Inc. All Rights Reserved.
 
+#pragma once
+
 /**
  * Factory for sprites
  */
diff --git a/Engine/Plugins/2D/Paper2D/Source/Paper2DEditor/Classes/PaperSpriteThumbnailRenderer.h b/Engine/Plugins/2D/Paper2D/Source/Paper2DEditor/Classes/PaperSpriteThumbnailRenderer.h
--- a/Engine/Plugins/2D/Paper2D/Source/Paper2DEditor/Classes/PaperSpriteThumbnailRenderer.h
+++ b/Engine/Plugins/2D/Paper2D/Source/Paper2DEditor/Classes/PaperSpriteThumbnailRenderer.h
@@ -1,5 +1,7 @@
 // Copyright 1998-2014 Epic Games, Inc. All Rights Reserved.
 
+#pragma once
+
 #include "PaperSpriteThumbnailRenderer.generated.h"
 
 UCLASS()
diff --git a/Engine/Plugins/2D/Paper2D/Source/Paper2DEditor/Classes/PaperTileSetThumbnailRenderer.h b/Engine/Plugins/2D/Paper2D/Source/Paper2DEditor/Classes/PaperTileSetThumbnailRenderer.h
--- a/Engine/Plugins/2D/Paper2D/Source/Paper2DEditor/Classes/PaperTileSetThumbnailRenderer.h
+++ b/Engine/Plugins/2D/Paper2D/Source/Paper2DEditor/Classes/PaperTileSetThumbnailRenderer.h
@@ -1,5 +1,7 @@
 // Copyright 1998-2014 Epic Games, Inc. All Rights Reserved.
 
+#pragma once
+
 #include "PaperTileSetThumbnailRenderer.generated.h"
 
 UCLASS()
diff --git a/Engine/Plugins/2D/Paper2D/Source/Paper2DEditor/Classes/TerrainSplineActorFactory.h b/Engine/Plugins/2D/Paper2D/Source/Paper2DEditor/Classes/TerrainSplineActorFactory.h
--- a/Engine/Plugins/2D/Paper2D/Source/Paper2DEditor/Classes/TerrainSplineActorFactory.h
+++ b/Engine/Plugins/2D/Paper2D/Source/Paper2DEditor/Classes/TerrainSplineActorFactory.h
@@ -1,5 +1,7 @@
 // Copyright 1998-2014 Epic Games, Inc. All Rights Reserved.
 
+#pragma once
+
 #include "TerrainSplineActorFactory.generated.h"
 
 UCLASS()
diff --git a/Engine/Plugins/2D/Paper2D/Source/Paper2DEditor/Classes/TileMapActorFactory.h b/Engine/Plugins/2D/Paper2D/Source/Paper2DEditor/Classes/TileMapActorFactory.h
--- a/Engine/Plugins/2D/Paper2D/Source/Paper2DEditor/Classes/TileMapActorFactory.h
+++ b/Engine/Plugins/2D/Paper2D/Source/Paper2DEditor/Classes/TileMapActorFactory.h
@@ -1,5 +1,7 @@
 // Copyright 1998-2014 Epic Games, Inc. All Rights Reserved.
 
+#pragma once
+
 #include "TileMapActorFactory.generated.h"
 
 UCLASS()
diff --git a/Engine/Plugins/2D/Paper2D/Source/Paper2DEditor/Classes/TileMapAssetImportData.h b/Engine/Plugins/2D/Paper2D/Source/Paper2DEditor/Classes/TileMapAssetImportData.h
--- a/Engine/Plugins/2D/Paper2D/Source/Paper2DEditor/Classes/TileMapAssetImportData.h
+++ b/Engine/Plugins/2D/Paper2D/Source/Paper2DEditor/Classes/TileMapAssetImportData.h
@@ -1,5 +1,7 @@
 // Copyright 1998-2014 Epic Games, Inc. All Rights Reserved.
 
+#pragma once
+
 #include "TileMapAssetImportData.generated.h"
 
 USTRUCT()
diff --git a/Engine/Plugins/2D/Paper2D/Source/Paper2DEditor/Private/FlipbookEditor/STimelineHeader.h b/Engine/Plugins/2D/Paper2D/Source/Paper2DEditor/Private/FlipbookEditor/STimelineHeader.h
--- a/Engine/Plugins/2D/Paper2D/Source/Paper2DEditor/Private/FlipbookEditor/STimelineHeader.h
+++ b/Engine/Plugins/2D/Paper2D/Source/Paper2DEditor/Private/FlipbookEditor/STimelineHeader.h
@@ -1,5 +1,7 @@
 // Copyright 1998-2014 Epic Games, Inc. All Rights Reserved.
 
+#pragma once
+
 //////////////////////////////////////////////////////////////////////////
 // STimelineHeader
 
diff --git a/Engine/Plugins/Editor/EpicSurvey/Source/EpicSurvey/Private/SSurveyNotification.h b/Engine/Plugins/Editor/EpicSurvey/Source/EpicSurvey/Private/SSurveyNotification.h
--- a/Engine/Plugins/Editor/EpicSurvey/Source/EpicSurvey/Private/SSurveyNotification.h
+++ b/Engine/Plugins/Editor/EpicSurvey/Source/EpicSurvey/Private/SSurveyNotification.h
@@ -1,4 +1,7 @@
 // Copyright 1998-2014 Epic Games, Inc. All Rights Reserved.
+
+#pragma once
+
 #include "EpicSurvey.h"
 
 #define LOCTEXT_NAMESPACE "EpicSurvey"
diff --git a/Engine/Plugins/Editor/PluginsEditor/Source/PluginsEditor/Private/PluginsEditorPrivatePCH.h b/Engine/Plugins/Editor/PluginsEditor/Source/PluginsEditor/Private/PluginsEditorPrivatePCH.h
--- a/Engine/Plugins/Editor/PluginsEditor/Source/PluginsEditor/Private/PluginsEditorPrivatePCH.h
+++ b/Engine/Plugins/Editor/PluginsEditor/Source/PluginsEditor/Private/PluginsEditorPrivatePCH.h
@@ -1,5 +1,7 @@
 // Copyright 1998-2014 Epic Games, Inc. All Rights Reserved.
 
+#pragma once
+
 #include "UnrealEd.h"
 #include "IPluginsEditor.h"
 
diff --git a/Engine/Source/Editor/PListEditor/Private/PListNode.h b/Engine/Source/Editor/PListEditor/Private/PListNode.h
--- a/Engine/Source/Editor/PListEditor/Private/PListNode.h
+++ b/Engine/Source/Editor/PListEditor/Private/PListNode.h
@@ -2,9 +2,6 @@
 
 #pragma once
 
-#ifndef __PListNode_h__
-#define __PlistNode_h__
-
 #include "SlateBasics.h"
 
 /** Forward declaration for the panel which we will need to access in delegates */
@@ -124,5 +121,3 @@ private:
 
 	int32 Depth;
 };
-
-#endif
diff --git a/Engine/Source/Editor/PListEditor/Private/PListNodeArray.h b/Engine/Source/Editor/PListEditor/Private/PListNodeArray.h
--- a/Engine/Source/Editor/PListEditor/Private/PListNodeArray.h
+++ b/Engine/Source/Editor/PListEditor/Private/PListNodeArray.h
@@ -2,9 +2,6 @@
 
 #pragma once
 
-#ifndef __PListNodeArray_h__
-#define __PlistNodeArray_h__
-
 #include "PListNode.h"
 
 /** A Node representing an array */
@@ -104,5 +101,3 @@ private:
 	/** Flag for valid key string */
 	bool bKeyValid;
 };
-
-#endif
diff --git a/Engine/Source/Editor/PListEditor/Private/PListNodeBoolean.h b/Engine/Source/Editor/PListEditor/Private/PListNodeBoolean.h
--- a/Engine/Source/Editor/PListEditor/Private/PListNodeBoolean.h
+++ b/Engine/Source/Editor/PListEditor/Private/PListNodeBoolean.h
@@ -2,9 +2,6 @@
 
 #pragma once
 
-#ifndef __PListNodeBoolean_h__
-#define __PlistNodeBoolean_h__
-
 #include "PListNode.h"
 
 /** A Node representing a boolean value */
@@ -97,5 +94,3 @@ private:
 	/** Flag for valid key string */
 	bool bKeyValid;
 };
-
-#endif
diff --git a/Engine/Source/Editor/PListEditor/Private/PListNodeDictionary.h b/Engine/Source/Editor/PListEditor/Private/PListNodeDictionary.h
--- a/Engine/Source/Editor/PListEditor/Private/PListNodeDictionary.h
+++ b/Engine/Source/Editor/PListEditor/Private/PListNodeDictionary.h
@@ -2,9 +2,6 @@
 
 #pragma once
 
-#ifndef __PListNodeDictionary_h__
-#define __PlistNodeDictionary_h__
-
 #include "PListNode.h"
 
 /** A Node representing a dictionary */
@@ -91,5 +88,3 @@ private:
 	/** Filtered Status */
 	bool bFiltered;
 };
-
-#endif
diff --git a/Engine/Source/Editor/PListEditor/Private/PListNodeFile.h b/Engine/Source/Editor/PListEditor/Private/PListNodeFile.h
--- a/Engine/Source/Editor/PListEditor/Private/PListNodeFile.h
+++ b/Engine/Source/Editor/PListEditor/Private/PListNodeFile.h
@@ -2,9 +2,6 @@
 
 #pragma once
 
-#ifndef __PListNodeFile_h__
-#define __PlistNodeFile_h__
-
 #include "PListNode.h"
 
 /** A Node representing the plist as a whole */
@@ -60,5 +57,3 @@ private:
 	/** Widget for displaying text on this row */
 	TSharedPtr<STextBlock> TextWidget;
 };
-
-#endif
diff --git a/Engine/Source/Editor/PListEditor/Private/PListNodeString.h b/Engine/Source/Editor/PListEditor/Private/PListNodeString.h
--- a/Engine/Source/Editor/PListEditor/Private/PListNodeString.h
+++ b/Engine/Source/Editor/PListEditor/Private/PListNodeString.h
@@ -2,9 +2,6 @@
 
 #pragma once
 
-#ifndef __PListNodeString_h__
-#define __PlistNodeString_h__
-
 #include "PListNode.h"
 
 /** A Node representing a string value */
@@ -103,5 +100,3 @@ private:
 	/** Flag for valid value string */
 	bool bValueValid;
 };
-
-#endif
diff --git a/Engine/Source/Editor/UnrealEd/Classes/Commandlets/ParticleSystemAuditCommandlet.h b/Engine/Source/Editor/UnrealEd/Classes/Commandlets/ParticleSystemAuditCommandlet.h
--- a/Engine/Source/Editor/UnrealEd/Classes/Commandlets/ParticleSystemAuditCommandlet.h
+++ b/Engine/Source/Editor/UnrealEd/Classes/Commandlets/ParticleSystemAuditCommandlet.h
@@ -1,5 +1,7 @@
 // Copyright 1998-2014 Epic Games, Inc. All Rights Reserved.
 
+#pragma once
+
 #include "Runtime/Engine/Classes/Particles/ParticleSystem.h"
 #include "ParticleDefinitions.h"
 #include "ParticleSystemAuditCommandlet.generated.h"
diff --git a/Engine/Source/Editor/UnrealEd/Classes/Factories/WorldFactory.h b/Engine/Source/Editor/UnrealEd/Classes/Factories/WorldFactory.h
--- a/Engine/Source/Editor/UnrealEd/Classes/Factories/WorldFactory.h
+++ b/Engine/Source/Editor/UnrealEd/Classes/Factories/WorldFactory.h
@@ -1,5 +1,7 @@
 // Copyright 1998-2013 Epic Games, Inc. All Rights Reserved.
 
+#pragma once
+
 /**
  *
  */
diff --git a/Engine/Source/Editor/UnrealEd/Private/NavigationBuildingNotification.h b/Engine/Source/Editor/UnrealEd/Private/NavigationBuildingNotification.h
--- a/Engine/Source/Editor/UnrealEd/Private/NavigationBuildingNotification.h
+++ b/Engine/Source/Editor/UnrealEd/Private/NavigationBuildingNotification.h
@@ -1,4 +1,7 @@
 // Copyright 1998-2014 Epic Games, Inc. All Rights Reserved.
+
+#pragma once
+
 #include "UnrealEd.h"
 
 /** Notification class for asynchronous shader compiling. */
diff --git a/Engine/Source/Runtime/Core/Public/Windows/WindowsPlatformCrashContext.h b/Engine/Source/Runtime/Core/Public/Windows/WindowsPlatformCrashContext.h
--- a/Engine/Source/Runtime/Core/Public/Windows/WindowsPlatformCrashContext.h
+++ b/Engine/Source/Runtime/Core/Public/Windows/WindowsPlatformCrashContext.h
@@ -1,5 +1,7 @@
 // Copyright 1998-2014 Epic Games, Inc. All Rights Reserved.
 
+#pragma once
+
 #include "GenericPlatform/GenericPlatformContext.h"
 
 #include "AllowWindowsPlatformTypes.h"
diff --git a/Engine/Source/Runtime/Engine/Classes/VisualLogger/VisualLoggerKismetLibrary.h b/Engine/Source/Runtime/Engine/Classes/VisualLogger/VisualLoggerKismetLibrary.h
--- a/Engine/Source/Runtime/Engine/Classes/VisualLogger/VisualLoggerKismetLibrary.h
+++ b/Engine/Source/Runtime/Engine/Classes/VisualLogger/VisualLoggerKismetLibrary.h
@@ -1,4 +1,7 @@
 // Copyright 1998-2014 Epic Games, Inc. All Rights Reserved.
+
+#pragma once
+
 #include "VisualLoggerKismetLibrary.generated.h"
 
 UCLASS(MinimalAPI)
diff --git a/Engine/Source/Runtime/Serialization/Private/Tests/StructSerializerTestTypes.h b/Engine/Source/Runtime/Serialization/Private/Tests/StructSerializerTestTypes.h
--- a/Engine/Source/Runtime/Serialization/Private/Tests/StructSerializerTestTypes.h
+++ b/Engine/Source/Runtime/Serialization/Private/Tests/StructSerializerTestTypes.h
@@ -1,7 +1,8 @@
 // Copyright 1998-2013 Epic Games, Inc. All Rights Reserved.
 
-#include "StructSerializerTestTypes.generated.h"
+#pragma once
 
+#include "StructSerializerTestTypes.generated.h"
 
 /**
  * Test structure for numeric properties.
