Fix unity build.

diff --git a/Engine/Source/Editor/UnrealEd/Private/Factories/EditorFactories.cpp b/Engine/Source/Editor/UnrealEd/Private/Factories/EditorFactories.cpp
--- a/Engine/Source/Editor/UnrealEd/Private/Factories/EditorFactories.cpp
+++ b/Engine/Source/Editor/UnrealEd/Private/Factories/EditorFactories.cpp
@@ -101,6 +101,7 @@
 #include "Camera/CameraAnim.h"
 #include "GameFramework/TouchInterface.h"
 #include "Engine/DataTable.h"
+#include "DataTableEditorUtils.h"
 
 DEFINE_LOG_CATEGORY_STATIC(LogEditorFactories, Log, All);
 
