Fix for issue raised in https://udn.unrealengine.com/questions/223502/fexportmaterialproxy-crash.html

Material resources can now generate their own shader map IDs which makes the code much neater and solves the crash.

diff --git a/Engine/Source/Editor/UnrealEd/Private/EditorExporters.cpp b/Engine/Source/Editor/UnrealEd/Private/EditorExporters.cpp
--- a/Engine/Source/Editor/UnrealEd/Private/EditorExporters.cpp
+++ b/Engine/Source/Editor/UnrealEd/Private/EditorExporters.cpp
@@ -729,43 +729,20 @@ public:
 		Material = InMaterialInterface->GetMaterial();
 		Material->AppendReferencedTextures(ReferencedTextures);
 		FPlatformMisc::CreateGuid(Id);
-		
-		// Have to properly handle compilation of static switches in MaterialInstance* cases...
-		UMaterialInstance* MaterialInstance = Cast<UMaterialInstance>(InMaterialInterface);
-
-		// Walk up the parent chain until we find the first MI with static parameters
-		while (MaterialInstance 
-			&& !MaterialInstance->bHasStaticPermutationResource
-			&& MaterialInstance->Parent 
-			&& MaterialInstance->Parent->IsA<UMaterialInstance>())
-		{
-			MaterialInstance = Cast<UMaterialInstance>(MaterialInstance->Parent);
-		}
-
-		// Special path for a MI with static parameters
-		if (MaterialInstance && MaterialInstance->bHasStaticPermutationResource && MaterialInstance->Parent)
-		{
-			FMaterialResource* MIResource = MaterialInstance->GetMaterialResource(GMaxRHIFeatureLevel);
 
-			// Use the shader map Id from the static permutation
-			// This allows us to create a deterministic yet unique Id for the shader map that will be compiled for this FLightmassMaterialProxy
-			FMaterialShaderMapId ResourceId;
-			//@todo - always use highest quality level for static lighting
-			MaterialInstance->GetMaterialResourceId(GMaxRHIShaderPlatform, EMaterialQualityLevel::Num, ResourceId);
+		FMaterialResource* Resource = InMaterialInterface->GetMaterialResource(GMaxRHIFeatureLevel);
 
-			// Override with a special usage so we won't re-use the shader map used by the MI for rendering
-			ResourceId.Usage = GetShaderMapUsage();
-
-			CacheShaders(ResourceId, GMaxRHIShaderPlatform, true);
-		}
-		else
+		FMaterialShaderMapId ResourceId;
+		Resource->GetShaderMapId(GMaxRHIShaderPlatform, ResourceId);
+		// Override with a special usage so we won't re-use the shader map used by the material for rendering
+		switch (InPropertyToCompile)
 		{
-			FMaterialResource* MaterialResource = Material->GetMaterialResource(GMaxRHIFeatureLevel);
-			
-			// Copy the material resource Id
-			// The FLightmassMaterialProxy's GetShaderMapUsage will set it apart from the MI's resource when it comes to finding a shader map
-			CacheShaders(GMaxRHIShaderPlatform, true);
-		}
+		case MP_DiffuseColor: ResourceId.Usage = EMaterialShaderMapUsage::MaterialExportDiffuseColor; break;
+		case MP_SpecularColor: ResourceId.Usage = EMaterialShaderMapUsage::MaterialExportSpecularColor; break;
+		case MP_Normal: ResourceId.Usage = EMaterialShaderMapUsage::MaterialExportNormal; break;
+		};
+		
+		CacheShaders(ResourceId, GMaxRHIShaderPlatform, true);
 	}
 
 	/** This override is required otherwise the shaders aren't ready for use when the surface is rendered resulting in a blank image */
diff --git a/Engine/Source/Editor/UnrealEd/Private/Lightmass/LightmassRender.cpp b/Engine/Source/Editor/UnrealEd/Private/Lightmass/LightmassRender.cpp
--- a/Engine/Source/Editor/UnrealEd/Private/Lightmass/LightmassRender.cpp
+++ b/Engine/Source/Editor/UnrealEd/Private/Lightmass/LightmassRender.cpp
@@ -190,43 +190,13 @@ public:
 
 		Material->AppendReferencedTextures(ReferencedTextures);
 
-		// Have to properly handle compilation of static switches in MaterialInstance* cases...
-		UMaterialInstance* MaterialInstance = Cast<UMaterialInstance>(InMaterialInterface);
-
-		// Walk up the parent chain until we find the first MI with static parameters
-		while (MaterialInstance 
-			&& !MaterialInstance->bHasStaticPermutationResource
-			&& MaterialInstance->Parent 
-			&& MaterialInstance->Parent->IsA<UMaterialInstance>())
-		{
-			MaterialInstance = Cast<UMaterialInstance>(MaterialInstance->Parent);
-		}
-
-		// Special path for a MI with static parameters
-		if (MaterialInstance && MaterialInstance->bHasStaticPermutationResource && MaterialInstance->Parent)
-		{
-			FMaterialResource* MIResource = MaterialInstance->GetMaterialResource(GMaxRHIFeatureLevel);
-
-			// Use the shader map Id from the static permutation
-			// This allows us to create a deterministic yet unique Id for the shader map that will be compiled for this FLightmassMaterialProxy
-			FMaterialShaderMapId ResourceId;
-			//@todo - always use highest quality level for static lighting
-			MaterialInstance->GetMaterialResourceId(GMaxRHIShaderPlatform, EMaterialQualityLevel::Num, ResourceId);
-
-			// Override with a special usage so we won't re-use the shader map used by the MI for rendering
-			ResourceId.Usage = GetShaderMapUsage();
+		FMaterialResource* Resource = InMaterialInterface->GetMaterialResource(GMaxRHIFeatureLevel);
 
-			CacheShaders(ResourceId, GMaxRHIShaderPlatform, true);
-		}
-		else
-		{
-			FMaterialResource* MaterialResource = Material->GetMaterialResource(GMaxRHIFeatureLevel);
-
-			// Copy the material resource Id
-			// The FLightmassMaterialProxy's GetShaderMapUsage will set it apart from the MI's resource when it comes to finding a shader map
-
-			CacheShaders(GMaxRHIShaderPlatform, true);
-		}
+		FMaterialShaderMapId ResourceId;
+		Resource->GetShaderMapId(GMaxRHIShaderPlatform, ResourceId);
+		// Override with a special usage so we won't re-use the shader map used by the material for rendering
+		ResourceId.Usage = GetShaderMapUsage();
+		CacheShaders(ResourceId, GMaxRHIShaderPlatform, true);
 	}
 
 	virtual const TArray<UTexture*>& GetReferencedTextures() const override
diff --git a/Engine/Source/Runtime/Engine/Classes/Materials/MaterialInstance.h b/Engine/Source/Runtime/Engine/Classes/Materials/MaterialInstance.h
--- a/Engine/Source/Runtime/Engine/Classes/Materials/MaterialInstance.h
+++ b/Engine/Source/Runtime/Engine/Classes/Materials/MaterialInstance.h
@@ -339,9 +339,6 @@ public:
 	 */
 	void CacheResourceShadersForCooking(EShaderPlatform ShaderPlatform, TArray<FMaterialResource*>& OutCachedMaterialResources);
 
-	/** Builds a FMaterialShaderMapId to represent this material instance for the given platform and quality level. */
-	ENGINE_API void GetMaterialResourceId(EShaderPlatform ShaderPlatform, EMaterialQualityLevel::Type QualityLevel, FMaterialShaderMapId& OutId);
-
 	/** 
 	 * Gathers actively used shader maps from all material resources used by this material instance
 	 * Note - not refcounting the shader maps so the references must not be used after material resources are modified (compilation, loading, etc)
@@ -406,9 +403,6 @@ protected:
 	/** Initialize the material instance's resources. */
 	ENGINE_API void InitResources();
 
-	/** Builds a FMaterialShaderMapId for the material instance. */
-	void GetMaterialResourceId(const FMaterialResource* Resource, EShaderPlatform ShaderPlatform, const FStaticParameterSet& CompositedStaticParameters, FMaterialShaderMapId& OutId);
-
 	/** 
 	 * Cache resource shaders for rendering on the given shader platform. 
 	 * If a matching shader map is not found in memory or the DDC, a new one will be compiled.
diff --git a/Engine/Source/Runtime/Engine/Private/Materials/Material.cpp b/Engine/Source/Runtime/Engine/Private/Materials/Material.cpp
--- a/Engine/Source/Runtime/Engine/Private/Materials/Material.cpp
+++ b/Engine/Source/Runtime/Engine/Private/Materials/Material.cpp
@@ -126,6 +126,10 @@ void FMaterialResource::GetShaderMapId(EShaderPlatform Platform, FMaterialShader
 	if(MaterialInstance)
 	{
 		MaterialInstance->GetBasePropertyOverridesHash(OutId.BasePropertyOverridesHash);
+
+		FStaticParameterSet CompositedStaticParameters;
+		MaterialInstance->GetStaticParameterValues(CompositedStaticParameters);
+		OutId.ParameterSet = CompositedStaticParameters;
 	}
 }
 
diff --git a/Engine/Source/Runtime/Engine/Private/Materials/MaterialInstance.cpp b/Engine/Source/Runtime/Engine/Private/Materials/MaterialInstance.cpp
--- a/Engine/Source/Runtime/Engine/Private/Materials/MaterialInstance.cpp
+++ b/Engine/Source/Runtime/Engine/Private/Materials/MaterialInstance.cpp
@@ -1276,27 +1276,6 @@ void UMaterialInstance::GetAllShaderMaps(TArray<FMaterialShaderMap*>& OutShaderM
 	}
 }
 
-void UMaterialInstance::GetMaterialResourceId(EShaderPlatform ShaderPlatform, EMaterialQualityLevel::Type QualityLevel, FMaterialShaderMapId& OutId)
-{
-	UMaterial* BaseMaterial = GetMaterial();
-
-	FStaticParameterSet CompositedStaticParameters;
-	GetStaticParameterValues(CompositedStaticParameters);
-
-	ERHIFeatureLevel::Type FeatureLevel = GetMaxSupportedFeatureLevel(ShaderPlatform);
-	const FMaterialResource* BaseResource = BaseMaterial->GetMaterialResource(FeatureLevel, QualityLevel);
-
-	GetMaterialResourceId(BaseResource, ShaderPlatform, CompositedStaticParameters, OutId);
-}
-
-void UMaterialInstance::GetMaterialResourceId(const FMaterialResource* Resource, EShaderPlatform ShaderPlatform, const FStaticParameterSet& CompositedStaticParameters, FMaterialShaderMapId& OutId)
-{
-	Resource->GetShaderMapId(ShaderPlatform, OutId);
-
-	//@todo - should the resource be able to generate its own shadermap id?
-	OutId.ParameterSet = CompositedStaticParameters;
-}
-
 void UMaterialInstance::UpdatePermutationAllocations()
 {
 	if (bHasStaticPermutationResource)
@@ -1404,9 +1383,6 @@ void UMaterialInstance::CacheResourceShadersForCooking(EShaderPlatform ShaderPla
 
 void UMaterialInstance::CacheShadersForResources(EShaderPlatform ShaderPlatform, const TArray<FMaterialResource*>& ResourcesToCache, bool bApplyCompletedShaderMapForRendering)
 {
-	FStaticParameterSet CompositedStaticParameters;
-	GetStaticParameterValues(CompositedStaticParameters);
-
 	UMaterial* BaseMaterial = GetMaterial();
 
 	BaseMaterial->CacheExpressionTextureReferences();
@@ -1416,7 +1392,7 @@ void UMaterialInstance::CacheShadersForResources(EShaderPlatform ShaderPlatform,
 		FMaterialResource* CurrentResource = ResourcesToCache[ResourceIndex];
 
 		FMaterialShaderMapId ShaderMapId;
-		GetMaterialResourceId(CurrentResource, ShaderPlatform, CompositedStaticParameters, ShaderMapId);
+		CurrentResource->GetShaderMapId(ShaderPlatform, ShaderMapId);
 
 		const bool bSuccess = CurrentResource->CacheShaders(ShaderMapId, ShaderPlatform, bApplyCompletedShaderMapForRendering);
 
diff --git a/Engine/Source/Runtime/Engine/Public/MaterialShared.h b/Engine/Source/Runtime/Engine/Public/MaterialShared.h
--- a/Engine/Source/Runtime/Engine/Public/MaterialShared.h
+++ b/Engine/Source/Runtime/Engine/Public/MaterialShared.h
@@ -313,7 +313,10 @@ namespace EMaterialShaderMapUsage
 		LightmassExportEmissive,
 		LightmassExportDiffuse,
 		LightmassExportOpacity,
-		LightmassExportNormal
+		LightmassExportNormal,
+		MaterialExportDiffuseColor,
+		MaterialExportSpecularColor,
+		MaterialExportNormal
 	};
 }
 
