Fixed: Texture streaming streams out all textures and in again during many editor operations
#UE-5981

diff --git a/Engine/Source/Editor/UnrealEd/Private/Editor.cpp b/Engine/Source/Editor/UnrealEd/Private/Editor.cpp
--- a/Engine/Source/Editor/UnrealEd/Private/Editor.cpp
+++ b/Engine/Source/Editor/UnrealEd/Private/Editor.cpp
@@ -1624,6 +1624,13 @@ bool UEditorEngine::UpdateSingleViewportClient(FEditorViewportClient* InViewport
 {
 	bool bUpdatedNonRealtimeViewport = false;
 
+	// Always submit view information for content streaming 
+	// otherwise content for editor view can be streamed out if there are other views (ex: thumbnails)
+	if (InViewportClient->IsPerspective())
+	{
+		IStreamingManager::Get().AddViewInformation( InViewportClient->GetViewLocation(), InViewportClient->Viewport->GetSizeXY().X, InViewportClient->Viewport->GetSizeXY().X / FMath::Tan(InViewportClient->ViewFOV) );
+	}
+	
 	// Only allow viewports to be drawn if we are not throttling for slate UI responsiveness or if the viewport client requested a redraw
 	// Note about bNeedsRedraw: Redraws can happen during some Slate events like checking a checkbox in a menu to toggle a view mode in the viewport.  In those cases we need to show the user the results immediately
 	if( FSlateThrottleManager::Get().IsAllowingExpensiveTasks() || InViewportClient->bNeedsRedraw )
@@ -1634,7 +1641,6 @@ bool UEditorEngine::UpdateSingleViewportClient(FEditorViewportClient* InViewport
 		// Add view information for perspective viewports.
 		if( InViewportClient->IsPerspective() )
 		{
-			IStreamingManager::Get().AddViewInformation( InViewportClient->GetViewLocation(), InViewportClient->Viewport->GetSizeXY().X, InViewportClient->Viewport->GetSizeXY().X / FMath::Tan(InViewportClient->ViewFOV) );
 			GWorld->ViewLocationsRenderedLastFrame.Add(InViewportClient->GetViewLocation());
 	
 			// If we're currently simulating in editor, then we'll need to make sure that sub-levels are streamed in.
