Added missing feature pack files

diff --git a/Engine/Source/Editor/UnrealEd/Private/FeaturePack/FeaturePack.cpp b/Engine/Source/Editor/UnrealEd/Private/FeaturePack/FeaturePack.cpp
new file mode 100644
--- /dev/null
+++ b/Engine/Source/Editor/UnrealEd/Private/FeaturePack/FeaturePack.cpp
@@ -0,0 +1,78 @@
+// Copyright 1998-2014 Epic Games, Inc. All Rights Reserved.
+
+#include "UnrealEd.h"
+#include "FeaturePack/FeaturePack.h"
+#include "AssetToolsModule.h"
+#include "AssetEditorManager.h"
+#include "AssetRegistryModule.h"
+#include "LevelEditor.h"
+#include "Toolkits/AssetEditorManager.h"
+
+
+DEFINE_LOG_CATEGORY_STATIC(LogFeaturePack, Log, All);
+
+FFeaturePack::FFeaturePack()
+{
+}
+
+void FFeaturePack::ImportPendingPacks()
+{	
+	bool bAddPacks;
+	if (GConfig->GetBool(TEXT("StartupActions"), TEXT("bAddPacks"), bAddPacks, GGameIni) == true)
+	{
+		if (bAddPacks == true)
+		{
+			ParsePacks();
+			UE_LOG(LogFeaturePack, Warning, TEXT("Inserted %d feature packs"), ImportSets.Num() );
+ 			GConfig->SetBool(TEXT("StartupActions"), TEXT("bAddPacks"), false, GGameIni);
+ 			GConfig->Flush(true, GGameIni);
+		}
+	}	
+}
+
+void FFeaturePack::ParsePacks()
+{
+	FAssetToolsModule& AssetToolsModule = FModuleManager::Get().LoadModuleChecked<FAssetToolsModule>("AssetTools");
+	TArray<FString> PacksToAdd;
+	GConfig->GetArray(TEXT("StartupActions"), TEXT("InsertPack"), PacksToAdd, GGameIni);
+	for (int32 iPackEntry = 0; iPackEntry < PacksToAdd.Num(); iPackEntry++)
+	{
+		PackData EachPackData;
+
+		TArray<FString> PackEntries;
+		PacksToAdd[iPackEntry].ParseIntoArray(&PackEntries, TEXT(","), true);
+		FString PackSource;
+		FString PackName;
+		for (int32 iEntry = 0; iEntry < PackEntries.Num(); iEntry++)
+		{
+			FString EachString = PackEntries[iEntry];
+			// String the parenthesis
+			EachString = EachString.Replace(TEXT("("), TEXT(""));
+			EachString = EachString.Replace(TEXT(")"), TEXT(""));
+			if (EachString.StartsWith(TEXT("PackSource=")) == true)
+			{
+				EachString = EachString.Replace(TEXT("PackSource="), TEXT(""));
+				EachString = EachString.TrimQuotes();
+				EachPackData.PackSource = EachString;
+			}
+			if (EachString.StartsWith(TEXT("PackName=")) == true)
+			{
+				EachString = EachString.Replace(TEXT("PackName="), TEXT(""));
+				EachString = EachString.TrimQuotes();
+				EachPackData.PackName = EachString;
+			}
+		}
+		
+		if ((EachPackData.PackSource.IsEmpty() == false) && (EachPackData.PackName.IsEmpty() == false))
+		{
+			TArray<FString> EachImport;
+			FString FullPath = FPaths::FeaturePackDir() + EachPackData.PackSource;
+			EachImport.Add(FullPath);
+			EachPackData.ImportedObjects = AssetToolsModule.Get().ImportAssets(EachImport, EachPackData.PackName, nullptr);
+			ImportSets.Add(EachPackData);
+		}
+	}
+}
+
+
+#undef  LOCTEXT_NAMESPACE
diff --git a/Engine/Source/Editor/UnrealEd/Public/FeaturePack/FeaturePack.h b/Engine/Source/Editor/UnrealEd/Public/FeaturePack/FeaturePack.h
new file mode 100644
--- /dev/null
+++ b/Engine/Source/Editor/UnrealEd/Public/FeaturePack/FeaturePack.h
@@ -0,0 +1,26 @@
+// Copyright 1998-2014 Epic Games, Inc. All Rights Reserved.
+
+#pragma once
+
+
+class FFeaturePack
+{
+public:
+	FFeaturePack();
+
+	/* Import any pending feature packs. */
+	void ImportPendingPacks();
+
+	/* Parse and insert any feature packs flagged to include and import them */
+	void ParsePacks();
+
+private:
+	struct PackData
+	{
+		FString PackSource;
+		FString PackName;
+		TArray<UObject*>	ImportedObjects;
+	};
+	
+	TArray<PackData>	ImportSets;
+};
