UE-5162 Changing the default value of a variable and undoing results in a crash if referenced blueprints are open

Transaction buffer doesn't store direct ptr to CDO.

diff --git a/Engine/Source/Editor/UnrealEd/Classes/Editor/Transactor.h b/Engine/Source/Editor/UnrealEd/Classes/Editor/Transactor.h
--- a/Engine/Source/Editor/UnrealEd/Classes/Editor/Transactor.h
+++ b/Engine/Source/Editor/UnrealEd/Classes/Editor/Transactor.h
@@ -64,6 +64,39 @@ protected:
 	class UNREALED_API FObjectRecord
 	{
 	public:
+
+		struct FPersistentObjectRef
+		{
+		private:
+			UObject* Object;
+			UClass* SourceCDO;
+		public:
+			FPersistentObjectRef()
+				: Object(NULL)
+				, SourceCDO(NULL)
+			{}
+
+			FPersistentObjectRef(UObject* InObject)
+			{
+				const bool bIsCDO = InObject && InObject->HasAllFlags(RF_ClassDefaultObject);
+				Object = bIsCDO ? NULL : InObject;
+				SourceCDO = bIsCDO ? InObject->GetClass() : NULL;
+			}
+
+			UObject* Get() const
+			{
+				checkSlow(!SourceCDO || !Object);
+				return SourceCDO ? SourceCDO->GetDefaultObject(false) : Object;
+			}
+
+			UObject* operator->() const
+			{
+				auto Obj = Get();
+				check(Obj);
+				return Obj;
+			}
+		};
+
 		// Variables.
 		/** The data stream used to serialize/deserialize record */
 		TArray<uint8>		Data;
@@ -72,7 +105,7 @@ protected:
 		/** FNames referenced in the object record */
 		TArray<FName>		ReferencedNames;
 		/** The object to track */
-		UObject*			Object;
+		FPersistentObjectRef	Object;
 		/** Annotation data for the object stored externally */
 		TSharedPtr<ITransactionObjectAnnotation> ObjectAnnotation;
 		/** Array: If an array object, reference to script array */
@@ -163,7 +196,7 @@ protected:
 				{
 					for( int32 i=0; i<Owner->Records.Num(); i++ )
 					{
-						if( Owner->Records[i].Object==InObject )
+						if( Owner->Records[i].Object.Get()==InObject )
 						{
 							Owner->Records[i].Restore( Owner );
 						}
diff --git a/Engine/Source/Editor/UnrealEd/Private/EditorTransaction.cpp b/Engine/Source/Editor/UnrealEd/Private/EditorTransaction.cpp
--- a/Engine/Source/Editor/UnrealEd/Private/EditorTransaction.cpp
+++ b/Engine/Source/Editor/UnrealEd/Private/EditorTransaction.cpp
@@ -40,8 +40,8 @@ void FTransaction::FObjectRecord::SerializeContents( FArchive& Ar, int32 InOper
 	{
 		//UE_LOG( LogEditorTransaction, Log, TEXT("Array %s %i*%i: %i"), Object ? *Object->GetFullName() : TEXT("Invalid Object"), Index, ElementSize, InOper);
 
-		check((SIZE_T)Array>=(SIZE_T)Object+sizeof(UObject));
-		check((SIZE_T)Array+sizeof(FScriptArray)<=(SIZE_T)Object+Object->GetClass()->GetPropertiesSize());
+		check((SIZE_T)Array >= (SIZE_T)Object.Get() + sizeof(UObject));
+		check((SIZE_T)Array + sizeof(FScriptArray) <= (SIZE_T)Object.Get() + Object->GetClass()->GetPropertiesSize());
 		check(ElementSize!=0);
 		check(DefaultConstructor!=NULL);
 		check(Serializer!=NULL);
@@ -152,9 +152,11 @@ void FTransaction::DumpObjectMap(FOutputDevice& Ar) const
 
 FArchive& operator<<( FArchive& Ar, FTransaction::FObjectRecord& R )
 {
-	check(R.Object);
+	UObject* Object = R.Object.Get();
+	check(Object);
 	FMemMark Mark(FMemStack::Get());
-	Ar << R.Object;
+	Ar << Object;
+	R.Object = Object;
 	Ar << R.Data;
 	Ar << R.ReferencedObjects;
 	Ar << R.ReferencedNames;
@@ -164,7 +166,9 @@ FArchive& operator<<( FArchive& Ar, FTransaction::FObjectRecord& R )
 
 void FTransaction::FObjectRecord::AddReferencedObjects( FReferenceCollector& Collector )
 {
-	Collector.AddReferencedObject( Object );
+	UObject* Obj = Object.Get();
+	Collector.AddReferencedObject(Obj);
+	Object = Obj;
 	for( int32 ObjIndex = 0; ObjIndex < ReferencedObjects.Num(); ObjIndex++ )
 	{
 		Collector.AddReferencedObject( ReferencedObjects[ ObjIndex ] );
@@ -247,13 +251,14 @@ void FTransaction::Apply()
 		FObjectRecord& Record = Records[i];
 		Record.bRestored = false;
 
-		if (!ChangedObjects.Contains(Record.Object))
+		UObject* Object = Record.Object.Get();
+		if (!ChangedObjects.Contains(Object))
 		{
-			Record.Object->CheckDefaultSubobjects();
-			Record.Object->PreEditUndo();
+			Object->CheckDefaultSubobjects();
+			Object->PreEditUndo();
 		}
 
-		ChangedObjects.Add(Record.Object, Record.ObjectAnnotation);
+		ChangedObjects.Add(Object, Record.ObjectAnnotation);
 	}
 	for( int32 i=Start; i!=End; i+=Inc )
 	{
@@ -320,10 +325,10 @@ void FTransaction::GetTransactionObjects(TArray<UObject*>& Objects)
 
 	for(int32 i=0; i<Records.Num(); i++)
 	{
-		UObject* obj = Records[i].Object;
-		if(obj)
+		UObject* Obj = Records[i].Object.Get();
+		if (Obj)
 		{
-			Objects.AddUnique(obj);
+			Objects.AddUnique(Obj);
 		}
 	}
 }
