Automation test for the launch on feature in the editor.

diff --git a/Engine/Source/Editor/UnrealEd/Classes/Editor/EditorEngine.h b/Engine/Source/Editor/UnrealEd/Classes/Editor/EditorEngine.h
--- a/Engine/Source/Editor/UnrealEd/Classes/Editor/EditorEngine.h
+++ b/Engine/Source/Editor/UnrealEd/Classes/Editor/EditorEngine.h
@@ -10,6 +10,7 @@
 #include "../Settings/LevelEditorViewportSettings.h"
 #include "EditorEngine.generated.h"
 
+
 class APlayerStart;
 class FAssetData;
 class FPoly;
@@ -19,24 +20,24 @@ class USkeleton;
 // Things to set in mapSetBrush.
 //
 UENUM()
-enum EMapSetBrushFlags				
+enum EMapSetBrushFlags
 {
 	// Set brush color.
-	MSB_BrushColor	= 1,
+	MSB_BrushColor = 1,
 	// Set group.
-	MSB_Group		= 2,
+	MSB_Group = 2,
 	// Set poly flags.
-	MSB_PolyFlags	= 4,
+	MSB_PolyFlags = 4,
 	// Set CSG operation.
-	MSB_BrushType	= 8,
+	MSB_BrushType = 8,
 };
 
 UENUM()
 enum EPasteTo
 {
-	PT_OriginalLocation	= 0,
-	PT_Here				= 1,
-	PT_WorldOrigin		= 2
+	PT_OriginalLocation = 0,
+	PT_Here = 1,
+	PT_WorldOrigin = 2
 };
 
 USTRUCT()
@@ -49,15 +50,15 @@ struct FSlatePlayInEditorInfo
 
 	/** The current play in editor SWindow if playing in a floating window */
 	TWeakPtr<class SWindow>				SlatePlayInEditorWindow;
-	
+
 	/** The current play in editor rendering and I/O viewport if playing in a floating window*/
 	TSharedPtr<class FSceneViewport>	SlatePlayInEditorWindowViewport;
-	
+
 	/** The slate viewport that should be used for play in viewport */
 	TWeakPtr<class ILevelViewport>		DestinationSlateViewport;
 
 	FSlatePlayInEditorInfo()
-	: SlatePlayInEditorWindow(NULL), DestinationSlateViewport(NULL)
+		: SlatePlayInEditorWindow(NULL), DestinationSlateViewport(NULL)
 	{}
 };
 
@@ -71,15 +72,15 @@ public:
 
 	GENERATED_USTRUCT_BODY()
 
-	/** Type of account. Needed to identity the auth method to use (epic, internal, facebook, etc) */
-	UPROPERTY()
-	FString Type;
+		/** Type of account. Needed to identity the auth method to use (epic, internal, facebook, etc) */
+		UPROPERTY()
+		FString Type;
 	/** Id of the user logging in (email, display name, facebook id, etc) */
 	UPROPERTY()
-	FString Id;
+		FString Id;
 	/** Credentials of the user logging in (password or auth token) */
 	UPROPERTY()
-	FString Token;
+		FString Token;
 };
 
 /**
@@ -125,16 +126,16 @@ struct FCopySelectedInfo
 	/** Do not cache this info, it is only valid after a call to CanCopySelectedActorsToClipboard has been made, and becomes redundant
 	when the current selection changes. Used to determine whether a copy can be performed based on the current selection state */
 	FCopySelectedInfo()
-		: bHasSelectedActors( false )
-		, bAllActorsInSameLevel( true )
-		, LevelAllActorsAreIn( NULL )
-		, bHasSelectedSurfaces( false )
-		, LevelWithSelectedSurface( NULL )
+	: bHasSelectedActors(false)
+	, bAllActorsInSameLevel(true)
+	, LevelAllActorsAreIn(NULL)
+	, bHasSelectedSurfaces(false)
+	, LevelWithSelectedSurface(NULL)
 	{}
 
 
 	/** Does the current selection contain actors */
-	bool bHasSelectedActors;		
+	bool bHasSelectedActors;
 
 	/** If we have selected actors, are they within the same level */
 	bool bAllActorsInSameLevel;
@@ -156,9 +157,9 @@ struct FCopySelectedInfo
 		// If there are selected actors and BSP surfaces AND all selected actors 
 		// and surfaces are in the same level, we can do a quick copy.
 		bool bCanPerformQuickCopy = false;
-		if( LevelAllActorsAreIn && LevelWithSelectedSurface )
+		if (LevelAllActorsAreIn && LevelWithSelectedSurface)
 		{
-			bCanPerformQuickCopy = ( LevelWithSelectedSurface == LevelAllActorsAreIn );
+			bCanPerformQuickCopy = (LevelWithSelectedSurface == LevelAllActorsAreIn);
 		}
 		// Else, if either we have only selected actors all in one level OR we have 
 		// only selected surfaces all in one level, then we can perform a quick copy. 
@@ -202,7 +203,7 @@ private:
  * Engine that drives the Editor.
  * Separate from UGameEngine because it may have much different functionality than desired for an instance of a game itself.
  */
-UCLASS(config=Engine, transient)
+UCLASS(config = Engine, transient)
 class UNREALED_API UEditorEngine : public UEngine
 {
 	GENERATED_UCLASS_BODY()
@@ -253,98 +254,98 @@ class UNREALED_API UEditorEngine : public UEngine
 
 	// Toggles.
 	UPROPERTY()
-	uint32 bFastRebuild:1;
+		uint32 bFastRebuild : 1;
 
 	UPROPERTY()
-	uint32 IsImportingT3D:1;
+		uint32 IsImportingT3D : 1;
 
 	// Other variables.
 	UPROPERTY()
-	uint32 ClickFlags;
+		uint32 ClickFlags;
 
 	UPROPERTY()
 	class UPackage* ParentContext;
 
 	UPROPERTY()
-	FVector ClickLocation;
+		FVector ClickLocation;
 
 	UPROPERTY()
-	FPlane ClickPlane;
+		FPlane ClickPlane;
 
 	UPROPERTY()
-	FVector MouseMovement;
+		FVector MouseMovement;
 
 	// Setting for the detail mode to show in the editor viewports
 	UPROPERTY()
-	TEnumAsByte<enum EDetailMode> DetailMode;
+		TEnumAsByte<enum EDetailMode> DetailMode;
 
 	// Advanced.
-	UPROPERTY(EditAnywhere, config, Category=Advanced)
-	uint32 UseSizingBox:1;
+	UPROPERTY(EditAnywhere, config, Category = Advanced)
+		uint32 UseSizingBox : 1;
 
-	UPROPERTY(EditAnywhere, config, Category=Advanced)
-	uint32 UseAxisIndicator:1;
+	UPROPERTY(EditAnywhere, config, Category = Advanced)
+		uint32 UseAxisIndicator : 1;
 
-	UPROPERTY(EditAnywhere, config, Category=Advanced)
-	uint32 GodMode:1;
+	UPROPERTY(EditAnywhere, config, Category = Advanced)
+		uint32 GodMode : 1;
 
 	/** The location to autosave to. */
-	UPROPERTY(EditAnywhere, config, Category=Advanced)
-	FString AutoSaveDir;
+	UPROPERTY(EditAnywhere, config, Category = Advanced)
+		FString AutoSaveDir;
 
-	UPROPERTY(EditAnywhere, config, Category=Advanced)
-	uint32 InvertwidgetZAxis:1;
+	UPROPERTY(EditAnywhere, config, Category = Advanced)
+		uint32 InvertwidgetZAxis : 1;
 
-	UPROPERTY(EditAnywhere, config, Category=Advanced)
-	FString GameCommandLine;
+	UPROPERTY(EditAnywhere, config, Category = Advanced)
+		FString GameCommandLine;
 
 	/** If true, show translucent marker polygons on the builder brush and volumes. */
-	UPROPERTY(EditAnywhere, config, Category=Advanced)
-	uint32 bShowBrushMarkerPolys:1;
+	UPROPERTY(EditAnywhere, config, Category = Advanced)
+		uint32 bShowBrushMarkerPolys : 1;
 
 	UPROPERTY()
-	uint32 bPrefabsLocked_DEPRECATED:1;
+		uint32 bPrefabsLocked_DEPRECATED : 1;
 
 	/** If true, socket snapping is enabled in the main level viewports. */
-	UPROPERTY(EditAnywhere, config, Category=Advanced)
-	uint32 bEnableSocketSnapping:1;
+	UPROPERTY(EditAnywhere, config, Category = Advanced)
+		uint32 bEnableSocketSnapping : 1;
 
 	/** If true, same type views will be camera tied, and ortho views will use perspective view for LOD parenting */
 	UPROPERTY()
-	uint32 bEnableLODLocking:1;
+		uint32 bEnableLODLocking : 1;
 
 	/** If true, actors can be grouped and grouping rules will be maintained. When deactivated, any currently existing groups will still be preserved.*/
-	UPROPERTY(EditAnywhere, config, Category=Advanced)
-	uint32 bGroupingActive:1;
+	UPROPERTY(EditAnywhere, config, Category = Advanced)
+		uint32 bGroupingActive : 1;
 
 	UPROPERTY(config)
-	FString HeightMapExportClassName;
+		FString HeightMapExportClassName;
 
 	/** Array of actor factories created at editor startup and used by context menu etc. */
 	UPROPERTY()
-	TArray<class UActorFactory*> ActorFactories;
+		TArray<class UActorFactory*> ActorFactories;
 
 	/** Actors that are being deleted and should processed in the global re-attach*/
 	UPROPERTY()
-	TArray<class AActor*> ActorsForGlobalReregister;
+		TArray<class AActor*> ActorsForGlobalReregister;
 
-	/** String that maps one class name to another, used to create hook for game-specific actors created through shortcuts etc 
+	/** String that maps one class name to another, used to create hook for game-specific actors created through shortcuts etc
 	 *  Pairing is "ORIGINALCLASS;DESIREDCLASS
 	 *  (ie APylon;AMyGamePylon)
 	 */
 	UPROPERTY(config)
-	TArray<FString> ClassMapPair;
+		TArray<FString> ClassMapPair;
 
 	/** The name of the file currently being opened in the editor. "" if no file is being opened. */
 	UPROPERTY()
-	FString UserOpenedFile;
+		FString UserOpenedFile;
 
 	///////////////////////////////
 	// "Play From Here" properties
-	
+
 	/** Additional per-user/per-game options set in the .ini file. Should be in the form "?option1=X?option2?option3=Y"					*/
-	UPROPERTY(EditAnywhere, config, Category=Advanced)
-	FString InEditorGameURLOptions;
+	UPROPERTY(EditAnywhere, config, Category = Advanced)
+		FString InEditorGameURLOptions;
 
 	/** A pointer to a UWorld that is the duplicated/saved-loaded to be played in with "Play From Here" 								*/
 	UPROPERTY()
@@ -352,83 +353,83 @@ class UNREALED_API UEditorEngine : public UEngine
 
 	/** An optional location for the starting location for "Play From Here"																*/
 	UPROPERTY()
-	FVector PlayWorldLocation;
+		FVector PlayWorldLocation;
 
 	/** An optional rotation for the starting location for "Play From Here"																*/
 	UPROPERTY()
-	FRotator PlayWorldRotation;
+		FRotator PlayWorldRotation;
 
 	/** Has a request for "Play From Here" been made?													 								*/
 	UPROPERTY()
-	uint32 bIsPlayWorldQueued:1;
+		uint32 bIsPlayWorldQueued : 1;
 
 	/** Has a request to toggle between PIE and SIE been made? */
 	UPROPERTY()
-	uint32 bIsToggleBetweenPIEandSIEQueued:1;
+		uint32 bIsToggleBetweenPIEandSIEQueued : 1;
 
 	/** True if we are requesting to start a simulation-in-editor session */
 	UPROPERTY()
-	uint32 bIsSimulateInEditorQueued:1;
+		uint32 bIsSimulateInEditorQueued : 1;
 
 	/** Allows multipel PIE worlds under a single instance. If false, you can only do multiple UE4 processes for pie networking */
 	UPROPERTY(globalconfig)
-	uint32 bAllowMultiplePIEWorlds:1;
+		uint32 bAllowMultiplePIEWorlds : 1;
 
 	/** The PlayerStart class used when spawning the player at the current camera location. */
 	UPROPERTY()
-	TSubclassOf<class ANavigationObjectBase>  PlayFromHerePlayerStartClass;
+		TSubclassOf<class ANavigationObjectBase>  PlayFromHerePlayerStartClass;
 
 	/** True if there is a pending end play map queued */
 	UPROPERTY()
-	uint32 bRequestEndPlayMapQueued:1;
+		uint32 bRequestEndPlayMapQueued : 1;
 
 	/** Did the request include the optional location and rotation?										 								*/
 	UPROPERTY()
-	uint32 bHasPlayWorldPlacement:1;
+		uint32 bHasPlayWorldPlacement : 1;
 
 	/** True to enable mobile preview mode when launching the game from the editor on PC platform */
 	UPROPERTY()
-	uint32 bUseMobilePreviewForPlayWorld:1;
+		uint32 bUseMobilePreviewForPlayWorld : 1;
 
 	/** True if we're Simulating In Editor, as opposed to Playing In Editor.  In this mode, simulation takes place right the level editing environment */
 	UPROPERTY()
-	uint32 bIsSimulatingInEditor:1;
+		uint32 bIsSimulatingInEditor : 1;
 
 	/** When Simulating In Editor, a pointer to the original (non-simulating) editor world */
 	UPROPERTY()
 	class UWorld* EditorWorld;
-	
+
 	/** When Simulating In Editor, an array of all actors that were selected when it began*/
 	UPROPERTY()
-	TArray<TWeakObjectPtr<class AActor> > ActorsThatWereSelected;
+		TArray<TWeakObjectPtr<class AActor> > ActorsThatWereSelected;
 
 	/** True to start movie capturing right away when launching the game from the editor on PC platform */
 	UPROPERTY()
-	uint32 bStartMovieCapture:1;
+		uint32 bStartMovieCapture : 1;
 
 	/** Where did the person want to play? Where to play the game - -1 means in editor, 0 or more is an index into the GConsoleSupportContainer	*/
 	UPROPERTY()
-	int32 PlayWorldDestination;
+		int32 PlayWorldDestination;
 
 	/** The current play world destination (I.E console).  -1 means no current play world destination, 0 or more is an index into the GConsoleSupportContainer	*/
 	UPROPERTY()
-	int32 CurrentPlayWorldDestination;
+		int32 CurrentPlayWorldDestination;
 
 	/** Mobile preview settings for what orientation to default to */
 	UPROPERTY(config)
-	uint32 bMobilePreviewPortrait:1;
+		uint32 bMobilePreviewPortrait : 1;
 
 	/** Currently targeted device for mobile previewer. */
 	UPROPERTY(config)
-	int32 BuildPlayDevice;
+		int32 BuildPlayDevice;
 
 	/** Enabled online PIE */
 	UPROPERTY(config)
-	bool bOnlinePIEEnabled;
+		bool bOnlinePIEEnabled;
 
 	/** Logins available for use when running PIE instances */
 	UPROPERTY(config)
-	TArray<FPIELoginInfo> PIELogins;
+		TArray<FPIELoginInfo> PIELogins;
 
 	/** Maps world contexts to their slate data */
 	TMap<FName, FSlatePlayInEditorInfo>	SlatePlayInEditorMap;
@@ -438,19 +439,19 @@ class UNREALED_API UEditorEngine : public UEngine
 
 	/** When set to anything other than -1, indicates a specific In-Editor viewport index that PIE should use */
 	UPROPERTY()
-	int32 PlayInEditorViewportIndex;
+		int32 PlayInEditorViewportIndex;
 
 	/** The width resolution that we want to use for the matinee capture */
 	UPROPERTY()
-	int32 MatineeCaptureResolutionX;
+		int32 MatineeCaptureResolutionX;
 
 	/** The height resolution that we want to use for the matinee capture */
 	UPROPERTY()
-	int32 MatineeCaptureResolutionY;
+		int32 MatineeCaptureResolutionY;
 
 	/** Play world url string edited by a user. */
 	UPROPERTY()
-	FString UserEditedPlayWorldURL;
+		FString UserEditedPlayWorldURL;
 
 	/** Temporary render target that can be used by the editor. */
 	UPROPERTY(transient)
@@ -485,31 +486,31 @@ class UNREALED_API UEditorEngine : public UEngine
 
 	/** The index of the mesh to use from the list of preview meshes. */
 	UPROPERTY()
-	int32 PreviewMeshIndex;
+		int32 PreviewMeshIndex;
 
 	/** When true, the preview mesh mode is activated. */
 	UPROPERTY()
-	uint32 bShowPreviewMesh:1;
+		uint32 bShowPreviewMesh : 1;
 
 	/** If "Camera Align" emitter handling uses a custom zoom or not */
 	UPROPERTY(config)
-	uint32 bCustomCameraAlignEmitter:1;
+		uint32 bCustomCameraAlignEmitter : 1;
 
 	/** The distance to place the camera from an emitter actor when custom zooming is enabled */
 	UPROPERTY(config)
-	float CustomCameraAlignEmitterDistance;
+		float CustomCameraAlignEmitterDistance;
 
 	/** If true, then draw sockets when socket snapping is enabled in 'g' mode */
 	UPROPERTY(config)
-	uint32 bDrawSocketsInGMode:1;
+		uint32 bDrawSocketsInGMode : 1;
 
 	/** If true, then draw particle debug helpers in editor viewports */
 	UPROPERTY(transient)
-	uint32 bDrawParticleHelpers:1;
+		uint32 bDrawParticleHelpers : 1;
 
 	/** Brush builders that have been created in the editor */
 	UPROPERTY(transient)
-	TArray<class UBrushBuilder*> BrushBuilders;	
+		TArray<class UBrushBuilder*> BrushBuilders;
 
 	/** Whether or not to recheck the current actor selection for lock actors the next time HasLockActors is called */
 	bool bCheckForLockActors;
@@ -538,14 +539,14 @@ public:
 	class FUObjectAnnotationSparseBool ObjectsThatExistInEditorWorld;
 
 	/** Called prior to a Blueprint compile */
-	DECLARE_EVENT_OneParam( UEditorEngine, FBlueprintPreCompileEvent, UBlueprint* );
+	DECLARE_EVENT_OneParam(UEditorEngine, FBlueprintPreCompileEvent, UBlueprint*);
 	FBlueprintPreCompileEvent& OnBlueprintPreCompile() { return BlueprintPreCompileEvent; }
 
 	/** Broadcasts that a Blueprint is about to be compiled */
 	void BroadcastBlueprintPreCompile(UBlueprint* BlueprintToCompile) { BlueprintPreCompileEvent.Broadcast(BlueprintToCompile); }
 
 	/** Called when a Blueprint compile is completed. */
-	DECLARE_EVENT( UEditorEngine, FBlueprintCompiledEvent );
+	DECLARE_EVENT(UEditorEngine, FBlueprintCompiledEvent);
 	FBlueprintCompiledEvent& OnBlueprintCompiled() { return BlueprintCompiledEvent; }
 
 	/**	Broadcasts that a blueprint just finished compiling. THIS SHOULD NOT BE PUBLIC */
@@ -560,56 +561,56 @@ public:
 
 	/** Called when uobjects have been replaced to allow others a chance to fix their references. */
 	typedef TMap<UObject*, UObject*> ReplacementObjectMap;
-	DECLARE_EVENT_OneParam( UEditorEngine, FObjectsReplacedEvent, const ReplacementObjectMap& );
+	DECLARE_EVENT_OneParam(UEditorEngine, FObjectsReplacedEvent, const ReplacementObjectMap&);
 	FObjectsReplacedEvent& OnObjectsReplaced() { return ObjectsReplacedEvent; }
 
 	/**	Broadcasts that objects have been replaced*/
 	void BroadcastBlueprintCompiled(const TMap<UObject*, UObject*>& ReplacementMap) { ObjectsReplacedEvent.Broadcast(ReplacementMap); }
 
 	/** Called when a package with data-driven classes becomes loaded or unloaded */
-	DECLARE_EVENT( UEditorEngine, FClassPackageLoadedOrUnloadedEvent );
+	DECLARE_EVENT(UEditorEngine, FClassPackageLoadedOrUnloadedEvent);
 	FClassPackageLoadedOrUnloadedEvent& OnClassPackageLoadedOrUnloaded() { return ClassPackageLoadedOrUnloadedEvent; }
 
 	/**	Broadcasts that a class package was just loaded or unloaded. THIS SHOULD NOT BE PUBLIC */
 	void BroadcastClassPackageLoadedOrUnloaded() { ClassPackageLoadedOrUnloadedEvent.Broadcast(); }
 
 	/** Called when an object is reimported. */
-	DECLARE_EVENT_OneParam( UEditorEngine, FObjectReimported, UObject* );
+	DECLARE_EVENT_OneParam(UEditorEngine, FObjectReimported, UObject*);
 	FObjectReimported& OnObjectReimported() { return ObjectReimportedEvent; }
 
 	/** Editor-only event triggered before an actor or component is moved, rotated or scaled by an editor system */
-	DECLARE_EVENT_OneParam( UEditorEngine, FOnBeginTransformObject, UObject& );
+	DECLARE_EVENT_OneParam(UEditorEngine, FOnBeginTransformObject, UObject&);
 	FOnBeginTransformObject& OnBeginObjectMovement() { return OnBeginObjectTransformEvent; }
 
 	/** Editor-only event triggered before after actor or component has moved, rotated or scaled by an editor system */
-	DECLARE_EVENT_OneParam( UEditorEngine, FOnEndTransformObject, UObject& );
+	DECLARE_EVENT_OneParam(UEditorEngine, FOnEndTransformObject, UObject&);
 	FOnEndTransformObject& OnEndObjectMovement() { return OnEndObjectTransformEvent; }
 
 	/** Delegate broadcast by the engine every tick when PIE/SIE is active, to check to see whether we need to
 		be able to capture state for simulating actor (for Sequencer recording features).  The single bool parameter
 		should be set to true if recording features are needed. */
-	DECLARE_EVENT_OneParam( UEditorEngine, FGetActorRecordingState, bool& /* bIsRecordingActive */ );
+	DECLARE_EVENT_OneParam(UEditorEngine, FGetActorRecordingState, bool& /* bIsRecordingActive */);
 	FGetActorRecordingState& GetActorRecordingState() { return GetActorRecordingStateEvent; }
 
-	/** 
-	 * Called before an actor or component is about to be translated, rotated, or scaled by the editor 
-	 * 
+	/**
+	 * Called before an actor or component is about to be translated, rotated, or scaled by the editor
+	 *
 	 * @param Object	The actor or component that will be moved
 	 */
-	void BroadcastBeginObjectMovement( UObject& Object ) const { OnBeginObjectTransformEvent.Broadcast( Object ); }
+	void BroadcastBeginObjectMovement(UObject& Object) const { OnBeginObjectTransformEvent.Broadcast(Object); }
 
 	/**
 	 * Called when an actor or component has been translated, rotated, or scaled by the editor
 	 *
 	 * @param Object	The actor or component that moved
 	 */
-	void BroadcastEndObjectMovement( UObject& Object ) const { OnEndObjectTransformEvent.Broadcast( Object ); }
+	void BroadcastEndObjectMovement(UObject& Object) const { OnEndObjectTransformEvent.Broadcast(Object); }
 
 	/**	Broadcasts that an object has been reimported. THIS SHOULD NOT BE PUBLIC */
 	void BroadcastObjectReimported(UObject* InObject) { ObjectReimportedEvent.Broadcast(InObject); }
 
 	// Begin UObject interface.
-	virtual void FinishDestroy() override;	
+	virtual void FinishDestroy() override;
 	virtual void PostEditChangeProperty(FPropertyChangedEvent& PropertyChangedEvent) override;
 	static void AddReferencedObjects(UObject* InThis, FReferenceCollector& Collector);
 	// End UObject interface.
@@ -618,55 +619,55 @@ public:
 	virtual void Init(IEngineLoop* InEngineLoop) override;
 	virtual void InitializeObjectReferences() override;
 	// End UEngine interface.
-	
+
 	// Begin FExec Interface
-	virtual bool Exec( UWorld* InWorld, const TCHAR* Cmd, FOutputDevice& Ar=*GLog ) override;
+	virtual bool Exec(UWorld* InWorld, const TCHAR* Cmd, FOutputDevice& Ar = *GLog) override;
 	// End FExec Interface
 
-	bool	CommandIsDeprecated( const TCHAR* CommandStr, FOutputDevice& Ar );
-	
+	bool	CommandIsDeprecated(const TCHAR* CommandStr, FOutputDevice& Ar);
+
 	/**
 	 * Exec command handlers
 	 */
-	bool	HandleBlueprintifyFunction( const TCHAR* Str , FOutputDevice& Ar );
-	bool	HandleCallbackCommand( const TCHAR* Str, FOutputDevice& Ar );
-	bool	HandleTestPropsCommand( const TCHAR* Str, FOutputDevice& Ar );
-	bool	HandleMapCommand( const TCHAR* Str, FOutputDevice& Ar, UWorld* InWorld );
-	bool	HandleSelectCommand( const TCHAR* Str, FOutputDevice& Ar, UWorld* InWorld );
-	bool	HandleDeleteCommand( const TCHAR* Str, FOutputDevice& Ar, UWorld* InWorld );
-	bool	HandlePrefabCommand( const TCHAR* Str, FOutputDevice& Ar );
-	bool	HandleLightmassDebugCommand( const TCHAR* Str, FOutputDevice& Ar );
-	bool	HandleLightmassStatsCommand( const TCHAR* Str, FOutputDevice& Ar );
-	bool	HandleSwarmDistributionCommand( const TCHAR* Str, FOutputDevice& Ar );
-	bool	HandleLightmassImmediateImportCommand( const TCHAR* Str, FOutputDevice& Ar );
-	bool	HandleLightmassImmediateProcessCommand( const TCHAR* Str, FOutputDevice& Ar );
-	bool	HandleLightmassSortCommand( const TCHAR* Str, FOutputDevice& Ar );
-	bool	HandleLightmassDebugMaterialCommand( const TCHAR* Str, FOutputDevice& Ar );
-	bool	HandleLightmassPaddingCommand( const TCHAR* Str, FOutputDevice& Ar );
-	bool	HandleLightmassDebugPaddingCommand( const TCHAR* Str, FOutputDevice& Ar );
-	bool	HandleLightmassProfileCommand( const TCHAR* Str, FOutputDevice& Ar );
-	bool	HandleSetReplacementCommand( const TCHAR* Str, FOutputDevice& Ar, UWorld* InWorld );
-	bool	HandleSelectNameCommand( const TCHAR* Str, FOutputDevice& Ar, UWorld* InWorld  );
-	bool	HandleDumpPublicCommand( const TCHAR* Str, FOutputDevice& Ar );
-	bool	HandleJumpToCommand( const TCHAR* Str, FOutputDevice& Ar );
-	bool	HandleBugItGoCommand( const TCHAR* Str, FOutputDevice& Ar );
-	bool	HandleTagSoundsCommand( const TCHAR* Str, FOutputDevice& Ar );
-	bool	HandlecheckSoundsCommand( const TCHAR* Str, FOutputDevice& Ar );
-	bool	HandleFixupBadAnimNotifiersCommand( const TCHAR* Str, FOutputDevice& Ar );
-	bool	HandleSetDetailModeCommand( const TCHAR* Str, FOutputDevice& Ar );
-	bool	HandleSetDetailModeViewCommand( const TCHAR* Str, FOutputDevice& Ar, UWorld* InWorld );
-	bool	HandleCleanBSPMaterialCommand( const TCHAR* Str, FOutputDevice& Ar, UWorld* InWorld  );
-	bool	HandleCreateMeshFromBSPCommand( const TCHAR* Str, FOutputDevice& Ar );
-	bool	HandleAutoMergeStaticMeshCommand( const TCHAR* Str, FOutputDevice& Ar );
-	bool	HandleAddSelectedCommand( const TCHAR* Str, FOutputDevice& Ar );
-	bool	HandleToggleSocketGModeCommand( const TCHAR* Str, FOutputDevice& Ar );
-	bool	HandleListMapPackageDependenciesCommand( const TCHAR* Str, FOutputDevice& Ar );
-	bool	HandleRebuildVolumesCommand( const TCHAR* Str, FOutputDevice& Ar, UWorld* InWorld );
-	bool	HandleRemoveArchtypeFlagCommand( const TCHAR* Str, FOutputDevice& Ar );
-	
+	bool	HandleBlueprintifyFunction(const TCHAR* Str, FOutputDevice& Ar);
+	bool	HandleCallbackCommand(const TCHAR* Str, FOutputDevice& Ar);
+	bool	HandleTestPropsCommand(const TCHAR* Str, FOutputDevice& Ar);
+	bool	HandleMapCommand(const TCHAR* Str, FOutputDevice& Ar, UWorld* InWorld);
+	bool	HandleSelectCommand(const TCHAR* Str, FOutputDevice& Ar, UWorld* InWorld);
+	bool	HandleDeleteCommand(const TCHAR* Str, FOutputDevice& Ar, UWorld* InWorld);
+	bool	HandlePrefabCommand(const TCHAR* Str, FOutputDevice& Ar);
+	bool	HandleLightmassDebugCommand(const TCHAR* Str, FOutputDevice& Ar);
+	bool	HandleLightmassStatsCommand(const TCHAR* Str, FOutputDevice& Ar);
+	bool	HandleSwarmDistributionCommand(const TCHAR* Str, FOutputDevice& Ar);
+	bool	HandleLightmassImmediateImportCommand(const TCHAR* Str, FOutputDevice& Ar);
+	bool	HandleLightmassImmediateProcessCommand(const TCHAR* Str, FOutputDevice& Ar);
+	bool	HandleLightmassSortCommand(const TCHAR* Str, FOutputDevice& Ar);
+	bool	HandleLightmassDebugMaterialCommand(const TCHAR* Str, FOutputDevice& Ar);
+	bool	HandleLightmassPaddingCommand(const TCHAR* Str, FOutputDevice& Ar);
+	bool	HandleLightmassDebugPaddingCommand(const TCHAR* Str, FOutputDevice& Ar);
+	bool	HandleLightmassProfileCommand(const TCHAR* Str, FOutputDevice& Ar);
+	bool	HandleSetReplacementCommand(const TCHAR* Str, FOutputDevice& Ar, UWorld* InWorld);
+	bool	HandleSelectNameCommand(const TCHAR* Str, FOutputDevice& Ar, UWorld* InWorld);
+	bool	HandleDumpPublicCommand(const TCHAR* Str, FOutputDevice& Ar);
+	bool	HandleJumpToCommand(const TCHAR* Str, FOutputDevice& Ar);
+	bool	HandleBugItGoCommand(const TCHAR* Str, FOutputDevice& Ar);
+	bool	HandleTagSoundsCommand(const TCHAR* Str, FOutputDevice& Ar);
+	bool	HandlecheckSoundsCommand(const TCHAR* Str, FOutputDevice& Ar);
+	bool	HandleFixupBadAnimNotifiersCommand(const TCHAR* Str, FOutputDevice& Ar);
+	bool	HandleSetDetailModeCommand(const TCHAR* Str, FOutputDevice& Ar);
+	bool	HandleSetDetailModeViewCommand(const TCHAR* Str, FOutputDevice& Ar, UWorld* InWorld);
+	bool	HandleCleanBSPMaterialCommand(const TCHAR* Str, FOutputDevice& Ar, UWorld* InWorld);
+	bool	HandleCreateMeshFromBSPCommand(const TCHAR* Str, FOutputDevice& Ar);
+	bool	HandleAutoMergeStaticMeshCommand(const TCHAR* Str, FOutputDevice& Ar);
+	bool	HandleAddSelectedCommand(const TCHAR* Str, FOutputDevice& Ar);
+	bool	HandleToggleSocketGModeCommand(const TCHAR* Str, FOutputDevice& Ar);
+	bool	HandleListMapPackageDependenciesCommand(const TCHAR* Str, FOutputDevice& Ar);
+	bool	HandleRebuildVolumesCommand(const TCHAR* Str, FOutputDevice& Ar, UWorld* InWorld);
+	bool	HandleRemoveArchtypeFlagCommand(const TCHAR* Str, FOutputDevice& Ar);
+
 
 	/** Get tick rate limitor. */
-	virtual float GetMaxTickRate( float DeltaTime, bool bAllowFrameRateSmoothing = true ) const override;
+	virtual float GetMaxTickRate(float DeltaTime, bool bAllowFrameRateSmoothing = true) const override;
 
 	/**
 	 * Initializes the Editor.
@@ -678,9 +679,9 @@ public:
 	 *
 	 * @param		InWorld			World in which to create the builder brush.
 	 */
-	void InitBuilderBrush( UWorld* InWorld );
+	void InitBuilderBrush(UWorld* InWorld);
 
-	virtual void Tick( float DeltaSeconds, bool bIdleMode ) override;
+	virtual void Tick(float DeltaSeconds, bool bIdleMode) override;
 
 	/** Returns the global instance of the editor user settings class. */
 	const UEditorUserSettings& GetEditorUserSettings() const;
@@ -715,7 +716,7 @@ public:
 	 * @param bLinkedOrthoMovement	True if orthographic viewport movement is linked
 	 * @return - Whether a NON-realtime viewport has updated in this call.  Used to help time-slice canvas redraws
 	 */
-	bool UpdateSingleViewportClient(FEditorViewportClient* InViewportClient, const bool bInAllowNonRealtimeViewportToDraw, bool bLinkedOrthoMovement );
+	bool UpdateSingleViewportClient(FEditorViewportClient* InViewportClient, const bool bInAllowNonRealtimeViewportToDraw, bool bLinkedOrthoMovement);
 
 	/** Used for generating status bar text */
 	enum EMousePositionType
@@ -742,7 +743,7 @@ public:
 	 *
 	 * @param InCancelled	New state for the cancelled flag.
 	 */
-	virtual void SetMapBuildCancelled( bool InCancelled ) override
+	virtual void SetMapBuildCancelled(bool InCancelled) override
 	{
 		// Intentionally empty.
 	}
@@ -752,16 +753,16 @@ public:
 	 *
 	 * @param InActor	Actor that is being drawn.
 	 */
-	virtual bool ShouldDrawBrushWireframe( AActor* InActor ) override;
+	virtual bool ShouldDrawBrushWireframe(AActor* InActor) override;
 
 	// Execute a command that is safe for rebuilds.
-	virtual bool SafeExec( UWorld* InWorld, const TCHAR* Cmd, FOutputDevice& Ar=*GLog );
+	virtual bool SafeExec(UWorld* InWorld, const TCHAR* Cmd, FOutputDevice& Ar = *GLog);
 	// Process an incoming network message meant for the editor server
-	bool Exec_StaticMesh( UWorld* InWorld, const TCHAR* Str, FOutputDevice& Ar );
-	bool Exec_Brush( UWorld* InWorld, const TCHAR* Str, FOutputDevice& Ar );
-	bool Exec_Poly( UWorld* InWorld, const TCHAR* Str, FOutputDevice& Ar );
-	bool Exec_Obj( const TCHAR* Str, FOutputDevice& Ar );
-	bool Exec_Camera( const TCHAR* Str, FOutputDevice& Ar );
+	bool Exec_StaticMesh(UWorld* InWorld, const TCHAR* Str, FOutputDevice& Ar);
+	bool Exec_Brush(UWorld* InWorld, const TCHAR* Str, FOutputDevice& Ar);
+	bool Exec_Poly(UWorld* InWorld, const TCHAR* Str, FOutputDevice& Ar);
+	bool Exec_Obj(const TCHAR* Str, FOutputDevice& Ar);
+	bool Exec_Camera(const TCHAR* Str, FOutputDevice& Ar);
 	bool Exec_Transaction(const TCHAR* Str, FOutputDevice& Ar);
 	bool Exec_Particle(const TCHAR* Str, FOutputDevice& Ar);
 
@@ -772,7 +773,7 @@ public:
 	 * @param InFilename The name of the file to load and execute
 	 * @param Ar Output device
 	 */
-	void ExecFile( UWorld* InWorld, const TCHAR* InFilename, FOutputDevice& Ar );
+	void ExecFile(UWorld* InWorld, const TCHAR* InFilename, FOutputDevice& Ar);
 
 	// Transaction interfaces.
 	int32 BeginTransaction(const TCHAR* SessionContext, const FText& Description, UObject* PrimaryObject);
@@ -792,9 +793,9 @@ public:
 	 */
 	enum EMapRebuildType
 	{
-		MRT_Current				= 0,
-		MRT_AllVisible			= 1,
-		MRT_AllDirtyForLighting	= 2,
+		MRT_Current = 0,
+		MRT_AllVisible = 1,
+		MRT_AllDirtyForLighting = 2,
 	};
 
 	/**
@@ -810,7 +811,7 @@ public:
 	 * @param Level	The level to be rebuilt.
 	 */
 	void RebuildLevel(ULevel& Level);
-	
+
 	/**
 	 * Builds up a model from a set of brushes. Used by RebuildLevel.
 	 *
@@ -837,9 +838,9 @@ public:
 	 * @param	Actor					Target actor.
 	 * @param	bActiveViewportOnly		If true, move/reorient only the active viewport.
 	 */
-	void MoveViewportCamerasToActor(AActor& Actor,  bool bActiveViewportOnly);
+	void MoveViewportCamerasToActor(AActor& Actor, bool bActiveViewportOnly);
 
-	/** 
+	/**
 	 * Snaps an actor in a direction.  Optionally will align with the trace normal.
 	 * @param InActor			Actor to move to the floor.
 	 * @param InAlign			sWhether or not to rotate the actor to align with the trace normal.
@@ -849,7 +850,7 @@ public:
 	 * @param InDestination		The destination actor we want to move this actor to, NULL assumes we just want to go towards the floor
 	 * @return					Whether or not the actor was moved.
 	 */
-	bool SnapActorTo( AActor* InActor, const bool InAlign, const bool InUseLineTrace, const bool InUseBounds, const bool InUsePivot, const AActor* InDestination = NULL );
+	bool SnapActorTo(AActor* InActor, const bool InAlign, const bool InUseLineTrace, const bool InUseBounds, const bool InUsePivot, const AActor* InDestination = NULL);
 
 	/**
 	 * Moves an actor in front of a camera specified by the camera's origin and direction.
@@ -860,7 +861,7 @@ public:
 	 * @param InCameraOrigin	The location of the camera in world space
 	 * @param InCameraDirection	The normalized direction vector of the camera
 	 */
-	void MoveActorInFrontOfCamera( AActor& InActor, const FVector& InCameraOrigin, const FVector& InCameraDirection );
+	void MoveActorInFrontOfCamera(AActor& InActor, const FVector& InCameraOrigin, const FVector& InCameraDirection);
 
 	/**
 	* Moves all viewport cameras to focus on the provided array of actors.
@@ -903,7 +904,7 @@ public:
 	 * @param	bIgnoreAxis				If true, leave the existing pivot unaffected for components of NewPivot that are 0.
 	 * @param	bAssignPivot			If true, assign the given pivot to any valid actors that retain it (defaults to false)
 	 */
-	virtual void SetPivot(FVector NewPivot, bool bSnapPivotToGrid, bool bIgnoreAxis, bool bAssignPivot=false) {}
+	virtual void SetPivot(FVector NewPivot, bool bSnapPivotToGrid, bool bIgnoreAxis, bool bAssignPivot = false) {}
 	virtual void ResetPivot() {}
 
 	//
@@ -917,7 +918,7 @@ public:
 	 * @param	Redraw			Whether to redraw viewports
 	 * @param	TransReset		Human readable reason for resetting the transaction system
 	 */
-	virtual void Cleanse( bool ClearSelection, bool Redraw, const FText& TransReset );
+	virtual void Cleanse(bool ClearSelection, bool Redraw, const FText& TransReset);
 	virtual void FinishAllSnaps() { }
 
 	/**
@@ -925,8 +926,8 @@ public:
 	 *
 	 * @param	bInvalidateHitProxies		[opt] If true (the default), invalidates cached hit proxies too.
 	 */
-	virtual void RedrawLevelEditingViewports(bool bInvalidateHitProxies=true) {}
-		
+	virtual void RedrawLevelEditingViewports(bool bInvalidateHitProxies = true) {}
+
 	/**
 	 * Triggers a high res screen shot for current editor viewports.
 	 *
@@ -961,7 +962,7 @@ public:
 	virtual UTransactor* CreateTrans();
 
 	/** Plays an editor sound, loading the sound on demand if necessary (if the user has sounds enabled.)  The reference to the sound asset is not retained. */
-	void PlayEditorSound( const FString& SoundAssetName );
+	void PlayEditorSound(const FString& SoundAssetName);
 
 	/**
 	 * Returns the preview audio component
@@ -974,7 +975,7 @@ public:
 	 * @param	Sound		A sound to attach to the audio component
 	 * @param	SoundNode	A sound node that is attached to the audio component when the sound is NULL
 	 */
-	UAudioComponent* ResetPreviewAudioComponent( class USoundBase* Sound = NULL, class USoundNode* SoundNode = NULL );
+	UAudioComponent* ResetPreviewAudioComponent(class USoundBase* Sound = NULL, class USoundNode* SoundNode = NULL);
 
 	/**
 	 * Plays a preview of a specified sound or node
@@ -989,7 +990,7 @@ public:
 	 */
 	void ClearPreviewComponents();
 
-	
+
 	/**
 	 * Close all the edit windows for assets that are owned by the world being closed
 	 */
@@ -1000,7 +1001,7 @@ public:
 	 *
 	 * @param	bInvalidateHitProxies		[opt] If true (the default), invalidates cached hit proxies too.
 	 */
-	void RedrawAllViewports(bool bInvalidateHitProxies=true);
+	void RedrawAllViewports(bool bInvalidateHitProxies = true);
 
 	/**
 	 * Invalidates all viewports parented to the specified view.
@@ -1008,7 +1009,7 @@ public:
 	 * @param	InParentView				The parent view whose child views should be invalidated.
 	 * @param	bInvalidateHitProxies		[opt] If true (the default), invalidates cached hit proxies too.
 	 */
-	void InvalidateChildViewports(FSceneViewStateInterface* InParentView, bool bInvalidateHitProxies=true);
+	void InvalidateChildViewports(FSceneViewStateInterface* InParentView, bool bInvalidateHitProxies = true);
 
 	/**
 	 * Looks for an appropriate actor factory for the specified UClass.
@@ -1016,7 +1017,7 @@ public:
 	 * @param	InClass		The class to find the factory for.
 	 * @return				A pointer to the factory to use.  NULL if no factories support this class.
 	 */
-	UActorFactory* FindActorFactoryForActorClass( const UClass* InClass );
+	UActorFactory* FindActorFactoryForActorClass(const UClass* InClass);
 
 	/**
 	 * Looks for an actor factory spawned from the specified class.
@@ -1024,7 +1025,7 @@ public:
 	 * @param	InClass		The factories class to find
 	 * @return				A pointer to the factory to use.  NULL if no factories support this class.
 	 */
-	UActorFactory* FindActorFactoryByClass( const UClass* InClass ) const;
+	UActorFactory* FindActorFactoryByClass(const UClass* InClass) const;
 
 	/**
 	* Looks for an actor factory spawned from the specified class, for the specified UClass for an actor.
@@ -1033,7 +1034,7 @@ public:
 	* @param	InActorClass		The class to find the factory for.
 	* @return				A pointer to the factory to use.  NULL if no factories support this class.
 	*/
-	UActorFactory* FindActorFactoryByClassForActorClass( const UClass* InFactoryClass, const UClass* InActorClass );
+	UActorFactory* FindActorFactoryByClassForActorClass(const UClass* InFactoryClass, const UClass* InActorClass);
 
 	/**
 	 * Uses the supplied factory to create an actor at the clicked location and adds to level.
@@ -1042,7 +1043,7 @@ public:
 	 * @param	ObjectFlags				[opt] The flags to apply to the actor when it is created
 	 * @return							A pointer to the new actor, or NULL on fail.
 	 */
-	AActor* UseActorFactoryOnCurrentSelection( UActorFactory* Factory, const FTransform* InActorTransform, EObjectFlags ObjectFlags = RF_Transactional );
+	AActor* UseActorFactoryOnCurrentSelection(UActorFactory* Factory, const FTransform* InActorTransform, EObjectFlags ObjectFlags = RF_Transactional);
 
 	/**
 	 * Uses the supplied factory to create an actor at the clicked location and adds to level.
@@ -1054,19 +1055,19 @@ public:
 	 * @param	ObjectFlags				[opt] The flags to apply to the actor when it is created
 	 * @return							A pointer to the new actor, or NULL on fail.
 	 */
-	AActor* UseActorFactory( UActorFactory* Factory, const FAssetData& AssetData, const FTransform* ActorLocation, EObjectFlags ObjectFlags = RF_Transactional );
+	AActor* UseActorFactory(UActorFactory* Factory, const FAssetData& AssetData, const FTransform* ActorLocation, EObjectFlags ObjectFlags = RF_Transactional);
 
 	/**
 	 * Replaces the selected Actors with the same number of a different kind of Actor
 	 * if a Factory is specified, it is used to spawn the requested Actors, otherwise NewActorClass is used (one or the other must be specified)
 	 * note that only Location, Rotation, Drawscale, Drawscale3D, Tag, and Group are copied from the old Actors
-	 * 
+	 *
 	 * @param Factory - the Factory to use to create Actors
 	 */
 	void ReplaceSelectedActors(UActorFactory* Factory, const FAssetData& AssetData, UClass* NewActorClass);
 
 	/**
-	 * Converts passed in brushes into a single static mesh actor. 
+	 * Converts passed in brushes into a single static mesh actor.
 	 * Note: This replaces all the brushes with a single actor. This actor will not be attached to anything unless a single brush was converted.
 	 *
 	 * @param	InStaticMeshPackageName		The name to save the brushes to.
@@ -1085,7 +1086,7 @@ public:
 	 *
 	 * @param	ConvertToClass	The light class we are going to convert to.
 	 */
-	void ConvertLightActors( UClass* ConvertToClass );
+	void ConvertLightActors(UClass* ConvertToClass);
 
 	/**
 	 * Converts passed in actors into new actors of the specified type.
@@ -1102,7 +1103,7 @@ public:
 	 * @param	bUseSpecialCases			If true, looks for classes that can be handled by hardcoded conversions
 	 * @param	InStaticMeshPackageName		The name to save the brushes to.
 	 */
-	void DoConvertActors( const TArray<AActor*>& ActorsToConvert, UClass* ConvertToClass, const TSet<FString>& ComponentsToConsider, bool bUseSpecialCases, const FString& InStaticMeshPackageName );
+	void DoConvertActors(const TArray<AActor*>& ActorsToConvert, UClass* ConvertToClass, const TSet<FString>& ComponentsToConsider, bool bUseSpecialCases, const FString& InStaticMeshPackageName);
 
 	/**
 	 * Sets up for a potentially deferred ConvertActors call, based on if any brushes are being converted to a static mesh. If one (or more)
@@ -1113,14 +1114,14 @@ public:
 	 * @param	ComponentsToConsider	Names of components to consider for property copying as well
 	 * @param	bUseSpecialCases		If true, looks for classes that can be handled by hardcoded conversions
 	 */
-	void ConvertActors( const TArray<AActor*>& ActorsToConvert, UClass* ConvertToClass, const TSet<FString>& ComponentsToConsider, bool bUseSpecialCases = false );
+	void ConvertActors(const TArray<AActor*>& ActorsToConvert, UClass* ConvertToClass, const TSet<FString>& ComponentsToConsider, bool bUseSpecialCases = false);
 
 	/**
 	 * Changes the state of preview mesh mode to on or off.
 	 *
 	 * @param	bState	Enables the preview mesh mode if true; Disables the preview mesh mode if false.
 	 */
-	void SetPreviewMeshMode( bool bState );
+	void SetPreviewMeshMode(bool bState);
 
 	/**
 	 * Updates the position of the preview mesh in the level.
@@ -1157,7 +1158,7 @@ public:
 	 * @param	InLevel			Level to place duplicate
 	 * @param	bUseOffset		Should the actor locations be offset after they are created?
 	 */
-	virtual void edactDuplicateSelected( ULevel* InLevel, bool bOffsetLocations ) {}
+	virtual void edactDuplicateSelected(ULevel* InLevel, bool bOffsetLocations) {}
 
 	/**
 	 * Deletes all selected actors
@@ -1166,7 +1167,7 @@ public:
 	 * @param		bVerifyDeletionCanHappen	[opt] If true (default), verify that deletion can be performed.
 	 * @return									true unless the delete operation was aborted.
 	 */
-	virtual bool edactDeleteSelected(UWorld* InWorld, bool bVerifyDeletionCanHappen=true) { return true; }
+	virtual bool edactDeleteSelected(UWorld* InWorld, bool bVerifyDeletionCanHappen = true) { return true; }
 
 	/**
 	 * Checks the state of the selected actors and notifies the user of any potentially unknown destructive actions which may occur as
@@ -1190,9 +1191,9 @@ public:
 	*
 	* Rebuild the level's Bsp from the level's CSG brushes.
 	*
-	* @param InWorld	The world in which the rebuild the CSG brushes for 
+	* @param InWorld	The world in which the rebuild the CSG brushes for
 	*/
-	virtual void csgRebuild( UWorld* InWorld );
+	virtual void csgRebuild(UWorld* InWorld);
 
 	/**
 	*
@@ -1200,11 +1201,11 @@ public:
 	*
 	* @param InModel	Model to get poly from
 	* @param iSurf		surface index
-	* @param Poly		
+	* @param Poly
 	*
 	* returns true if poly not available
 	*/
-	virtual bool polyFindMaster( UModel* InModel, int32 iSurf, FPoly& Poly );
+	virtual bool polyFindMaster(UModel* InModel, int32 iSurf, FPoly& Poly);
 
 	/**
 	 * Update a the master brush EdPoly corresponding to a newly-changed
@@ -1212,36 +1213,36 @@ public:
 	 *
 	 * Doesn't do any transaction tracking.
 	 */
-	virtual void polyUpdateMaster( UModel* Model, int32 iSurf, int32 UpdateTexCoords );
+	virtual void polyUpdateMaster(UModel* Model, int32 iSurf, int32 UpdateTexCoords);
 
 	/**
 	 * Populates a list with all polys that are linked to the specified poly.  The
 	 * resulting list includes the original poly.
 	 */
-	virtual void polyGetLinkedPolys( ABrush* InBrush, FPoly* InPoly, TArray<FPoly>* InPolyList );
+	virtual void polyGetLinkedPolys(ABrush* InBrush, FPoly* InPoly, TArray<FPoly>* InPolyList);
 
 	/**
 	 * Takes a list of polygons and returns a list of the outside edges (edges which are not shared
 	 * by other polys in the list).
 	 */
-	virtual void polyGetOuterEdgeList( TArray<FPoly>* InPolyList, TArray<FEdge>* InEdgeList );
+	virtual void polyGetOuterEdgeList(TArray<FPoly>* InPolyList, TArray<FEdge>* InEdgeList);
 
 	/**
 	 * Takes a list of polygons and creates a new list of polys which have no overlapping edges.  It splits
 	 * edges as necessary to achieve this.
 	 */
-	virtual void polySplitOverlappingEdges( TArray<FPoly>* InPolyList, TArray<FPoly>* InResult );
+	virtual void polySplitOverlappingEdges(TArray<FPoly>* InPolyList, TArray<FPoly>* InResult);
 
 	/**
 	 * Sets and clears all Bsp node flags.  Affects all nodes, even ones that don't
 	 * really exist.
 	 */
-	virtual void polySetAndClearPolyFlags( UModel* Model, uint32 SetBits, uint32 ClearBits, bool SelectedOnly, bool UpdateMaster );
+	virtual void polySetAndClearPolyFlags(UModel* Model, uint32 SetBits, uint32 ClearBits, bool SelectedOnly, bool UpdateMaster);
 
 	// Selection.
-	virtual void SelectActor(AActor* Actor, bool bInSelected, bool bNotify, bool bSelectEvenIfHidden=false) {}
-	virtual bool CanSelectActor(AActor* Actor, bool bInSelected, bool bSelectEvenIfHidden=false, bool bWarnIfLevelLocked=false) const { return true; }
-	virtual void SelectGroup(class AGroupActor* InGroupActor, bool bForceSelection=false, bool bInSelected=true, bool bNotify=true) {}
+	virtual void SelectActor(AActor* Actor, bool bInSelected, bool bNotify, bool bSelectEvenIfHidden = false) {}
+	virtual bool CanSelectActor(AActor* Actor, bool bInSelected, bool bSelectEvenIfHidden = false, bool bWarnIfLevelLocked = false) const { return true; }
+	virtual void SelectGroup(class AGroupActor* InGroupActor, bool bForceSelection = false, bool bInSelected = true, bool bNotify = true) {}
 
 	/**
 	 * Replaces the components in ActorsToReplace with an primitive component in Replacement
@@ -1250,7 +1251,7 @@ public:
 	 * @param Replacement The first usable component in Replacement will be the ReplacementPrimitive for the actors
 	 * @param ClassToReplace If this is set, only components will of this class will be used/replaced
 	 */
-	virtual void AssignReplacementComponentsByActors(TArray<AActor*>& ActorsToReplace, AActor* Replacement, UClass* ClassToReplace=NULL);
+	virtual void AssignReplacementComponentsByActors(TArray<AActor*>& ActorsToReplace, AActor* Replacement, UClass* ClassToReplace = NULL);
 
 	/**
 	 * Selects or deselects a BSP surface in the persistent level's UModel.  Does nothing if GEdSelectionLock is true.
@@ -1268,18 +1269,18 @@ public:
 	 * @param	bNoteSelectionChange		If true, call NoteSelectionChange().
 	 * @param	bDeselectBSPSurfs			If true, also deselect all BSP surfaces.
 	 */
-	virtual void SelectNone(bool bNoteSelectionChange, bool bDeselectBSPSurfs, bool WarnAboutManyActors=true) {}
+	virtual void SelectNone(bool bNoteSelectionChange, bool bDeselectBSPSurfs, bool WarnAboutManyActors = true) {}
 
 	// Bsp Poly selection virtuals from EditorCsg.cpp.
-	virtual void polySelectAll ( UModel* Model );
-	virtual void polySelectMatchingGroups( UModel* Model );
-	virtual void polySelectMatchingItems( UModel* Model );
-	virtual void polySelectCoplanars( UWorld* InWorld, UModel* Model );
-	virtual void polySelectAdjacents( UWorld* InWorld, UModel* Model );
-	virtual void polySelectAdjacentWalls( UWorld* InWorld, UModel* Model );
-	virtual void polySelectAdjacentFloors( UWorld* InWorld, UModel* Model );
-	virtual void polySelectAdjacentSlants( UWorld* InWorld, UModel* Model );
-	virtual void polySelectMatchingBrush( UModel* Model );
+	virtual void polySelectAll(UModel* Model);
+	virtual void polySelectMatchingGroups(UModel* Model);
+	virtual void polySelectMatchingItems(UModel* Model);
+	virtual void polySelectCoplanars(UWorld* InWorld, UModel* Model);
+	virtual void polySelectAdjacents(UWorld* InWorld, UModel* Model);
+	virtual void polySelectAdjacentWalls(UWorld* InWorld, UModel* Model);
+	virtual void polySelectAdjacentFloors(UWorld* InWorld, UModel* Model);
+	virtual void polySelectAdjacentSlants(UWorld* InWorld, UModel* Model);
+	virtual void polySelectMatchingBrush(UModel* Model);
 
 	/**
 	 * Selects surfaces whose material matches that of any selected surfaces.
@@ -1295,23 +1296,23 @@ public:
 	 */
 	virtual void polySelectMatchingResolution(UWorld* InWorld, bool bCurrentLevelOnly);
 
-	virtual void polySelectReverse( UModel* Model );
-	virtual void polyMemorizeSet( UModel* Model );
-	virtual void polyRememberSet( UModel* Model );
-	virtual void polyXorSet( UModel* Model );
-	virtual void polyUnionSet( UModel* Model );
-	virtual void polyIntersectSet( UModel* Model );
-	virtual void polySelectZone( UModel *Model );
+	virtual void polySelectReverse(UModel* Model);
+	virtual void polyMemorizeSet(UModel* Model);
+	virtual void polyRememberSet(UModel* Model);
+	virtual void polyXorSet(UModel* Model);
+	virtual void polyUnionSet(UModel* Model);
+	virtual void polyIntersectSet(UModel* Model);
+	virtual void polySelectZone(UModel *Model);
 
 	// Pan textures on selected polys.  Doesn't do transaction tracking.
-	virtual void polyTexPan( UModel* Model, int32 PanU, int32 PanV, int32 Absolute );
+	virtual void polyTexPan(UModel* Model, int32 PanU, int32 PanV, int32 Absolute);
 
 	// Scale textures on selected polys. Doesn't do transaction tracking.
-	virtual void polyTexScale( UModel* Model,float UU, float UV, float VU, float VV, bool Absolute );
+	virtual void polyTexScale(UModel* Model, float UU, float UV, float VU, float VV, bool Absolute);
 
 	// Map brush selection virtuals from EditorCsg.cpp.
-	virtual void MapSelectOperation( UWorld* InWorld, EBrushType BrushType );
-	virtual void MapSelectFlags( UWorld* InWorld, uint32 Flags );
+	virtual void MapSelectOperation(UWorld* InWorld, EBrushType BrushType);
+	virtual void MapSelectFlags(UWorld* InWorld, uint32 Flags);
 
 	// Put the first selected brush into the current Brush.
 	virtual void MapBrushGet(UWorld* InWorld);
@@ -1329,27 +1330,27 @@ public:
 	 * Swaps position in the actor list for the first two selected actors in the current level
 	 */
 	virtual void mapSendToSwap(UWorld* InWorld);
-	virtual void MapSetBrush( UWorld* InWorld, EMapSetBrushFlags PropertiesMask, uint16 BrushColor, FName Group, uint32 SetPolyFlags, uint32 ClearPolyFlags, uint32 BrushType, int32 DrawType );
+	virtual void MapSetBrush(UWorld* InWorld, EMapSetBrushFlags PropertiesMask, uint16 BrushColor, FName Group, uint32 SetPolyFlags, uint32 ClearPolyFlags, uint32 BrushType, int32 DrawType);
 
 	// Bsp virtuals from Bsp.cpp.
-	virtual void bspRepartition( UWorld* InWorld, int32 iNode );
+	virtual void bspRepartition(UWorld* InWorld, int32 iNode);
 
 	/** Convert a Bsp node to an EdPoly.  Returns number of vertices in Bsp node. */
-	virtual int32 bspNodeToFPoly( UModel* Model, int32 iNode, FPoly* EdPoly );
+	virtual int32 bspNodeToFPoly(UModel* Model, int32 iNode, FPoly* EdPoly);
 
 	/**
 	 * Clean up all nodes after a CSG operation.  Resets temporary bit flags and unlinks
 	 * empty leaves.  Removes zero-vertex nodes which have nonzero-vertex coplanars.
 	 */
-	virtual void bspCleanup( UModel* Model );
+	virtual void bspCleanup(UModel* Model);
 
-	/** 
+	/**
 	 * Build EdPoly list from a model's Bsp. Not transactional.
-	 * @param DestArray helps build bsp FPolys in non-main threads. It also allows to perform this action without GUndo 
+	 * @param DestArray helps build bsp FPolys in non-main threads. It also allows to perform this action without GUndo
 	 *	interfering. Temporary results will be written to DestArray. Defaults to Model->Polys->Element
 	 */
-	virtual void bspBuildFPolys( UModel* Model, bool SurfLinks, int32 iNode, TArray<FPoly>* DestArray = NULL );
-	virtual void bspMergeCoplanars( UModel* Model, bool RemapLinks, bool MergeDisparateTextures );
+	virtual void bspBuildFPolys(UModel* Model, bool SurfLinks, int32 iNode, TArray<FPoly>* DestArray = NULL);
+	virtual void bspMergeCoplanars(UModel* Model, bool RemapLinks, bool MergeDisparateTextures);
 	/**
 	 * Performs any CSG operation between the brush and the world.
 	 *
@@ -1364,15 +1365,15 @@ public:
 	 * @param	bShowProgressBar				If true, display progress bar for complex brushes
 	 * @return									0 if nothing happened, 1 if the operation was error-free, or 1+N if N CSG errors occurred.
 	 */
-	virtual int32 bspBrushCSG( ABrush* Actor, UModel* Model, uint32 PolyFlags, EBrushType BrushType, ECsgOper CSGOper, bool bBuildBounds, bool bMergePolys, bool bReplaceNULLMaterialRefs, bool bShowProgressBar=true );
+	virtual int32 bspBrushCSG(ABrush* Actor, UModel* Model, uint32 PolyFlags, EBrushType BrushType, ECsgOper CSGOper, bool bBuildBounds, bool bMergePolys, bool bReplaceNULLMaterialRefs, bool bShowProgressBar = true);
 
 	/**
 	 * Optimize a level's Bsp, eliminating T-joints where possible, and building side
-	 * links.  This does not always do a 100% perfect job, mainly due to imperfect 
+	 * links.  This does not always do a 100% perfect job, mainly due to imperfect
 	 * levels, however it should never fail or return incorrect results.
 	 */
-	virtual void bspOptGeom( UModel* Model );
-	
+	virtual void bspOptGeom(UModel* Model);
+
 	/**
 	 * Builds lighting information depending on passed in options.
 	 *
@@ -1400,21 +1401,21 @@ public:
 	 * @param Filename	The FBX filename
 	 * @param bImportMorphTracks	true to import any morph curve data.
 	 */
-	static UAnimSequence * ImportFbxAnimation( USkeleton* Skeleton, UObject* Outer, class UFbxAnimSequenceImportData* ImportData, const TCHAR* InFilename, const TCHAR* AnimName, bool bImportMorphTracks );
+	static UAnimSequence * ImportFbxAnimation(USkeleton* Skeleton, UObject* Outer, class UFbxAnimSequenceImportData* ImportData, const TCHAR* InFilename, const TCHAR* AnimName, bool bImportMorphTracks);
 	/**
-	 * Reimport animation using SourceFilePath and SourceFileStamp 
+	 * Reimport animation using SourceFilePath and SourceFileStamp
 	 *
 	 * @param Skeleton	The skeleton that animation is import into
 	 * @param Filename	The FBX filename
 	 */
-	static bool ReimportFbxAnimation( USkeleton* Skeleton, UAnimSequence* AnimSequence, class UFbxAnimSequenceImportData* ImportData, const TCHAR* InFilename);
+	static bool ReimportFbxAnimation(USkeleton* Skeleton, UAnimSequence* AnimSequence, class UFbxAnimSequenceImportData* ImportData, const TCHAR* InFilename);
 
 
 	// Object management.
-	virtual void RenameObject(UObject* Object,UObject* NewOuter,const TCHAR* NewName, ERenameFlags Flags=REN_None);
+	virtual void RenameObject(UObject* Object, UObject* NewOuter, const TCHAR* NewName, ERenameFlags Flags = REN_None);
 
 	// Level management.
-	void AnalyzeLevel(ULevel* Level,FOutputDevice& Ar);
+	void AnalyzeLevel(ULevel* Level, FOutputDevice& Ar);
 
 	/**
 	 * Removes all components from the current level's scene.
@@ -1426,8 +1427,8 @@ public:
 	 */
 	void EditorUpdateComponents();
 
-	/** 
-	 * Displays a modal message dialog 
+	/**
+	 * Displays a modal message dialog
 	 * @param	InMessage	Type of the message box
 	 * @param	InText		Message to display
 	 * @param	InTitle		Title for the message box
@@ -1435,8 +1436,8 @@ public:
 	 */
 	EAppReturnType::Type OnModalMessageDialog(EAppMsgType::Type InMessage, const FText& InText, const FText& InTitle);
 
-	/** 
-	 * Returns whether an object should replace an exisiting one or not 
+	/**
+	 * Returns whether an object should replace an exisiting one or not
 	 * @param	Filename		Filename of the package
 	 * @return	Returns whether the objects should replace the already existing ones.
 	 */
@@ -1452,7 +1453,7 @@ public:
 	 * @param	bUseMobilePreview		True to enable mobile preview mode (PC platform only)
 	 * @param	bMovieCapture			True to start with movie capture recording (PC platform only)
 	 */
-	virtual void PlayMap( const FVector* StartLocation = NULL, const FRotator* StartRotation = NULL, int32 DestinationConsole = -1, int32 InPlayInViewportIndex = -1, bool bUseMobilePreview = false, bool bMovieCapture = false );
+	virtual void PlayMap(const FVector* StartLocation = NULL, const FRotator* StartRotation = NULL, int32 DestinationConsole = -1, int32 InPlayInViewportIndex = -1, bool bUseMobilePreview = false, bool bMovieCapture = false);
 
 
 
@@ -1464,7 +1465,7 @@ public:
 	/**
 	 * Start cook by the book in the editor process space
 	 */
-	virtual void StartCookByTheBookInEditor( const TArray<ITargetPlatform*> &TargetPlatforms, const TArray<FString> &CookMaps, const TArray<FString> &CookDirectories, const TArray<FString> &CookCultures, const TArray<FString> &IniMapSections ) { }
+	virtual void StartCookByTheBookInEditor(const TArray<ITargetPlatform*> &TargetPlatforms, const TArray<FString> &CookMaps, const TArray<FString> &CookDirectories, const TArray<FString> &CookCultures, const TArray<FString> &IniMapSections) { }
 
 	/**
 	 * Checks if the cook by the book is finished
@@ -1487,15 +1488,15 @@ public:
 	 * @param	DestinationConsole		Where to play the game - -1 means in editor, 0 or more is an index into the GConsoleSupportContainer
 	 * @param	bUseMobilePreview		True to enable mobile preview mode (PC platform only)
 	 */
-	void RequestPlaySession( bool bAtPlayerStart, TSharedPtr<class ILevelViewport> DestinationViewport, bool bInSimulateInEditor, const FVector* StartLocation = NULL, const FRotator* StartRotation = NULL, int32 DestinationConsole = -1, bool bUseMobilePreview = false );
+	void RequestPlaySession(bool bAtPlayerStart, TSharedPtr<class ILevelViewport> DestinationViewport, bool bInSimulateInEditor, const FVector* StartLocation = NULL, const FRotator* StartRotation = NULL, int32 DestinationConsole = -1, bool bUseMobilePreview = false);
 
 	// @todo gmp: temp hack for Rocket demo
-	void RequestPlaySession( const FVector* StartLocation, const FRotator* StartRotation, bool MobilePreview );
+	void RequestPlaySession(const FVector* StartLocation, const FRotator* StartRotation, bool MobilePreview);
 
 	/**
-	 * Request to play a game on a remote device 
+	 * Request to play a game on a remote device
 	 */
-	void RequestPlaySession( const FString& DeviceId, const FString& DeviceName );
+	void RequestPlaySession(const FString& DeviceId, const FString& DeviceName);
 
 	/**
 	 * Cancel request to start a play session
@@ -1522,13 +1523,13 @@ public:
 	virtual void StartQueuedPlayMapRequest();
 
 	/**
-	 * Request that the current PIE/SIE session should end.  
+	 * Request that the current PIE/SIE session should end.
 	 * This should be used to end the game if its not safe to directly call EndPlayMap in your stack frame
 	 */
 	void RequestEndPlayMap();
 
 	/**
-	 * @return true if there is an end play map request queued 
+	 * @return true if there is an end play map request queued
 	 */
 	bool ShouldEndPlayMap() const { return bRequestEndPlayMapQueued; }
 
@@ -1543,14 +1544,14 @@ public:
 	virtual bool SavePlayWorldPackages(UWorld* InWorld, const TCHAR* Prefix, TArray<FString>& OutSavedFilenames);
 
 	/**
-	 * Builds a URL for game spawned by the editor (not including map name!). 
+	 * Builds a URL for game spawned by the editor (not including map name!).
 	 * @param	MapName			The name of the map to put into the URL
 	 * @param	bSpecatorMode	If true, the player starts in spectator mode
 	 * @param   AdditionalURLOptions	Additional URL Options to append (e.g, ?OptionX?OptionY). This is in addition to InEditorGameURLOptions.
 	 *
 	 * @return	The URL for the game
 	 */
-	virtual FString BuildPlayWorldURL(const TCHAR* MapName, bool bSpectatorMode = false, FString AdditionalURLOptions=FString());
+	virtual FString BuildPlayWorldURL(const TCHAR* MapName, bool bSpectatorMode = false, FString AdditionalURLOptions = FString());
 
 	/**
 	 * Starts a Play In Editor session
@@ -1558,7 +1559,7 @@ public:
 	 * @param	InWorld		World context
 	 * @param	bInSimulateInEditor	True to start an in-editor simulation session, or false to start a play-in-editor session
 	 */
-	virtual void PlayInEditor( UWorld* InWorld, bool bInSimulateInEditor );
+	virtual void PlayInEditor(UWorld* InWorld, bool bInSimulateInEditor);
 
 	virtual UGameInstance* CreatePIEGameInstance(int32 PIEInstance, bool bInSimulateInEditor, bool bAnyBlueprintErrors, bool bStartInSpectatorMode, bool bPlayNetDedicated, float PIEStartTime);
 
@@ -1572,7 +1573,7 @@ public:
 	 */
 	virtual void EndPlayMap();
 
-	/** 
+	/**
 	 * Destroy the current play session and perform miscellaneous cleanup
 	 */
 	virtual void TeardownPlaySession(FWorldContext &PieWorldContext);
@@ -1604,7 +1605,7 @@ public:
 	virtual bool ShouldThrottleCPUUsage() const override;
 
 	/**
-	 * @return true if all windows are hidden (including minimized)                                                         
+	 * @return true if all windows are hidden (including minimized)
 	 */
 	bool AreAllWindowsHidden() const;
 
@@ -1612,7 +1613,7 @@ public:
 	 *	Returns pointer to a temporary render target.
 	 *	If it has not already been created, does so here.
 	 */
-	UTextureRenderTarget2D* GetScratchRenderTarget( const uint32 MinSize );
+	UTextureRenderTarget2D* GetScratchRenderTarget(const uint32 MinSize);
 
 	/**
 	 *  Returns the Editors timer manager instance.
@@ -1622,12 +1623,12 @@ public:
 	/**
 	 * Handles freezing/unfreezing of rendering
 	 */
-	virtual void ProcessToggleFreezeCommand( UWorld* InWorld ) override;
+	virtual void ProcessToggleFreezeCommand(UWorld* InWorld) override;
 
 	/**
 	 * Handles frezing/unfreezing of streaming
 	 */
-	virtual void ProcessToggleFreezeStreamingCommand( UWorld* InWorld ) override;
+	virtual void ProcessToggleFreezeStreamingCommand(UWorld* InWorld) override;
 
 	// Editor specific
 
@@ -1635,7 +1636,7 @@ public:
 	* Closes the main editor frame.
 	*/
 	virtual void CloseEditor() {}
-	virtual void GetPackageList( TArray<UPackage*>* InPackages, UClass* InClass ) {}
+	virtual void GetPackageList(TArray<UPackage*>* InPackages, UClass* InClass) {}
 
 	/**
 	 * Returns the number of currently selected actors.
@@ -1651,7 +1652,7 @@ public:
 	bool IsWorldSettingsSelected();
 
 	/** Function to return list of assets currently selected in content browser (loaded/not loaded) */
-	void GetContentBrowserSelections( TArray<UClass*>& Selection ) const;
+	void GetContentBrowserSelections(TArray<UClass*>& Selection) const;
 
 	/**
 	 * Returns an FSelectionIterator that iterates over the set of selected actors.
@@ -1667,13 +1668,13 @@ public:
 	/**
 	 * @return the appropriate selection set for the specified object class.
 	 */
-	class USelection* GetSelectedSet( const UClass* Class ) const;
+	class USelection* GetSelectedSet(const UClass* Class) const;
 
 	/**
 	 * @param	RequiredParentClass The parent class this type must implement, or null to accept anything
 	 * @return	The first selected class (either a UClass type itself, or the UClass generated by a blueprint), or null if there are no class or blueprint types selected
 	 */
-	const UClass* GetFirstSelectedClass( const UClass* const RequiredParentClass ) const;
+	const UClass* GetFirstSelectedClass(const UClass* const RequiredParentClass) const;
 
 	/**
 	 * Clears out the current map, if any, and creates a new blank map.
@@ -1695,7 +1696,7 @@ public:
 	 *
 	 * @param	InLevel		The destination level.
 	 */
-	void MoveSelectedActorsToLevel( ULevel* InLevel );
+	void MoveSelectedActorsToLevel(ULevel* InLevel);
 
 	/**
 	 * Checks to see whether it's possible to perform a copy operation on the selected actors.
@@ -1704,7 +1705,7 @@ public:
 	 * @param OutCopySelected	Can be NULL, copies the internal results of the copy check to this struct
 	 * @return			true if it's possible to do a copy operation based on the selection
 	 */
-	bool CanCopySelectedActorsToClipboard( UWorld* InWorld, FCopySelectedInfo* OutCopySelected = NULL );
+	bool CanCopySelectedActorsToClipboard(UWorld* InWorld, FCopySelectedInfo* OutCopySelected = NULL);
 
 	/**
 	 * Copies selected actors to the clipboard.  Supports copying actors from multiple levels.
@@ -1713,7 +1714,7 @@ public:
 	 * @param InWorld		World to get the selected actors from
 	 * @param bShouldCut If true, deletes the selected actors after copying them to the clipboard
 	 */
-	void CopySelectedActorsToClipboard( UWorld* InWorld, const bool bShouldCut );
+	void CopySelectedActorsToClipboard(UWorld* InWorld, const bool bShouldCut);
 
 	/**
 	 * Checks to see whether it's possible to perform a paste operation.
@@ -1721,7 +1722,7 @@ public:
 	 * @param InWorld		World to paste into
 	 * @return				true if it's possible to do a Paste operation
 	 */
-	bool CanPasteSelectedActorsFromClipboard( UWorld* InWorld );
+	bool CanPasteSelectedActorsFromClipboard(UWorld* InWorld);
 
 	/**
 	 * Pastes selected actors from the clipboard.
@@ -1730,7 +1731,7 @@ public:
 	 * @param InWorld		World to get the selected actors from
 	 * @param PasteTo		Where to paste the content too
 	 */
-	void PasteSelectedActorsFromClipboard( UWorld* InWorld, const FText& TransDescription, const EPasteTo PasteTo );
+	void PasteSelectedActorsFromClipboard(UWorld* InWorld, const FText& TransDescription, const EPasteTo PasteTo);
 
 	/**
 	 * Computes a color to use for property coloration for the given object.
@@ -1784,11 +1785,11 @@ public:
 	 * @param	bIncludePersistentLvl	If true, the persistent level will also be checked for visibility
 	 * @return							false if the user selects "No", true otherwise.
 	 */
-	bool WarnAboutHiddenLevels( UWorld* InWorld, bool bIncludePersistentLvl) const;
+	bool WarnAboutHiddenLevels(UWorld* InWorld, bool bIncludePersistentLvl) const;
 
-	void ApplyDeltaToActor(AActor* InActor, bool bDelta, const FVector* InTranslation, const FRotator* InRotation, const FVector* InScaling, bool bAltDown=false, bool bShiftDown=false, bool bControlDown=false) const;
+	void ApplyDeltaToActor(AActor* InActor, bool bDelta, const FVector* InTranslation, const FRotator* InRotation, const FVector* InScaling, bool bAltDown = false, bool bShiftDown = false, bool bControlDown = false) const;
 
-	void ApplyDeltaToComponent(USceneComponent* InComponent, bool bDelta, const FVector* InTranslation, const FRotator* InRotation, const FVector* InScaling, const FVector& PivotLocation ) const;
+	void ApplyDeltaToComponent(USceneComponent* InComponent, bool bDelta, const FVector* InTranslation, const FRotator* InRotation, const FVector* InScaling, const FVector& PivotLocation) const;
 
 	/**
 	 * Disable actor/component modification during delta movement.
@@ -1816,11 +1817,11 @@ public:
 	virtual bool Game_Map_Check_Actor(const TCHAR* Str, FOutputDevice& Ar, bool bCheckDeprecatedOnly, AActor* InActor) { return true; }
 
 	/** Map given class name key to value in ClassMapPair array - if no key match just return the class of given name */
-	class UClass* GetClassFromPairMap( FString ClassName );
+	class UClass* GetClassFromPairMap(FString ClassName);
 
 	/**
 	 * Auto merge all staticmeshes that are able to be merged
-	*/
+	 */
 	void AutoMergeStaticMeshes();
 
 	/**
@@ -1848,7 +1849,7 @@ public:
 	 *	@param	OutErrorMsg			if an errors occurs, this is the localized reason
 	 *	@return	true if the package is not a map file
 	 */
-	bool PackageIsAMapFile( const TCHAR* PackageFilename, FText& OutNotMapReason );
+	bool PackageIsAMapFile(const TCHAR* PackageFilename, FText& OutNotMapReason);
 
 	/**
 	 *	Searches through the given ULevel for any external references. Prints any found UObjects to the log and informs the user
@@ -1857,15 +1858,15 @@ public:
 	 *	@param	bAddForMapCheck		Optional flag to add any found references to the Map Check dialog (defaults to false)
 	 *	@return	true if the given package has external references and the user does not want to ignore the warning (via prompt)
 	 */
-	bool PackageUsingExternalObjects(ULevel* LevelToCheck, bool bAddForMapCheck=false);
+	bool PackageUsingExternalObjects(ULevel* LevelToCheck, bool bAddForMapCheck = false);
 
 	/**
 	 * Synchronizes the content or generic browser's selected objects to the collection specified.
 	 *
 	 * @param	ObjectSet	the list of objects to sync to
 	 */
-	void SyncBrowserToObjects( TArray<UObject*>& InObjectsToSync );
-	void SyncBrowserToObjects( TArray<class FAssetData>& InAssetsToSync );
+	void SyncBrowserToObjects(TArray<UObject*>& InObjectsToSync);
+	void SyncBrowserToObjects(TArray<class FAssetData>& InAssetsToSync);
 
 	/**
 	 * Syncs the selected actors objects to the content browser
@@ -1892,16 +1893,16 @@ public:
 	/**
 	 * Opens the objects specialized editor (Same as double clicking on the item in the content browser)
 	 *
-	 * @param ObjectToEdit	The object to open the editor for 
+	 * @param ObjectToEdit	The object to open the editor for
 	 */
-	void EditObject( UObject* ObjectToEdit );
+	void EditObject(UObject* ObjectToEdit);
 
 	/**
 	 * Selects the currently selected actor(s) levels in the level browser
 	 *
 	 * @param bDeselectOthers	Whether or not to deselect the current level selection first
 	 */
-	void SelectLevelInLevelBrowser( bool bDeselectOthers );
+	void SelectLevelInLevelBrowser(bool bDeselectOthers);
 
 	/**
 	 * Deselects the currently selected actor(s) levels in the level browser
@@ -1918,7 +1919,7 @@ public:
 	 *
 	 * @param bArchetype	true to only select actors of the same class AND same Archetype
 	 */
-	void SelectAllActorsWithClass( bool bArchetype );
+	void SelectAllActorsWithClass(bool bArchetype);
 
 	/**
 	 * Finds all references to the currently selected actors, and reports results in a find message log
@@ -1952,9 +1953,9 @@ public:
 	virtual bool IsPackageOKToSave(UPackage* InPackage, const FString& InFilename, FOutputDevice* Error);
 
 	/** The editor wrapper for UPackage::SavePackage. Auto-adds files to source control when necessary */
-	bool SavePackage( UPackage* InOuter, UObject* Base, EObjectFlags TopLevelFlags, const TCHAR* Filename, 
-		FOutputDevice* Error=GError, ULinkerLoad* Conform=NULL, bool bForceByteSwapping=false, bool bWarnOfLongFilename=true, 
-		uint32 SaveFlags=SAVE_None, const class ITargetPlatform* TargetPlatform = NULL, const FDateTime& FinalTimeStamp = FDateTime::MinValue(), bool bSlowTask = true );
+	bool SavePackage(UPackage* InOuter, UObject* Base, EObjectFlags TopLevelFlags, const TCHAR* Filename,
+		FOutputDevice* Error = GError, ULinkerLoad* Conform = NULL, bool bForceByteSwapping = false, bool bWarnOfLongFilename = true,
+		uint32 SaveFlags = SAVE_None, const class ITargetPlatform* TargetPlatform = NULL, const FDateTime& FinalTimeStamp = FDateTime::MinValue(), bool bSlowTask = true);
 
 	/** Invoked before a UWorld is saved to update editor systems */
 	virtual void OnPreSaveWorld(uint32 SaveFlags, UWorld* World);
@@ -1975,44 +1976,44 @@ public:
 	 */
 	void OnSourceControlDialogClosed(bool bEnabled);
 
-	/** 
+	/**
 	 * @return - returns the currently selected positional snap grid setting
 	 */
 	float GetGridSize();
 
-	/** 
+	/**
 	 * @return - if the grid size is part of the 1,2,48,16,.. list or not
 	 */
 	bool IsGridSizePowerOfTwo() const;
 
-	/** 
+	/**
 	 * Sets the selected positional snap grid setting
 	 *
 	 * @param InIndex - The index of the selected grid setting
 	 */
 	void SetGridSize(int32 InIndex);
 
-	/** 
+	/**
 	 * Increase the positional snap grid setting (will not increase passed the maximum value)
 	 */
 	void GridSizeIncrement();
-	
-	/** 
+
+	/**
 	 * Decreased the positional snap grid setting (will not decrease passed the minimum value)
 	 */
 	void GridSizeDecrement();
 
-	/** 
+	/**
 	 * Accesses the array of snap grid position options
 	 */
 	const TArray<float>& GetCurrentPositionGridArray() const;
-	
-	/** 
+
+	/**
 	 * @return - returns the currently selected rotational snap grid setting
 	 */
 	FRotator GetRotGridSize();
 
-	/** 
+	/**
 	 * Sets the selected Rotational snap grid setting
 	 *
 	 * @param InIndex - The index of the selected grid setting
@@ -2020,32 +2021,32 @@ public:
 	 */
 	void SetRotGridSize(int32 InIndex, enum ERotationGridMode InGridMode);
 
-	/** 
+	/**
 	 * Increase the rotational snap grid setting (will not increase passed the maximum value)
 	 */
 	void RotGridSizeIncrement();
-	
-	/** 
+
+	/**
 	 * Decreased the rotational snap grid setting (will not decrease passed the minimum value)
 	 */
 	void RotGridSizeDecrement();
-	
-	/** 
+
+	/**
 	 * Accesses the array of snap grid rotation options
 	 */
 	const TArray<float>& GetCurrentRotationGridArray() const;
-	
+
 	float GetScaleGridSize();
 	void SetScaleGridSize(int32 InIndex);
 
 	float GetGridInterval();
-	
-	 /**
-	  * Function to convert selected brushes into volumes of the provided class.
-	  *
-	  * @param	VolumeClass	Class of volume that selected brushes should be converted into
-	  */
-	void ConvertSelectedBrushesToVolumes( UClass* VolumeClass );
+
+	/**
+	 * Function to convert selected brushes into volumes of the provided class.
+	 *
+	 * @param	VolumeClass	Class of volume that selected brushes should be converted into
+	 */
+	void ConvertSelectedBrushesToVolumes(UClass* VolumeClass);
 
 	/**
 	 * Called to convert actors of one class type to another
@@ -2053,23 +2054,23 @@ public:
 	 * @param FromClass The class converting from
 	 * @param ToClass	The class converting to
 	 */
-	void ConvertActorsFromClass( UClass* FromClass, UClass* ToClass );
+	void ConvertActorsFromClass(UClass* FromClass, UClass* ToClass);
 
-	/** 
-	 * Show a (Suppressable) warning dialog to remind the user he is about to lose his undo buffer 
+	/**
+	 * Show a (Suppressable) warning dialog to remind the user he is about to lose his undo buffer
 	 *
 	 * @param MatineeActor	The actor we wish to check (can be null)
 	 * @returns true if the user wishes to proceed
 	 */
 	bool ShouldOpenMatinee(AMatineeActor* MatineeActor) const;
 
-	/** 
+	/**
 	 * Open the Matinee tool to edit the supplied MatineeActor. Will check that MatineeActor has an InterpData attached.
 	 *
 	 * @param MatineeActor	The actor we wish to edit
 	 * @param bWarnUser		If true, calls ShouldOpenMatinee as part of the open process
 	 */
-	void OpenMatinee(class AMatineeActor* MatineeActor, bool bWarnUser=true);
+	void OpenMatinee(class AMatineeActor* MatineeActor, bool bWarnUser = true);
 
 	void UpdateReflectionCaptures();
 
@@ -2079,7 +2080,7 @@ public:
 	 * Convenience method for adding a Slate modal window that is parented to the main frame (if it exists)
 	 * This function does not return until the modal window is closed.
 	 */
-	void EditorAddModalWindow( TSharedRef<SWindow> InModalWindow ) const;
+	void EditorAddModalWindow(TSharedRef<SWindow> InModalWindow) const;
 
 	/**
 	 * Finds a brush builder of the provided class.  If it doesnt exist one will be created
@@ -2087,7 +2088,7 @@ public:
 	 * @param BrushBuilderClass	The class to get the builder brush from
 	 * @return The found or created brush builder
 	 */
-	UBrushBuilder* FindBrushBuilder( UClass* BrushBuilderClass );
+	UBrushBuilder* FindBrushBuilder(UClass* BrushBuilderClass);
 
 	/**
 	 * Parents one actor to another
@@ -2098,14 +2099,14 @@ public:
 	 * @param SocketName	An optional socket name to attach to in the parent (use NAME_None when there is no socket)
 	 * @param Component		Actual Component included in the ParentActor which is a usually blueprint actor
 	 */
-	void ParentActors( AActor* ParentActor, AActor* ChildActor, const FName SocketName, USceneComponent* Component=NULL );
+	void ParentActors(AActor* ParentActor, AActor* ChildActor, const FName SocketName, USceneComponent* Component = NULL);
 
 	/**
 	 * Detaches selected actors from there parents
 	 *
 	 * @return				true if a detachment occurred
 	 */
-	bool DetachSelectedActors();	
+	bool DetachSelectedActors();
 
 	/**
 	 * Checks the validity of parenting one actor to another
@@ -2115,9 +2116,9 @@ public:
 	 * @param ReasonText	Optional text to receive a description of why the function returned false.
 	 * @return				true if the parenting action is valid and will succeed, false if invalid.
 	 */
-	bool CanParentActors( const AActor* ParentActor, const AActor* ChildActor, FText* ReasonText = NULL);
+	bool CanParentActors(const AActor* ParentActor, const AActor* ChildActor, FText* ReasonText = NULL);
 
-	/** 
+	/**
 	 * turns all navigable static geometry of ULevel into polygon soup stored in passed Level (ULevel::StaticNavigableGeometry)
 	 */
 	virtual void RebuildStaticNavigableGeometry(ULevel* Level);
@@ -2127,23 +2128,23 @@ public:
 	 *
 	 * @param Objects	Array to be filled with objects which can be browsed to
 	 */
-	void GetObjectsToSyncToContentBrowser( TArray<UObject*>& Objects );
+	void GetObjectsToSyncToContentBrowser(TArray<UObject*>& Objects);
 
 	/**
 	 * Queries for a list of assets that are referenced by the current editor selection (actors, surfaces, etc.)
 	 *
 	 * @param	Objects	Array to be filled with asset objects referenced by the current editor selection
 	 */
-	void GetReferencedAssetsForEditorSelection( TArray<UObject*>& Objects );
+	void GetReferencedAssetsForEditorSelection(TArray<UObject*>& Objects);
 
-	/** Returns the WorldContext for the editor world. For now, there will always be exactly 1 of these in the editor. 
+	/** Returns the WorldContext for the editor world. For now, there will always be exactly 1 of these in the editor.
 	 *
 	 * @param	bEnsureIsGWorld		Temporarily allow callers to validate that this maps to GWorld to help track down any conversion bugs
 	 */
 	FWorldContext &GetEditorWorldContext(bool bEnsureIsGWorld = false);
 
 
-	/** 
+	/**
 	 * mostly done to check if PIE is being set up, go GWorld is going to change, and it's not really _the_G_World_
 	 * NOTE: hope this goes away once PIE and regular game triggering are not that separate code paths
 	 */
@@ -2163,7 +2164,7 @@ public:
 	 */
 	FViewport* GetPIEViewport();
 
-	/** 
+	/**
 	 *	Checks for any player starts and returns the first one found.
 	 *
 	 *	@return		The first player start location found.
@@ -2199,37 +2200,37 @@ public:
 	 * @param	NewActorLabel			The new label string to assign to the actor.  If empty, the actor will have a default label.
 	 * @param	InExistingActorLabels	(optional) Pointer to a set of actor labels that are currently in use
 	 */
-	void SetActorLabelUnique( AActor* Actor, const FString& NewActorLabel, const FCachedActorLabels* InExistingActorLabels = nullptr ) const;
+	void SetActorLabelUnique(AActor* Actor, const FString& NewActorLabel, const FCachedActorLabels* InExistingActorLabels = nullptr) const;
 
 	/**
 	 * Gets the user-friendly, localized (if exists) name of a property
 	 *
-	 * @param	Property	the property we want to try to et the friendly name of	
+	 * @param	Property	the property we want to try to et the friendly name of
 	 * @param	OwnerClass	if specified, uses this class's loc file instead of the property's owner class
 	 *						useful for overriding the friendly name given a property inherited from a parent class.
 	 *
 	 * @return	the friendly name for the property.  localized first, then metadata, then the property's name.
 	 */
-	static FString GetFriendlyName( const UProperty* Property, UStruct* OwnerClass = NULL );
+	static FString GetFriendlyName(const UProperty* Property, UStruct* OwnerClass = NULL);
 
 	/**
-	 * Register a client tool to receive undo events 
+	 * Register a client tool to receive undo events
 	 * @param UndoClient	An object wanting to receive PostUndo/PostRedo events
 	 */
-	void RegisterForUndo (class FEditorUndoClient* UndoClient );
+	void RegisterForUndo(class FEditorUndoClient* UndoClient);
 
 	/**
-	 * Unregister a client from receiving undo events 
+	 * Unregister a client from receiving undo events
 	 * @param UndoClient	An object wanting to unsubscribe from PostUndo/PostRedo events
 	 */
-	void UnregisterForUndo ( class FEditorUndoClient* UndoEditor );
+	void UnregisterForUndo(class FEditorUndoClient* UndoEditor);
 
-	/** 
+	/**
 	 * Are we playing on a local PC session?
 	 */
 	bool IsPlayingOnLocalPCSession() const { return bPlayOnLocalPcSession && !bIsPlayWorldQueued; }
 
-	/** 
+	/**
 	 * Are we playing via the Launcher?
 	 */
 	bool IsPlayingViaLauncher() const { return bPlayUsingLauncher && !bIsPlayWorldQueued; }
@@ -2245,7 +2246,7 @@ public:
 	 * @param Assets		An array of assets to load and select
 	 * @param TypeOfAsset	An optional class to restrict the types of assets loaded & selected
 	 */
-	void LoadAndSelectAssets( TArray<FAssetData>& Assets, UClass* TypeOfAsset=NULL );
+	void LoadAndSelectAssets(TArray<FAssetData>& Assets, UClass* TypeOfAsset = NULL);
 
 	/**
 	 * @return True if percentage based scaling is enabled
@@ -2257,7 +2258,7 @@ public:
 	 */
 	virtual bool WorldIsPIEInNewViewport(UWorld *InWorld) override;
 
-	virtual void FocusNextPIEWorld(UWorld *CurrentPieWorld, bool previous=false) override;
+	virtual void FocusNextPIEWorld(UWorld *CurrentPieWorld, bool previous = false) override;
 
 	DECLARE_DELEGATE(FPIEInstanceWindowSwitch);
 
@@ -2271,12 +2272,12 @@ private:
 	// Map execs.
 	//
 
-	bool Map_Select( UWorld* InWorld, const TCHAR* Str, FOutputDevice& Ar);
-	bool Map_Brush( UWorld* InWorld, const TCHAR* Str, FOutputDevice& Ar);
-	bool Map_Sendto( UWorld* InWorld, const TCHAR* Str, FOutputDevice& Ar);
+	bool Map_Select(UWorld* InWorld, const TCHAR* Str, FOutputDevice& Ar);
+	bool Map_Brush(UWorld* InWorld, const TCHAR* Str, FOutputDevice& Ar);
+	bool Map_Sendto(UWorld* InWorld, const TCHAR* Str, FOutputDevice& Ar);
 	bool Map_Rebuild(UWorld* InWorld, const TCHAR* Str, FOutputDevice& Ar);
 	bool Map_Load(const TCHAR* Str, FOutputDevice& Ar);
-	bool Map_Import( UWorld* InWorld, const TCHAR* Str, FOutputDevice& Ar );
+	bool Map_Import(UWorld* InWorld, const TCHAR* Str, FOutputDevice& Ar);
 
 	struct EMapCheckNotification
 	{
@@ -2291,9 +2292,9 @@ private:
 	/**
 	 * Checks map for common errors.
 	 */
-	bool Map_Check( UWorld* InWorld, const TCHAR* Str, FOutputDevice& Ar, bool bCheckDeprecatedOnly, EMapCheckNotification::Type Notification = EMapCheckNotification::DisplayResults );
-	bool Map_Scale( UWorld* InWorld, const TCHAR* Str, FOutputDevice& Ar );
-	bool Map_Setbrush( UWorld* InWorld, const TCHAR* Str, FOutputDevice& Ar );
+	bool Map_Check(UWorld* InWorld, const TCHAR* Str, FOutputDevice& Ar, bool bCheckDeprecatedOnly, EMapCheckNotification::Type Notification = EMapCheckNotification::DisplayResults);
+	bool Map_Scale(UWorld* InWorld, const TCHAR* Str, FOutputDevice& Ar);
+	bool Map_Setbrush(UWorld* InWorld, const TCHAR* Str, FOutputDevice& Ar);
 
 	/**
 	 * Attempts to load a preview static mesh from the array preview static meshes at the given index.
@@ -2301,7 +2302,7 @@ private:
 	 * @param	Index	The index of the name in the PlayerPreviewMeshNames array from the editor user settings.
 	 * @return	true if a static mesh was loaded; false, otherwise.
 	 */
-	bool LoadPreviewMesh( int32 Index );
+	bool LoadPreviewMesh(int32 Index);
 
 	/**
 	 * Moves selected actors to the current level - NB. This is only intended to be called from MoveSelectedActorsToCurrentLevel()
@@ -2309,10 +2310,10 @@ private:
 	 *
 	 * @param	InLevel		Destination level.
 	 */
-	void DoMoveSelectedActorsToLevel( ULevel* InLevel );
+	void DoMoveSelectedActorsToLevel(ULevel* InLevel);
 public:
 	/**
-	 * Creates a PIE world by duplicating the editor world	 
+	 * Creates a PIE world by duplicating the editor world
 	 */
 	virtual UWorld* CreatePIEWorldByDuplication(FWorldContext &WorldContext, UWorld* InWorld, FString &PlayWorldMapName) override;
 private:
@@ -2336,7 +2337,7 @@ private:
 
 	/**
 	 * Delegate called as each PIE instance login is complete, continues creating the PIE world for a single instance
-	 * 
+	 *
 	 * @param LocalUserNum local user id, for PIE is going to be 0 (there is no splitscreen)
 	 * @param bWasSuccessful was the login successful
 	 * @param UserId userid of the logged in account
@@ -2356,7 +2357,7 @@ public:
 	void CreatePIEWorldFromLogin(FWorldContext& PieWorldContext, EPlayNetMode PlayNetMode, FPieLoginStruct& DataStruct);
 
 	/*
-	 * Handler for when viewport close request is made. 
+	 * Handler for when viewport close request is made.
 	 *
 	 * @param InViewport the viewport being closed.
 	 */
@@ -2380,7 +2381,7 @@ private:
 	/**
 	 * Toggles PIE to SIE or vice-versa
 	 */
-	void ToggleBetweenPIEandSIE( bool bNewSession = false);
+	void ToggleBetweenPIEandSIE(bool bNewSession = false);
 
 	/**
 	 * Hack to switch worlds for the PIE window before and after a slate event
@@ -2388,12 +2389,12 @@ private:
 	 * @param WorldID	The id of the world to restore or -1 if no world
 	 * @return The ID of the world to restore later or -1 if no world to restore
 	 */
-	int32 OnSwitchWorldForSlatePieWindow( int32 WorldID );
+	int32 OnSwitchWorldForSlatePieWindow(int32 WorldID);
 
 	/**
 	 * Called via a delegate to toggle between the editor and pie world
 	 */
-	void OnSwitchWorldsForPIE( bool bSwitchToPieWorld );
+	void OnSwitchWorldsForPIE(bool bSwitchToPieWorld);
 
 public:
 	/**
@@ -2403,7 +2404,7 @@ public:
 	 *
 	 * @return	true if spawn succeeded, false if failed
 	 */
-	bool SpawnPlayFromHereStart(UWorld* World, AActor*& PlayerStartPIE, const FVector& StartLocation, const FRotator& StartRotation );
+	bool SpawnPlayFromHereStart(UWorld* World, AActor*& PlayerStartPIE, const FVector& StartLocation, const FRotator& StartRotation);
 
 
 private:
@@ -2435,23 +2436,23 @@ private:
 	 * @param Component		Actual Component which has the Socket, this component will be used when ParentActor is a blueprint actor
 	 * @return true if successful, false otherwise
 	 */
-	bool SnapToSocket( AActor* ParentActor, AActor* ChildActor, const FName SocketName, USceneComponent* Component=NULL );
+	bool SnapToSocket(AActor* ParentActor, AActor* ChildActor, const FName SocketName, USceneComponent* Component = NULL);
 
 	/**
 	 * Delegate definition for the execute function that follows
 	 */
-	DECLARE_DELEGATE_OneParam( FSelectCommand, UModel* );
-	DECLARE_DELEGATE_TwoParams( FSelectInWorldCommand, UWorld*, UModel* );
-	
+	DECLARE_DELEGATE_OneParam(FSelectCommand, UModel*);
+	DECLARE_DELEGATE_TwoParams(FSelectInWorldCommand, UWorld*, UModel*);
+
 	/**
 	 * Utility method call a select command for each level model in the worlds level list.
 	 *
 	 * @param InWorld			World containing the levels
 	 * @param InSelectCommand	Select command delegate
 	 */
-	void ExecuteCommandForAllLevelModels( UWorld* InWorld, FSelectCommand InSelectCommand, const FText& TransDesription );
-	void ExecuteCommandForAllLevelModels( UWorld* InWorld, FSelectInWorldCommand InSelectCommand, const FText& TransDesription );
-	
+	void ExecuteCommandForAllLevelModels(UWorld* InWorld, FSelectCommand InSelectCommand, const FText& TransDesription);
+	void ExecuteCommandForAllLevelModels(UWorld* InWorld, FSelectInWorldCommand InSelectCommand, const FText& TransDesription);
+
 public:
 
 	/**
@@ -2459,8 +2460,8 @@ public:
 	 *
 	 * @param InWorld			World containing the levels
 	 */
-	 void FlagModifyAllSelectedSurfacesInLevels( UWorld* InWorld );
-	 
+	void FlagModifyAllSelectedSurfacesInLevels(UWorld* InWorld);
+
 private:
 
 	/**
@@ -2472,7 +2473,7 @@ private:
 	 * @param	NewWorld	An optional new world to keep in memory after destroying the world referenced in the Context.
 	 *						This world and it's sublevels will remain in memory.
 	 */
-	void EditorDestroyWorld( FWorldContext & Context, const FText& CleanseText, UWorld* NewWorld = nullptr );
+	void EditorDestroyWorld(FWorldContext & Context, const FText& CleanseText, UWorld* NewWorld = nullptr);
 
 	/** Returns the GameViewport widget */
 	virtual TSharedPtr<SViewport> GetGameViewportWidget() const override;
@@ -2485,13 +2486,13 @@ private:
 	 *
 	 * @return	true if the label ends with a number.
 	 */
-	bool SplitActorLabel( FString& InOutLabel, int32& OutIdx ) const;
+	bool SplitActorLabel(FString& InOutLabel, int32& OutIdx) const;
 
-	ULevel* CreateTransLevelMoveBuffer( UWorld* InWorld );
+	ULevel* CreateTransLevelMoveBuffer(UWorld* InWorld);
 
 	/**	Broadcasts that an undo has just occurred. */
 	void BroadcastPostUndo(const FString& UndoContext, UObject* PrimaryObject, bool bUndoSuccess);
-	
+
 	/**	Broadcasts that an redo has just occurred. */
 	void BroadcastPostRedo(const FString& RedoContext, UObject* PrimaryObject, bool bRedoSuccess);
 
@@ -2502,7 +2503,7 @@ private:
 	void BroadcastObjectsReplaced(const TMap<UObject*, UObject*>& ReplacementMap) { ObjectsReplacedEvent.Broadcast(ReplacementMap); }
 
 	/** Internal helper functions */
-	virtual void PostUndo (bool bSuccess);
+	virtual void PostUndo(bool bSuccess);
 
 	/** Delegate callback: the world origin is going to be moved. */
 	void PreWorldOriginOffset(UWorld* InWorld, FIntVector InSrcOrigin, FIntVector InDstOrigin);
@@ -2525,22 +2526,22 @@ private:
 	virtual void TriggerStreamingDataRebuild() override;
 
 	/** Remaps a network path for PIE Networking under multiple worlds */
-	virtual bool NetworkRemapPath( UWorld *InWorld, FString &Str, bool reading=true) override;
+	virtual bool NetworkRemapPath(UWorld *InWorld, FString &Str, bool reading = true) override;
 
 	/** Remaps a network path for PIE Networking under multiple worlds */
-	virtual bool NetworkRemapPath( UPendingNetGame *PendingNetGame, FString &Str, bool reading=true) override;
+	virtual bool NetworkRemapPath(UPendingNetGame *PendingNetGame, FString &Str, bool reading = true) override;
 
 	/** Handles user setting changes. */
-	void HandleSettingChanged( FName Name );
+	void HandleSettingChanged(FName Name);
 
 	/** Callback for handling undo and redo transactions before they happen. */
-	void HandleTransactorBeforeRedoUndo( FUndoSessionContext SessionContext );
+	void HandleTransactorBeforeRedoUndo(FUndoSessionContext SessionContext);
 
 	/** Callback for finished redo transactions. */
-	void HandleTransactorRedo( FUndoSessionContext SessionContext, bool Succeeded );
+	void HandleTransactorRedo(FUndoSessionContext SessionContext, bool Succeeded);
 
 	/** Callback for finished undo transactions. */
-	void HandleTransactorUndo( FUndoSessionContext SessionContext, bool Succeeded );
+	void HandleTransactorUndo(FUndoSessionContext SessionContext, bool Succeeded);
 
 private:
 
@@ -2620,10 +2621,10 @@ private:
 
 	/** Cached version of the view location at the point the PIE session was ended */
 	FVector LastViewLocation;
-	
+
 	/** Cached version of the view location at the point the PIE session was ended */
 	FRotator LastViewRotation;
-	
+
 	/** Are the lastview/rotation variables valid */
 	bool	bLastViewAndLocationValid;
 
@@ -2667,8 +2668,15 @@ public:
 
 private:
 	/** Handler for when any asset is loaded in the editor */
-	void OnAssetLoaded( UObject* Asset );
+	void OnAssetLoaded(UObject* Asset);
 
 	/** Gets the init values for worlds opened via Map_Load in the editor */
 	UWorld::InitializationValues GetEditorWorldInitializationValues() const;
+
+public:
+	// Launcher Worker
+	TSharedPtr<class ILauncherWorker> LauncherWorker;
+	
+	/** Function to run the Play On command for automation testing. */
+	void AutomationPlayUsingLauncher(const FString& InLauncherDeviceId);	
 };
diff --git a/Engine/Source/Editor/UnrealEd/Private/PlayLevel.cpp b/Engine/Source/Editor/UnrealEd/Private/PlayLevel.cpp
--- a/Engine/Source/Editor/UnrealEd/Private/PlayLevel.cpp
+++ b/Engine/Source/Editor/UnrealEd/Private/PlayLevel.cpp
@@ -1285,6 +1285,8 @@ void UEditorEngine::HandleStageStarted(const FString& InStage, TWeakPtr<SNotific
 
 void UEditorEngine::HandleStageCompleted(const FString& InStage, double StageTime, bool bHasCode, TWeakPtr<SNotificationItem> NotificationItemPtr)
 {
+	UE_LOG(LogPlayLevel, Log, TEXT("Completed Launch On Stage: %s, Time: %f"), *InStage, StageTime);
+
 	// analytics for launch on
 	TArray<FAnalyticsEventAttribute> ParamArray;
 	ParamArray.Add(FAnalyticsEventAttribute(TEXT("Time"), StageTime));
@@ -1335,6 +1337,8 @@ void UEditorEngine::HandleLaunchCompleted(bool Succeeded, double TotalTime, int3
 		TArray<FAnalyticsEventAttribute> ParamArray;
 		ParamArray.Add(FAnalyticsEventAttribute(TEXT("Time"), TotalTime));
 		FEditorAnalytics::ReportEvent(TEXT( "Editor.LaunchOn.Completed" ), PlayUsingLauncherDeviceId.Left(PlayUsingLauncherDeviceId.Find(TEXT("@"))), bHasCode, ParamArray);
+
+		UE_LOG(LogPlayLevel, Log, TEXT("Launch On Completed. Time: %f"), TotalTime);
 	}
 	else
 	{
@@ -1348,7 +1352,7 @@ void UEditorEngine::HandleLaunchCompleted(bool Succeeded, double TotalTime, int3
 		{
 			CompletionMsg = LOCTEXT("LauncherTaskFailed", "Launch failed!");
 		}
-
+		
 		MessageLog->Error()
 			->AddToken(FTextToken::Create(CompletionMsg))
 			->AddToken(FTextToken::Create(FText::FromString(FEditorAnalytics::TranslateErrorCode(ErrorCode))));
@@ -1385,6 +1389,8 @@ void UEditorEngine::PlayUsingLauncher()
 		ILauncherDeviceGroupRef DeviceGroup = LauncherServicesModule.CreateDeviceGroup(FGuid::NewGuid(), TEXT("PlayOnDevices"));
 		DeviceGroup->AddDevice(PlayUsingLauncherDeviceId);
 
+		UE_LOG(LogPlayLevel, Log, TEXT("Launcher Device ID: %s"), *PlayUsingLauncherDeviceId);
+
 		// does the project have any code?
 		FGameProjectGenerationModule& GameProjectModule = FModuleManager::LoadModuleChecked<FGameProjectGenerationModule>(TEXT("GameProjectGeneration"));
 		bool bHasCode = GameProjectModule.Get().ProjectHasCodeFiles();
@@ -1470,7 +1476,7 @@ void UEditorEngine::PlayUsingLauncher()
 		}
 
 		ILauncherPtr Launcher = LauncherServicesModule.CreateLauncher();
-		ILauncherWorkerPtr LauncherWorker = Launcher->Launch(TargetDeviceServicesModule.GetDeviceProxyManager(), LauncherProfile);
+		GEditor->LauncherWorker = Launcher->Launch(TargetDeviceServicesModule.GetDeviceProxyManager(), LauncherProfile);
 
 		// create notification item
 		FText LaunchingText = LOCTEXT("LauncherTaskInProgressNotificationNoDevice", "Launching...");
@@ -1491,7 +1497,7 @@ void UEditorEngine::PlayUsingLauncher()
 			FNotificationButtonInfo(
 				LOCTEXT("LauncherTaskCancel", "Cancel"),
 				LOCTEXT("LauncherTaskCancelToolTip", "Cancels execution of this task."),
-				FSimpleDelegate::CreateStatic(HandleCancelButtonClicked, LauncherWorker)
+				FSimpleDelegate::CreateStatic(HandleCancelButtonClicked, GEditor->LauncherWorker)
 			)
 		);
 
@@ -1507,22 +1513,22 @@ void UEditorEngine::PlayUsingLauncher()
 		FEditorAnalytics::ReportEvent(TEXT( "Editor.LaunchOn.Started" ), PlayUsingLauncherDeviceId.Left(PlayUsingLauncherDeviceId.Find(TEXT("@"))), bHasCode);
 
 		NotificationItem->SetCompletionState(SNotificationItem::CS_Pending);
-
+		
 		TWeakPtr<SNotificationItem> NotificationItemPtr(NotificationItem);
-		if (LauncherWorker.IsValid() && LauncherWorker->GetStatus() != ELauncherWorkerStatus::Completed)
+		if (GEditor->LauncherWorker.IsValid() && GEditor->LauncherWorker->GetStatus() != ELauncherWorkerStatus::Completed)
 		{
 			TSharedPtr<FMessageLog> MessageLog = MakeShareable(new FMessageLog("PackagingResults"));
 
 			GEditor->PlayEditorSound(TEXT("/Engine/EditorSounds/Notifications/CompileStart_Cue.CompileStart_Cue"));
-			LauncherWorker->OnOutputReceived().AddStatic(HandleOutputReceived);
-			LauncherWorker->OnStageStarted().AddUObject(this, &UEditorEngine::HandleStageStarted, NotificationItemPtr);
-			LauncherWorker->OnStageCompleted().AddUObject(this, &UEditorEngine::HandleStageCompleted, bHasCode, NotificationItemPtr);
-			LauncherWorker->OnCompleted().AddUObject(this, &UEditorEngine::HandleLaunchCompleted, bHasCode, NotificationItemPtr, MessageLog);
-			LauncherWorker->OnCanceled().AddUObject(this, &UEditorEngine::HandleLaunchCanceled, bHasCode, NotificationItemPtr);
+			GEditor->LauncherWorker->OnOutputReceived().AddStatic(HandleOutputReceived);
+			GEditor->LauncherWorker->OnStageStarted().AddUObject(this, &UEditorEngine::HandleStageStarted, NotificationItemPtr);
+			GEditor->LauncherWorker->OnStageCompleted().AddUObject(this, &UEditorEngine::HandleStageCompleted, bHasCode, NotificationItemPtr);
+			GEditor->LauncherWorker->OnCompleted().AddUObject(this, &UEditorEngine::HandleLaunchCompleted, bHasCode, NotificationItemPtr, MessageLog);
+			GEditor->LauncherWorker->OnCanceled().AddUObject(this, &UEditorEngine::HandleLaunchCanceled, bHasCode, NotificationItemPtr);
 		}
 		else
 		{
-			LauncherWorker.Reset();
+			GEditor->LauncherWorker.Reset();
 			GEditor->PlayEditorSound(TEXT("/Engine/EditorSounds/Notifications/CompileFailed_Cue.CompileFailed_Cue"));
 
 			NotificationItem->SetText(LOCTEXT("LauncherTaskFailedNotification", "Failed to launch task!"));
@@ -3315,6 +3321,11 @@ void UEditorEngine::RemapGamepadControllerIdForPIE(class UGameViewportClient* Ga
 	}
 }
 
-
+void UEditorEngine::AutomationPlayUsingLauncher(const FString& InLauncherDeviceId)
+{
+	PlayUsingLauncherDeviceId = InLauncherDeviceId;
+	PlayUsingLauncherDeviceName = PlayUsingLauncherDeviceId.Right(PlayUsingLauncherDeviceId.Find(TEXT("@")));
+	PlayUsingLauncher();
+}
 
 #undef LOCTEXT_NAMESPACE
diff --git a/Engine/Source/Editor/UnrealEd/Private/Tests/AutomationEditorCommon.cpp b/Engine/Source/Editor/UnrealEd/Private/Tests/AutomationEditorCommon.cpp
--- a/Engine/Source/Editor/UnrealEd/Private/Tests/AutomationEditorCommon.cpp
+++ b/Engine/Source/Editor/UnrealEd/Private/Tests/AutomationEditorCommon.cpp
@@ -49,7 +49,18 @@
 #include "Engine/TextureRenderTarget2D.h"
 #include "Engine/UserDefinedEnum.h"
 #include "Engine/World.h"
+#include "AssetSelection.h"
 
+#include "Interfaces/ILauncherDeviceGroup.h"
+#include "LauncherServices.h"
+#include "TargetDeviceServices.h"
+#include "Editor/EditorEngine.h"
+#include "PlatformInfo.h"
+
+#include "CookOnTheSide/CookOnTheFlyServer.h"
+#include "Interfaces/ILauncherProfile.h"
+
+#define COOK_TIMEOUT 3600
 
 DEFINE_LOG_CATEGORY_STATIC(LogAutomationEditorCommon, Log, All);
 
@@ -431,6 +442,65 @@ namespace AutomationEditorCommonUtils
 		}
 		return true;
 	}
+
+
+	void GetLaunchOnDeviceID(FString& OutDeviceID, const FString& InMapName)
+	{
+		UAutomationTestSettings const* AutomationTestSettings = GetDefault<UAutomationTestSettings>();
+		check(AutomationTestSettings);
+
+		OutDeviceID = "None";
+
+		FString LaunchOnDeviceId;
+		for (auto LaunchIter = AutomationTestSettings->LaunchOnSettings.CreateConstIterator(); LaunchIter; LaunchIter++)
+		{
+			FString LaunchOnSettings = LaunchIter->DeviceID;
+			FString LaunchOnMap = FPaths::GetBaseFilename(LaunchIter->LaunchOnTestmap.FilePath);
+			if (LaunchOnMap.Equals(InMapName))
+			{
+				// shared devices section
+				TSharedPtr<ITargetDeviceServicesModule> TargetDeviceServicesModule = StaticCastSharedPtr<ITargetDeviceServicesModule>(FModuleManager::Get().LoadModule(TEXT("TargetDeviceServices")));
+				// for each platform...
+				TArray<ITargetDeviceProxyPtr> DeviceProxies;
+				TargetDeviceServicesModule->GetDeviceProxyManager()->GetProxies(FName(*LaunchOnSettings), true, DeviceProxies);
+				// for each proxy...
+				for (auto DeviceProxyIt = DeviceProxies.CreateIterator(); DeviceProxyIt; ++DeviceProxyIt)
+				{
+					ITargetDeviceProxyPtr DeviceProxy = *DeviceProxyIt;
+					if (DeviceProxy->IsConnected())
+					{
+						OutDeviceID = DeviceProxy->GetTargetDeviceId((FName)*LaunchOnSettings);
+						break;
+					}
+				}
+			}
+		}
+	}
+
+	void GetLaunchOnDeviceID(FString& OutDeviceID, const FString& InMapName, const FString& InDeviceName)
+	{
+		UAutomationTestSettings const* AutomationTestSettings = GetDefault<UAutomationTestSettings>();
+		check(AutomationTestSettings);
+
+		//Output device name will default to "None".
+		OutDeviceID = "None";
+
+		// shared devices section
+		TSharedPtr<ITargetDeviceServicesModule> TargetDeviceServicesModule = StaticCastSharedPtr<ITargetDeviceServicesModule>(FModuleManager::Get().LoadModule(TEXT("TargetDeviceServices")));
+		// for each platform...
+		TArray<ITargetDeviceProxyPtr> DeviceProxies;
+		TargetDeviceServicesModule->GetDeviceProxyManager()->GetProxies(FName(*InDeviceName), true, DeviceProxies);
+		// for each proxy...
+		for (auto DeviceProxyIt = DeviceProxies.CreateIterator(); DeviceProxyIt; ++DeviceProxyIt)
+		{
+			ITargetDeviceProxyPtr DeviceProxy = *DeviceProxyIt;
+			if (DeviceProxy->IsConnected())
+			{
+				OutDeviceID = DeviceProxy->GetTargetDeviceId((FName)*InDeviceName);
+				break;
+			}
+		}
+	}
 }
 
 ///////////////////////////////////////////////////////////////////////
@@ -588,6 +658,128 @@ bool FChangeViewportToFirstAvailableBookmarkCommand::Update()
 	}
 	return true;
 }
+
+/**
+* Latent command that adds a static mesh to the worlds origin.
+*/
+bool FAddStaticMeshCommand::Update()
+{
+	//Gather assets.
+	UObject* Cube = (UStaticMesh*)StaticLoadObject(UStaticMesh::StaticClass(), NULL, TEXT("/Engine/EngineMeshes/Cube.Cube"), NULL, LOAD_None, NULL);
+	//Add Cube mesh to the world
+	AStaticMeshActor* StaticMesh = Cast<AStaticMeshActor>(FActorFactoryAssetProxy::AddActorForAsset(Cube));
+	StaticMesh->TeleportTo(FVector(0.0f, 0.0f, 0.0f), FRotator(0, 0, 0));
+	StaticMesh->SetActorRelativeScale3D(FVector(1.0f, 1.0f, 1.0f));
+
+	UE_LOG(LogEditorAutomationTests, Log, TEXT("Static Mesh cube has been added to 0, 0, 0."))
+
+		return true;
+}
+
+/**
+* Latent command that builds lighting for the current level.
+*/
+bool FBuildLightingCommand::Update()
+{
+	//If we are running with -NullRHI then we have to skip this step.
+	if (GUsingNullRHI)
+	{
+		UE_LOG(LogEditorAutomationTests, Warning, TEXT("SKIPPED Build Lighting Step.  You're currently running with -NullRHI."));
+		return true;
+	}
+
+	if (GUnrealEd->WarnIfLightingBuildIsCurrentlyRunning())
+	{
+		UE_LOG(LogEditorAutomationTests, Warning, TEXT("Lighting is already being built."));
+		return true;
+	}
+
+	FLightingBuildOptions LightingBuildOptions;
+
+	// Retrieve settings from ini.
+	GConfig->GetBool(TEXT("LightingBuildOptions"), TEXT("OnlyBuildSelected"), LightingBuildOptions.bOnlyBuildSelected, GEditorUserSettingsIni);
+	GConfig->GetBool(TEXT("LightingBuildOptions"), TEXT("OnlyBuildCurrentLevel"), LightingBuildOptions.bOnlyBuildCurrentLevel, GEditorUserSettingsIni);
+	GConfig->GetBool(TEXT("LightingBuildOptions"), TEXT("OnlyBuildSelectedLevels"), LightingBuildOptions.bOnlyBuildSelectedLevels, GEditorUserSettingsIni);
+	GConfig->GetBool(TEXT("LightingBuildOptions"), TEXT("OnlyBuildVisibility"), LightingBuildOptions.bOnlyBuildVisibility, GEditorUserSettingsIni);
+	GConfig->GetBool(TEXT("LightingBuildOptions"), TEXT("UseErrorColoring"), LightingBuildOptions.bUseErrorColoring, GEditorUserSettingsIni);
+	GConfig->GetBool(TEXT("LightingBuildOptions"), TEXT("ShowLightingBuildInfo"), LightingBuildOptions.bShowLightingBuildInfo, GEditorUserSettingsIni);
+	int32 QualityLevel;
+	GConfig->GetInt(TEXT("LightingBuildOptions"), TEXT("QualityLevel"), QualityLevel, GEditorUserSettingsIni);
+	QualityLevel = FMath::Clamp<int32>(QualityLevel, Quality_Preview, Quality_Production);
+	LightingBuildOptions.QualityLevel = Quality_Production;
+
+	UE_LOG(LogEditorAutomationTests, Log, TEXT("Building lighting in Production Quality."));
+	GUnrealEd->BuildLighting(LightingBuildOptions);
+
+	return true;
+}
+
+
+bool FSaveLevelCommand::Update()
+{
+	if (!GUnrealEd->IsLightingBuildCurrentlyExporting() && !GUnrealEd->IsLightingBuildCurrentlyRunning())
+	{
+		UWorld* World = GEditor->GetEditorWorldContext().World();
+		ULevel* Level = World->GetCurrentLevel();
+		MapName += TEXT("_Copy.umap");
+		FString TempMapLocation = FPaths::Combine(*FPaths::GameContentDir(), TEXT("Maps"), TEXT("Automation_TEMP"), *MapName);
+		FEditorFileUtils::SaveLevel(Level, TempMapLocation);
+
+		return true;
+	}
+	return false;
+}
+
+bool FLaunchOnCommand::Update()
+{
+	GUnrealEd->AutomationPlayUsingLauncher(InLauncherDeviceID);
+	return true;
+}
+
+bool FWaitToFinishCookByTheBookCommand::Update()
+{
+	if (!GUnrealEd->CookServer->IsCookByTheBookRunning())
+	{
+		if (GUnrealEd->IsCookByTheBookInEditorFinished())
+		{
+			UE_LOG(LogEditorAutomationTests, Log, TEXT("The cook by the book operation has finished."));
+		}
+		return true;
+	}
+	else if ((FPlatformTime::Seconds() - StartTime) == COOK_TIMEOUT)
+	{
+		GUnrealEd->CancelCookByTheBookInEditor();
+		UE_LOG(LogEditorAutomationTests, Error, TEXT("It has been an hour or more since the cook has started."));
+		return false;
+	}
+	return false;
+}
+
+bool FDeleteDirCommand::Update()
+{
+	FString FullFolderPath = FPaths::ConvertRelativePathToFull(*InFolderLocation);
+	if (IFileManager::Get().DirectoryExists(*FullFolderPath))
+	{
+		IFileManager::Get().DeleteDirectory(*FullFolderPath, false, true);
+	}
+	return true;
+}
+
+bool FWaitToFinishBuildDeployCommand::Update()
+{
+	if (GEditor->LauncherWorker->GetStatus() == ELauncherWorkerStatus::Completed)
+	{
+		UE_LOG(LogEditorAutomationTests, Log, TEXT("The build game and deploy operation has finished."));
+		return true;
+	}
+	else if (GEditor->LauncherWorker->GetStatus() == ELauncherWorkerStatus::Canceled || GEditor->LauncherWorker->GetStatus() == ELauncherWorkerStatus::Canceling)
+	{
+		UE_LOG(LogEditorAutomationTests, Warning, TEXT("The build was canceled."));
+		return true;
+	}
+	return false;
+}
+
 //////////////////////////////////////////////////////////////////////
 //Find Asset Commands
 
@@ -795,4 +987,4 @@ void FEditorAutomationTestUtilities::CollectGameContentByClass(const UClass * Cl
 			OutAssetList.Add(Asset.ObjectPath.ToString());
 		}
 	}
-}
\ No newline at end of file
+}
diff --git a/Engine/Source/Editor/UnrealEd/Private/Tests/EditorAutomationTests.cpp b/Engine/Source/Editor/UnrealEd/Private/Tests/EditorAutomationTests.cpp
--- a/Engine/Source/Editor/UnrealEd/Private/Tests/EditorAutomationTests.cpp
+++ b/Engine/Source/Editor/UnrealEd/Private/Tests/EditorAutomationTests.cpp
@@ -13,6 +13,7 @@
 #include "LevelEditor.h"
 #include "ModuleManager.h"
 #include "MainFrame.h"
+#include "FileManager.h"
 
 #include "Kismet2/KismetEditorUtilities.h"
 #include "Kismet2/BlueprintEditorUtils.h"
@@ -1448,14 +1449,12 @@ bool FLightPlacement::RunTest(const FString& Parameters)
 		}
 	}
 
-	{
-		//Gather assets.
-		UObject* Cube = (UStaticMesh*)StaticLoadObject(UStaticMesh::StaticClass(), NULL, TEXT("/Engine/EngineMeshes/Cube.Cube"), NULL, LOAD_None, NULL);
-		//Add Cube mesh to the world
-		AStaticMeshActor* StaticMesh = Cast<AStaticMeshActor>(FActorFactoryAssetProxy::AddActorForAsset(Cube));
-		StaticMesh->TeleportTo(FVector(0.0f, 0.0f, 0.0f), FRotator(0, 0, 0));
-		StaticMesh->SetActorRelativeScale3D(FVector(3.0f, 3.0f, 1.75f));
-	}
+	//Gather assets.
+	UObject* Cube = (UStaticMesh*)StaticLoadObject(UStaticMesh::StaticClass(), NULL, TEXT("/Engine/EngineMeshes/Cube.Cube"), NULL, LOAD_None, NULL);
+	//Add Cube mesh to the world
+	AStaticMeshActor* StaticMesh = Cast<AStaticMeshActor>(FActorFactoryAssetProxy::AddActorForAsset(Cube));
+	StaticMesh->TeleportTo(FVector(0.0f, 0.0f, 0.0f), FRotator(0, 0, 0));
+	StaticMesh->SetActorRelativeScale3D(FVector(3.0f, 3.0f, 1.75f));
 
 	//Create the point light and set it's mobility, brightness, and light color.
 	const FTransform Transform(FVector(0.0f, 0.0f, 400.0f));
@@ -2073,39 +2072,123 @@ bool FStaticMeshUVsTest::RunTest(const FString& Parameters)
 	return true;
 }
 
-
 /**
 * Launches a map onto a specified device after making a change to it.
 */
-IMPLEMENT_COMPLEX_AUTOMATION_TEST(FLaunchOnTestWithMapIteration, "Performance.Launch Test With Map Iteration", (EAutomationTestFlags::ATF_Editor | EAutomationTestFlags::ATF_RequiresUser))
+IMPLEMENT_COMPLEX_AUTOMATION_TEST(FLaunchOnTest, "Editor.Launch On Test", (EAutomationTestFlags::ATF_Editor | EAutomationTestFlags::ATF_RequiresUser))
 
-void FLaunchOnTestWithMapIteration::GetTests(TArray<FString>& OutBeautifiedNames, TArray <FString>& OutTestCommands) const
+void FLaunchOnTest::GetTests(TArray<FString>& OutBeautifiedNames, TArray <FString>& OutTestCommands) const
 {
-	//Get Map names from preferences (engine.ini)
+	{
+		UAutomationTestSettings const* AutomationTestSettings = GetDefault<UAutomationTestSettings>();
+		check(AutomationTestSettings);
 
-	//Check to see if the map has a set launch on device.
-		//If then don't add it as a test.
-}
+		TArray<FString> MapToLaunch;
+		TArray<FString> DeviceToUse;
+		for (auto Iter = AutomationTestSettings->LaunchOnSettings.CreateConstIterator(); Iter; ++Iter)
+		{
+			if (Iter->LaunchOnTestmap.FilePath.Len() > 0 && !Iter->DeviceID.IsEmpty())
+			{
+				MapToLaunch.Add(Iter->LaunchOnTestmap.FilePath);
+				DeviceToUse.Add(Iter->DeviceID);
+			}
 
-bool FLaunchOnTestWithMapIteration::RunTest(const FString& Parameters)
-{
-	//Get the map Name
+			for (int32 i = 0; i < MapToLaunch.Num(); ++i)
+			{
+				//Get the location of the map being used.
+				FString Filename = FPaths::ConvertRelativePathToFull(MapToLaunch[i]);
 
-	//Load Map
-	
-	//Make an adjustment to the map
-		
-	//Rebuild Map
+				//Get the DeviceID
+				FString DeviceID;
+				AutomationEditorCommonUtils::GetLaunchOnDeviceID(DeviceID, FPaths::GetBaseFilename(MapToLaunch[i]), DeviceToUse[i]);
+
+				if (!DeviceID.IsEmpty() && !DeviceID.Equals(TEXT("None")))
+				{
+					if (Filename.Contains(TEXT("/Engine/"), ESearchCase::IgnoreCase, ESearchDir::FromStart))
+					{
+						//If true it will proceed to add the asset to the test list.
+						//This will be false if the map is on a different drive.
+						if (FPaths::MakePathRelativeTo(Filename, *FPaths::EngineContentDir()))
+						{
+							FString ShortName = FPaths::GetBaseFilename(Filename);
+							FString PathName = FPaths::GetPath(Filename);
+							FString AssetName = FString::Printf(TEXT("/Game/%s/%s.%s %s"), *PathName, *ShortName, *ShortName, *DeviceID);
+
+							ShortName += (TEXT(" ( ") + DeviceID.Left(DeviceID.Find(TEXT("@"))) + TEXT(" ) "));
+
+							OutBeautifiedNames.Add(ShortName);
+							OutTestCommands.Add(AssetName);
+						}
+						else
+						{
+							UE_LOG(LogEditorAutomationTests, Error, TEXT("Invalid asset path: %s."), *Filename);
+						}
+					}
+					else
+					{
+						//If true it will proceed to add the asset to the test list.
+						//This will be false if the map is on a different drive.
+						if (FPaths::MakePathRelativeTo(Filename, *FPaths::GameContentDir()))
+						{
+							FString ShortName = FPaths::GetBaseFilename(Filename);
+							FString PathName = FPaths::GetPath(Filename);
+							FString AssetName = FString::Printf(TEXT("/Game/%s/%s.%s %s"), *PathName, *ShortName, *ShortName, *DeviceID);
 
-	//Check out Map and then save it
+							ShortName += (TEXT(" (") + DeviceID.Left(DeviceID.Find(TEXT("@"))) + TEXT(") "));
 
-	//Set the device to be launched onto
+							OutBeautifiedNames.Add(ShortName);
+							OutTestCommands.Add(AssetName);
+						}
+						else
+						{
+							UE_LOG(LogEditorAutomationTests, Error, TEXT("Invalid asset path: %s."), *Filename);
+						}
+					}
+				}
+			}
+		}
+	}
+}
+
+bool FLaunchOnTest::RunTest(const FString& Parameters)
+{
+	//Get the map name and device id from the parameters.
+	FString MapName = Parameters.Left(Parameters.Find(TEXT(" ")));
+	FString DeviceID = Parameters.RightChop(Parameters.Find(TEXT(" ")));
+	DeviceID.Trim();
+
+	//Delete the Cooked, StagedBuilds, and Automation_TEMP folder if they exist.
+	FString CookedLocation = FPaths::Combine(*FPaths::GameSavedDir(), TEXT("Cooked"));
+	FString StagedBuildsLocation = FPaths::Combine(*FPaths::GameSavedDir(), TEXT("StagedBuilds"));
+	FString TempMapLocation = FPaths::Combine(*FPaths::GameContentDir(), TEXT("Maps"), TEXT("Automation_TEMP"));
+	ADD_LATENT_AUTOMATION_COMMAND(FDeleteDirCommand(CookedLocation));
+	ADD_LATENT_AUTOMATION_COMMAND(FDeleteDirCommand(StagedBuildsLocation));
+	ADD_LATENT_AUTOMATION_COMMAND(FDeleteDirCommand(TempMapLocation));
+	ADD_LATENT_AUTOMATION_COMMAND(FWaitLatentCommand(2.0f));
+
+	//Load Map and get the time it took to take to load the map.
+	ADD_LATENT_AUTOMATION_COMMAND(FEditorLoadMap(MapName));
+	ADD_LATENT_AUTOMATION_COMMAND(FWaitLatentCommand(1.0f));
+
+	//Make an adjustment to the map and rebuild its lighting.
+	ADD_LATENT_AUTOMATION_COMMAND(FAddStaticMeshCommand);
+	ADD_LATENT_AUTOMATION_COMMAND(FWaitLatentCommand(1.0f));
+	ADD_LATENT_AUTOMATION_COMMAND(FBuildLightingCommand);
+
+	//Save a copy of the map to the automation temp map folder location once the lighting build has finish.
+	ADD_LATENT_AUTOMATION_COMMAND(FSaveLevelCommand(FPaths::GetBaseFilename(MapName)));
 
 	//Launch onto device and get launch on times and cook times
+	ADD_LATENT_AUTOMATION_COMMAND(FLaunchOnCommand(DeviceID));
+	ADD_LATENT_AUTOMATION_COMMAND(FWaitToFinishCookByTheBookCommand);
+	ADD_LATENT_AUTOMATION_COMMAND(FWaitToFinishBuildDeployCommand);
 
-	//Verify the launch worked.
+	//@todo: Close the Launched on Game.
 
-	//Revert Map
+	//Delete the Cooked, StagedBuilds, and Automation_TEMP folder if they exist.
+	ADD_LATENT_AUTOMATION_COMMAND(FDeleteDirCommand(CookedLocation));
+	ADD_LATENT_AUTOMATION_COMMAND(FDeleteDirCommand(StagedBuildsLocation));
+	ADD_LATENT_AUTOMATION_COMMAND(FDeleteDirCommand(TempMapLocation));
 
 	return true;
 }
\ No newline at end of file
diff --git a/Engine/Source/Editor/UnrealEd/Public/Tests/AutomationEditorCommon.h b/Engine/Source/Editor/UnrealEd/Public/Tests/AutomationEditorCommon.h
--- a/Engine/Source/Editor/UnrealEd/Public/Tests/AutomationEditorCommon.h
+++ b/Engine/Source/Editor/UnrealEd/Public/Tests/AutomationEditorCommon.h
@@ -135,6 +135,25 @@ namespace AutomationEditorCommonUtils
 	* @param InArchiveName - is the name of the archive to be used.
 	*/
 	bool IsArchiveWriteable(const FString& InFilePath, const FArchive* InArchiveName);
+		
+	/**
+	* Returns the first DeviceID 'Platform@Device'
+	* Searches the automation preferences for the platform to use.
+
+	* @param OutDeviceID - The variable that will hold the device ID.
+	* @param InMapName - The map name to check against in the automation preferences.
+	*/
+	void GetLaunchOnDeviceID(FString& OutDeviceID, const FString& InMapName);
+
+	/**
+	* Returns the DeviceID 'Platform@Device'
+	* Searches the automation preferences for the platform to use.
+
+	* @param OutDeviceID - The variable that will hold the device ID.
+	* @param InMapName - The map name to check against in the automation preferences.
+	* @param InDeviceName - Device Name 
+	*/
+	void GetLaunchOnDeviceID(FString& OutDeviceID, const FString& InMapName, const FString& InDeviceName);
 }
 
 
@@ -182,9 +201,44 @@ DEFINE_LATENT_AUTOMATION_COMMAND(FWaitForShadersToFinishCompiling);
 */
 DEFINE_LATENT_AUTOMATION_COMMAND(FChangeViewportToFirstAvailableBookmarkCommand);
 
+/**
+* Latent command that adds a static mesh to the worlds origin.
+*/
+DEFINE_LATENT_AUTOMATION_COMMAND(FAddStaticMeshCommand);
+
+/**
+* Latent command that builds the lighting for the current level.
+*/
+DEFINE_LATENT_AUTOMATION_COMMAND(FBuildLightingCommand);
+
+/**
+* Latent command that saves a copy of the level to a transient folder.
+*/
+DEFINE_LATENT_AUTOMATION_COMMAND_ONE_PARAMETER(FSaveLevelCommand, FString, MapName);
+
+/**
+* Triggers a launch on using a specified device ID.
+*/
+DEFINE_LATENT_AUTOMATION_COMMAND_ONE_PARAMETER(FLaunchOnCommand, FString, InLauncherDeviceID);
+
+/**
+* Wait for a cook by the book to finish.  Will time out after 3600 seconds (1 hr).
+*/
+DEFINE_LATENT_AUTOMATION_COMMAND(FWaitToFinishCookByTheBookCommand);
+
+/**
+* Wait for a build and deploy to finish before moving on.
+*/
+DEFINE_LATENT_AUTOMATION_COMMAND(FWaitToFinishBuildDeployCommand);
+
+/**
+* Latent command to delete a directory.
+*/
+DEFINE_LATENT_AUTOMATION_COMMAND_ONE_PARAMETER(FDeleteDirCommand, FString, InFolderLocation);
+
+
 //////////////////////////////////////////////////////////////////////////
 // FEditorAutomationTestUtilities
-
 class FEditorAutomationTestUtilities
 {
 public:
diff --git a/Engine/Source/Runtime/Engine/Classes/Tests/AutomationTestSettings.h b/Engine/Source/Runtime/Engine/Classes/Tests/AutomationTestSettings.h
--- a/Engine/Source/Runtime/Engine/Classes/Tests/AutomationTestSettings.h
+++ b/Engine/Source/Runtime/Engine/Classes/Tests/AutomationTestSettings.h
@@ -291,6 +291,23 @@ struct FEditorMapPerformanceTestDefinition
 	int32 TestTimer;
 };
 
+/**
+* Holds settings for the editor Launch On With Map Iterations test.
+*/
+USTRUCT()
+struct FLaunchOnTestSettings
+{
+	GENERATED_USTRUCT_BODY()
+
+	/** Map to be used for the Launch On test **/
+	UPROPERTY(config, EditAnywhere, Category = Automation, meta = (FilePathFilter = "umap"))
+	FFilePath LaunchOnTestmap;
+
+	/** Device to be used **/
+	UPROPERTY(config, EditAnywhere, Category = Automation, meta = (ToolTip = "This is the device to be used for launch on. Example: WindowsNoEditor, Android, IOS, Linux"))
+	FString DeviceID;
+};
+
 /**
  * Implements the Editor's user settings.
  */
@@ -355,7 +372,10 @@ public:
 	UPROPERTY(EditAnywhere, config, Category = Automation)
 	TArray<FEditorImportExportTestDefinition> ImportExportTestDefinitions;
 
-
-
+	/**
+	* The map and device type to be used for the editor Launch On With Map Iterations test.
+	*/
+	UPROPERTY(config, EditAnywhere, Category = Automation, meta = (FilePathFilter = "umap"))
+	TArray<FLaunchOnTestSettings> LaunchOnSettings;
 
 };
