Fixing issues with context menu in Sound and Material Editors

Implemented DoesSupportEventDispatcher to return true in the relevant schemas so that context menu can be opened

diff --git a/Engine/Source/Editor/UnrealEd/Classes/MaterialGraph/MaterialGraphSchema.h b/Engine/Source/Editor/UnrealEd/Classes/MaterialGraph/MaterialGraphSchema.h
--- a/Engine/Source/Editor/UnrealEd/Classes/MaterialGraph/MaterialGraphSchema.h
+++ b/Engine/Source/Editor/UnrealEd/Classes/MaterialGraph/MaterialGraphSchema.h
@@ -195,6 +195,7 @@ class UMaterialGraphSchema : public UEdGraphSchema
 	virtual void GetContextMenuActions(const UEdGraph* CurrentGraph, const UEdGraphNode* InGraphNode, const UEdGraphPin* InGraphPin, class FMenuBuilder* MenuBuilder, bool bIsDebugging) const override;
 	virtual const FPinConnectionResponse CanCreateConnection(const UEdGraphPin* A, const UEdGraphPin* B) const override;
 	virtual bool TryCreateConnection(UEdGraphPin* A, UEdGraphPin* B) const override;
+	virtual bool DoesSupportEventDispatcher() const override { return true; }
 	virtual bool ShouldHidePinDefaultValue(UEdGraphPin* Pin) const override { return true; }
 	virtual FLinearColor GetPinTypeColor(const FEdGraphPinType& PinType) const override;
 	virtual void BreakNodeLinks(UEdGraphNode& TargetNode) const override;
diff --git a/Engine/Source/Editor/UnrealEd/Classes/SoundClassGraph/SoundClassGraphSchema.h b/Engine/Source/Editor/UnrealEd/Classes/SoundClassGraph/SoundClassGraphSchema.h
--- a/Engine/Source/Editor/UnrealEd/Classes/SoundClassGraph/SoundClassGraphSchema.h
+++ b/Engine/Source/Editor/UnrealEd/Classes/SoundClassGraph/SoundClassGraphSchema.h
@@ -46,6 +46,7 @@ class USoundClassGraphSchema : public UEdGraphSchema
 	virtual void GetContextMenuActions(const UEdGraph* CurrentGraph, const UEdGraphNode* InGraphNode, const UEdGraphPin* InGraphPin, class FMenuBuilder* MenuBuilder, bool bIsDebugging) const override;
 	virtual const FPinConnectionResponse CanCreateConnection(const UEdGraphPin* PinA, const UEdGraphPin* PinB) const override;
 	virtual bool TryCreateConnection(UEdGraphPin* PinA, UEdGraphPin* PinB) const override;
+	virtual bool DoesSupportEventDispatcher() const override { return true; }
 	virtual bool ShouldHidePinDefaultValue(UEdGraphPin* Pin) const override;
 	virtual FLinearColor GetPinTypeColor(const FEdGraphPinType& PinType) const override;
 	virtual void BreakNodeLinks(UEdGraphNode& TargetNode) const override;
diff --git a/Engine/Source/Editor/UnrealEd/Classes/SoundCueGraph/SoundCueGraphSchema.h b/Engine/Source/Editor/UnrealEd/Classes/SoundCueGraph/SoundCueGraphSchema.h
--- a/Engine/Source/Editor/UnrealEd/Classes/SoundCueGraph/SoundCueGraphSchema.h
+++ b/Engine/Source/Editor/UnrealEd/Classes/SoundCueGraph/SoundCueGraphSchema.h
@@ -114,6 +114,7 @@ class USoundCueGraphSchema : public UEdGraphSchema
 	virtual void CreateDefaultNodesForGraph(UEdGraph& Graph) const override;
 	virtual const FPinConnectionResponse CanCreateConnection(const UEdGraphPin* A, const UEdGraphPin* B) const override;
 	virtual bool TryCreateConnection(UEdGraphPin* A, UEdGraphPin* B) const override;
+	virtual bool DoesSupportEventDispatcher() const override { return true; }
 	virtual bool ShouldHidePinDefaultValue(UEdGraphPin* Pin) const override;
 	virtual FLinearColor GetPinTypeColor(const FEdGraphPinType& PinType) const override;
 	virtual void BreakNodeLinks(UEdGraphNode& TargetNode) const override;
