Added support for ini settings invalidating cooked content for cook in editor.

diff --git a/Engine/Source/Editor/UnrealEd/Classes/CookOnTheSide/CookOnTheFlyServer.h b/Engine/Source/Editor/UnrealEd/Classes/CookOnTheSide/CookOnTheFlyServer.h
--- a/Engine/Source/Editor/UnrealEd/Classes/CookOnTheSide/CookOnTheFlyServer.h
+++ b/Engine/Source/Editor/UnrealEd/Classes/CookOnTheSide/CookOnTheFlyServer.h
@@ -42,6 +42,7 @@ namespace ECookMode
 	enum Type
 	{
 		CookOnTheFly,				// default mode, handles requests from network
+		CookOnTheFlyFromTheEditor,	// cook on the side
 		CookByTheBookFromTheEditor,	// precook all resources while in the editor
 		CookByTheBook,				// cooking by the book (not in the editor)
 	};
@@ -381,7 +382,7 @@ private:
 			// return FilesProcessed.Contains(Filename);
 		}
 		// two versions of this function so I don't have to create temporary FFIleplatformRequest in some cases to call the exists function
-		bool Exists( const FName& Filename, const TArray<FName>& PlatformNames )
+		bool Exists( const FName& Filename, const TArray<FName>& PlatformNames ) const
 		{
 			FScopeLock ScopeLock(&SynchronizationObject);
 
@@ -404,7 +405,17 @@ private:
 			return true;
 		}
 
-		bool GetCookedPlatforms( const FName& Filename, TArray<FName>& PlatformList )
+		void RemoveAllCookedFilesForPlatform( const FName& PlatformName )
+		{
+			FScopeLock ScopeLock(&SynchronizationObject);
+
+			for ( auto& Request : FilesProcessed )
+			{	
+				Request.Value.RemovePlatform( PlatformName );
+			}
+		}
+
+		bool GetCookedPlatforms( const FName& Filename, TArray<FName>& PlatformList ) const
 		{
 			FScopeLock ScopeLock( &SynchronizationObject );
 			const FFilePlatformRequest* Request = FilesProcessed.Find(Filename);
@@ -618,6 +629,9 @@ private:
 	FThreadSafeUnsolicitedPackagesList UnsolicitedCookedPackages;
 	FThreadSafeFilenameSet CookedPackages; // set of files which have been cooked when needing to recook a file the entry will need to be removed from here
 
+	// declared mutable as it's used purely as a cache and don't want to have to declare all the functions as non const just because of this cache
+	mutable TMap<FName, TArray<FString>> CachedIniVersionStringsMap;
+
 public:
 
 	void WarmCookedPackages(const FString& AssetRegistryPath, const TArray<FName>& TargetPlatformNames);
@@ -697,6 +711,20 @@ public:
 	 */
 	uint32 TickCookOnTheSide( const float TimeSlice, uint32 &CookedPackagesCount );
 	
+
+	/**
+	 * Clear all the previously cooked data all cook requests from now on will be considered recook requests
+	 */
+	void ClearAllCookedData();
+
+	/**
+	 * Clear all the previously cooked data for the platform passed in 
+	 * 
+	 * @param name of the platform to clear the cooked packages for
+	 */
+	void ClearPlatformCookedData( const FName& PlatformName );
+
+
 	/**
 	 * Force stop whatever pending cook requests are going on and clear all the cooked data
 	 * Note cook on the side / cook on the fly clients may not be able to recover from this if they are waiting on a cook request to complete
@@ -721,7 +749,21 @@ public:
 	 * 
 	 * @return returns true if this cooker is running in realtime mode like in the editor
 	 */
-	bool IsRealtimeMode();
+	bool IsRealtimeMode() const;
+
+	/**
+	 * Helper function returns if we are in any cook by the book mode
+	 *
+	 * @return if the cook mode is a cook by the book mode
+	 */
+	bool IsCookByTheBookMode() const;
+
+	/**
+	 * Helper function returns if we are in any cook on the fly mode
+	 *
+	 * @return if the cook mode is a cook on the fly mode
+	 */
+	bool IsCookOnTheFlyMode() const;
 	
 
 	virtual void BeginDestroy() override;
@@ -766,16 +808,6 @@ private:
 	 */
 	void CookByTheBookFinished();
 
-	/**
-	 * Helper function returns if we are in any cook by the book mode
-	 *
-	 * @return if the cook mode is a cook by the book mode
-	 */
-	inline bool IsCookByTheBookMode() const 
-	{ 
-		return CurrentCookMode == ECookMode::CookByTheBookFromTheEditor || CurrentCookMode == ECookMode::CookByTheBook; 
-	}
-
 
 	//////////////////////////////////////////////////////////////////////////
 	// cook on the fly specific functions
@@ -846,15 +878,7 @@ private:
 	 * @param TargetPlatforms to check if out of date
 	 * @param OutOfDateTargetPlatforms return list of out of date target platforms which should be cleaned
 	 */
-	bool IniSettingsOutOfDate( const TArray<ITargetPlatform*>& TargetPlatforms, TArray<ITargetPlatform*>& OutOfDateTargetPlatforms ) const;
-
-	/**
-	 * SaveIniVersionStrings save out the ini version strings for all platforms so next cook can use them
-	 * -iterate uses this to figure out if content needs to be recooked
-	 * 
-	 * @param TargetPlatforms the target platforms to generate ini version info for
-	 */
-	void SaveIniVersionStrings( const TArray<ITargetPlatform*>& TargetPlatforms ) const;
+	bool IniSettingsOutOfDate( const TArray<ITargetPlatform*>& TargetPlatforms, TArray<ITargetPlatform*>& OutOfDateTargetPlatforms, bool bShouldSaveCookedVersions = true ) const;
 
 	/**
 	 * IsCookFlagSet
diff --git a/Engine/Source/Editor/UnrealEd/Private/CookOnTheFlyServer.cpp b/Engine/Source/Editor/UnrealEd/Private/CookOnTheFlyServer.cpp
--- a/Engine/Source/Editor/UnrealEd/Private/CookOnTheFlyServer.cpp
+++ b/Engine/Source/Editor/UnrealEd/Private/CookOnTheFlyServer.cpp
@@ -38,7 +38,7 @@
 DEFINE_LOG_CATEGORY_STATIC(LogCookOnTheFly, Log, All);
 
 #define DEBUG_COOKONTHEFLY 0
-#define OUTPUT_TIMING 0
+#define OUTPUT_TIMING 1
 
 #if OUTPUT_TIMING
 
@@ -207,6 +207,12 @@ void OutputTimers()
 		UE_LOG( LogCookOnTheFly, Display, TEXT("%s\t%.2f"), *TimerInfo.Name, TimerInfo.Length * 1000.0f );
 	}
 
+	// first item is the total
+	if ( GTimerInfo.Num() > 0 && ( ( GTimerInfo[0].Length * 1000.0f ) > 40.0f ) )
+	{
+		UE_LOG( LogCookOnTheFly, Display, TEXT("Cook tick exceeded 40ms by  %f"), GTimerInfo[0].Length * 1000.0f  );
+	}
+
 	GTimerInfo.Empty();
 }
 
@@ -436,7 +442,7 @@ TStatId UCookOnTheFlyServer::GetStatId() const
 
 bool UCookOnTheFlyServer::StartNetworkFileServer( const bool BindAnyPort )
 {
-	check( CurrentCookMode == ECookMode::CookOnTheFly );
+	check( IsCookOnTheFlyMode() );
 	//GetDerivedDataCacheRef().WaitForQuiescence(false);
 
 	// start the listening thread
@@ -782,9 +788,19 @@ void UCookOnTheFlyServer::GenerateManifestInfo( UPackage* Package, const TArray<
 	}
 }
 
-bool UCookOnTheFlyServer::IsRealtimeMode()
+bool UCookOnTheFlyServer::IsRealtimeMode() const 
 {
-	return CurrentCookMode == ECookMode::CookByTheBookFromTheEditor || CurrentCookMode == ECookMode::CookByTheBook;
+	return CurrentCookMode == ECookMode::CookByTheBookFromTheEditor || CurrentCookMode == ECookMode::CookOnTheFlyFromTheEditor;
+}
+
+bool UCookOnTheFlyServer::IsCookByTheBookMode() const
+{
+	return CurrentCookMode == ECookMode::CookByTheBookFromTheEditor || CurrentCookMode == ECookMode::CookByTheBook; 
+}
+
+bool UCookOnTheFlyServer::IsCookOnTheFlyMode() const
+{
+	return CurrentCookMode == ECookMode::CookOnTheFly || CurrentCookMode == ECookMode::CookOnTheFlyFromTheEditor; 
 }
 
 uint32 UCookOnTheFlyServer::TickCookOnTheSide( const float TimeSlice, uint32 &CookedPackageCount )
@@ -1054,22 +1070,25 @@ uint32 UCookOnTheFlyServer::TickCookOnTheSide( const float TimeSlice, uint32 &Co
 					Obj->BeginCacheForCookedPlatformData( TargetPlatform );
 					if ( Obj->IsCachedCookedPlatformDataLoaded(TargetPlatform) == false )
 					{
+#if DEBUG_COOKONTHEFLY || 1
 						UE_LOG(LogCookOnTheFly, Display, TEXT("Object %s isn't cached yet"), *Obj->GetFullName());
+#endif
 						bIsAllDataCached = false;
-						break;
 					}
 
 					if ( Timer.IsTimeUp() && IsRealtimeMode() )
 					{
+#if DEBUG_COOKONTHEFLY || 1
 						UE_LOG(LogCookOnTheFly, Display, TEXT("Object %s took too long to cache"), *Obj->GetFullName());
+#endif
 						bIsAllDataCached = false;
 						break;
 					}
 				}
 
-				if ( bIsAllDataCached == false )
+				if ( Timer.IsTimeUp() && IsRealtimeMode() && (bIsAllDataCached == false) )
 				{
-					break;
+					break; // break out of the target platform loop (can't break out of the main package tick loop here without requeing the request)
 				}
 			}
 		}
@@ -1080,7 +1099,7 @@ uint32 UCookOnTheFlyServer::TickCookOnTheSide( const float TimeSlice, uint32 &Co
 			{
 				// reque the current task and process it next tick
 				CookRequests.EnqueueUnique(ToBuild, true);
-				break;
+				break; // break out of the package tick loop
 			}
 		}
 
@@ -1747,9 +1766,6 @@ void UCookOnTheFlyServer::Initialize( ECookMode::Type DesiredCookMode, ECookInit
 
 	CleanSandbox(Platforms);
 
-	// generate version info for use by -iterate to figure out what files need recooking
-	SaveIniVersionStrings(Platforms);
-
 	// always generate the asset registry before starting to cook, for either method
 	GenerateAssetRegistry(Platforms);
 
@@ -1889,44 +1905,6 @@ bool UCookOnTheFlyServer::GetCurrentIniVersionStrings( const ITargetPlatform* Ta
 	GetVersionFormatNumbersForIniVersionStrings( IniVersionStrings, TEXT("AudioFormat"), TPM->GetAudioFormats() );
 	GetVersionFormatNumbersForIniVersionStrings( IniVersionStrings, TEXT("TextureFormat"), TPM->GetTextureFormats() );
 	GetVersionFormatNumbersForIniVersionStrings( IniVersionStrings, TEXT("ShaderFormat"), TPM->GetShaderFormats() );
-	/*
-	for ( const auto& AudioFormat : TPM->GetAudioFormats() )
-	{
-		TArray<FName> SupportedFormats;
-		AudioFormat.GetSupportedFormats(SupportedFormats);
-		for ( const auto& SupportedFormat : SupportedFormats )
-		{
-			int32 VersionNumber = AudioFormat.GetVersion(SupportedFormat);
-			FString AudioString = FString::Printf( TEXT("AudioFormat:%s:VersionNumber%d"), SupportedFormat.ToString(), VersionNumber);
-			IniVersionStrings.Emplace( AudioString );
-		}
-	}
-
-	for ( const auto& TextureFormat : TPM->GetTextureFormats() )
-	{
-		TArray<FName> SupportedFormats;
-		TextureFormat.GetSupportedFormats(SupportedFormats);
-		for ( const auto& SupportedFormat : SupportedFormats )
-		{
-			int32 VersionNumber = AudioFormat.GetVersion(SupportedFormat);
-			FString AudioString = FString::Printf( TEXT("AudioFormat:%s:VersionNumber%d"), SupportedFormat.ToString(), VersionNumber);
-			IniVersionStrings.Emplace( AudioString );
-		}
-	}
-
-	for ( const auto& AudioFormat : TPM->GetAudioFormats() )
-	{
-		TArray<FName> SupportedFormats;
-		AudioFormat.GetSupportedFormats(SupportedFormats);
-		for ( const auto& SupportedFormat : SupportedFormats )
-		{
-			int32 VersionNumber = AudioFormat.GetVersion(SupportedFormat);
-			FString AudioString = FString::Printf( TEXT("AudioFormat:%s:VersionNumber%d"), SupportedFormat.ToString(), VersionNumber);
-			IniVersionStrings.Emplace( AudioString );
-		}
-	}
-	*/
-
 
 	return true;
 }
@@ -1945,9 +1923,10 @@ bool UCookOnTheFlyServer::GetCookedIniVersionStrings( const ITargetPlatform* Tar
 }
 
 
-bool UCookOnTheFlyServer::IniSettingsOutOfDate( const TArray<ITargetPlatform*>& TargetPlatforms, TArray<ITargetPlatform*>& OutOfDateTargetPlatforms ) const
+bool UCookOnTheFlyServer::IniSettingsOutOfDate( const TArray<ITargetPlatform*>& TargetPlatforms, TArray<ITargetPlatform*>& OutOfDateTargetPlatforms, bool bShouldSaveCookedVersions ) const
 {
-
+	const FString CookedIni = FPaths::GameDir() / TEXT("CookedIniVersion.txt");
+	const FString SandboxCookedIni = SandboxFile->ConvertToAbsolutePathForExternalAppForWrite(*CookedIni);
 
 	for ( const auto& TargetPlatform : TargetPlatforms )
 	{
@@ -1959,45 +1938,45 @@ bool UCookOnTheFlyServer::IniSettingsOutOfDate( const TArray<ITargetPlatform*>&
 			continue;
 		}
 
-		TArray<FString> CookedIniVersionStrings;
-		if ( GetCookedIniVersionStrings( TargetPlatform, CookedIniVersionStrings ) == false )
+		// check if the cached ones are filled out
+		const FName TargetPlatformName = FName(*TargetPlatform->PlatformName());
+		TArray<FString>* FoundCookedIniVersionStrings = CachedIniVersionStringsMap.Find( TargetPlatformName );
+
+		if ( FoundCookedIniVersionStrings == NULL )
 		{
-			// can't even get the cooked version strings for this platform add it to the out of date list
-			OutOfDateTargetPlatforms.Add(TargetPlatform);
-			continue;
+			TArray<FString> CookedIniVersionStrings;
+			GetCookedIniVersionStrings( TargetPlatform, CookedIniVersionStrings );
+			FoundCookedIniVersionStrings = &CachedIniVersionStringsMap.Emplace( TargetPlatformName, MoveTemp(CookedIniVersionStrings) );
 		}
 
+		check( FoundCookedIniVersionStrings );
+		bool bCurrentIniSettingsChanged = false;
 		for ( const auto& CurrentVersionString : CurrentIniVersionStrings ) 
 		{
-			if ( CookedIniVersionStrings.Contains(CurrentVersionString) == false )
+			if ( FoundCookedIniVersionStrings->Contains(CurrentVersionString) == false )
 			{
-				OutOfDateTargetPlatforms.Add(TargetPlatform);
+				
+				bCurrentIniSettingsChanged = true;
 				break;
 			}
 		}
-	}
-
-	return true;
-}
-
-void UCookOnTheFlyServer::SaveIniVersionStrings( const TArray<ITargetPlatform*>& TargetPlatforms ) const
-{
-	const FString CookedIni = FPaths::GameDir() / TEXT("CookedIniVersion.txt");
-	const FString SandboxCookedIni = SandboxFile->ConvertToAbsolutePathForExternalAppForWrite(*CookedIni);
 
+		if ( bCurrentIniSettingsChanged )
+		{
+			OutOfDateTargetPlatforms.Add(TargetPlatform);
+			*FoundCookedIniVersionStrings = CurrentIniVersionStrings;
 
-	for ( const auto& TargetPlatform : TargetPlatforms )
-	{
-		// there is a list of important ini settings in the Editor config 
-		TArray<FString> IniVersionStrings;
-		GetCurrentIniVersionStrings(TargetPlatform, IniVersionStrings);
-
-		const FString PlatformSandboxCookedIni = SandboxCookedIni.Replace(TEXT("[Platform]"), *TargetPlatform->PlatformName());
-
-		GConfig->SetArray(TEXT("CookSettings"), TEXT("VersionedIniParams"), IniVersionStrings, PlatformSandboxCookedIni);
-
-		GConfig->Flush(false, PlatformSandboxCookedIni);
+			// save the iniversion strings
+			if ( bShouldSaveCookedVersions )
+			{
+				const FString PlatformSandboxCookedIni = SandboxCookedIni.Replace(TEXT("[Platform]"), *TargetPlatform->PlatformName());
+				GConfig->SetArray(TEXT("CookSettings"), TEXT("VersionedIniParams"), CurrentIniVersionStrings, PlatformSandboxCookedIni);
+				GConfig->Flush(false, PlatformSandboxCookedIni);
+			}
+		}
 	}
+
+	return true;
 }
 
 void UCookOnTheFlyServer::CleanSandbox(const TArray<ITargetPlatform*>& Platforms)
@@ -2464,6 +2443,27 @@ void UCookOnTheFlyServer::StopAndClearCookedData()
 	CookedPackages.Empty(); // set of files which have been cooked when needing to recook a file the entry will need to be removed from here
 }
 
+void UCookOnTheFlyServer::ClearAllCookedData()
+{
+	// if we are going to clear the cooked packages it is conceivable that we will recook the packages which we just cooked 
+	// that means it's also conceivable that we will recook the same package which currently has an outstanding async write request
+	UPackage::WaitForAsyncFileWrites();
+
+	UnsolicitedCookedPackages.Empty();
+	CookedPackages.Empty(); // set of files which have been cooked when needing to recook a file the entry will need to be removed from here
+}
+
+void UCookOnTheFlyServer::ClearPlatformCookedData( const FName& PlatformName )
+{
+	// if we are going to clear the cooked packages it is conceivable that we will recook the packages which we just cooked 
+	// that means it's also conceivable that we will recook the same package which currently has an outstanding async write request
+	UPackage::WaitForAsyncFileWrites();
+
+	CookedPackages.RemoveAllCookedFilesForPlatform( PlatformName );
+	TArray<FName> PackageNames;
+	UnsolicitedCookedPackages.GetPackagesForPlatformAndRemove(PlatformName, PackageNames);
+}
+
 void UCookOnTheFlyServer::StartCookByTheBook(const TArray<ITargetPlatform*>& TargetPlatforms, 
 											 const TArray<FString> &CookMaps, const TArray<FString> &CookDirectories, 
 											 const TArray<FString> &CookCultures, const TArray<FString> &IniMapSections, 
@@ -2502,6 +2502,13 @@ void UCookOnTheFlyServer::StartCookByTheBook(const TArray<ITargetPlatform*>& Tar
 
 	*/
 
+	TArray<ITargetPlatform*> ModifiedTargetPlatforms;
+	IniSettingsOutOfDate(TargetPlatforms, ModifiedTargetPlatforms);
+	for ( const auto& ModifiedTargetPlatform : ModifiedTargetPlatforms )
+	{
+		ClearPlatformCookedData( FName( *ModifiedTargetPlatform->PlatformName() ) );
+	}
+
 	CookByTheBookOptions->bRunning = true;
 	CookByTheBookOptions->bCancel = false;
 	CookByTheBookOptions->CookTime = 0.0f;
@@ -2514,7 +2521,6 @@ void UCookOnTheFlyServer::StartCookByTheBook(const TArray<ITargetPlatform*>& Tar
 	CookByTheBookOptions->bLeakTest = (CookOptions & ECookByTheBookOptions::LeakTest) != ECookByTheBookOptions::None; // this won't work from the editor this needs to be standalone
 	check( !CookByTheBookOptions->bLeakTest || CurrentCookMode == ECookMode::CookByTheBook );
 
-
 	CookByTheBookOptions->LastGCItems.Empty();
 	if (CookByTheBookOptions->bLeakTest )
 	{
@@ -2524,8 +2530,6 @@ void UCookOnTheFlyServer::StartCookByTheBook(const TArray<ITargetPlatform*>& Tar
 		}
 	}
 
-
-	
 	// allow the game to fill out the asset registry, as well as get a list of objects to always cook
 	TArray<FString> FilesInPath;
 	FGameDelegates::Get().GetCookModificationDelegate().ExecuteIfBound(FilesInPath);
diff --git a/Engine/Source/Editor/UnrealEd/Private/UnrealEdEngine.cpp b/Engine/Source/Editor/UnrealEd/Private/UnrealEdEngine.cpp
--- a/Engine/Source/Editor/UnrealEd/Private/UnrealEdEngine.cpp
+++ b/Engine/Source/Editor/UnrealEd/Private/UnrealEdEngine.cpp
@@ -128,7 +128,7 @@ void UUnrealEdEngine::Init(IEngineLoop* InEngineLoop)
 	else if ( ExperimentalSettings->bCookOnTheSide )
 	{
 		CookServer = ConstructObject<UCookOnTheFlyServer>( UCookOnTheFlyServer::StaticClass() );
-		CookServer->Initialize( ECookMode::CookOnTheFly, BaseCookingFlags );
+		CookServer->Initialize( ECookMode::CookOnTheFlyFromTheEditor, BaseCookingFlags );
 		CookServer->StartNetworkFileServer( false );
 
 		FCoreUObjectDelegates::OnObjectPropertyChanged.AddUObject(CookServer, &UCookOnTheFlyServer::OnObjectPropertyChanged);
