Graph Editor  - a cleaner way to deal with Drag&Drop action for EventDispatcher.

diff --git a/Engine/Source/Editor/AnimGraph/Classes/AnimationStateMachineSchema.h b/Engine/Source/Editor/AnimGraph/Classes/AnimationStateMachineSchema.h
--- a/Engine/Source/Editor/AnimGraph/Classes/AnimationStateMachineSchema.h
+++ b/Engine/Source/Editor/AnimGraph/Classes/AnimationStateMachineSchema.h
@@ -78,6 +78,5 @@ class UAnimationStateMachineSchema : public UEdGraphSchema
 	virtual void DroppedAssetsOnPin(const TArray<FAssetData>& Assets, const FVector2D& GraphPosition, UEdGraphPin* Pin) const override;
 	virtual void GetAssetsNodeHoverMessage(const TArray<FAssetData>& Assets, const UEdGraphNode* HoverNode, FString& OutTooltipText, bool& OutOkIcon) const override;
 	virtual void GetAssetsPinHoverMessage(const TArray<FAssetData>& Assets, const UEdGraphPin* HoverPin, FString& OutTooltipText, bool& OutOkIcon) const override;
-	virtual bool DoesSupportEventDispatcher() const	override {	return false; }
 	// End UEdGraphSchema interface
 };
diff --git a/Engine/Source/Editor/BlueprintGraph/Classes/EdGraphSchema_K2.h b/Engine/Source/Editor/BlueprintGraph/Classes/EdGraphSchema_K2.h
--- a/Engine/Source/Editor/BlueprintGraph/Classes/EdGraphSchema_K2.h
+++ b/Engine/Source/Editor/BlueprintGraph/Classes/EdGraphSchema_K2.h
@@ -414,9 +414,16 @@ public:
 	virtual void BackwardCompatibilityNodeConversion(UEdGraph* Graph, bool bOnlySafeChanges) const override;
 	virtual void SplitPin(UEdGraphPin* Pin) const override;
 	virtual void RecombinePin(UEdGraphPin* Pin) const override;
-	virtual bool DoesSupportEventDispatcher() const	override { return true; }
 	// End EdGraphSchema Interface
 
+	/**
+	 *
+	 *	Determine if this graph supports event dispatcher
+	 * 
+	 *	@return True if this schema supports Event Dispatcher 
+	 */
+	virtual bool DoesSupportEventDispatcher() const { return true; }
+
 	/**
 	* Configure the supplied variable node based on the supplied info
 	*
diff --git a/Engine/Source/Editor/GraphEditor/Private/SGraphPanel.cpp b/Engine/Source/Editor/GraphEditor/Private/SGraphPanel.cpp
--- a/Engine/Source/Editor/GraphEditor/Private/SGraphPanel.cpp
+++ b/Engine/Source/Editor/GraphEditor/Private/SGraphPanel.cpp
@@ -808,10 +808,9 @@ FReply SGraphPanel::OnDrop( const FGeometry& MyGeometry, const FDragDropEvent& D
 	if (Operation->IsOfType<FGraphEditorDragDropAction>())
 	{
 		check(GraphObj);
-		auto Schema = GraphObj->GetSchema();
-		if (Schema && Schema->DoesSupportEventDispatcher())
+		TSharedPtr<FGraphEditorDragDropAction> DragConn = StaticCastSharedPtr<FGraphEditorDragDropAction>(Operation);
+		if (DragConn.IsValid() && DragConn->IsSupportedBySchema(GraphObj->GetSchema()))
 		{
-			TSharedPtr<FGraphEditorDragDropAction> DragConn = StaticCastSharedPtr<FGraphEditorDragDropAction>(Operation);
 			return DragConn->DroppedOnPanel(SharedThis(this), DragDropEvent.GetScreenSpacePosition(), NodeAddPosition, *GraphObj);
 		}
 
diff --git a/Engine/Source/Editor/GraphEditor/Public/GraphEditorDragDropAction.h b/Engine/Source/Editor/GraphEditor/Public/GraphEditorDragDropAction.h
--- a/Engine/Source/Editor/GraphEditor/Public/GraphEditorDragDropAction.h
+++ b/Engine/Source/Editor/GraphEditor/Public/GraphEditorDragDropAction.h
@@ -27,6 +27,8 @@ public:
 	virtual void OnDragBegin(const TSharedRef<class SGraphPin>& InPin) {}
 	// End of interface to override
 	
+	virtual bool IsSupportedBySchema(const class UEdGraphSchema* Schema) const { return true; }
+
 protected:
 	void SetFeedbackMessage(const TSharedPtr<SWidget>& Message);
 	void SetSimpleFeedbackMessage(const FSlateBrush* Icon, const FSlateColor& IconColor, const FText& Message);
diff --git a/Engine/Source/Editor/Kismet/Private/BPDelegateDragDropAction.cpp b/Engine/Source/Editor/Kismet/Private/BPDelegateDragDropAction.cpp
--- a/Engine/Source/Editor/Kismet/Private/BPDelegateDragDropAction.cpp
+++ b/Engine/Source/Editor/Kismet/Private/BPDelegateDragDropAction.cpp
@@ -142,4 +142,10 @@ bool FKismetDelegateDragDropAction::IsValid() const
 		(NULL != FindField<UMulticastDelegateProperty>(VariableSource.Get(), VariableName));
 }
 
+bool FKismetDelegateDragDropAction::IsSupportedBySchema(const class UEdGraphSchema* Schema) const
+{
+	auto SchemaK2 = Cast<UEdGraphSchema_K2>(Schema);
+	return SchemaK2 && SchemaK2->DoesSupportEventDispatcher();
+}
+
 #undef LOCTEXT_NAMESPACE
diff --git a/Engine/Source/Editor/Kismet/Private/BPDelegateDragDropAction.h b/Engine/Source/Editor/Kismet/Private/BPDelegateDragDropAction.h
--- a/Engine/Source/Editor/Kismet/Private/BPDelegateDragDropAction.h
+++ b/Engine/Source/Editor/Kismet/Private/BPDelegateDragDropAction.h
@@ -17,6 +17,8 @@ public:
 	virtual FReply DroppedOnNode(FVector2D ScreenPosition, FVector2D GraphPosition) override { return FGraphEditorDragDropAction::DroppedOnNode(ScreenPosition, GraphPosition); }
 
 	virtual FReply DroppedOnPanel(const TSharedRef< SWidget >& Panel, FVector2D ScreenPosition, FVector2D GraphPosition, UEdGraph& Graph) override;
+
+	virtual bool IsSupportedBySchema(const class UEdGraphSchema* Schema) const override;
 	// End of FGraphEditorDragDropAction
 
 	bool IsValid() const;
diff --git a/Engine/Source/Editor/UnrealEd/Classes/MaterialGraph/MaterialGraphSchema.h b/Engine/Source/Editor/UnrealEd/Classes/MaterialGraph/MaterialGraphSchema.h
--- a/Engine/Source/Editor/UnrealEd/Classes/MaterialGraph/MaterialGraphSchema.h
+++ b/Engine/Source/Editor/UnrealEd/Classes/MaterialGraph/MaterialGraphSchema.h
@@ -195,7 +195,6 @@ class UMaterialGraphSchema : public UEdGraphSchema
 	virtual void GetContextMenuActions(const UEdGraph* CurrentGraph, const UEdGraphNode* InGraphNode, const UEdGraphPin* InGraphPin, class FMenuBuilder* MenuBuilder, bool bIsDebugging) const override;
 	virtual const FPinConnectionResponse CanCreateConnection(const UEdGraphPin* A, const UEdGraphPin* B) const override;
 	virtual bool TryCreateConnection(UEdGraphPin* A, UEdGraphPin* B) const override;
-	virtual bool DoesSupportEventDispatcher() const override { return true; }
 	virtual bool ShouldHidePinDefaultValue(UEdGraphPin* Pin) const override { return true; }
 	virtual FLinearColor GetPinTypeColor(const FEdGraphPinType& PinType) const override;
 	virtual void BreakNodeLinks(UEdGraphNode& TargetNode) const override;
diff --git a/Engine/Source/Editor/UnrealEd/Classes/SoundClassGraph/SoundClassGraphSchema.h b/Engine/Source/Editor/UnrealEd/Classes/SoundClassGraph/SoundClassGraphSchema.h
--- a/Engine/Source/Editor/UnrealEd/Classes/SoundClassGraph/SoundClassGraphSchema.h
+++ b/Engine/Source/Editor/UnrealEd/Classes/SoundClassGraph/SoundClassGraphSchema.h
@@ -46,7 +46,6 @@ class USoundClassGraphSchema : public UEdGraphSchema
 	virtual void GetContextMenuActions(const UEdGraph* CurrentGraph, const UEdGraphNode* InGraphNode, const UEdGraphPin* InGraphPin, class FMenuBuilder* MenuBuilder, bool bIsDebugging) const override;
 	virtual const FPinConnectionResponse CanCreateConnection(const UEdGraphPin* PinA, const UEdGraphPin* PinB) const override;
 	virtual bool TryCreateConnection(UEdGraphPin* PinA, UEdGraphPin* PinB) const override;
-	virtual bool DoesSupportEventDispatcher() const override { return true; }
 	virtual bool ShouldHidePinDefaultValue(UEdGraphPin* Pin) const override;
 	virtual FLinearColor GetPinTypeColor(const FEdGraphPinType& PinType) const override;
 	virtual void BreakNodeLinks(UEdGraphNode& TargetNode) const override;
diff --git a/Engine/Source/Editor/UnrealEd/Classes/SoundCueGraph/SoundCueGraphSchema.h b/Engine/Source/Editor/UnrealEd/Classes/SoundCueGraph/SoundCueGraphSchema.h
--- a/Engine/Source/Editor/UnrealEd/Classes/SoundCueGraph/SoundCueGraphSchema.h
+++ b/Engine/Source/Editor/UnrealEd/Classes/SoundCueGraph/SoundCueGraphSchema.h
@@ -114,7 +114,6 @@ class USoundCueGraphSchema : public UEdGraphSchema
 	virtual void CreateDefaultNodesForGraph(UEdGraph& Graph) const override;
 	virtual const FPinConnectionResponse CanCreateConnection(const UEdGraphPin* A, const UEdGraphPin* B) const override;
 	virtual bool TryCreateConnection(UEdGraphPin* A, UEdGraphPin* B) const override;
-	virtual bool DoesSupportEventDispatcher() const override { return true; }
 	virtual bool ShouldHidePinDefaultValue(UEdGraphPin* Pin) const override;
 	virtual FLinearColor GetPinTypeColor(const FEdGraphPinType& PinType) const override;
 	virtual void BreakNodeLinks(UEdGraphNode& TargetNode) const override;
diff --git a/Engine/Source/Runtime/Engine/Classes/EdGraph/EdGraphSchema.h b/Engine/Source/Runtime/Engine/Classes/EdGraph/EdGraphSchema.h
--- a/Engine/Source/Runtime/Engine/Classes/EdGraph/EdGraphSchema.h
+++ b/Engine/Source/Runtime/Engine/Classes/EdGraph/EdGraphSchema.h
@@ -516,13 +516,6 @@ class ENGINE_API UEdGraphSchema : public UObject
 	 */
 	virtual bool DoesSupportPinWatching() const	{ return false; }
 
-	/**
-	 *
-	 *	Determine if this graph supports event dispatcher
-	 * 
-	 *	@return True if this schema supports Event Dispatcher 
-	 */
-	virtual bool DoesSupportEventDispatcher() const	{ return false; }
 	/**
 	 * Checks to see if the specified pin is being watched by the graph's debug system.
 	 * 
