EditorAutomationTests.cpp
--Add a todo for verifying and closing a session for the launch on test.

diff --git a/Engine/Source/Editor/UnrealEd/Private/Tests/EditorAutomationTests.cpp b/Engine/Source/Editor/UnrealEd/Private/Tests/EditorAutomationTests.cpp
--- a/Engine/Source/Editor/UnrealEd/Private/Tests/EditorAutomationTests.cpp
+++ b/Engine/Source/Editor/UnrealEd/Private/Tests/EditorAutomationTests.cpp
@@ -2183,6 +2183,8 @@ bool FLaunchOnTest::RunTest(const FString& Parameters)
 	ADD_LATENT_AUTOMATION_COMMAND(FWaitToFinishCookByTheBookCommand);
 	ADD_LATENT_AUTOMATION_COMMAND(FWaitToFinishBuildDeployCommand);
 
+	//@todo: Verify the game launched.
+
 	//@todo: Close the Launched on Game.
 
 	//Delete the Cooked, StagedBuilds, and Automation_TEMP folder if they exist.
