Editor: Fix non-unity error in UnrealEd (add include to AutomationEditorCommon.cpp)

diff --git a/Engine/Source/Editor/UnrealEd/Private/Tests/AutomationEditorCommon.cpp b/Engine/Source/Editor/UnrealEd/Private/Tests/AutomationEditorCommon.cpp
--- a/Engine/Source/Editor/UnrealEd/Private/Tests/AutomationEditorCommon.cpp
+++ b/Engine/Source/Editor/UnrealEd/Private/Tests/AutomationEditorCommon.cpp
@@ -59,6 +59,7 @@
 
 #include "CookOnTheSide/CookOnTheFlyServer.h"
 #include "Interfaces/ILauncherProfile.h"
+#include "LightingBuildOptions.h"
 
 #define COOK_TIMEOUT 3600
 
