Added support for dlc plugins directory to be searched for content and dlc to output to dlc cooked directory.
Removed some cabbage code.

diff --git a/Engine/Source/Editor/UnrealEd/Classes/CookOnTheSide/CookOnTheFlyServer.h b/Engine/Source/Editor/UnrealEd/Classes/CookOnTheSide/CookOnTheFlyServer.h
--- a/Engine/Source/Editor/UnrealEd/Classes/CookOnTheSide/CookOnTheFlyServer.h
+++ b/Engine/Source/Editor/UnrealEd/Classes/CookOnTheSide/CookOnTheFlyServer.h
@@ -157,62 +157,6 @@ public:
 		 */
 	};
 
-	// pending delete: don't think I need this anymore
-	/*template<typename Type>
-	struct FLookupQueue
-	{
-	private:
-		TSet<Type> Set;
-		FQueue<Type> Queue;
-
-	public:
-		void Enqueue(const Type& Item)
-		{
-			Set.Add(Item);
-			Queue.Enqueue(Item);
-		}
-		void EnqueueUnique( const Type& Item )
-		{
-			if ( Set.Find(Item) == NULL )
-			{
-				Enqueue( Item );
-			}
-		}
-
-		bool Dequeue(Type* Result)
-		{
-			if (Queue.Num())
-			{
-				Queue.Dequeue(Result);
-				Set.Remove(*Result);
-				return true;
-			}
-			return false;
-		}
-		void DequeueAll(TArray<Type>& Results)
-		{
-			Queue.DequeueAll(Results);
-			Set.Empty();
-		}
-
-		bool HasItems() const
-		{
-			return Queue.Num() > 0;
-		}
-
-		void Remove( const Type& Item ) 
-		{
-			Queue.Remove( Item );
-			Set.Remove(Item);
-		}
-
-		int Num() const 
-		{
-			return Queue.Num();
-		}
-
-	};*/
-
 public:
 	/** cooked file requests which includes platform which file is requested for */
 	struct FFilePlatformRequest
@@ -577,6 +521,36 @@ private:
 			CookedPackages.Empty();
 		}
 	};
+
+
+	struct FCachedPackageFilename
+	{
+	public:
+		FCachedPackageFilename(FString &&InPackageFilename, FString &&InStandardFilename, FName InStandardFileFName ) :
+			PackageFilename( MoveTemp( InPackageFilename )),
+			StandardFilename(MoveTemp(InStandardFilename)),
+			StandardFileFName( InStandardFileFName )
+		{
+		}
+
+		FCachedPackageFilename( const FCachedPackageFilename &In )
+		{
+			PackageFilename = In.PackageFilename;
+			StandardFilename = In.StandardFilename;
+			StandardFileFName = In.StandardFileFName;
+		}
+
+		FCachedPackageFilename( FCachedPackageFilename &&In )
+		{
+			PackageFilename = MoveTemp(In.PackageFilename);
+			StandardFilename = MoveTemp(In.StandardFilename);
+			StandardFileFName = In.StandardFileFName;
+		}
+
+		FString PackageFilename; // this is also full path
+		FString StandardFilename;
+		FName StandardFileFName;
+	};
 private:
 	/** Current cook mode the cook on the fly server is running in */
 	ECookMode::Type CurrentCookMode;
@@ -600,6 +574,8 @@ private:
 		bool bRunning;
 		/** Cancel has been queued will be processed next tick */
 		bool bCancel;
+		/** DlcName setup if we are cooking dlc will be used as the directory to save cooked files to */
+		FString DlcName;
 		/** Leak test: last gc items (only valid when running from commandlet requires gc between each cooked package) */
 		TSet<FWeakObjectPtr> LastGCItems;
 		/** Map of platform name to manifest generator, manifest is only used in cook by the book however it needs to be maintained across multiple cook by the books. */
@@ -629,6 +605,18 @@ private:
 	FThreadSafeUnsolicitedPackagesList UnsolicitedCookedPackages;
 	FThreadSafeFilenameSet CookedPackages; // set of files which have been cooked when needing to recook a file the entry will need to be removed from here
 
+	TMap<FName, FCachedPackageFilename> PackageFilenameCache; // filename cache (only process the string operations once)
+
+	FString GetCachedPackageFilename( const FName& PackageName );
+	FString GetCachedStandardPackageFilename( const FName& PackageName );
+	FName GetCachedStandardPackageFileFName( const FName& PackageName );
+	FString GetCachedPackageFilename( UPackage* Package );
+	FString GetCachedStandardPackageFilename( UPackage* Package );
+	FName GetCachedStandardPackageFileFName( UPackage* Package );
+	const FString& GetCachedSandboxFilename( UPackage* Package, TAutoPtr<class FSandboxPlatformFile>& SandboxFile );
+	const FCachedPackageFilename& Cache(const FName& PackageName);
+	void ClearPackageFilenameCache();
+
 	// declared mutable as it's used purely as a cache and don't want to have to declare all the functions as non const just because of this cache
 	mutable TMap<FName, TArray<FString>> CachedIniVersionStringsMap;
 
@@ -688,7 +676,7 @@ public:
 	void StartCookByTheBook(const TArray<ITargetPlatform*>& TargetPlatforms, 
 		const TArray<FString>& CookMaps, const TArray<FString>& CookDirectories, 
 		const TArray<FString>& CookCultures, const TArray<FString>& IniMapSections, 
-		ECookByTheBookOptions CookOptions = ECookByTheBookOptions::None );
+		ECookByTheBookOptions CookOptions = ECookByTheBookOptions::None, const FString &DLCName = FString() );
 
 	/**
 	 * Queue a cook by the book cancel (you might want to do this instead of calling cancel directly so that you don't have to be in the game thread when canceling
@@ -872,6 +860,24 @@ private:
 	 */
 	bool GetCookedIniVersionStrings( const ITargetPlatform* TargetPlatform, TArray<FString>& IniVersionStrings ) const;
 
+
+	/**
+	 * Convert a path to a full sandbox path
+	 * This function should be used instead of calling the 
+	 */
+	FString ConvertToFullSandboxPath( const FString &FileName, bool bForWrite = false ) const;
+
+	inline bool IsCookingDLC() const
+	{
+		// can only cook dlc in cook by the book
+		// we are cooking dlc when the dlc name is setup
+		if ( CookByTheBookOptions )
+		{
+			return  !CookByTheBookOptions->DlcName.IsEmpty();
+		}
+		return false;
+	}
+
 	/**
 	 * Checks if important ini settings have changed since last cook for each target platform 
 	 * 
@@ -899,6 +905,12 @@ private:
 	 */
 	FString GetOutputDirectory( const FString& PlatformName ) const;
 
+
+	/**
+	 * Gets the asset registry path and name
+	 */
+	FString GetAssetRegistryPath() const;
+
 	/**
 	 *	Get the given packages 'cooked' timestamp (i.e. account for dependencies)
 	 *
diff --git a/Engine/Source/Editor/UnrealEd/Private/Commandlets/CookCommandlet.cpp b/Engine/Source/Editor/UnrealEd/Private/Commandlets/CookCommandlet.cpp
--- a/Engine/Source/Editor/UnrealEd/Private/Commandlets/CookCommandlet.cpp
+++ b/Engine/Source/Editor/UnrealEd/Private/Commandlets/CookCommandlet.cpp
@@ -43,6 +43,7 @@ static FString GetPackageFilename( UPackage* Package )
 	if (FPackageName::DoesPackageExist(Package->GetName(), NULL, &Filename))
 	{
 		Filename = FPaths::ConvertRelativePathToFull(Filename);
+		FPaths::RemoveDuplicateSlashes(Filename);
 	}
 	return Filename;
 }
@@ -1034,7 +1035,10 @@ bool UCookCommandlet::NewCook( const TArray<ITargetPlatform*>& Platforms, TArray
 	//////////////////////////////////////////////////////////////////////////
 	// parse commandline options 
 
-	FString AssetRegistry;
+	FString DLCName;
+	FParse::Value( *Params, TEXT("DLCNAME"), DLCName);
+
+	/*FString AssetRegistry;
 	if (FParse::Value(*Params, TEXT("SHIPPEDASSETREGISTRY="), AssetRegistry))
 	{
 		TArray<FName> TargetPlatformNames;
@@ -1044,7 +1048,7 @@ bool UCookCommandlet::NewCook( const TArray<ITargetPlatform*>& Platforms, TArray
 			TargetPlatformNames.Add(PlatformName); // build list of all target platform names
 		}
 		CookOnTheFlyServer->WarmCookedPackages(FPaths::GameContentDir() / AssetRegistry, TargetPlatformNames);
-	}
+	}*/
 
 	TArray<FString> CmdLineIniSections;
 	FString SectionStr;
@@ -1133,7 +1137,7 @@ bool UCookCommandlet::NewCook( const TArray<ITargetPlatform*>& Platforms, TArray
 		MapList.Add( MapName );
 	}
 
-	CookOnTheFlyServer->StartCookByTheBook(Platforms, MapList, CmdLineDirEntries, CmdLineCultEntries, CmdLineIniSections, CookOptions );
+	CookOnTheFlyServer->StartCookByTheBook(Platforms, MapList, CmdLineDirEntries, CmdLineCultEntries, CmdLineIniSections, CookOptions, DLCName );
 
 	// Garbage collection should happen when either
 	//	1. We have cooked a map
diff --git a/Engine/Source/Editor/UnrealEd/Private/CookOnTheFlyServer.cpp b/Engine/Source/Editor/UnrealEd/Private/CookOnTheFlyServer.cpp
--- a/Engine/Source/Editor/UnrealEd/Private/CookOnTheFlyServer.cpp
+++ b/Engine/Source/Editor/UnrealEd/Private/CookOnTheFlyServer.cpp
@@ -35,7 +35,7 @@
 #include "Engine/LevelStreaming.h"
 #include "TextureLODSettings.h"
 
-DEFINE_LOG_CATEGORY_STATIC(LogCookOnTheFly, Log, All);
+DEFINE_LOG_CATEGORY_STATIC(LogCook, Log, All);
 
 #define DEBUG_COOKONTHEFLY 0
 #define OUTPUT_TIMING 1
@@ -200,17 +200,17 @@ void OutputTimers()
 
 	OutputDevice->Flush();
 
-	UE_LOG( LogCookOnTheFly, Display, TEXT("Timing information for cook") );
-	UE_LOG( LogCookOnTheFly, Display, TEXT("Name\tlength(ms)") );
+	UE_LOG( LogCook, Display, TEXT("Timing information for cook") );
+	UE_LOG( LogCook, Display, TEXT("Name\tlength(ms)") );
 	for ( auto TimerInfo : GTimerInfo )
 	{
-		UE_LOG( LogCookOnTheFly, Display, TEXT("%s\t%.2f"), *TimerInfo.Name, TimerInfo.Length * 1000.0f );
+		UE_LOG( LogCook, Display, TEXT("%s\t%.2f"), *TimerInfo.Name, TimerInfo.Length * 1000.0f );
 	}
 
 	// first item is the total
 	if ( GTimerInfo.Num() > 0 && ( ( GTimerInfo[0].Length * 1000.0f ) > 40.0f ) )
 	{
-		UE_LOG( LogCookOnTheFly, Display, TEXT("Cook tick exceeded 40ms by  %f"), GTimerInfo[0].Length * 1000.0f  );
+		UE_LOG( LogCook, Display, TEXT("Cook tick exceeded 40ms by  %f"), GTimerInfo[0].Length * 1000.0f  );
 	}
 
 	GTimerInfo.Empty();
@@ -249,112 +249,77 @@ void OutputTimers()
 ///////////////////////////////////////////////////////////////
 
 
-/* Static functions
- *****************************************************************************/
 
-struct FCachedPackageFilename
-{
-public:
-	FCachedPackageFilename(FString &&InPackageFilename, FString &&InStandardFilename, FName InStandardFileFName ) :
-		PackageFilename( MoveTemp( InPackageFilename )),
-		StandardFilename(MoveTemp(InStandardFilename)),
-		StandardFileFName( InStandardFileFName )
-	{
-	}
-
-	FCachedPackageFilename( const FCachedPackageFilename &In )
-	{
-		PackageFilename = In.PackageFilename;
-		StandardFilename = In.StandardFilename;
-		StandardFileFName = In.StandardFileFName;
-	}
+/* helper structs functions
+ *****************************************************************************/
 
-	FCachedPackageFilename( FCachedPackageFilename &&In )
+/** Helper to pass a recompile request to game thread */
+struct FRecompileRequest
 {
-		PackageFilename = MoveTemp(In.PackageFilename);
-		StandardFilename = MoveTemp(In.StandardFilename);
-		StandardFileFName = In.StandardFileFName;
-	}
-
-	FString PackageFilename; // this is also full path
-	FString StandardFilename;
-	FName StandardFileFName;
+	struct FShaderRecompileData RecompileData;
+	bool bComplete;
 };
 
-static TMap<FName, FCachedPackageFilename> PackageFilenameCache;
-
-static const FCachedPackageFilename &Cache(const FName& PackageName)
-{
-	FCachedPackageFilename *Cached = PackageFilenameCache.Find( PackageName );
-	if ( Cached != NULL )
-	{
-		return *Cached;
-	}
-
-	FString Filename;
-	FString PackageFilename;
-	FString StandardFilename;
-	FName StandardFileFName = NAME_None;
-	if (FPackageName::DoesPackageExist(PackageName.ToString(), NULL, &Filename))
+/**
+ * Uses the FMessageLog to log a message
+ * 
+ * @param Message to log
+ * @param Severity of the message
+ */
+void LogCookerMessage( const FString& MessageText, EMessageSeverity::Type Severity)
 {
-		StandardFilename = PackageFilename = FPaths::ConvertRelativePathToFull(Filename);
+	FMessageLog MessageLog("CookResults");
 
+	TSharedRef<FTokenizedMessage> Message = FTokenizedMessage::Create(Severity);
 
-		FPaths::MakeStandardFilename(StandardFilename);
-		StandardFileFName = FName(*StandardFilename);
-	}
+	Message->AddToken( FTextToken::Create( FText::FromString(MessageText) ) );
+	// Message->AddToken(FTextToken::Create(MessageLogTextDetail)); 
+	// Message->AddToken(FDocumentationToken::Create(TEXT("https://docs.unrealengine.com/latest/INT/Platforms/iOS/QuickStart/6/index.html"))); 
+	MessageLog.AddMessage(Message);
 
-	return PackageFilenameCache.Emplace( PackageName, MoveTemp(FCachedPackageFilename(MoveTemp(PackageFilename),MoveTemp(StandardFilename), StandardFileFName)) );
+	MessageLog.Notify();
 }
 
 
-/*static FString GetPackageFilename( UPackage* Package )
-{
-	FString Filename;
-	if (FPackageName::DoesPackageExist(Package->GetName(), NULL, &Filename))
-	{
-		Filename = FPaths::ConvertRelativePathToFull(Filename);
-	}
 
-	return Filename;
-}*/
+/* FIlename caching functions
+ *****************************************************************************/
 
-static FString GetCachedPackageFilename( const FName& PackageName )
+FString UCookOnTheFlyServer::GetCachedPackageFilename( const FName& PackageName )
 {
 	return Cache( PackageName ).PackageFilename;
 }
 
-static FString GetCachedStandardPackageFilename( const FName& PackageName )
+FString UCookOnTheFlyServer::GetCachedStandardPackageFilename( const FName& PackageName )
 {
 	return Cache( PackageName ).StandardFilename;
 }
 
-static FName GetCachedStandardPackageFileFName( const FName& PackageName )
+FName UCookOnTheFlyServer::GetCachedStandardPackageFileFName( const FName& PackageName )
 {
 	return Cache( PackageName ).StandardFileFName;
 }
 
 
-static FString GetCachedPackageFilename( UPackage* Package )
+FString UCookOnTheFlyServer::GetCachedPackageFilename( UPackage* Package )
 {
 	check( Package->GetName() == Package->GetFName().ToString() );
 	return Cache( Package->GetFName() ).PackageFilename;
 }
 
-static FString GetCachedStandardPackageFilename( UPackage* Package )
+FString UCookOnTheFlyServer::GetCachedStandardPackageFilename( UPackage* Package )
 {
 	check( Package->GetName() == Package->GetFName().ToString() );
 	return Cache( Package->GetFName() ).StandardFilename;
 }
 
-static FName GetCachedStandardPackageFileFName( UPackage* Package )
+FName UCookOnTheFlyServer::GetCachedStandardPackageFileFName( UPackage* Package )
 {
 	check( Package->GetName() == Package->GetFName().ToString() );
 	return Cache( Package->GetFName() ).StandardFileFName;
 }
 
-
-static const FString& GetCachedSandboxFilename( UPackage* Package, TAutoPtr<class FSandboxPlatformFile>& SandboxFile )
+const FString& UCookOnTheFlyServer::GetCachedSandboxFilename( UPackage* Package, TAutoPtr<class FSandboxPlatformFile>& SandboxFile )
 {
 	FName PackageFName = Package->GetFName();
 	static TMap<FName, FString> CachedSandboxFilenames;
@@ -363,45 +328,43 @@ static const FString& GetCachedSandboxFilename( UPackage* Package, TAutoPtr<clas
 		return *CachedSandboxFilename;
 
 	const FString& PackageFilename = GetCachedPackageFilename(Package);
-	FString SandboxFilename = SandboxFile->ConvertToAbsolutePathForExternalAppForWrite(*PackageFilename );
+	FString SandboxFilename = ConvertToFullSandboxPath(*PackageFilename, true );
 
 	return CachedSandboxFilenames.Add( PackageFName, MoveTemp(SandboxFilename) );
 }
 
-/* helper structs
- *****************************************************************************/
-
-/** Helper to pass a recompile request to game thread */
-struct FRecompileRequest
+const UCookOnTheFlyServer::FCachedPackageFilename& UCookOnTheFlyServer::Cache(const FName& PackageName)
 {
-	struct FShaderRecompileData RecompileData;
-	bool bComplete;
-};
+	FCachedPackageFilename *Cached = PackageFilenameCache.Find( PackageName );
+	if ( Cached != NULL )
+	{
+		return *Cached;
+	}
+	// cache all the things, like it's your birthday!
 
+	FString Filename;
+	FString PackageFilename;
+	FString StandardFilename;
+	FName StandardFileFName = NAME_None;
+	if (FPackageName::DoesPackageExist(PackageName.ToString(), NULL, &Filename))
+{
+		StandardFilename = PackageFilename = FPaths::ConvertRelativePathToFull(Filename);
 
 
-/**
- * Uses the FMessageLog to log a message
- * 
- * @param Message to log
- * @param Severity of the message
- */
-void LogCookerMessage( const FString& MessageText, EMessageSeverity::Type Severity)
-	{
-	FMessageLog MessageLog("CookResults");
-
-	TSharedRef<FTokenizedMessage> Message = FTokenizedMessage::Create(Severity);
+		FPaths::MakeStandardFilename(StandardFilename);
+		StandardFileFName = FName(*StandardFilename);
+	}
 
-	Message->AddToken( FTextToken::Create( FText::FromString(MessageText) ) );
-	// Message->AddToken(FTextToken::Create(MessageLogTextDetail)); 
-	// Message->AddToken(FDocumentationToken::Create(TEXT("https://docs.unrealengine.com/latest/INT/Platforms/iOS/QuickStart/6/index.html"))); 
-	MessageLog.AddMessage(Message);
+	return PackageFilenameCache.Emplace( PackageName, MoveTemp(FCachedPackageFilename(MoveTemp(PackageFilename),MoveTemp(StandardFilename), StandardFileFName)) );
+}
 
-	MessageLog.Notify();
+void UCookOnTheFlyServer::ClearPackageFilenameCache()
+	{
+	PackageFilenameCache.Empty();
 }
 
 
-/* UCookOnTheFlyServer structors
+/* UCookOnTheFlyServer functions
  *****************************************************************************/
 
 UCookOnTheFlyServer::UCookOnTheFlyServer(const FObjectInitializer& ObjectInitializer)
@@ -480,7 +443,7 @@ bool UCookOnTheFlyServer::BroadcastFileserverPresence( const FGuid &InstanceId )
 		if ((NetworkFileServer == NULL || !NetworkFileServer->IsItReadyToAcceptConnections() || !NetworkFileServer->GetAddressList(AddressList)))
 		{
 			LogCookerMessage( FString(TEXT("Failed to create network file server")), EMessageSeverity::Error );
-			UE_LOG(LogCookOnTheFly, Error, TEXT("Failed to create network file server"));
+			UE_LOG(LogCook, Error, TEXT("Failed to create network file server"));
 			continue;
 		}
 
@@ -635,12 +598,12 @@ void UCookOnTheFlyServer::GetDependencies( const TSet<UPackage*>& Packages, TSet
 	{
 		if (*It)
 		{
-			UPackage* Package = It->GetOutermost();
-			if ( Packages.Find( Package ) )
-			{
-				RootSet.Add(*It);
-				Found.Add(*It);
-			}
+		UPackage* Package = It->GetOutermost();
+		if ( Packages.Find( Package ) )
+		{
+			RootSet.Add(*It);
+			Found.Add(*It);
+		}
 		}
 	}
 
@@ -721,9 +684,10 @@ void UCookOnTheFlyServer::GenerateManifestInfo( UPackage* Package, const TArray<
 	{
 		for (const auto& StreamingLevel : World->StreamingLevels)
 		{
-			// Dependencies.Add( StreamingLevel->GetLoadedLevel() );
-			RootPackages.Add(StreamingLevel->GetLoadedLevel()->GetOutermost() );
-			//GetDependencies( StreamingLevel->GetLoadedLevel()->GetOutermost(), Dependencies );
+			if ( StreamingLevel->GetLoadedLevel() )
+			{
+				RootPackages.Add(StreamingLevel->GetLoadedLevel()->GetOutermost() );
+			}
 		}
 
 		TArray<FString> NewPackagesToCook;
@@ -780,8 +744,8 @@ void UCookOnTheFlyServer::GenerateManifestInfo( UPackage* Package, const TArray<
 			if (!Filename.IsEmpty())
 			{
 				// Populate streaming install manifests
-				FString SandboxFilename = SandboxFile->ConvertToAbsolutePathForExternalAppForWrite(*Filename);
-				//UE_LOG(LogCookOnTheFly, Display, TEXT("Adding package to manifest %s, %s, %s"), *DependentPackage->GetName(), *SandboxFilename, *LastLoadedMapName);
+				FString SandboxFilename = ConvertToFullSandboxPath(*Filename, true);
+				//UE_LOG(LogCook, Display, TEXT("Adding package to manifest %s, %s, %s"), *DependentPackage->GetName(), *SandboxFilename, *LastLoadedMapName);
 				ManifestGenerator->AddPackageToChunkManifest(DependentPackage, SandboxFilename, LastLoadedMapName, SandboxFile.GetOwnedPointer());
 			}
 		}
@@ -795,7 +759,7 @@ bool UCookOnTheFlyServer::IsRealtimeMode() const
 
 bool UCookOnTheFlyServer::IsCookByTheBookMode() const
 {
-	return CurrentCookMode == ECookMode::CookByTheBookFromTheEditor || CurrentCookMode == ECookMode::CookByTheBook; 
+	return CurrentCookMode == ECookMode::CookByTheBookFromTheEditor || CurrentCookMode == ECookMode::CookByTheBook;
 }
 
 bool UCookOnTheFlyServer::IsCookOnTheFlyMode() const
@@ -872,12 +836,14 @@ uint32 UCookOnTheFlyServer::TickCookOnTheSide( const float TimeSlice, uint32 &Co
 
 		if (CookedPackages.Exists(ToBuild))
 		{
-			UE_LOG(LogCookOnTheFly, Display, TEXT("Package for platform already cooked %s, discarding request"), *ToBuild.GetFilename().ToString());
+#if DEBUG_COOKONTHEFLY
+			UE_LOG(LogCook, Display, TEXT("Package for platform already cooked %s, discarding request"), *ToBuild.GetFilename().ToString());
+#endif
 			continue;
 		}
-
-		UE_LOG(LogCookOnTheFly, Display, TEXT("Processing package %s"), *ToBuild.GetFilename().ToString());
-
+#if DEBUG_COOKONTHEFLY
+		UE_LOG(LogCook, Display, TEXT("Processing package %s"), *ToBuild.GetFilename().ToString());
+#endif
 		SCOPE_TIMER(TickCookOnTheSide);
 
 		check( ToBuild.IsValid() );
@@ -887,7 +853,7 @@ uint32 UCookOnTheFlyServer::TickCookOnTheSide( const float TimeSlice, uint32 &Co
 
 #if OUTPUT_TIMING
 		//FScopeTimer PackageManualTimer( ToBuild.GetFilename().ToString(), false );
-		UE_LOG(LogCookOnTheFly, Display,  TEXT("ProcessingPackage %s"), *ToBuild.GetFilename().ToString() );
+		UE_LOG(LogCook, Display,  TEXT("ProcessingPackage %s"), *ToBuild.GetFilename().ToString() );
 #endif
 
 		for ( const auto &PlatformName : TargetPlatformNames )
@@ -923,7 +889,9 @@ uint32 UCookOnTheFlyServer::TickCookOnTheSide( const float TimeSlice, uint32 &Co
 				Package = FindObject<UPackage>( ANY_PACKAGE, *PackageName );
 			}
 
-			UE_LOG( LogCookOnTheFly, Display, TEXT("Processing request %s"), *BuildFilename);
+#if DEBUG_COOKONTHEFLY
+			UE_LOG( LogCook, Display, TEXT("Processing request %s"), *BuildFilename);
+#endif
 
 			//  if the package is already loaded then try to avoid reloading it :)
 			if ( ( Package == NULL ) || ( Package->IsFullyLoaded() == false ) )
@@ -939,13 +907,13 @@ uint32 UCookOnTheFlyServer::TickCookOnTheSide( const float TimeSlice, uint32 &Co
 			}
 			else
 			{
-				UE_LOG(LogCookOnTheFly, Display, TEXT("Package already loaded %s avoiding reload"), *BuildFilename );
+				UE_LOG(LogCook, Display, TEXT("Package already loaded %s avoiding reload"), *BuildFilename );
 			}
 
 			if( Package == NULL )
 			{
 				LogCookerMessage( FString::Printf(TEXT("Error loading %s!"), *BuildFilename), EMessageSeverity::Error );
-				UE_LOG(LogCookOnTheFly, Error, TEXT("Error loading %s!"), *BuildFilename );
+				UE_LOG(LogCook, Error, TEXT("Error loading %s!"), *BuildFilename );
 				Result |= COSR_ErrorLoadingPackage;
 			}
 			else
@@ -1006,7 +974,7 @@ uint32 UCookOnTheFlyServer::TickCookOnTheSide( const float TimeSlice, uint32 &Co
 					//  mark the original request as processed (if this isn't actually the file they were requesting then it will fail)
 					//	and then also save our new request as processed so we don't do it again
 					PackagesToSave.AddUnique( Package );
-					UE_LOG( LogCookOnTheFly, Display, TEXT("Request for %s received going to save %s"), *BuildFilename, *PackageFilename );
+					UE_LOG( LogCook, Display, TEXT("Request for %s received going to save %s"), *BuildFilename, *PackageFilename );
 					// CookedPackages.Add( ToBuild );
 					
 					// ToBuild.SetFilename( PackageFilename );
@@ -1026,7 +994,7 @@ uint32 UCookOnTheFlyServer::TickCookOnTheSide( const float TimeSlice, uint32 &Co
 			// this could also happen if the source file doesn't exist which is often as we request files with different extensions when we are searching for files
 			// just return that we processed the cook request
 			// the network file manager will then handle the missing file and search somewhere else
-			UE_LOG(LogCookOnTheFly, Display, TEXT("Not cooking package %s"), *ToBuild.GetFilename().ToString());
+			UE_LOG(LogCook, Display, TEXT("Not cooking package %s"), *ToBuild.GetFilename().ToString());
 			CookedPackages.Add( FFilePlatformRequest( ToBuild.GetFilename(), TargetPlatformNames ) );
 			continue;
 		}
@@ -1071,7 +1039,7 @@ uint32 UCookOnTheFlyServer::TickCookOnTheSide( const float TimeSlice, uint32 &Co
 					if ( Obj->IsCachedCookedPlatformDataLoaded(TargetPlatform) == false )
 					{
 #if DEBUG_COOKONTHEFLY || 1
-						UE_LOG(LogCookOnTheFly, Display, TEXT("Object %s isn't cached yet"), *Obj->GetFullName());
+						UE_LOG(LogCook, Display, TEXT("Object %s isn't cached yet"), *Obj->GetFullName());
 #endif
 						bIsAllDataCached = false;
 					}
@@ -1079,7 +1047,7 @@ uint32 UCookOnTheFlyServer::TickCookOnTheSide( const float TimeSlice, uint32 &Co
 					if ( Timer.IsTimeUp() && IsRealtimeMode() )
 					{
 #if DEBUG_COOKONTHEFLY || 1
-						UE_LOG(LogCookOnTheFly, Display, TEXT("Object %s took too long to cache"), *Obj->GetFullName());
+						UE_LOG(LogCook, Display, TEXT("Object %s took too long to cache"), *Obj->GetFullName());
 #endif
 						bIsAllDataCached = false;
 						break;
@@ -1142,20 +1110,6 @@ uint32 UCookOnTheFlyServer::TickCookOnTheSide( const float TimeSlice, uint32 &Co
 						if ( CookedPackages.Exists(StandardPackageFName,AllTargetPlatformNames) )
 							continue;
 						ACCUMULATE_TIMER_STOP(UnsolicitedPackageAlreadyCooked);
-						// this has been replaced to use the GenerateDependencies call
-						/*if ( CookByTheBookOptions && CookByTheBookOptions->ManifestGenerator != NULL )
-						{
-							check( CurrentCookMode == ECookMode::CookByTheBookFromTheEditor || CurrentCookMode == ECookMode::CookByTheBook );
-							// if we are generating a manifest add the package to the manifest :)
-							FString Filename = GetCachedPackageFilename(Package);
-							if (!Filename.IsEmpty())
-							{
-								// Populate streaming install manifests
-								FString SandboxFilename = SandboxFile->ConvertToAbsolutePathForExternalAppForWrite(*Filename);
-								CookByTheBookOptions->ManifestGenerator->AddPackageToChunkManifest(Package, SandboxFilename, LastLoadedMapName, SandboxFile.GetOwnedPointer());
-							}
-						}*/
-
 
 						if ( CookByTheBookOptions )
 						{
@@ -1168,7 +1122,6 @@ uint32 UCookOnTheFlyServer::TickCookOnTheSide( const float TimeSlice, uint32 &Co
 								if (!Filename.IsEmpty())
 								{
 									// Populate streaming install manifests
-									// FString SandboxFilename = SandboxFile->ConvertToAbsolutePathForExternalAppForWrite(*Filename);
 									const FString& SandboxFilename = GetCachedSandboxFilename( Package, SandboxFile );
 									Manifest->AddUnassignedPackageToManifest(Package, SandboxFilename);
 								}
@@ -1274,14 +1227,6 @@ uint32 UCookOnTheFlyServer::TickCookOnTheSide( const float TimeSlice, uint32 &Co
 				}
 
 
-				/*{
-					SCOPE_TIMER(GenerateManifestInfo);
-					// update manifest with cooked package info
-					GenerateManifestInfo( Package, AllTargetPlatformNames );
-				}*/
-				
-
-
 				//@todo ResetLoaders outside of this (ie when Package is NULL) causes problems w/ default materials
 				if (Package->HasAnyFlags(RF_RootSet) == false && ((CurrentCookMode==ECookMode::CookOnTheFly)) )
 				{
@@ -1304,7 +1249,9 @@ uint32 UCookOnTheFlyServer::TickCookOnTheSide( const float TimeSlice, uint32 &Co
 							(bWasUpToDate == false))
 						{
 							UnsolicitedCookedPackages.AddCookedPackage( FileRequest );
-							UE_LOG(LogCookOnTheFly, Display, TEXT("UnsolicitedCookedPackages: %s"), *FileRequest.GetFilename().ToString());
+#if DEBUG_COOKONTHEFLY
+							UE_LOG(LogCook, Display, TEXT("UnsolicitedCookedPackages: %s"), *FileRequest.GetFilename().ToString());
+#endif
 						}
 					}
 				}
@@ -1325,15 +1272,14 @@ uint32 UCookOnTheFlyServer::TickCookOnTheSide( const float TimeSlice, uint32 &Co
 		if ( CookByTheBookOptions && CookByTheBookOptions->bLeakTest && bFinishedSave )
 		{
 			check( CurrentCookMode == ECookMode::CookByTheBook);
-			UE_LOG(LogCookOnTheFly, Display, TEXT("Full GC..."));
+			UE_LOG(LogCook, Display, TEXT("Full GC..."));
 
 			CollectGarbage( RF_Native );
-
 			for (FObjectIterator It; It; ++It)
 			{
 				if (!CookByTheBookOptions->LastGCItems.Contains(FWeakObjectPtr(*It)))
 				{
-					UE_LOG(LogCookOnTheFly, Warning, TEXT("\tLeaked %s"), *(It->GetFullName()));
+					UE_LOG(LogCook, Warning, TEXT("\tLeaked %s"), *(It->GetFullName()));
 					CookByTheBookOptions->LastGCItems.Add(FWeakObjectPtr(*It));
 				}
 			}
@@ -1394,7 +1340,9 @@ void UCookOnTheFlyServer::MarkPackageDirtyForCooker( UPackage *Package )
 		/*FString PackageFilename(GetPackageFilename(Package));
 		FPaths::MakeStandardFilename(PackageFilename);*/
 		FString PackageFilename = GetCachedStandardPackageFilename(Package);
-		// UE_LOG(LogCookOnTheFly, Display, TEXT("Modification detected to package %s"), *PackageFilename);
+#if DEBUG_COOKONTHEFLY
+		UE_LOG(LogCook, Display, TEXT("Modification detected to package %s"), *PackageFilename);
+#endif
 		const FName PackageFFileName = FName(*PackageFilename);
 
 		if ( CurrentCookMode == ECookMode::CookByTheBookFromTheEditor )
@@ -1405,12 +1353,12 @@ void UCookOnTheFlyServer::MarkPackageDirtyForCooker( UPackage *Package )
 			{
 				if ( IsCookByTheBookRunning() )
 				{
-					// if this package was previously cooked and we are doing a cook by the book 
-					// we need to recook this package before finishing cook by the book
-					CookRequests.EnqueueUnique( FFilePlatformRequest(PackageFFileName, CookedPlatforms) );	
-				}
+				// if this package was previously cooked and we are doing a cook by the book 
+				// we need to recook this package before finishing cook by the book
+				CookRequests.EnqueueUnique( FFilePlatformRequest(PackageFFileName, CookedPlatforms) );
 			}
 		}
+		}
 		CookedPackages.RemoveAll( PackageFFileName );
 	}
 }
@@ -1502,12 +1450,12 @@ bool UCookOnTheFlyServer::ShouldCook(const FString& InFileName, const FName &InP
 
 		if (GetPackageTimestamp(FPaths::GetBaseFilename(PkgFilename, false), DependentTimeStamp) == false)
 		{
-			UE_LOG(LogCookOnTheFly, Display, TEXT("Failed to find dependency timestamp for: %s"), *PkgFilename);
+			UE_LOG(LogCook, Display, TEXT("Failed to find dependency timestamp for: %s"), *PkgFilename);
 		}
 	}
 
 	// Use SandboxFile to do path conversion to properly handle sandbox paths (outside of standard paths in particular).
-	PkgFilename = SandboxFile->ConvertToAbsolutePathForExternalAppForWrite(*PkgFilename);
+	PkgFilename = ConvertToFullSandboxPath(*PkgFilename, true);
 
 	ITargetPlatformManagerModule& TPM = GetTargetPlatformManagerRef();
 
@@ -1568,12 +1516,14 @@ bool UCookOnTheFlyServer::SaveCookedPackage( UPackage* Package, uint32 SaveFlags
 
 			if (GetPackageTimestamp(FPaths::GetBaseFilename(PkgFilename, false), DependentTimeStamp) == false)
 			{
-				UE_LOG(LogCookOnTheFly, Display, TEXT("Failed to find depedency timestamp for: %s"), *PkgFilename);
+#if DEBUG_COOKONTHEFLY
+				UE_LOG(LogCook, Display, TEXT("Failed to find depedency timestamp for: %s"), *PkgFilename);
+#endif
 			}
 		}
 
 		// Use SandboxFile to do path conversion to properly handle sandbox paths (outside of standard paths in particular).
-		Filename = SandboxFile->ConvertToAbsolutePathForExternalAppForWrite(*Filename);
+		Filename = ConvertToFullSandboxPath(*Filename, true);
 
 		uint32 OriginalPackageFlags = Package->PackageFlags;
  		UWorld* World = NULL;
@@ -1651,7 +1601,8 @@ bool UCookOnTheFlyServer::SaveCookedPackage( UPackage* Package, uint32 SaveFlags
 					{
 						LogCookerMessage( FString::Printf(TEXT("Package %s supposed to be fully loaded but isn't. RF_WasLoaded is %s"), 
 							*Package->GetName(), Package->HasAnyFlags(RF_WasLoaded) ? TEXT("set") : TEXT("not set")), EMessageSeverity::Warning);
-						UE_LOG(LogCookOnTheFly, Warning, TEXT("Package %s supposed to be fully loaded but isn't. RF_WasLoaded is %s"), 
+
+						UE_LOG(LogCook, Warning, TEXT("Package %s supposed to be fully loaded but isn't. RF_WasLoaded is %s"), 
 							*Package->GetName(), Package->HasAnyFlags(RF_WasLoaded) ? TEXT("set") : TEXT("not set"));
 					}
 					bPackageFullyLoaded = true;
@@ -1665,7 +1616,7 @@ bool UCookOnTheFlyServer::SaveCookedPackage( UPackage* Package, uint32 SaveFlags
 					Flags = World ? RF_NoFlags : RF_Standalone;
 				}
 
-				UE_LOG(LogCookOnTheFly, Display, TEXT("Cooking %s -> %s"), *Package->GetName(), *PlatFilename);
+				UE_LOG(LogCook, Display, TEXT("Cooking %s -> %s"), *Package->GetName(), *PlatFilename);
 
 				bool bSwap = (!Target->IsLittleEndian()) ^ (!PLATFORM_LITTLE_ENDIAN);
 
@@ -1695,7 +1646,7 @@ bool UCookOnTheFlyServer::SaveCookedPackage( UPackage* Package, uint32 SaveFlags
 				if( FullFilename.Len() >= PLATFORM_MAX_FILEPATH_LENGTH )
 				{
 					LogCookerMessage( FString::Printf(TEXT("Couldn't save package, filename is too long: %s"), *PlatFilename), EMessageSeverity::Error );
-					UE_LOG( LogCookOnTheFly, Error, TEXT( "Couldn't save package, filename is too long :%s" ), *PlatFilename );
+					UE_LOG( LogCook, Error, TEXT( "Couldn't save package, filename is too long :%s" ), *PlatFilename );
 					bSavedCorrectly = false;
 				}
 				else
@@ -1720,7 +1671,9 @@ bool UCookOnTheFlyServer::SaveCookedPackage( UPackage* Package, uint32 SaveFlags
 			}
 			else
 			{
-				UE_LOG(LogCookOnTheFly, Display, TEXT("Up to date: %s"), *PlatFilename);
+#if DEBUG_COOKONTHEFLY
+				UE_LOG(LogCook, Display, TEXT("Up to date: %s"), *PlatFilename);
+#endif
 
 				bOutWasUpToDate = true;
 			}
@@ -1840,7 +1793,9 @@ bool UCookOnTheFlyServer::GetCurrentIniVersionStrings( const ITargetPlatform* Ta
 
 		if ( IniVersionedArray.Num() != 3 )
 		{
-			// UE_LOG(LogCookOnTheFly, Warning, TEXT("Invalid entry in CookSettings, VersionedIniParams %s"), *IniVersioned);
+#if DEBUG_COOKONTHEFLY
+			UE_LOG(LogCook, Warning, TEXT("Invalid entry in CookSettings, VersionedIniParams %s"), *IniVersioned);
+#endif
 			return false;
 		}
 
@@ -1861,7 +1816,9 @@ bool UCookOnTheFlyServer::GetCurrentIniVersionStrings( const ITargetPlatform* Ta
 		FString Value;
 		if ( !PlatformIniFile->GetString(*Section, *Key, Value) )
 		{
-			// UE_LOG(LogCookOnTheFly, Warning, TEXT("Unable to find entry in CookSettings, VersionedIniParams %s, assume default is being used"), *IniVersioned);
+#if DEBUG_COOKONTHEFLY
+			UE_LOG(LogCook, Warning, TEXT("Unable to find entry in CookSettings, VersionedIniParams %s, assume default is being used"), *IniVersioned);
+#endif
 			continue;
 		}
 
@@ -1913,7 +1870,7 @@ bool UCookOnTheFlyServer::GetCookedIniVersionStrings( const ITargetPlatform* Tar
 {
 
 	const FString EditorIni = FPaths::GameDir() / TEXT("CookedIniVersion.txt");
-	const FString SandboxEditorIni = SandboxFile->ConvertToAbsolutePathForExternalAppForWrite(*EditorIni);
+	const FString SandboxEditorIni = ConvertToFullSandboxPath(*EditorIni, true);
 
 
 	const FString PlatformSandboxEditorIni = SandboxEditorIni.Replace(TEXT("[Platform]"), *TargetPlatform->PlatformName());
@@ -1944,7 +1901,7 @@ bool UCookOnTheFlyServer::IniSettingsOutOfDate( const TArray<ITargetPlatform*>&
 
 		if ( FoundCookedIniVersionStrings == NULL )
 		{
-			TArray<FString> CookedIniVersionStrings;
+		TArray<FString> CookedIniVersionStrings;
 			GetCookedIniVersionStrings( TargetPlatform, CookedIniVersionStrings );
 			FoundCookedIniVersionStrings = &CachedIniVersionStringsMap.Emplace( TargetPlatformName, MoveTemp(CookedIniVersionStrings) );
 		}
@@ -1962,18 +1919,18 @@ bool UCookOnTheFlyServer::IniSettingsOutOfDate( const TArray<ITargetPlatform*>&
 		}
 
 		if ( bCurrentIniSettingsChanged )
-		{
-			OutOfDateTargetPlatforms.Add(TargetPlatform);
-			*FoundCookedIniVersionStrings = CurrentIniVersionStrings;
+{
+		OutOfDateTargetPlatforms.Add(TargetPlatform);
+		*FoundCookedIniVersionStrings = CurrentIniVersionStrings;
 
 			// save the iniversion strings
 			if ( bShouldSaveCookedVersions )
-			{
-				const FString PlatformSandboxCookedIni = SandboxCookedIni.Replace(TEXT("[Platform]"), *TargetPlatform->PlatformName());
+	{
+		const FString PlatformSandboxCookedIni = SandboxCookedIni.Replace(TEXT("[Platform]"), *TargetPlatform->PlatformName());
 				GConfig->SetArray(TEXT("CookSettings"), TEXT("VersionedIniParams"), CurrentIniVersionStrings, PlatformSandboxCookedIni);
-				GConfig->Flush(false, PlatformSandboxCookedIni);
-			}
-		}
+		GConfig->Flush(false, PlatformSandboxCookedIni);
+	}
+}
 	}
 
 	return true;
@@ -2038,8 +1995,9 @@ void UCookOnTheFlyServer::CleanSandbox(const TArray<ITargetPlatform*>& Platforms
 
 						if (Diff < 0.0)
 						{
-							UE_LOG(LogCookOnTheFly, Display, TEXT("Deleting out of date cooked file: %s"), *CookedFilename);
-
+#if DEBUG_COOKONTHEFLY
+							UE_LOG(LogCook, Display, TEXT("Deleting out of date cooked file: %s"), *CookedFilename);
+#endif
 							IFileManager::Get().Delete(*CookedFilename);
 						}
 					}
@@ -2050,8 +2008,9 @@ void UCookOnTheFlyServer::CleanSandbox(const TArray<ITargetPlatform*>& Platforms
 			CollectGarbage(RF_Native);
 		}
 	}
-
-	UE_LOG(LogCookOnTheFly, Display, TEXT("Sandbox cleanup took %5.3f seconds"), SandboxCleanTime);
+#if DEBUG_COOKONTHEFLY
+	UE_LOG(LogCook, Display, TEXT("Sandbox cleanup took %5.3f seconds"), SandboxCleanTime);
+#endif
 }
 
 void UCookOnTheFlyServer::GenerateAssetRegistry(const TArray<ITargetPlatform*>& Platforms)
@@ -2063,8 +2022,9 @@ void UCookOnTheFlyServer::GenerateAssetRegistry(const TArray<ITargetPlatform*>&
 	double GenerateAssetRegistryTime = 0.0;
 	{
 		SCOPE_SECONDS_COUNTER(GenerateAssetRegistryTime);
-		UE_LOG(LogCookOnTheFly, Display, TEXT("Creating asset registry [is editor: %d]"), GIsEditor);
-
+#if DEBUG_COOKONTHEFLY
+		UE_LOG(LogCook, Display, TEXT("Creating asset registry [is editor: %d]"), GIsEditor);
+#endif
 		// Perform a synchronous search of any .ini based asset paths (note that the per-game delegate may
 		// have already scanned paths on its own)
 		// We want the registry to be fully initialized when generating streaming manifests too.
@@ -2090,12 +2050,12 @@ void UCookOnTheFlyServer::GenerateAssetRegistry(const TArray<ITargetPlatform*>&
 			// write it out to a memory archive
 			FArrayWriter SerializedAssetRegistry;
 			AssetRegistry.Serialize(SerializedAssetRegistry);
-			UE_LOG(LogCookOnTheFly, Display, TEXT("Generated asset registry size is %5.2fkb"), (float)SerializedAssetRegistry.Num() / 1024.f);
+			UE_LOG(LogCook, Display, TEXT("Generated asset registry size is %5.2fkb"), (float)SerializedAssetRegistry.Num() / 1024.f);
 
 			// now save it in each cooked directory
 			FString RegistryFilename = FPaths::GameDir() / TEXT("AssetRegistry.bin");
 			// Use SandboxFile to do path conversion to properly handle sandbox paths (outside of standard paths in particular).
-			FString SandboxFilename = SandboxFile->ConvertToAbsolutePathForExternalAppForWrite(*RegistryFilename);
+			FString SandboxFilename = ConvertToFullSandboxPath(*RegistryFilename, true);
 
 			for (int32 Index = 0; Index < Platforms.Num(); Index++)
 			{
@@ -2104,7 +2064,7 @@ void UCookOnTheFlyServer::GenerateAssetRegistry(const TArray<ITargetPlatform*>&
 			}
 		}
 	}
-	UE_LOG(LogCookOnTheFly, Display, TEXT("Done creating registry. It took %5.2fs."), GenerateAssetRegistryTime);
+	UE_LOG(LogCook, Display, TEXT("Done creating registry. It took %5.2fs."), GenerateAssetRegistryTime);
 }
 
 void UCookOnTheFlyServer::GenerateLongPackageNames(TArray<FString>& FilesInPath)
@@ -2128,7 +2088,7 @@ void UCookOnTheFlyServer::GenerateLongPackageNames(TArray<FString>& FilesInPath)
 			else
 			{
 				LogCookerMessage( FString::Printf(TEXT("Unable to generate long package name for %s"), *FileInPath), EMessageSeverity::Warning);
-				UE_LOG(LogCookOnTheFly, Warning, TEXT("Unable to generate long package name for %s"), *FileInPath);
+				UE_LOG(LogCook, Warning, TEXT("Unable to generate long package name for %s"), *FileInPath);
 			}
 		}
 	}
@@ -2181,7 +2141,7 @@ void UCookOnTheFlyServer::CollectFilesToCook(TArray<FString>& FilesInPath, const
 			if (FPackageName::SearchForPackageOnDisk(CurrEntry, NULL, &OutFilename) == false)
 			{
 				LogCookerMessage( FString::Printf(TEXT("Unable to find package for map %s."), *CurrEntry), EMessageSeverity::Warning);
-				UE_LOG(LogCookOnTheFly, Warning, TEXT("Unable to find package for map %s."), *CurrEntry);
+				UE_LOG(LogCook, Warning, TEXT("Unable to find package for map %s."), *CurrEntry);
 			}
 			else
 			{
@@ -2214,6 +2174,29 @@ void UCookOnTheFlyServer::CollectFilesToCook(TArray<FString>& FilesInPath, const
 		}
 	}
 
+	if ( IsCookingDLC() )
+	{
+		// get the dlc and make sure we cook that directory 
+		FString DLCPath = FPaths::GamePluginsDir() / CookByTheBookOptions->DlcName;
+
+		TArray<FString> Files;
+		IFileManager::Get().FindFilesRecursive(Files, *DLCPath, *(FString(TEXT("*")) + FPackageName::GetAssetPackageExtension()), true, false);
+		IFileManager::Get().FindFilesRecursive(Files, *DLCPath, *(FString(TEXT("*")) + FPackageName::GetMapPackageExtension()), true, false);
+		for (int32 Index = 0; Index < Files.Num(); Index++)
+		{
+			FString StdFile = Files[Index];
+			FPaths::MakeStandardFilename(StdFile);
+			AddFileToCook( FilesInPath,StdFile);
+
+			// this asset may not be in our currently mounted content directories, so try to mount a new one now
+			FString LongPackageName;
+			if(!FPackageName::IsValidLongPackageName(StdFile) && !FPackageName::TryConvertFilenameToLongPackageName(StdFile, LongPackageName))
+			{
+				FPackageName::RegisterMountPoint(ExternalMountPointName, DLCPath);
+			}
+		}
+	}
+
 	if ((FilesInPath.Num() == 0) || bCookAll)
 	{
 		TArray<FString> Tokens;
@@ -2239,7 +2222,7 @@ void UCookOnTheFlyServer::CollectFilesToCook(TArray<FString>& FilesInPath, const
 			TArray<FString> TokenFiles;
 			if ( !NormalizePackageNames( Unused, TokenFiles, Tokens[TokenIndex], PackageFilter) )
 			{
-				UE_LOG(LogCookOnTheFly, Display, TEXT("No packages found for parameter %i: '%s'"), TokenIndex, *Tokens[TokenIndex]);
+				UE_LOG(LogCook, Display, TEXT("No packages found for parameter %i: '%s'"), TokenIndex, *Tokens[TokenIndex]);
 				continue;
 			}
 
@@ -2369,6 +2352,26 @@ void UCookOnTheFlyServer::SaveGlobalShaderMapFiles(const TArray<ITargetPlatform*
 }
 
 
+FString UCookOnTheFlyServer::ConvertToFullSandboxPath( const FString &FileName, bool bForWrite ) const
+{
+	FString Result;
+	if ( bForWrite)
+	{
+		Result = SandboxFile->ConvertToAbsolutePathForExternalAppForWrite(*FileName);
+	}
+	else
+	{
+		Result = SandboxFile->ConvertToAbsolutePathForExternalAppForRead(*FileName);
+	}
+
+	if ( IsCookByTheBookRunning() && IsCookingDLC() )
+	{
+		Result.ReplaceInline(TEXT("/Cooked/"), *FString::Printf(TEXT("/CookedDLC_%s/"), *CookByTheBookOptions->DlcName) );
+	}
+	return Result;
+}
+
+
 void UCookOnTheFlyServer::CookByTheBookFinished()
 {
 	check( IsInGameThread() );
@@ -2382,11 +2385,11 @@ void UCookOnTheFlyServer::CookByTheBookFinished()
 	{
 		// Save modified asset registry with all streaming chunk info generated during cook
 		FString RegistryFilename = FPaths::GameDir() / TEXT("AssetRegistry.bin");
-		FString SandboxRegistryFilename = SandboxFile->ConvertToAbsolutePathForExternalAppForWrite(*RegistryFilename);
+		FString SandboxRegistryFilename = ConvertToFullSandboxPath(*RegistryFilename, true);
 		// the registry filename will be modified when we call save asset registry
 		
 		FString CookedAssetRegistry = FPaths::GameDir() / TEXT("CookedAssetRegistry.json");
-		FString SandboxCookedAssetRegistryFilename = SandboxFile->ConvertToAbsolutePathForExternalAppForWrite(*CookedAssetRegistry);
+		FString SandboxCookedAssetRegistryFilename = ConvertToFullSandboxPath(*CookedAssetRegistry, true);
 
 		for ( auto& Manifest : CookByTheBookOptions->ManifestGenerators )
 		{
@@ -2400,7 +2403,7 @@ void UCookOnTheFlyServer::CookByTheBookFinished()
 
 	CookByTheBookOptions->LastGCItems.Empty();
 	const float TotalCookTime = (float)(FPlatformTime::Seconds() - CookByTheBookOptions->CookStartTime);
-	UE_LOG(LogCookOnTheFly, Display, TEXT("Cook by the book total time in tick %fs total time %f"), CookByTheBookOptions->CookTime, TotalCookTime);
+	UE_LOG(LogCook, Display, TEXT("Cook by the book total time in tick %fs total time %f"), CookByTheBookOptions->CookTime, TotalCookTime);
 
 	CookByTheBookOptions->bRunning = false;
 }
@@ -2443,6 +2446,11 @@ void UCookOnTheFlyServer::StopAndClearCookedData()
 	CookedPackages.Empty(); // set of files which have been cooked when needing to recook a file the entry will need to be removed from here
 }
 
+FString UCookOnTheFlyServer::GetAssetRegistryPath() const
+{
+	return FPaths::GameDir() / TEXT("AssetRegistry.bin");
+}
+
 void UCookOnTheFlyServer::ClearAllCookedData()
 {
 	// if we are going to clear the cooked packages it is conceivable that we will recook the packages which we just cooked 
@@ -2467,7 +2475,7 @@ void UCookOnTheFlyServer::ClearPlatformCookedData( const FName& PlatformName )
 void UCookOnTheFlyServer::StartCookByTheBook(const TArray<ITargetPlatform*>& TargetPlatforms, 
 											 const TArray<FString> &CookMaps, const TArray<FString> &CookDirectories, 
 											 const TArray<FString> &CookCultures, const TArray<FString> &IniMapSections, 
-											 ECookByTheBookOptions CookOptions )
+											 ECookByTheBookOptions CookOptions, const FString &DLCName )
 {
 	check( IsInGameThread() );
 
@@ -2513,6 +2521,7 @@ void UCookOnTheFlyServer::StartCookByTheBook(const TArray<ITargetPlatform*>& Tar
 	CookByTheBookOptions->bCancel = false;
 	CookByTheBookOptions->CookTime = 0.0f;
 	CookByTheBookOptions->CookStartTime = FPlatformTime::Seconds();
+	CookByTheBookOptions->DlcName = DLCName;
 
 	bool bCookAll = (CookOptions & ECookByTheBookOptions::CookAll) != ECookByTheBookOptions::None;
 	bool bMapsOnly = (CookOptions & ECookByTheBookOptions::MapsOnly) != ECookByTheBookOptions::None;
@@ -2530,20 +2539,6 @@ void UCookOnTheFlyServer::StartCookByTheBook(const TArray<ITargetPlatform*>& Tar
 		}
 	}
 
-	// allow the game to fill out the asset registry, as well as get a list of objects to always cook
-	TArray<FString> FilesInPath;
-	FGameDelegates::Get().GetCookModificationDelegate().ExecuteIfBound(FilesInPath);
-
-	SaveGlobalShaderMapFiles(TargetPlatforms);
-
-	CollectFilesToCook(FilesInPath, CookMaps, CookDirectories, CookCultures, IniMapSections, bCookAll, bMapsOnly, bNoDev );
-	if (FilesInPath.Num() == 0)
-	{
-		LogCookerMessage( FString::Printf(TEXT("No files found to cook.")), EMessageSeverity::Warning );
-		UE_LOG(LogCookOnTheFly, Warning, TEXT("No files found."));
-	}
-
-
 	TArray<FName> TargetPlatformNames;
 	for( const auto &Platform : TargetPlatforms )
 	{
@@ -2565,6 +2560,33 @@ void UCookOnTheFlyServer::StartCookByTheBook(const TArray<ITargetPlatform*>& Tar
 		}
 	}
 
+
+	if ( IsCookingDLC() )
+	{
+		ClearPackageFilenameCache();
+
+		FString OriginalSandboxRegistryFilename = SandboxFile->ConvertToAbsolutePathForExternalAppForRead(*GetAssetRegistryPath());
+		// load the original asset registry 
+		// don't recook any of the already cooked assets
+		// warm the cooked packages
+		WarmCookedPackages( OriginalSandboxRegistryFilename, TargetPlatformNames );
+	}
+
+	
+	// allow the game to fill out the asset registry, as well as get a list of objects to always cook
+	TArray<FString> FilesInPath;
+	FGameDelegates::Get().GetCookModificationDelegate().ExecuteIfBound(FilesInPath);
+
+	SaveGlobalShaderMapFiles(TargetPlatforms);
+
+	CollectFilesToCook(FilesInPath, CookMaps, CookDirectories, CookCultures, IniMapSections, bCookAll, bMapsOnly, bNoDev );
+	if (FilesInPath.Num() == 0)
+	{
+		LogCookerMessage( FString::Printf(TEXT("No files found to cook.")), EMessageSeverity::Warning );
+		UE_LOG(LogCook, Warning, TEXT("No files found."));
+	}
+
+
 	GenerateLongPackageNames(FilesInPath);
 
 	// add all the files for the requested platform to the cook list
@@ -2580,10 +2602,13 @@ void UCookOnTheFlyServer::StartCookByTheBook(const TArray<ITargetPlatform*>& Tar
 		else
 		{
 			LogCookerMessage( FString::Printf(TEXT("Unable to find package for cooking %s"), *FileName), EMessageSeverity::Warning );
-			UE_LOG(LogCookOnTheFly, Warning, TEXT("Unable to find package for cooking %s"), *FileName)
+			UE_LOG(LogCook, Warning, TEXT("Unable to find package for cooking %s"), *FileName)
 		}	
 	}
 
+	// this is to support canceling cooks from the editor
+	// this is required to make sure that the cooker is in a good state after cancel occurs
+	// if too many packages are being recooked after resume then we may need to figure out a different way to do this
 	for ( const auto& PreviousRequest : CookByTheBookOptions->PreviousCookRequests )
 	{
 		CookRequests.EnqueueUnique( MoveTemp( PreviousRequest ) );
@@ -2639,7 +2664,7 @@ void UCookOnTheFlyServer::HandleNetworkFileServerFileRequest( const FString& Fil
 	}
 	while (!CookedPackages.Exists(FileRequest));
 
-	UE_LOG( LogCookOnTheFly, Display, TEXT("Cook complete %s"), *FileRequest.GetFilename().ToString())
+	UE_LOG( LogCook, Display, TEXT("Cook complete %s"), *FileRequest.GetFilename().ToString())
 
 	TArray<FName> UnsolicitedFilenames;
 	UnsolicitedCookedPackages.GetPackagesForPlatformAndRemove(PlatformFname, UnsolicitedFilenames);
@@ -2656,7 +2681,7 @@ void UCookOnTheFlyServer::HandleNetworkFileServerFileRequest( const FString& Fil
 
 
 #if DEBUG_COOKONTHEFLY
-	UE_LOG( LogCookOnTheFly, Display, TEXT("Processed file request %s"), *Filename );
+	UE_LOG( LogCook, Display, TEXT("Processed file request %s"), *Filename );
 #endif
 
 }
@@ -2669,7 +2694,7 @@ void UCookOnTheFlyServer::HandleNetworkFileServerRecompileShaders(const FShaderR
 	// if we aren't in the game thread, we need to push this over to the game thread and wait for it to finish
 	if (!IsInGameThread())
 	{
-		UE_LOG(LogCookOnTheFly, Display, TEXT("Got a recompile request on non-game thread"));
+		UE_LOG(LogCook, Display, TEXT("Got a recompile request on non-game thread"));
 
 		// make a new request
 		FRecompileRequest* Request = new FRecompileRequest;
@@ -2685,7 +2710,7 @@ void UCookOnTheFlyServer::HandleNetworkFileServerRecompileShaders(const FShaderR
 			FPlatformProcess::Sleep(0);
 		}
 		delete Request;
-		UE_LOG(LogCookOnTheFly, Display, TEXT("Completed recompile..."));
+		UE_LOG(LogCook, Display, TEXT("Completed recompile..."));
 
 		// at this point, we are done on the game thread, and ModifiedFiles will have been filled out
 		return;
@@ -2730,7 +2755,7 @@ void UCookOnTheFlyServer::WarmCookedPackages(const FString& AssetRegistryPath, c
 			// load it
 			SerializedAssetData << *NewAssetData;
 
-			UE_LOG(LogCookOnTheFly, Verbose, TEXT("Read package %s from %s"), *GetCachedStandardPackageFilename(NewAssetData->ObjectPath), *AssetRegistryPath);
+			UE_LOG(LogCook, Verbose, TEXT("Read package %s from %s"), *GetCachedStandardPackageFilename(NewAssetData->ObjectPath), *AssetRegistryPath);
 
 			FileRequest.SetFilename(GetCachedStandardPackageFilename(NewAssetData->ObjectPath));
 			if (FileRequest.IsValid())
