[UE-5075] Adding the ability to reload blueprints from disk.

diff --git a/Engine/Source/Editor/BlueprintGraph/Private/BlueprintActionDatabase.cpp b/Engine/Source/Editor/BlueprintGraph/Private/BlueprintActionDatabase.cpp
--- a/Engine/Source/Editor/BlueprintGraph/Private/BlueprintActionDatabase.cpp
+++ b/Engine/Source/Editor/BlueprintGraph/Private/BlueprintActionDatabase.cpp
@@ -424,13 +424,29 @@ namespace BlueprintActionDatabaseImpl
 	static void OnAssetsPendingDelete(TArray<UObject*> const& ObjectsForDelete);
 
 	/**
-	 * Callback to refresh the database when an object has been delete (to clear 
-	 * any related classes that were stored in the database) 
+	 * Callback to refresh the database when an object has been deleted (to  
+	 * clear any related entries that were stored in the database) 
 	 * 
 	 * @param  AssetInfo	Data regarding the freshly removed asset.
 	 */
 	static void OnAssetRemoved(FAssetData const& AssetInfo);
 
+	/**
+	 * Callback to refresh the database when an object has been deleted/unloaded
+	 * (to clear any related entries that were stored in the database) 
+	 * 
+	 * @param  AssetObject	The object being removed/deleted/unloaded.
+	 */
+	static void OnAssetRemoved(UObject* AssetObject);
+
+	/**
+	 * Callback to refresh the database when a blueprint has been unloaded
+	 * (to clear any related entries that were stored in the database) 
+	 * 
+	 * @param  BlueprintObj	The blueprint being unloaded.
+	 */
+	static void OnBlueprintUnloaded(UBlueprint* BlueprintObj);
+
 	/**
 	 * Callback to refresh the database when an object has been renamed (to clear 
 	 * any related classes that were stored in the database under the old name) 
@@ -839,8 +855,7 @@ static void BlueprintActionDatabaseImpl::OnAssetRemoved(FAssetData const& AssetI
 	if (AssetInfo.IsAssetLoaded())
 	{
 		UObject* AssetObject = AssetInfo.GetAsset();
-		// the delete went through, so we don't need to track these for re-add
-		PendingDelete.Remove(AssetObject);
+		OnAssetRemoved(AssetObject);
 	}
 	else
 	{
@@ -848,6 +863,22 @@ static void BlueprintActionDatabaseImpl::OnAssetRemoved(FAssetData const& AssetI
 	}
 }
 
+//------------------------------------------------------------------------------
+static void BlueprintActionDatabaseImpl::OnAssetRemoved(UObject* AssetObject)
+{
+	FBlueprintActionDatabase& ActionDatabase = FBlueprintActionDatabase::Get();
+	ActionDatabase.ClearAssetActions(AssetObject);
+
+	// the delete went through, so we don't need to track these for re-add
+	PendingDelete.Remove(AssetObject);
+}
+
+//------------------------------------------------------------------------------
+static void BlueprintActionDatabaseImpl::OnBlueprintUnloaded(UBlueprint* BlueprintObj)
+{
+	OnAssetRemoved(BlueprintObj);
+}
+
 //------------------------------------------------------------------------------
 static void BlueprintActionDatabaseImpl::OnAssetRenamed(FAssetData const& AssetInfo, const FString& InOldName)
 {
@@ -915,6 +946,7 @@ FBlueprintActionDatabase::FBlueprintActionDatabase()
 	AssetRegistry.OnAssetRenamed().AddStatic(&BlueprintActionDatabaseImpl::OnAssetRenamed);
 
 	FEditorDelegates::OnAssetsPreDelete.AddStatic(&BlueprintActionDatabaseImpl::OnAssetsPendingDelete);
+	FKismetEditorUtilities::OnBlueprintUnloaded.AddStatic(&BlueprintActionDatabaseImpl::OnBlueprintUnloaded);
 
 	GEngine->OnWorldDestroyed().AddStatic(&BlueprintActionDatabaseImpl::OnWorldDestroyed);
 }
@@ -1145,7 +1177,11 @@ void FBlueprintActionDatabase::RefreshAssetActions(UObject* const AssetObject)
 	// Will clear up any unloaded asset actions associated with this object, if any
 	ClearUnloadedAssetActions(*AssetObject->GetPathName());
 
-	if (AssetActionList.Num() > 0)
+	if (AssetObject->IsPendingKill())
+	{
+		ClearAssetActions(AssetObject);
+	}
+	else if (AssetActionList.Num() > 0)
 	{
 		// queue these assets for priming
 		ActionPrimingQueue.Add(AssetObject, 0);
diff --git a/Engine/Source/Editor/Kismet/Private/BlueprintEditor.cpp b/Engine/Source/Editor/Kismet/Private/BlueprintEditor.cpp
--- a/Engine/Source/Editor/Kismet/Private/BlueprintEditor.cpp
+++ b/Engine/Source/Editor/Kismet/Private/BlueprintEditor.cpp
@@ -1070,6 +1070,7 @@ static auto CVarBlueprintExpertMode	= IConsoleManager::Get().RegisterConsoleVari
 
 FBlueprintEditor::FBlueprintEditor()
 	: bSaveIntermediateBuildProducts(false)
+	, bPendingDeferredClose(false)
 	, bRequestedSavingOpenDocumentState(false)
 	, bBlueprintModifiedOnOpen (false)
 	, PinVisibility(SGraphEditor::Pin_Show)
@@ -1272,6 +1273,8 @@ void FBlueprintEditor::CommonInitialization(const TArray<UBlueprint*>& InitBluep
 
 	CreateDefaultCommands();
 	CreateDefaultTabContents(InitBlueprints);
+
+	FKismetEditorUtilities::OnBlueprintUnloaded.AddSP(this, &FBlueprintEditor::OnBlueprintUnloaded);
 }
 
 void FBlueprintEditor::LoadLibrariesFromAssetRegistry()
@@ -2635,6 +2638,19 @@ void FBlueprintEditor::OnBlueprintChanged(UBlueprint* InBlueprint)
 	}
 }
 
+void FBlueprintEditor::OnBlueprintUnloaded(UBlueprint* InBlueprint)
+{
+	for (UObject* EditingObj : GetEditingObjects())
+	{
+		if (Cast<UBlueprint>(EditingObj) == InBlueprint)
+		{
+			// give the editor a chance to open a replacement
+			bPendingDeferredClose = true;
+			break;
+		}
+	}
+}
+
 void FBlueprintEditor::Compile()
 {
 	UBlueprint* BlueprintObj = GetBlueprintObj();
@@ -5899,6 +5915,13 @@ void FBlueprintEditor::Tick(float DeltaTime)
 	{
 		InstructionsFadeCountdown -= DeltaTime;
 	}
+
+	if (bPendingDeferredClose)
+	{
+		IAssetEditorInstance* EditorInst = FAssetEditorManager::Get().FindEditorForAsset(GetBlueprintObj(), /*bFocusIfOpen =*/false);
+		check(EditorInst != nullptr);
+		EditorInst->CloseWindow();
+	}
 }
 TStatId FBlueprintEditor::GetStatId() const
 {
diff --git a/Engine/Source/Editor/Kismet/Public/BlueprintEditor.h b/Engine/Source/Editor/Kismet/Public/BlueprintEditor.h
--- a/Engine/Source/Editor/Kismet/Public/BlueprintEditor.h
+++ b/Engine/Source/Editor/Kismet/Public/BlueprintEditor.h
@@ -576,6 +576,9 @@ protected:
 	/** Called when the Blueprint we are editing has changed */
 	virtual void OnBlueprintChanged(UBlueprint* InBlueprint);
 
+	/** Handles the unloading of Blueprints (by closing the editor, if it operating on the Blueprint being unloaded)*/
+	void OnBlueprintUnloaded(UBlueprint* InBlueprint);
+
 	/** Get title for Inspector 2 tab*/
 	virtual FString GetDefaultEditorTitle();
 
@@ -910,13 +913,11 @@ protected:
 	// Should intermediate build products be saved when recompiling?
 	bool bSaveIntermediateBuildProducts;
 
+	/** Flags if this blueprint editor should close on its next tick. */
+	bool bPendingDeferredClose;
+
 	/** Currently focused graph editor */
 	TWeakPtr<class SGraphEditor> FocusedGraphEdPtr;
-
-
-
-
-
 	
 	// Factory that spawns graph editors; used to look up all tabs spawned by it.
 	TWeakPtr<FDocumentTabFactory> GraphEditorTabFactoryPtr;
diff --git a/Engine/Source/Editor/UnrealEd/Private/Kismet2/Kismet2.cpp b/Engine/Source/Editor/UnrealEd/Private/Kismet2/Kismet2.cpp
--- a/Engine/Source/Editor/UnrealEd/Private/Kismet2/Kismet2.cpp
+++ b/Engine/Source/Editor/UnrealEd/Private/Kismet2/Kismet2.cpp
@@ -33,6 +33,7 @@
 #include "DlgPickAssetPath.h"
 #include "ComponentAssetBroker.h"
 #include "BlueprintEditorSettings.h"
+#include "Editor/UnrealEd/Classes/Editor/Transactor.h"
 
 #define LOCTEXT_NAMESPACE "UnrealEd.Editor"
 
@@ -163,6 +164,7 @@ void FObjectsBeingDebuggedIterator::FindNextLevelScriptActor()
 
 // Static variable definition
 TArray<FString> FKismetEditorUtilities::TrackedBlueprintParentList;
+FKismetEditorUtilities::FOnBlueprintUnloaded FKismetEditorUtilities::OnBlueprintUnloaded;
 
 /** Create the correct event graphs for this blueprint */
 void FKismetEditorUtilities::CreateDefaultEventGraphs(UBlueprint* Blueprint)
@@ -336,6 +338,152 @@ UBlueprint* FKismetEditorUtilities::CreateBlueprint(UClass* ParentClass, UObject
 	return NewBP;
 }
 
+static bool IsReferencedByUndoBuffer(UBlueprint* Blueprint)
+{
+	UObject* BlueprintObj = Blueprint;
+	FReferencerInformationList ReferencesIncludingUndo;
+	IsReferenced(BlueprintObj, GARBAGE_COLLECTION_KEEPFLAGS, /*bCheckSubObjects =*/true, &ReferencesIncludingUndo);
+
+	FReferencerInformationList ReferencesExcludingUndo;
+	// Determine the in-memory references, *excluding* the undo buffer
+	GEditor->Trans->DisableObjectSerialization();
+	IsReferenced(BlueprintObj, GARBAGE_COLLECTION_KEEPFLAGS, /*bCheckSubObjects =*/true, &ReferencesExcludingUndo);
+	GEditor->Trans->EnableObjectSerialization();
+
+	// see if this object is the transaction buffer - set a flag so we know we need to clear the undo stack
+	const int32 TotalReferenceCount   = ReferencesIncludingUndo.ExternalReferences.Num() + ReferencesIncludingUndo.InternalReferences.Num();
+	const int32 NonUndoReferenceCount = ReferencesExcludingUndo.ExternalReferences.Num() + ReferencesExcludingUndo.InternalReferences.Num();
+
+	return (TotalReferenceCount > NonUndoReferenceCount);
+}
+
+static void UnloadBlueprint(UBlueprint* const Blueprint)
+{
+	UPackage* const OldPackage = Cast<UPackage>(Blueprint->GetOutermost());
+
+	UPackage* const TransientPackage = GetTransientPackage();
+	// move the blueprint to the transient package (to be picked up by garbage collection later)
+	FName UnloadedName = MakeUniqueObjectName(TransientPackage, UBlueprint::StaticClass(), Blueprint->GetFName());
+	Blueprint->Rename(*UnloadedName.ToString(), TransientPackage, REN_DontCreateRedirectors | REN_DoNotDirty);
+
+	// make sure the blueprint is properly trashed
+	Blueprint->SetFlags(RF_Transient);
+	Blueprint->ClearFlags(RF_Standalone | RF_RootSet | RF_Transactional);
+	Blueprint->RemoveFromRoot();
+	Blueprint->MarkPendingKill();
+
+	FKismetEditorUtilities::OnBlueprintUnloaded.Broadcast(Blueprint);
+
+	// reset the blueprint's original package/linker so that we can get by
+	// any early returns (in the load code), and reload its exports as if new 
+	ResetLoaders(OldPackage);
+	OldPackage->ClearFlags(RF_WasLoaded);
+	OldPackage->bHasBeenFullyLoaded = false;
+	OldPackage->GetMetaData()->RemoveMetaDataOutsidePackage();
+
+	if (IsReferencedByUndoBuffer(Blueprint))
+	{
+		GEditor->Trans->Reset(LOCTEXT("ReloadedBlueprint", "Reloaded Blueprint"));
+	}
+
+	// handled in FBlueprintEditor (from the OnBlueprintUnloaded event)
+// 	IAssetEditorInstance* EditorInst = FAssetEditorManager::Get().FindEditorForAsset(Blueprint, /*bFocusIfOpen =*/false);
+// 	if (EditorInst != nullptr)
+// 	{
+// 		EditorInst->CloseWindow();
+// 	}
+}
+
+UBlueprint* FKismetEditorUtilities::ReloadBlueprint(UBlueprint* StaleBlueprint)
+{
+	check(StaleBlueprint->IsAsset());
+
+	UClass* OldGeneratedClass = StaleBlueprint->GeneratedClass;
+	UObject* OldCDO = nullptr; // will be filled out below
+	UClass* OldSkeletonClass = StaleBlueprint->SkeletonGeneratedClass;
+	UObject* OldSkelCDO = nullptr; // will be filled out below
+
+	//--------------------------------------
+	// Collect objects that need replacing
+	//--------------------------------------
+
+	TArray<UObject*> StaleObjects;
+	if (OldGeneratedClass != nullptr)
+	{
+		StaleObjects.Add(OldGeneratedClass);
+
+		OldCDO = OldGeneratedClass->GetDefaultObject(/*bCreateIfNeeded =*/false);
+		if (OldCDO != nullptr)
+		{
+			StaleObjects.Add(OldCDO);
+		}		
+	}
+	if (OldSkeletonClass != nullptr)
+	{
+		StaleObjects.Add(OldSkeletonClass);
+
+		OldSkelCDO = OldSkeletonClass->GetDefaultObject(/*bCreateIfNeeded =*/false);
+		if (OldSkelCDO != nullptr)
+		{
+			StaleObjects.Add(OldSkelCDO);
+		}
+	}
+	StaleObjects.Add(StaleBlueprint);
+
+	//--------------------------------------
+	// Unload/Load Blueprint
+	//--------------------------------------
+
+	FStringAssetReference BlueprintAssetRef(StaleBlueprint);
+	UnloadBlueprint(StaleBlueprint);
+	
+	UBlueprint* ReloadedBlueprint = Cast<UBlueprint>(StaticLoadObject(UBlueprint::StaticClass(), /*Outer =*/nullptr, *BlueprintAssetRef.AssetLongPathname));
+
+	//--------------------------------------
+	// Construct redirects
+	//--------------------------------------
+
+	TMap<UObject*, UObject*> Redirects;
+	if (OldGeneratedClass != nullptr)
+	{
+		Redirects.Add(OldGeneratedClass, ReloadedBlueprint->GeneratedClass);
+	}
+	if (OldCDO != nullptr)
+	{
+		Redirects.Add(OldCDO, ReloadedBlueprint->GeneratedClass->GetDefaultObject());
+	}
+	if (OldSkeletonClass != nullptr)
+	{
+		Redirects.Add(OldSkeletonClass, ReloadedBlueprint->SkeletonGeneratedClass);
+	}
+	if (OldSkelCDO != nullptr)
+	{
+		Redirects.Add(OldSkelCDO, ReloadedBlueprint->SkeletonGeneratedClass->GetDefaultObject());
+	}
+	Redirects.Add(StaleBlueprint, ReloadedBlueprint);
+
+	//--------------------------------------
+	// Replace old references
+	//--------------------------------------
+
+	TArray<UObject*> Referencers;
+	// find all objects, still referencing the old blueprint/class/cdo/etc.
+	for (auto Referencer : TFindObjectReferencers<UObject>(StaleObjects, /*PackageToCheck =*/nullptr, /*bIgnoreTemplates =*/false))
+	{
+		Referencers.Add(Referencer.Value);
+	}
+
+	StaleBlueprint->SetObjectBeingDebugged(nullptr);
+	FBlueprintCompileReinstancer::ReplaceInstancesOfClass(OldGeneratedClass, ReloadedBlueprint->GeneratedClass);
+
+	for (UObject* Referencer : Referencers)
+	{
+		FArchiveReplaceObjectRef<UObject>(Referencer, Redirects, /*bNullPrivateRefs=*/false, /*bIgnoreOuterRef=*/false, /*bIgnoreArchetypeRef=*/false);
+	}
+	
+	return ReloadedBlueprint;
+}
+
 void FKismetEditorUtilities::CompileBlueprint(UBlueprint* BlueprintObj, bool bIsRegeneratingOnLoad, bool bSkipGarbageCollection, bool bSaveIntermediateProducts, FCompilerResultsLog* pResults)
 {
 	// Broadcast pre-compile
diff --git a/Engine/Source/Editor/UnrealEd/Private/PlayLevel.cpp b/Engine/Source/Editor/UnrealEd/Private/PlayLevel.cpp
--- a/Engine/Source/Editor/UnrealEd/Private/PlayLevel.cpp
+++ b/Engine/Source/Editor/UnrealEd/Private/PlayLevel.cpp
@@ -1835,7 +1835,8 @@ void UEditorEngine::PlayInEditor( UWorld* InWorld, bool bInSimulateInEditor )
 				const bool bIsDirtyAndShouldBeRecompiled = Blueprint->IsPossiblyDirty();
 				if( !FBlueprintEditorUtils::IsDataOnlyBlueprint(Blueprint) 
 					&& (bIsDirtyAndShouldBeRecompiled || FBlueprintEditorUtils::IsLevelScriptBlueprint(Blueprint))
-					&& (Blueprint->Status != BS_Unknown) )
+					&& (Blueprint->Status != BS_Unknown)
+					&& !Blueprint->IsPendingKill() )
 				{
 					BlueprintsToRecompile.Add(Blueprint);
 				}
diff --git a/Engine/Source/Editor/UnrealEd/Public/Kismet2/KismetEditorUtilities.h b/Engine/Source/Editor/UnrealEd/Public/Kismet2/KismetEditorUtilities.h
--- a/Engine/Source/Editor/UnrealEd/Public/Kismet2/KismetEditorUtilities.h
+++ b/Engine/Source/Editor/UnrealEd/Public/Kismet2/KismetEditorUtilities.h
@@ -30,6 +30,22 @@ public:
 	 */
 	static UBlueprint* CreateBlueprint(UClass* ParentClass, UObject* Outer, const FName NewBPName, enum EBlueprintType BlueprintType, TSubclassOf<UBlueprint> BlueprintClassType, TSubclassOf<UBlueprintGeneratedClass> BlueprintGeneratedClassType, FName CallingContext = NAME_None);
 
+	/** 
+	 * Event that's broadcast anytime a blueprint is unloaded, and becomes 
+	 * invalid (with calls to ReloadBlueprint(), for example).
+	 */
+	DECLARE_EVENT_OneParam(FKismetEditorUtilities, FOnBlueprintUnloaded, UBlueprint*);
+	static FOnBlueprintUnloaded OnBlueprintUnloaded;
+
+	/** 
+	 * Unloads the supplied Blueprint (marking it pending-kill, and removing it 
+	 * from its outer package). Then proceeds to reload from disk.
+
+	 * @param  TargetBlueprint	The Blueprint you want to unload and replace.
+	 * @return The freshly loaded Blueprint (replacing the, now invalid, input).
+	 */
+	static UBlueprint* ReloadBlueprint(UBlueprint* TargetBlueprint);
+
 	/** Create the correct event graphs for this blueprint */
 	static void CreateDefaultEventGraphs(UBlueprint* Blueprint);
 
