Fix crash on retargeting of edited animations

diff --git a/Engine/Source/Editor/UnrealEd/Private/EditorAnimUtils.cpp b/Engine/Source/Editor/UnrealEd/Private/EditorAnimUtils.cpp
--- a/Engine/Source/Editor/UnrealEd/Private/EditorAnimUtils.cpp
+++ b/Engine/Source/Editor/UnrealEd/Private/EditorAnimUtils.cpp
@@ -178,6 +178,9 @@ namespace EditorAnimUtils
 			if (OldSkeleton)
 			{
 				EditorAnimUtils::CopyAnimCurves(OldSkeleton, NewSkeleton, AssetToRetarget, USkeleton::AnimCurveMappingName, FRawCurveTracks::FloatType);
+
+				// clear transform curves since those curves won't work in new skeleton
+				AssetToRetarget->RawCurveData.TransformCurves.Empty();
 				// I can't copy transform curves yet because transform curves need retargeting. 
 				//EditorAnimUtils::CopyAnimCurves(OldSkeleton, NewSkeleton, AssetToRetarget, USkeleton::AnimTrackCurveMappingName, FRawCurveTracks::TransformType);
 			}	
