Remove unnecessary include from middle of a function.

diff --git a/Engine/Source/Editor/UnrealEd/Private/Factories/ApexDestructibleAssetImport.cpp b/Engine/Source/Editor/UnrealEd/Private/Factories/ApexDestructibleAssetImport.cpp
--- a/Engine/Source/Editor/UnrealEd/Private/Factories/ApexDestructibleAssetImport.cpp
+++ b/Engine/Source/Editor/UnrealEd/Private/Factories/ApexDestructibleAssetImport.cpp
@@ -850,8 +850,6 @@ bool SetApexDestructibleAsset(UDestructibleMesh& DestructibleMesh, NxDestructibl
 		TArray<int32> LODPointToRawMap;
 		SkelMeshImportDataPtr->CopyLODImportData(LODPoints,LODWedges,LODFaces,LODInfluences,LODPointToRawMap);
 
-		#include "Developer/MeshUtilities/Public/MeshUtilities.h"
-
 		IMeshUtilities& MeshUtilities = FModuleManager::Get().LoadModuleChecked<IMeshUtilities>("MeshUtilities");
 
 		// Create actual rendering data.
