Change default PVRTC quality to 1 ('fast')

Fixes UE-3498 -  IOS: TEMPLATE: Odd block shadows

diff --git a/Engine/Source/Editor/UnrealEd/Private/Commandlets/CookCommandlet.cpp b/Engine/Source/Editor/UnrealEd/Private/Commandlets/CookCommandlet.cpp
--- a/Engine/Source/Editor/UnrealEd/Private/Commandlets/CookCommandlet.cpp
+++ b/Engine/Source/Editor/UnrealEd/Private/Commandlets/CookCommandlet.cpp
@@ -30,7 +30,7 @@ DEFINE_LOG_CATEGORY_STATIC(LogCookCommandlet, Log, All);
 UCookerSettings::UCookerSettings(const FObjectInitializer& ObjectInitializer)
 	: Super(ObjectInitializer)
 {
-	DefaultPVRTCQuality = 0;
+	DefaultPVRTCQuality = 1;
 }
 
 
