Portable object(.po) entries will now properly take the case of the ID into account when comparing the equality of two entries.

OPP-1473 Sign out button is not localized for users in Korean or Japanese

diff --git a/Engine/Source/Editor/UnrealEd/Classes/Commandlets/InternationalizationExportCommandlet.h b/Engine/Source/Editor/UnrealEd/Classes/Commandlets/InternationalizationExportCommandlet.h
--- a/Engine/Source/Editor/UnrealEd/Classes/Commandlets/InternationalizationExportCommandlet.h
+++ b/Engine/Source/Editor/UnrealEd/Classes/Commandlets/InternationalizationExportCommandlet.h
@@ -138,8 +138,9 @@ public:
 	
 	bool operator==( const FPortableObjectEntry& Other ) 
 	{ 
-		//@TODO: Reference materials are not fully clear if the MsgIdPlural member needs to match for equality.  This seems logical so we do it here.
-		return ( MsgId == Other.MsgId && MsgIdPlural == Other.MsgIdPlural && MsgCtxt == Other.MsgCtxt ); 
+		return ( MsgId.Equals(Other.MsgId, ESearchCase::CaseSensitive) &&
+				MsgIdPlural.Equals(Other.MsgIdPlural, ESearchCase::CaseSensitive) &&
+				MsgCtxt.Equals(Other.MsgCtxt, ESearchCase::IgnoreCase) ); 
 	}
 	
 	/**
