Made changes to the way preview actors are applied so that materials in multiple selections are not applied immediately.

diff --git a/Engine/Source/Editor/UnrealEd/Private/LevelEditorViewport.cpp b/Engine/Source/Editor/UnrealEd/Private/LevelEditorViewport.cpp
--- a/Engine/Source/Editor/UnrealEd/Private/LevelEditorViewport.cpp
+++ b/Engine/Source/Editor/UnrealEd/Private/LevelEditorViewport.cpp
@@ -495,10 +495,12 @@ static UObject* GetOrCreateMaterialFromTexture( UTexture* UnrealTexture )
  * @param	ObjToUse				Object to attempt to apply as specific asset
  * @param	ActorToApplyTo			Actor to whom the asset should be applied
  * @param   TargetMaterialSlot      When dealing with submeshes this will represent the target section/slot to apply materials to.
+ * @param	bTest					Whether to test if the object would be successfully applied without actually doing it.
+ * @param	bCreateDropPreview		Whether this is just a drop preview.
  *
  * @return	true if the provided object was successfully applied to the provided actor
  */
-static bool AttemptApplyObjToActor( UObject* ObjToUse, AActor* ActorToApplyTo, int32 TargetMaterialSlot = -1, bool bTest = false )
+static bool AttemptApplyObjToActor( UObject* ObjToUse, AActor* ActorToApplyTo, int32 TargetMaterialSlot = -1, bool bTest = false, bool bCreateDropPreview = false )
 {
 	bool bResult = false;
 
@@ -507,7 +509,7 @@ static bool AttemptApplyObjToActor( UObject* ObjToUse, AActor* ActorToApplyTo, i
 		UTexture* DroppedObjAsTexture = Cast<UTexture>( ObjToUse );
 		if ( DroppedObjAsTexture != NULL )
 		{
-			if ( bTest )
+			if ( bTest || bCreateDropPreview )
 			{
 				bResult = true;
 			}
@@ -521,7 +523,7 @@ static bool AttemptApplyObjToActor( UObject* ObjToUse, AActor* ActorToApplyTo, i
 		UMaterialInterface* DroppedObjAsMaterial = Cast<UMaterialInterface>( ObjToUse );
 		if ( DroppedObjAsMaterial )
 		{
-			if (bTest)
+			if (bTest || bCreateDropPreview)
 			{
 				bResult = true;
 			}
@@ -857,7 +859,8 @@ bool FLevelEditorViewportClient::DropObjectsOnActor(FViewportCursorLocation& Cur
 
 	for ( auto DroppedObject : DroppedObjects )
 	{
-		const bool bAppliedToActor = ( FactoryToUse == NULL ) ? AttemptApplyObjToActor( DroppedObject, DroppedUponActor, DroppedUponSlot ) : false;
+		const bool bTest = false;
+		const bool bAppliedToActor = ( FactoryToUse == NULL ) ? AttemptApplyObjToActor( DroppedObject, DroppedUponActor, DroppedUponSlot, bTest, bCreateDropPreview ) : false;
 
 		if (!bAppliedToActor)
 		{
@@ -891,12 +894,6 @@ bool FLevelEditorViewportClient::DropObjectsOnBSPSurface(FSceneView* View, FView
 		return false;
 	}
 
-	// Attempt to apply the dropped asset as a material to the BSP surface
-	if( ( FactoryToUse == NULL ) && AttemptApplyObjAsMaterialToSurface( DroppedObjects[ 0 ], TargetProxy, Cursor ) )
-	{
-		return true;
-	}
-
 	bool bSuccess = false;
 
 	const bool bTransacted = !bCreateDropPreview && !AreAllDroppedObjectsBrushBuilders(DroppedObjects);
@@ -909,12 +906,21 @@ bool FLevelEditorViewportClient::DropObjectsOnBSPSurface(FSceneView* View, FView
 
 	for (auto DroppedObject : DroppedObjects)
 	{
-		// Attempt to create actors from the dropped object
-		TArray<AActor*> NewActors = AttemptDropObjAsActors(GetWorld()->GetCurrentLevel(), DroppedObject, Cursor, bSelectActors, ObjectFlags, FactoryToUse);
+		const bool bAppliedToActor = (!bCreateDropPreview && FactoryToUse == NULL) ? AttemptApplyObjAsMaterialToSurface(DroppedObject, TargetProxy, Cursor) : false;
 
-		if (NewActors.Num() > 0)
+		if (!bAppliedToActor)
+		{
+			// Attempt to create actors from the dropped object
+			TArray<AActor*> NewActors = AttemptDropObjAsActors(GetWorld()->GetCurrentLevel(), DroppedObject, Cursor, bSelectActors, ObjectFlags, FactoryToUse);
+
+			if (NewActors.Num() > 0)
+			{
+				OutNewActors.Append(NewActors);
+				bSuccess = true;
+			}
+		}
+		else
 		{
-			OutNewActors.Append(NewActors);
 			bSuccess = true;
 		}
 	}
@@ -1163,15 +1169,6 @@ bool FLevelEditorViewportClient::DropObjectsAtCoordinates(int32 MouseX, int32 Mo
 	// Make sure the placement dragging actor is cleaned up.
 	DestroyDropPreviewActors();
 
-	if ( bCreateDropPreview && DroppedObjects.Num() > 0 && FactoryToUse == NULL )
-	{
-		if ( Cast<UMaterialInterface>( DroppedObjects[0] ) || Cast<UTexture>( DroppedObjects[0] ) )
-		{
-			// Do not create a drop preview for material interfaces or textures
-			return false;
-		}
-	}
-
 	if(DroppedObjects.Num() > 0)
 	{
 		Viewport->InvalidateHitProxy();
