Initialize billboard scale in ULevelEditorViewportSettings::PostEditChangeProperty, to reflect the config value just loaded.

diff --git a/Engine/Source/Editor/UnrealEd/Classes/Settings/LevelEditorViewportSettings.h b/Engine/Source/Editor/UnrealEd/Classes/Settings/LevelEditorViewportSettings.h
--- a/Engine/Source/Editor/UnrealEd/Classes/Settings/LevelEditorViewportSettings.h
+++ b/Engine/Source/Editor/UnrealEd/Classes/Settings/LevelEditorViewportSettings.h
@@ -476,6 +476,7 @@ protected:
 
 	// UObject overrides
 
+	virtual void PostInitProperties() override;
 	virtual void PostEditChangeProperty( struct FPropertyChangedEvent& PropertyChangedEvent ) override;
 
 private:
diff --git a/Engine/Source/Editor/UnrealEd/Private/Settings/SettingsClasses.cpp b/Engine/Source/Editor/UnrealEd/Private/Settings/SettingsClasses.cpp
--- a/Engine/Source/Editor/UnrealEd/Private/Settings/SettingsClasses.cpp
+++ b/Engine/Source/Editor/UnrealEd/Private/Settings/SettingsClasses.cpp
@@ -193,6 +193,12 @@ ULevelEditorViewportSettings::ULevelEditorViewportSettings( const FObjectInitial
 	PreviewMeshes.Add(FStringAssetReference("/Engine/EditorMeshes/ColorCalibrator/SM_ColorCalibrator.SM_ColorCalibrator"));
 }
 
+void ULevelEditorViewportSettings::PostInitProperties()
+{
+	Super::PostInitProperties();
+	UBillboardComponent::SetEditorScale(BillboardScale);
+	UArrowComponent::SetEditorScale(BillboardScale);
+}
 
 void ULevelEditorViewportSettings::PostEditChangeProperty( struct FPropertyChangedEvent& PropertyChangedEvent )
 {
diff --git a/Engine/Source/Runtime/CoreUObject/Public/UObject/UObject.h b/Engine/Source/Runtime/CoreUObject/Public/UObject/UObject.h
--- a/Engine/Source/Runtime/CoreUObject/Public/UObject/UObject.h
+++ b/Engine/Source/Runtime/CoreUObject/Public/UObject/UObject.h
@@ -136,7 +136,7 @@ protected:
 	virtual FString GetDetailedInfoInternal() const { return TEXT("No_Detailed_Info_Specified"); }
 public:
 	/**
-	 * Called after the C++ constructor and after the properties have been initialized, but before the config has been loaded, etc.
+	 * Called after the C++ constructor and after the properties have been initialized, including those loaded from config.
 	 * mainly this is to emulate some behavior of when the constructor was called after the properties were intialized.
 	 */
 	virtual void PostInitProperties();
