Was supposed to be part of CL 2374534 - P4 moved the file out of the original CL when resolving.

diff --git a/Engine/Source/Editor/UnrealEd/Private/UnrealEdEngine.cpp b/Engine/Source/Editor/UnrealEd/Private/UnrealEdEngine.cpp
--- a/Engine/Source/Editor/UnrealEd/Private/UnrealEdEngine.cpp
+++ b/Engine/Source/Editor/UnrealEd/Private/UnrealEdEngine.cpp
@@ -900,7 +900,7 @@ void UUnrealEdEngine::GetSortedVolumeClasses( TArray< UClass* >* VolumeClasses )
 	// Add all of the volume classes to the passed in array and then sort it
 	for( auto* Class : TObjectRange<UClass>() )
 	{
-		if (Class->IsChildOf(AVolume::StaticClass()) && !Class->HasAnyClassFlags(CLASS_Deprecated | CLASS_Abstract | CLASS_NotPlaceable))
+		if (Class->IsChildOf(AVolume::StaticClass()) && !Class->HasAnyClassFlags(CLASS_Deprecated | CLASS_Abstract | CLASS_NotPlaceable) && Class->ClassGeneratedBy == nullptr)
 		{
 			VolumeClasses->AddUnique( Class );
 		}
