Fix empty files not being marked for delete when not already checked out in resave packages and fixup redirects commandlets
Minor ResavePackages cleanup

diff --git a/Engine/Source/Editor/UnrealEd/Classes/Commandlets/ResavePackagesCommandlet.h b/Engine/Source/Editor/UnrealEd/Classes/Commandlets/ResavePackagesCommandlet.h
--- a/Engine/Source/Editor/UnrealEd/Classes/Commandlets/ResavePackagesCommandlet.h
+++ b/Engine/Source/Editor/UnrealEd/Classes/Commandlets/ResavePackagesCommandlet.h
@@ -40,10 +40,6 @@ class UNREALED_API UResavePackagesCommandlet : public UCommandlet
 	 */
 	int32 MaxPackagesToResave;
 
-	/** if true, only save packages that require sound conversion */
-	bool bSoundConversionOnly;
-	bool bSoundWasDirty;
-
 	/** allows users to save only packages with a particular class in them (useful for fixing content) */
 	TArray<FName> ResaveClasses;
 
diff --git a/Engine/Source/Editor/UnrealEd/Private/Commandlets/ContentCommandlets.cpp b/Engine/Source/Editor/UnrealEd/Private/Commandlets/ContentCommandlets.cpp
--- a/Engine/Source/Editor/UnrealEd/Private/Commandlets/ContentCommandlets.cpp
+++ b/Engine/Source/Editor/UnrealEd/Private/Commandlets/ContentCommandlets.cpp
@@ -168,11 +168,6 @@ int32 UResavePackagesCommandlet::InitializeResaveParameters( const TArray<FStrin
 		}
 	}
 
-	if ( Switches.Contains(TEXT("SOUNDCONVERSIONONLY")) )
-	{
-		bSoundConversionOnly = true;
-	}
-
 	FString ClassList;
 	for ( int32 SwitchIdx = 0; SwitchIdx < Switches.Num(); SwitchIdx++ )
 	{
@@ -309,7 +304,6 @@ void UResavePackagesCommandlet::LoadAndSaveOnePackage(const FString& Filename)
 			VerboseMessage(TEXT("Post PerformAdditionalOperations"));
 
 			// Check for any special per object operations
-			bSoundWasDirty = false;
 			for( FObjectIterator ObjectIt; ObjectIt; ++ObjectIt )
 			{
 				if( ObjectIt->IsIn( Package ) )
@@ -320,14 +314,6 @@ void UResavePackagesCommandlet::LoadAndSaveOnePackage(const FString& Filename)
 			
 			VerboseMessage(TEXT("Post PerformAdditionalOperations Loop"));
 
-			if (bSoundConversionOnly == true)
-			{
-				if (bSoundWasDirty == false)
-				{
-					bSavePackage = false;
-				}
-			}
-
 			if (bStripEditorOnlyContent)
 			{
 				UE_LOG(LogContentCommandlet, Log, TEXT("Removing editor only data"));
@@ -341,14 +327,13 @@ void UResavePackagesCommandlet::LoadAndSaveOnePackage(const FString& Filename)
 
 				TArray<UObject *> ObjectsInOuter;
 				GetObjectsWithOuter(Package, ObjectsInOuter);
-				for (int32 Index = 0; Index < ObjectsInOuter.Num(); Index++)
+				for (UObject* Obj : ObjectsInOuter)
 				{
-					UObject* Obj = ObjectsInOuter[Index];
-
 					if (!Obj->IsA(UMetaData::StaticClass()))
 					{
 						// This package has a real object
 						bIsEmpty = false;
+						break;
 					}
 				}
 
@@ -378,7 +363,7 @@ void UResavePackagesCommandlet::LoadAndSaveOnePackage(const FString& Filename)
 						UE_LOG(LogContentCommandlet, Display, TEXT("Deleting '%s' from source control..."), *Filename);
 						SourceControlProvider.Execute(ISourceControlOperation::Create<FDelete>(), PackageFilename);
 					}
-					else if (SourceControlState.IsValid() && (SourceControlState->CanEdit() || !SourceControlState->IsCurrent()))
+					else if (SourceControlState.IsValid() && SourceControlState->CanCheckout())
 					{
 						UE_LOG(LogContentCommandlet, Display, TEXT("Deleting '%s' from source control..."), *Filename);
 						SourceControlProvider.Execute(ISourceControlOperation::Create<FDelete>(), PackageFilename);
diff --git a/Engine/Source/Editor/UnrealEd/Private/Commandlets/FixupRedirectsCommandlet.cpp b/Engine/Source/Editor/UnrealEd/Private/Commandlets/FixupRedirectsCommandlet.cpp
--- a/Engine/Source/Editor/UnrealEd/Private/Commandlets/FixupRedirectsCommandlet.cpp
+++ b/Engine/Source/Editor/UnrealEd/Private/Commandlets/FixupRedirectsCommandlet.cpp
@@ -696,7 +696,7 @@ int32 UFixupRedirectsCommandlet::Main( const FString& Params )
 					CLEAR_WARN_COLOR();
 					SourceControlProvider.Execute(ISourceControlOperation::Create<FDelete>(), FileName);
 				}
-				else if(SourceControlState.IsValid() && (SourceControlState->CanEdit() || !SourceControlState->IsCurrent()))
+				else if(SourceControlState.IsValid() && SourceControlState->CanCheckout())
 				{
 					SET_WARN_COLOR(COLOR_GREEN);
 					UE_LOG(LogFixupRedirectsCommandlet, Warning, TEXT("Deleting '%s' from source control..."), *Filename);
