UE4: Checking in Maciej Mroz's fix for old event nodes. Support for old event nodes, wihout GUID.

diff --git a/Engine/Source/Editor/UnrealEd/Private/Kismet2/BlueprintEditorUtils.cpp b/Engine/Source/Editor/UnrealEd/Private/Kismet2/BlueprintEditorUtils.cpp
--- a/Engine/Source/Editor/UnrealEd/Private/Kismet2/BlueprintEditorUtils.cpp
+++ b/Engine/Source/Editor/UnrealEd/Private/Kismet2/BlueprintEditorUtils.cpp
@@ -6769,41 +6769,46 @@ FName FBlueprintEditorUtils::GetFunctionNameFromClassByGuid(const UClass* InClas
 
 bool FBlueprintEditorUtils::GetFunctionGuidFromClassByFieldName(const UClass* InClass, const FName FunctionName, FGuid& FunctionGuid)
 {
-	TArray<UBlueprint*> Blueprints;
-	UBlueprint::GetBlueprintHierarchyFromClass(InClass, Blueprints);
-
-	for (int32 BPIndex = 0; BPIndex < Blueprints.Num(); ++BPIndex)
+	if (FunctionName != NAME_None)
 	{
-		UBlueprint* Blueprint = Blueprints[BPIndex];
-		for (int32 FunctionIndex = 0; FunctionIndex < Blueprint->FunctionGraphs.Num(); ++FunctionIndex)
+		TArray<UBlueprint*> Blueprints;
+		UBlueprint::GetBlueprintHierarchyFromClass(InClass, Blueprints);
+
+		for (int32 BPIndex = 0; BPIndex < Blueprints.Num(); ++BPIndex)
 		{
-			UEdGraph* FunctionGraph = Blueprint->FunctionGraphs[FunctionIndex];
-			if (FunctionGraph && FunctionGraph->GetFName() == FunctionName)
+			UBlueprint* Blueprint = Blueprints[BPIndex];
+			for (int32 FunctionIndex = 0; FunctionIndex < Blueprint->FunctionGraphs.Num(); ++FunctionIndex)
 			{
-				FunctionGuid = FunctionGraph->GraphGuid;
-				return true;
+				UEdGraph* FunctionGraph = Blueprint->FunctionGraphs[FunctionIndex];
+				if (FunctionGraph && FunctionGraph->GetFName() == FunctionName)
+				{
+					FunctionGuid = FunctionGraph->GraphGuid;
+					return true;
+				}
 			}
-		}
 
-		TArray<UEdGraph*> UberGraphs;
-		FUberGraphHelper::GetAll(Blueprint, UberGraphs);
-		for (const auto UberGraph : UberGraphs)
-		{
-			TArray<UK2Node_CustomEvent*> CustomEvents;
-			UberGraph->GetNodesOfClass(CustomEvents);
-			for (const auto CustomEvent : CustomEvents)
+			TArray<UEdGraph*> UberGraphs;
+			FUberGraphHelper::GetAll(Blueprint, UberGraphs);
+			for (const auto UberGraph : UberGraphs)
 			{
-				if (!CustomEvent->bOverrideFunction && (CustomEvent->CustomFunctionName == FunctionName))
+				TArray<UK2Node_CustomEvent*> CustomEvents;
+				UberGraph->GetNodesOfClass(CustomEvents);
+				for (const auto CustomEvent : CustomEvents)
 				{
-					ensure(CustomEvent->CustomFunctionName != NAME_None);
-					ensure(CustomEvent->NodeGuid.IsValid());
-					FunctionGuid = CustomEvent->NodeGuid;
-					return true;
+					if (!CustomEvent->bOverrideFunction
+						&& (CustomEvent->CustomFunctionName == FunctionName)
+						&& CustomEvent->NodeGuid.IsValid())
+					{
+						FunctionGuid = CustomEvent->NodeGuid;
+						return true;
+					}
 				}
 			}
 		}
 	}
 
+	FunctionGuid.Invalidate();
+
 	return false;
 }
 
