Added simplification details to LOD maps generation. Include options for roughness/specular map generation

diff --git a/Engine/Source/Editor/UnrealEd/Private/EditorExporters.cpp b/Engine/Source/Editor/UnrealEd/Private/EditorExporters.cpp
--- a/Engine/Source/Editor/UnrealEd/Private/EditorExporters.cpp
+++ b/Engine/Source/Editor/UnrealEd/Private/EditorExporters.cpp
@@ -2691,6 +2691,76 @@ namespace MaterialExportUtils
 		OutFlattenMaterial.MaterialId = InMaterial->GetLightingGuid();
 		return true;
 	}
+
+	void RenderSceneToTexture(
+		FSceneInterface* Scene,
+		const FName& VisualizationMode, 
+		const FMatrix& ViewMatrix, 
+		const FMatrix& ProjectionMatrix,  
+		const TSet<FPrimitiveComponentId>& HiddenPrimitives, 
+		FIntPoint TargetSize, 
+		TArray<FColor>& OutSamples)
+	{
+		UTextureRenderTarget2D* RenderTargetTexture = new UTextureRenderTarget2D(FObjectInitializer());
+		check(RenderTargetTexture);
+		RenderTargetTexture->AddToRoot();
+		RenderTargetTexture->ClearColor = FLinearColor(0.0f, 0.0f, 0.0f, 0.0f);
+		RenderTargetTexture->InitCustomFormat(TargetSize.X, TargetSize.Y, PF_B8G8R8A8, true);
+		FTextureRenderTargetResource* RenderTargetResource = RenderTargetTexture->GameThread_GetRenderTargetResource();
+
+		FSceneViewFamilyContext ViewFamily(
+			FSceneViewFamily::ConstructionValues(RenderTargetResource, Scene, FEngineShowFlags(ESFIM_Game))
+				.SetWorldTimes(FApp::GetCurrentTime() - GStartTime, FApp::GetDeltaTime(), FApp::GetCurrentTime() - GStartTime)
+			);
+
+		ViewFamily.EngineShowFlags.DisableAdvancedFeatures();
+		ViewFamily.EngineShowFlags.MotionBlur = 0;
+		ViewFamily.EngineShowFlags.Lighting = 0;
+		ViewFamily.EngineShowFlags.PostProcessing = 0;
+		ViewFamily.EngineShowFlags.LightFunctions = 0;
+		ViewFamily.EngineShowFlags.DynamicShadows = 0;
+		ViewFamily.EngineShowFlags.Atmosphere = 0;
+		// To enable visualization mode
+		ViewFamily.EngineShowFlags.PostProcessing = 1;
+		ViewFamily.EngineShowFlags.PostProcessMaterial = 1;
+		ViewFamily.EngineShowFlags.VisualizeBuffer = 1; 
+
+		FSceneViewInitOptions ViewInitOptions;
+		ViewInitOptions.SetViewRectangle(FIntRect(0, 0, TargetSize.X, TargetSize.Y));
+		ViewInitOptions.ViewFamily = &ViewFamily;
+		ViewInitOptions.HiddenPrimitives = HiddenPrimitives;
+		ViewInitOptions.ViewMatrix = ViewMatrix;
+		ViewInitOptions.ProjectionMatrix = ProjectionMatrix;
+		
+		FSceneView* NewView = new FSceneView(ViewInitOptions);
+		NewView->CurrentBufferVisualizationMode = VisualizationMode;
+		ViewFamily.Views.Add(NewView);
+					
+		FCanvas Canvas(RenderTargetResource, NULL, FApp::GetCurrentTime() - GStartTime, FApp::GetDeltaTime(), FApp::GetCurrentTime() - GStartTime, Scene->GetFeatureLevel());
+		Canvas.Clear(FLinearColor::Black);
+		GetRendererModule().BeginRenderingViewFamily(&Canvas, &ViewFamily);
+
+		ENQUEUE_UNIQUE_RENDER_COMMAND_ONEPARAMETER(
+			UpdateThumbnailRTCommand,
+			FTextureRenderTargetResource*, RenderTargetResource, RenderTargetResource,
+		{
+			// Copy (resolve) the rendered thumbnail from the render target to its texture
+			RHICmdList.CopyToResolveTarget(
+				RenderTargetResource->GetRenderTargetTexture(),		// Source texture
+				RenderTargetResource->TextureRHI,					// Dest texture
+				false,												// Do we need the source image content again?
+				FResolveParams() );									// Resolve parameters
+		});
+
+		// Copy the contents of the remote texture to system memory
+		// NOTE: OutRawImageData must be a preallocated buffer!
+		OutSamples.Init(TargetSize.X*TargetSize.Y);
+		RenderTargetResource->ReadPixelsPtr(OutSamples.GetData(), FReadSurfaceDataFlags(), FIntRect(0, 0, TargetSize.X, TargetSize.Y));
+		FlushRenderingCommands();
+					
+		RenderTargetTexture->RemoveFromRoot();
+		RenderTargetTexture = nullptr;
+	}
 	
 	bool ExportMaterial(ALandscapeProxy* InLandscape, FFlattenMaterial& OutFlattenMaterial)
 	{
@@ -2716,180 +2786,86 @@ namespace MaterialExportUtils
 				PrimitivesToHide.Add(PrimitiveComp->SceneProxy->GetPrimitiveComponentId());
 			}
 		}
-		
-		FSceneInterface* Scene = InLandscape->GetWorld()->Scene;
-		{
-			// Render diffuse texture
-			if (OutFlattenMaterial.DiffuseSize.X > 0 && 
-				OutFlattenMaterial.DiffuseSize.Y > 0)
-			{
-				UTextureRenderTarget2D* RenderTargetTexture = new UTextureRenderTarget2D(FObjectInitializer());
-				check(RenderTargetTexture);
-				RenderTargetTexture->AddToRoot();
-				RenderTargetTexture->ClearColor = FLinearColor(0.0f, 0.0f, 0.0f, 0.0f);
-				RenderTargetTexture->InitCustomFormat(
-					OutFlattenMaterial.DiffuseSize.X, 
-					OutFlattenMaterial.DiffuseSize.Y, PF_B8G8R8A8, true);
-				FTextureRenderTargetResource* RenderTargetResource = RenderTargetTexture->GameThread_GetRenderTargetResource();
-			
-				FSceneViewFamilyContext ViewFamily(
-					FSceneViewFamily::ConstructionValues(RenderTargetResource, Scene, FEngineShowFlags(ESFIM_Game))
-						.SetWorldTimes(FApp::GetCurrentTime() - GStartTime, FApp::GetDeltaTime(), FApp::GetCurrentTime() - GStartTime)
-					);
-
-				ViewFamily.EngineShowFlags.DisableAdvancedFeatures();
-				ViewFamily.EngineShowFlags.MotionBlur = 0;
-				ViewFamily.EngineShowFlags.Lighting = 0;
-				ViewFamily.EngineShowFlags.PostProcessing = 0;
-				ViewFamily.EngineShowFlags.LightFunctions = 0;
-				ViewFamily.EngineShowFlags.DynamicShadows = 0;
-				ViewFamily.EngineShowFlags.Atmosphere = 0;
-
-				FSceneViewInitOptions ViewInitOptions;
-				ViewInitOptions.SetViewRectangle(FIntRect(0, 0, OutFlattenMaterial.DiffuseSize.X, OutFlattenMaterial.DiffuseSize.Y));
-				ViewInitOptions.ViewFamily = &ViewFamily;
-				ViewInitOptions.HiddenPrimitives = PrimitivesToHide;
-
-				ViewInitOptions.ViewMatrix = FTranslationMatrix(-LandscapeCenter);
-				ViewInitOptions.ViewMatrix*= FInverseRotationMatrix(InLandscape->GetActorRotation());
-				ViewInitOptions.ViewMatrix*= FMatrix(	FPlane(1,	0,	0,	0),
-														FPlane(0,	-1,	0,	0),
-														FPlane(0,	0,	-1,	0),
-														FPlane(0,	0,	0,	1));
+
+		FMatrix ViewMatrix = FTranslationMatrix(-LandscapeCenter);
+		ViewMatrix*= FInverseRotationMatrix(InLandscape->GetActorRotation());
+		ViewMatrix*= FMatrix(FPlane(1,	0,	0,	0),
+							 FPlane(0,	-1,	0,	0),
+							 FPlane(0,	0,	-1,	0),
+							 FPlane(0,	0,	0,	1));
 				
-				const float ZOffset = WORLD_MAX;
-				ViewInitOptions.ProjectionMatrix =  FReversedZOrthoMatrix(
-					LandscapeExtent.X,
-					LandscapeExtent.Y,
-					0.5f / ZOffset,
-					ZOffset);
-
-				ViewFamily.Views.Add(new FSceneView(ViewInitOptions));
-					
-				FCanvas Canvas(RenderTargetResource, NULL, FApp::GetCurrentTime() - GStartTime, FApp::GetDeltaTime(), FApp::GetCurrentTime() - GStartTime, Scene->GetFeatureLevel());
-				Canvas.Clear(FLinearColor::Black);
-				GetRendererModule().BeginRenderingViewFamily(&Canvas, &ViewFamily);
+		const float ZOffset = WORLD_MAX;
+		FMatrix ProjectionMatrix =  FReversedZOrthoMatrix(
+			LandscapeExtent.X,
+			LandscapeExtent.Y,
+			0.5f / ZOffset,
+			ZOffset);
 
-				ENQUEUE_UNIQUE_RENDER_COMMAND_ONEPARAMETER(
-					UpdateThumbnailRTCommand,
-					FTextureRenderTargetResource*, RenderTargetResource, RenderTargetResource,
-				{
-					// Copy (resolve) the rendered thumbnail from the render target to its texture
-					RHICmdList.CopyToResolveTarget(
-						RenderTargetResource->GetRenderTargetTexture(),		// Source texture
-						RenderTargetResource->TextureRHI,					// Dest texture
-						false,												// Do we need the source image content again?
-						FResolveParams() );									// Resolve parameters
-				});
-
-				OutFlattenMaterial.DiffuseSamples.Empty();
-				OutFlattenMaterial.DiffuseSamples.AddUninitialized(OutFlattenMaterial.DiffuseSize.X*OutFlattenMaterial.DiffuseSize.Y);
-
-				// Copy the contents of the remote texture to system memory
-				// NOTE: OutRawImageData must be a preallocated buffer!
-				RenderTargetResource->ReadPixelsPtr(
-					(FColor*)OutFlattenMaterial.DiffuseSamples.GetData(), 
-					FReadSurfaceDataFlags(), 
-					FIntRect(0, 0, OutFlattenMaterial.DiffuseSize.X, OutFlattenMaterial.DiffuseSize.Y)
-					);
-			
-				FlushRenderingCommands();
-					
-				RenderTargetTexture->RemoveFromRoot();
-				RenderTargetTexture = nullptr;
-			}
+		FSceneInterface* Scene = InLandscape->GetWorld()->Scene;
+		
+		// Render diffuse texture using BufferVisualizationMode=BaseColor
+		if (OutFlattenMaterial.DiffuseSize.X > 0 && 
+			OutFlattenMaterial.DiffuseSize.Y > 0)
+		{
+			static const FName BaseColorName("BaseColor");
+			RenderSceneToTexture(Scene, BaseColorName, ViewMatrix, ProjectionMatrix, PrimitivesToHide, 
+				OutFlattenMaterial.DiffuseSize, OutFlattenMaterial.DiffuseSamples);
+		}
 
-			// Render normal map using BufferVisualizationMode=WorldNormal
-			// Final material should use world space instead of tangent space for normals
-			if (OutFlattenMaterial.NormalSize.X > 0 && 
-				OutFlattenMaterial.NormalSize.Y > 0)
-			{
-				UTextureRenderTarget2D* RenderTargetTexture = new UTextureRenderTarget2D(FObjectInitializer());
-				check(RenderTargetTexture);
-				RenderTargetTexture->AddToRoot();
-				RenderTargetTexture->ClearColor = FLinearColor(0.0f, 0.0f, 0.0f, 0.0f);
-				RenderTargetTexture->InitCustomFormat(
-					OutFlattenMaterial.NormalSize.X, 
-					OutFlattenMaterial.NormalSize.Y, PF_B8G8R8A8, true);
-				FTextureRenderTargetResource* RenderTargetResource = RenderTargetTexture->GameThread_GetRenderTargetResource();
-
-				FSceneViewFamilyContext ViewFamily(
-					FSceneViewFamily::ConstructionValues(RenderTargetResource, Scene, FEngineShowFlags(ESFIM_Game))
-						.SetWorldTimes(FApp::GetCurrentTime() - GStartTime, FApp::GetDeltaTime(), FApp::GetCurrentTime() - GStartTime)
-					);
-
-				ViewFamily.EngineShowFlags.DisableAdvancedFeatures();
-				ViewFamily.EngineShowFlags.MotionBlur = 0;
-				ViewFamily.EngineShowFlags.Lighting = 0;
-				ViewFamily.EngineShowFlags.LightFunctions = 0;
-				ViewFamily.EngineShowFlags.DynamicShadows = 0;
-				ViewFamily.EngineShowFlags.Atmosphere = 0;
-				// To enable WorldNormal visualization mode
-				ViewFamily.EngineShowFlags.PostProcessing = 1;
-				ViewFamily.EngineShowFlags.PostProcessMaterial = 1;
-				ViewFamily.EngineShowFlags.VisualizeBuffer = 1; 
-				
-				FSceneViewInitOptions ViewInitOptions;
-				ViewInitOptions.SetViewRectangle(FIntRect(0, 0, OutFlattenMaterial.NormalSize.X, OutFlattenMaterial.NormalSize.Y));
-				ViewInitOptions.ViewFamily = &ViewFamily;
-				ViewInitOptions.HiddenPrimitives = PrimitivesToHide;
-
-				ViewInitOptions.ViewMatrix = FTranslationMatrix(-LandscapeCenter);
-				ViewInitOptions.ViewMatrix*= FInverseRotationMatrix(InLandscape->GetActorRotation());
-				ViewInitOptions.ViewMatrix*= FMatrix(	FPlane(1,	0,	0,	0),
-														FPlane(0,	-1,	0,	0),
-														FPlane(0,	0,	-1,	0),
-														FPlane(0,	0,	0,	1));
-				
-				const float ZOffset = WORLD_MAX;
-				ViewInitOptions.ProjectionMatrix =  FReversedZOrthoMatrix(
-					LandscapeExtent.X,
-					LandscapeExtent.Y,
-					0.5f / ZOffset,
-					ZOffset);
-
-				FSceneView* NewView = new FSceneView(ViewInitOptions);
-				NewView->CurrentBufferVisualizationMode = FName("WorldNormal");
-				ViewFamily.Views.Add(NewView);
-										
-				FCanvas Canvas(RenderTargetResource, NULL, FApp::GetCurrentTime() - GStartTime, FApp::GetDeltaTime(), FApp::GetCurrentTime() - GStartTime, Scene->GetFeatureLevel());
-				Canvas.Clear(FLinearColor::Black);
-				GetRendererModule().BeginRenderingViewFamily(&Canvas, &ViewFamily);
-
-				ENQUEUE_UNIQUE_RENDER_COMMAND_ONEPARAMETER(
-					UpdateThumbnailRTCommand,
-					FTextureRenderTargetResource*, RenderTargetResource, RenderTargetResource,
-				{
-					// Copy (resolve) the rendered thumbnail from the render target to its texture
-					RHICmdList.CopyToResolveTarget(
-						RenderTargetResource->GetRenderTargetTexture(),		// Source texture
-						RenderTargetResource->TextureRHI,					// Dest texture
-						false,												// Do we need the source image content again?
-						FResolveParams() );									// Resolve parameters
-				});
-
-				OutFlattenMaterial.NormalSamples.Empty();
-				OutFlattenMaterial.NormalSamples.AddUninitialized(OutFlattenMaterial.NormalSize.X*OutFlattenMaterial.NormalSize.Y);
-
-				// Copy the contents of the remote texture to system memory
-				// NOTE: OutRawImageData must be a preallocated buffer!
-				RenderTargetResource->ReadPixelsPtr(
-					(FColor*)OutFlattenMaterial.NormalSamples.GetData(), 
-					FReadSurfaceDataFlags(), 
-					FIntRect(0, 0, OutFlattenMaterial.NormalSize.X, OutFlattenMaterial.NormalSize.Y)
-					);
-			
-				FlushRenderingCommands();
+		// Render normal map using BufferVisualizationMode=WorldNormal
+		// Final material should use world space instead of tangent space for normals
+		if (OutFlattenMaterial.NormalSize.X > 0 && 
+			OutFlattenMaterial.NormalSize.Y > 0)
+		{
+			static const FName WorldNormalName("WorldNormal");
+			RenderSceneToTexture(Scene, WorldNormalName, ViewMatrix, ProjectionMatrix, PrimitivesToHide, 
+				OutFlattenMaterial.NormalSize, OutFlattenMaterial.NormalSamples);
+		}
 
-				RenderTargetTexture->RemoveFromRoot();
-				RenderTargetTexture = nullptr;
-			}
+		// Render roughness map using BufferVisualizationMode=Roughness
+		if (OutFlattenMaterial.RoughnessSize.X > 0 && 
+			OutFlattenMaterial.RoughnessSize.Y > 0)
+		{
+			static const FName RoughnessName("Roughness");
+			RenderSceneToTexture(Scene, RoughnessName, ViewMatrix, ProjectionMatrix, PrimitivesToHide, 
+				OutFlattenMaterial.RoughnessSize, OutFlattenMaterial.RoughnessSamples);
 		}
 
+		// Render roughness map using BufferVisualizationMode=Specular
+		if (OutFlattenMaterial.SpecularSize.X > 0 && 
+			OutFlattenMaterial.SpecularSize.Y > 0)
+		{
+			static const FName SpecularName("Specular");
+			RenderSceneToTexture(Scene, SpecularName, ViewMatrix, ProjectionMatrix, PrimitivesToHide, 
+				OutFlattenMaterial.SpecularSize, OutFlattenMaterial.SpecularSamples);
+		}
+				
 		OutFlattenMaterial.MaterialId = InLandscape->GetLandscapeGuid();
 		return true;
 	}
 
+	UTexture2D* CreateTexture(UPackage* Outer, const FString& AssetName, FIntPoint Size, const TArray<FColor>& Samples, TextureCompressionSettings CompressionSettings, TextureGroup LODGroup, EObjectFlags Flags)
+	{
+		FCreateTexture2DParameters TexParams;
+		TexParams.bUseAlpha = false;
+		TexParams.CompressionSettings = CompressionSettings;
+		TexParams.bDeferCompression = true;
+		TexParams.bSRGB = false;
+
+		if (Outer == nullptr)
+		{
+			Outer = CreatePackage(NULL, *AssetName);
+			Outer->FullyLoad();
+			Outer->Modify();
+		}
+
+		UTexture2D* Texture = FImageUtils::CreateTexture2D(Size.X, Size.Y, Samples, Outer, AssetName, Flags, TexParams);
+		Texture->LODGroup = LODGroup;
+		Texture->PostEditChange();
+		
+		return Texture;
+	}
+
 	UMaterial* CreateMaterial(const FFlattenMaterial& InFlattenMaterial, UPackage* Outer, const FString& BaseName, EObjectFlags Flags)
 	{
 		const FString AssetBaseName = FPackageName::GetShortName(BaseName);
@@ -2910,58 +2886,14 @@ namespace MaterialExportUtils
 		Material->TwoSided = false;
 		Material->SetShadingModel(MSM_DefaultLit);
 
-		//Set Metallic as a default constant
-		{
-			UMaterialExpressionConstant* MetallicExpression = ConstructObject<UMaterialExpressionConstant>(UMaterialExpressionConstant::StaticClass(), Material);
-			MetallicExpression->R = 0.0f;
-			Material->Expressions.Add(MetallicExpression);
-			Material->Metallic.Expression = MetallicExpression;
-
-			MetallicExpression->MaterialExpressionEditorX = -250;
-			MetallicExpression->MaterialExpressionEditorY = 0;
-		}
-		
-		//Set Roughness as default constant
-		{
-			UMaterialExpressionConstant* RoughnessExpression = ConstructObject<UMaterialExpressionConstant>(UMaterialExpressionConstant::StaticClass(), Material);
-			RoughnessExpression->R = 0.8f;
-			Material->Expressions.Add(RoughnessExpression);
-			Material->Roughness.Expression = RoughnessExpression;
-
-			RoughnessExpression->MaterialExpressionEditorX = -250;
-			RoughnessExpression->MaterialExpressionEditorY = 150;
-		}
-
-		//Build the Diffuse UTexture
+		// Set BaseColor
 		if (InFlattenMaterial.DiffuseSamples.Num() > 1)
 		{
-			FCreateTexture2DParameters TexParams;
-			TexParams.bUseAlpha = false;
-			TexParams.CompressionSettings = TC_Default;
-			TexParams.bDeferCompression = true;
-			TexParams.bSRGB = false;
-
-			const FString TextureAssetName = TEXT("T_") + AssetBaseName + TEXT("_D");
-			UPackage* TextureOuter = Outer;
-			if (TextureOuter == NULL)
-			{
-				TextureOuter = CreatePackage(NULL, *(AssetBasePath + TextureAssetName));
-				TextureOuter->FullyLoad();
-				TextureOuter->Modify();
-			}
-
-			UTexture2D* DiffuseTexture = FImageUtils::CreateTexture2D(
-				InFlattenMaterial.DiffuseSize.X, 
-				InFlattenMaterial.DiffuseSize.Y,
-				InFlattenMaterial.DiffuseSamples,
-				TextureOuter,
-				TextureAssetName,
-				Flags, 
-				TexParams);
-
-			//Create BaseColor Expression of the diffuse texture
+			const FString AssetName = TEXT("T_") + AssetBaseName + TEXT("_D");
+			UTexture2D* Texture = CreateTexture(Outer, AssetName, InFlattenMaterial.DiffuseSize, InFlattenMaterial.DiffuseSamples, TC_Default, TEXTUREGROUP_World, Flags);
+			
 			UMaterialExpressionTextureSample* BasecolorExpression = ConstructObject<UMaterialExpressionTextureSample>(UMaterialExpressionTextureSample::StaticClass(), Material);
-			BasecolorExpression->Texture = DiffuseTexture;
+			BasecolorExpression->Texture = Texture;
 			BasecolorExpression->SamplerType = EMaterialSamplerType::SAMPLERTYPE_LinearColor;
 			BasecolorExpression->MaterialExpressionEditorX = -400;
 			BasecolorExpression->MaterialExpressionEditorY = -150;
@@ -2969,46 +2901,72 @@ namespace MaterialExportUtils
 			Material->BaseColor.Expression = BasecolorExpression;
 		}
 
-		//Build the Normal UTexture
+		// Set Normal
 		if (InFlattenMaterial.NormalSamples.Num() > 1)
 		{
-			FCreateTexture2DParameters TexParams;
-			TexParams.bUseAlpha = false;
-			TexParams.CompressionSettings = TC_Normalmap;
-			TexParams.bDeferCompression = true;
-			TexParams.bSRGB = false;
-
-			const FString TextureAssetName = TEXT("T_") + AssetBaseName + TEXT("_N");
-			UPackage* TextureOuter = Outer;
-			if (TextureOuter == NULL)
-			{
-				TextureOuter = CreatePackage(NULL, *(AssetBasePath + TextureAssetName));
-				TextureOuter->FullyLoad();
-				TextureOuter->Modify();
-			}
-
-			UTexture2D* NormalTexture = FImageUtils::CreateTexture2D(
-				InFlattenMaterial.NormalSize.X, 
-				InFlattenMaterial.NormalSize.Y,
-				InFlattenMaterial.NormalSamples,
-				TextureOuter,
-				TextureAssetName,
-				Flags, 
-				TexParams);
-
-			NormalTexture->LODGroup = TEXTUREGROUP_WorldNormalMap;
-			NormalTexture->PostEditChange();
-						
-			//Assign the normalmap to the material
+			const FString AssetName = TEXT("T_") + AssetBaseName + TEXT("_N");
+			UTexture2D* Texture = CreateTexture(Outer, AssetName, InFlattenMaterial.NormalSize, InFlattenMaterial.NormalSamples, TC_Normalmap, TEXTUREGROUP_WorldNormalMap, Flags);
+			
 			UMaterialExpressionTextureSample* NormalExpression = ConstructObject<UMaterialExpressionTextureSample>(UMaterialExpressionTextureSample::StaticClass(), Material);
-			NormalExpression->Texture = NormalTexture;
+			NormalExpression->Texture = Texture;
 			NormalExpression->SamplerType = EMaterialSamplerType::SAMPLERTYPE_Normal;
 			NormalExpression->MaterialExpressionEditorX = -400;
-			NormalExpression->MaterialExpressionEditorY = 300;
+			NormalExpression->MaterialExpressionEditorY = 0;
 			Material->Expressions.Add(NormalExpression);
 			Material->Normal.Expression = NormalExpression;
 		}
+
+		if (InFlattenMaterial.SpecularSamples.Num() > 1)
+		{
+			const FString AssetName = TEXT("T_") + AssetBaseName + TEXT("_S");
+			UTexture2D* Texture = CreateTexture(Outer, AssetName, InFlattenMaterial.SpecularSize, InFlattenMaterial.SpecularSamples, TC_Default, TEXTUREGROUP_World, Flags);
+			
+			UMaterialExpressionTextureSample* SpecularExpression = ConstructObject<UMaterialExpressionTextureSample>(UMaterialExpressionTextureSample::StaticClass(), Material);
+			SpecularExpression->Texture = Texture;
+			SpecularExpression->SamplerType = EMaterialSamplerType::SAMPLERTYPE_LinearColor;
+			SpecularExpression->MaterialExpressionEditorX = -400;
+			SpecularExpression->MaterialExpressionEditorY = 300;
+			Material->Expressions.Add(SpecularExpression);
+			Material->Specular.Expression = SpecularExpression;
+		}
+		
+		// Set Roughness
+		if (InFlattenMaterial.RoughnessSamples.Num() > 1)
+		{
+			const FString AssetName = TEXT("T_") + AssetBaseName + TEXT("_R");
+			UTexture2D* Texture = CreateTexture(Outer, AssetName, InFlattenMaterial.RoughnessSize, InFlattenMaterial.RoughnessSamples, TC_Default, TEXTUREGROUP_World, Flags);
 			
+			UMaterialExpressionTextureSample* RoughnessExpression = ConstructObject<UMaterialExpressionTextureSample>(UMaterialExpressionTextureSample::StaticClass(), Material);
+			RoughnessExpression->Texture = Texture;
+			RoughnessExpression->SamplerType = EMaterialSamplerType::SAMPLERTYPE_LinearColor;
+			RoughnessExpression->MaterialExpressionEditorX = -250;
+			RoughnessExpression->MaterialExpressionEditorY = 150;
+			Material->Expressions.Add(RoughnessExpression);
+			Material->Roughness.Expression = RoughnessExpression;
+		}
+		else
+		{
+			float Roughness = InFlattenMaterial.RoughnessSamples.Num() ? InFlattenMaterial.RoughnessSamples[0].R : 0.8f;
+			// Set Roughness to constant
+			UMaterialExpressionConstant* RoughnessExpression = ConstructObject<UMaterialExpressionConstant>(UMaterialExpressionConstant::StaticClass(), Material);
+			RoughnessExpression->R = Roughness;
+			RoughnessExpression->MaterialExpressionEditorX = -250;
+			RoughnessExpression->MaterialExpressionEditorY = 150;
+			Material->Expressions.Add(RoughnessExpression);
+			Material->Roughness.Expression = RoughnessExpression;
+		}
+
+		//Set Metallic as a default constant
+		{
+			UMaterialExpressionConstant* MetallicExpression = ConstructObject<UMaterialExpressionConstant>(UMaterialExpressionConstant::StaticClass(), Material);
+			MetallicExpression->R = 0.0f;
+			Material->Expressions.Add(MetallicExpression);
+			Material->Metallic.Expression = MetallicExpression;
+
+			MetallicExpression->MaterialExpressionEditorX = -250;
+			MetallicExpression->MaterialExpressionEditorY = 0;
+		}
+							
 		Material->PostEditChange();
 		return Material;
 	}
diff --git a/Engine/Source/Editor/UnrealEd/Public/MaterialExportUtils.h b/Engine/Source/Editor/UnrealEd/Public/MaterialExportUtils.h
--- a/Engine/Source/Editor/UnrealEd/Public/MaterialExportUtils.h
+++ b/Engine/Source/Editor/UnrealEd/Public/MaterialExportUtils.h
@@ -11,14 +11,21 @@ namespace MaterialExportUtils
 		FFlattenMaterial()
 			: DiffuseSize(512, 512)
 			, NormalSize(512, 512)
+			, RoughnessSize(0, 0)
+			, SpecularSize(0, 0)
 		{}
 		
 		FGuid			MaterialId;
+		
 		FIntPoint		DiffuseSize;
-		FIntPoint		NormalSize;	
+		FIntPoint		NormalSize;
+		FIntPoint		RoughnessSize;	
+		FIntPoint		SpecularSize;	
 			
 		TArray<FColor>	DiffuseSamples;
 		TArray<FColor>	NormalSamples;
+		TArray<FColor>	RoughnessSamples;
+		TArray<FColor>	SpecularSamples;
 	};
 	
 	/**
diff --git a/Engine/Source/Editor/WorldBrowser/Private/Tiles/WorldTileCollectionModel.cpp b/Engine/Source/Editor/WorldBrowser/Private/Tiles/WorldTileCollectionModel.cpp
--- a/Engine/Source/Editor/WorldBrowser/Private/Tiles/WorldTileCollectionModel.cpp
+++ b/Engine/Source/Editor/WorldBrowser/Private/Tiles/WorldTileCollectionModel.cpp
@@ -1836,8 +1836,7 @@ bool FWorldTileCollectionModel::GenerateLODLevels(FLevelModelList InLevelList, i
 			TileModel->SetVisible(true);
 		}
 				
-		FWorldTileInfo TileInfo = TileModel->TileDetails->GetInfo();
-		FWorldTileLODInfo LODInfo = TileModel->TileDetails->GetInfo().LODList[TargetLODIndex];
+		FLevelSimplificationDetails SimplificationDetails = TileModel->GetLevelObject()->LevelSimplification[TargetLODIndex];
 		
 		// Source level package name
 		FString SourceLongPackageName = TileModel->TileDetails->PackageName.ToString();
@@ -1865,7 +1864,7 @@ bool FWorldTileCollectionModel::GenerateLODLevels(FLevelModelList InLevelList, i
 		TArray<FTransform>			AssetsToSpawnTransform;
 		TArray<AActor*>				Actors;
 		TArray<ALandscapeProxy*>	LandscapeActors;
-		// Separate flies from cutlets
+		// Separate landscape actors from all others
 		for (AActor* Actor : TileModel->GetLevelObject()->Actors)
 		{
 			if (Actor)
@@ -1888,7 +1887,7 @@ bool FWorldTileCollectionModel::GenerateLODLevels(FLevelModelList InLevelList, i
 			GWarn->StatusUpdate(0, 10, LOCTEXT("GeneratingProxyMesh", "Generating Proxy Mesh"));
 
 			FMeshProxySettings ProxySettings;
-			ProxySettings.ScreenSize = ProxySettings.ScreenSize*(LODInfo.GenDetailsPercentage/100.f);
+			ProxySettings.ScreenSize = ProxySettings.ScreenSize*(SimplificationDetails.DetailsPercentage/100.f);
 			TArray<UObject*> OutAssets;
 			FVector OutProxyLocation;
 			FString ProxyPackageName = TEXT("PROXY_") + FPackageName::GetShortName(TileModel->TileDetails->PackageName);
@@ -1908,7 +1907,7 @@ bool FWorldTileCollectionModel::GenerateLODLevels(FLevelModelList InLevelList, i
 
 		using namespace MaterialExportUtils;
 
-		// Convert landscape actors into static meshes and apply mesh reduction 
+		// Convert landscape actors into static meshes
 		int32 LandscapeActorIndex = 0;
 		for (ALandscapeProxy* Landscape : LandscapeActors)
 		{
@@ -1918,36 +1917,33 @@ bool FWorldTileCollectionModel::GenerateLODLevels(FLevelModelList InLevelList, i
 			FFlattenMaterial LandscapeFlattenMaterial;
 			FVector LandscapeWorldLocation = Landscape->GetActorLocation();
 		
-			Landscape->ExportToRawMesh(MAX_int32, LandscapeRawMesh);
+			Landscape->ExportToRawMesh(SimplificationDetails.LandscapeExportLOD, LandscapeRawMesh);
 		
 			for (FVector& VertexPos : LandscapeRawMesh.VertexPositions)
 			{
 				VertexPos-= LandscapeWorldLocation;
 			}
 								
-			// This is texture resolution for a landscape, probably need to be calculated using landscape size
-			LandscapeFlattenMaterial.DiffuseSize	= FIntPoint(1024, 1024);
-			LandscapeFlattenMaterial.NormalSize		= FIntPoint(1024, 1024);
+			// This is texture resolution for a landscape, probably needs to be calculated using landscape size
+			LandscapeFlattenMaterial.DiffuseSize = FIntPoint(1024, 1024);
+			
+			if (SimplificationDetails.bGenerateLandscapeNormalMap)
+			{
+				LandscapeFlattenMaterial.NormalSize = FIntPoint(1024, 1024);
+			}
+
+			if (SimplificationDetails.bGenerateLandscapeRoughnessMap)
+			{
+				LandscapeFlattenMaterial.RoughnessSize = FIntPoint(1024, 1024);
+			}
+
+			if (SimplificationDetails.bGenerateLandscapeSpecularMap)
+			{
+				LandscapeFlattenMaterial.SpecularSize = FIntPoint(1024, 1024);
+			}
+			
 			ExportMaterial(Landscape, LandscapeFlattenMaterial);
 		
-			// TODO: Disabled landscape mesh simplification, does not really needed since exporting highest landscape LOD already provide low triangle count
-			// Reduce landscape mesh
-			//if (LODInfo.GenDetailsPercentage != 100.f)
-			//{
-			//	FMeshReductionSettings Settings;
-			//	Settings.PercentTriangles = LODInfo.GenDetailsPercentage/100.f;
-			//	float OutMaxDeviation = 0.f;
-			//	FRawMesh LandscapeReducedMesh;
-
-			//	MeshUtilities.GetMeshReductionInterface()->Reduce(
-			//		LandscapeReducedMesh,
-			//		OutMaxDeviation,
-			//		LandscapeRawMesh,
-			//		Settings);
-
-			//	LandscapeRawMesh = LandscapeReducedMesh;
-			//}
-
 			FString LandscapeBaseAssetName = Landscape->GetName();
 			// Construct landscape material
 			UMaterial* StaticLandscapeMaterial = MaterialExportUtils::CreateMaterial(
diff --git a/Engine/Source/Editor/WorldBrowser/Private/Tiles/WorldTileDetails.cpp b/Engine/Source/Editor/WorldBrowser/Private/Tiles/WorldTileDetails.cpp
--- a/Engine/Source/Editor/WorldBrowser/Private/Tiles/WorldTileDetails.cpp
+++ b/Engine/Source/Editor/WorldBrowser/Private/Tiles/WorldTileDetails.cpp
@@ -13,8 +13,6 @@ FTileLODEntryDetails::FTileLODEntryDetails()
 	// Initialize properties with default values from FWorldTileLODInfo
 	: LODIndex(0)
 	, Distance(FWorldTileLODInfo().RelativeStreamingDistance)
-	, DetailsPercentage(FWorldTileLODInfo().GenDetailsPercentage)
-	, MaxDeviation(FWorldTileLODInfo().GenMaxDeviation)
 {
 }
 
@@ -29,21 +27,7 @@ UWorldTileDetails::UWorldTileDetails(const FObjectInitializer& ObjectInitializer
 	LOD4.LODIndex = 3;
 }
 
-static void SyncLODSettings(FTileLODEntryDetails& To, const FWorldTileLODInfo& From)
-{
-	To.Distance				= From.RelativeStreamingDistance;
-	To.DetailsPercentage	= From.GenDetailsPercentage;
-	To.MaxDeviation			= From.GenMaxDeviation;
-}
-
-static void SyncLODSettings(FWorldTileLODInfo& To, const FTileLODEntryDetails& From)
-{
-	To.RelativeStreamingDistance	= From.Distance;			
-	To.GenDetailsPercentage			= From.DetailsPercentage;
-	To.GenMaxDeviation				= From.MaxDeviation;		
-}
-
-void UWorldTileDetails::SetInfo(const FWorldTileInfo& Info)
+void UWorldTileDetails::SetInfo(const FWorldTileInfo& Info, ULevel* Level)
 {
 	ParentPackageName	= FName(*Info.ParentTilePackageName);
 	Position			= Info.Position;
@@ -55,9 +39,19 @@ void UWorldTileDetails::SetInfo(const FWorldTileInfo& Info)
 	// Sync LOD settings
 	NumLOD				= Info.LODList.Num();
 	FTileLODEntryDetails* LODEntries[WORLDTILE_LOD_MAX_INDEX] = {&LOD1, &LOD2, &LOD3, &LOD4};
-	for (int32 LODIdx = 0; LODIdx < Info.LODList.Num(); ++LODIdx)
+
+	for (int32 i = 0; i < WORLDTILE_LOD_MAX_INDEX; ++i)
 	{
-		SyncLODSettings(*LODEntries[LODIdx], Info.LODList[LODIdx]);
+		if (Info.LODList.IsValidIndex(i))
+		{
+			LODEntries[i]->Distance = Info.LODList[i].RelativeStreamingDistance;
+			LODEntries[i]->SimplificationDetails = Level ? Level->LevelSimplification[i] : FLevelSimplificationDetails();
+		}
+		else
+		{
+			LODEntries[i]->Distance = FWorldTileLODInfo().RelativeStreamingDistance;
+			LODEntries[i]->SimplificationDetails = FLevelSimplificationDetails();
+		}
 	}
 }
 	
@@ -75,9 +69,9 @@ FWorldTileInfo UWorldTileDetails::GetInfo() const
 	// Sync LOD settings
 	Info.LODList.SetNum(FMath::Clamp(NumLOD, 0, WORLDTILE_LOD_MAX_INDEX));
 	const FTileLODEntryDetails* LODEntries[WORLDTILE_LOD_MAX_INDEX] = {&LOD1, &LOD2, &LOD3, &LOD4};
-	for (int32 LODIdx = 0; LODIdx < Info.LODList.Num(); ++LODIdx)
+	for (int32 i = 0; i < Info.LODList.Num(); ++i)
 	{
-		SyncLODSettings(Info.LODList[LODIdx], *LODEntries[LODIdx]);
+		Info.LODList[i].RelativeStreamingDistance = LODEntries[i]->Distance;
 	}
 	
 	return Info;
diff --git a/Engine/Source/Editor/WorldBrowser/Private/Tiles/WorldTileDetails.h b/Engine/Source/Editor/WorldBrowser/Private/Tiles/WorldTileDetails.h
--- a/Engine/Source/Editor/WorldBrowser/Private/Tiles/WorldTileDetails.h
+++ b/Engine/Source/Editor/WorldBrowser/Private/Tiles/WorldTileDetails.h
@@ -22,13 +22,8 @@ struct FTileLODEntryDetails
 	UPROPERTY(Category=LODDetails, EditAnywhere, meta=(ClampMin = "10", ClampMax = "10000000", UIMin = "10", UIMax = "500000"))
 	int32		Distance;
 
-	// Percentage of details relative to main tile details
-	UPROPERTY(Category=ReductionSettings, EditAnywhere, meta=(ClampMin = "0", ClampMax = "100", UIMin = "0", UIMax = "100"))	
-	float		DetailsPercentage;
-	
-	// Maximum deviation of details percentage
-	UPROPERTY(Category=ReductionSettings, EditAnywhere, meta=(ClampMin = "0", ClampMax = "100", UIMin = "0", UIMax = "100"))	
-	float		MaxDeviation;
+	UPROPERTY(Category=ReductionSettings, EditAnywhere)	
+	FLevelSimplificationDetails SimplificationDetails;
 
 	FTileLODEntryDetails();
 };
@@ -107,7 +102,7 @@ class UWorldTileDetails : public UObject
 
 public:
 	// Initialize tile details with values stored in FWorldTileInfo object
-	void SetInfo(const FWorldTileInfo& Info);
+	void SetInfo(const FWorldTileInfo& Info, ULevel* Level);
 	
 	// Gets the initialized FWorldTileInfo from this details values
 	FWorldTileInfo GetInfo() const;
diff --git a/Engine/Source/Editor/WorldBrowser/Private/Tiles/WorldTileDetailsCustomization.cpp b/Engine/Source/Editor/WorldBrowser/Private/Tiles/WorldTileDetailsCustomization.cpp
--- a/Engine/Source/Editor/WorldBrowser/Private/Tiles/WorldTileDetailsCustomization.cpp
+++ b/Engine/Source/Editor/WorldBrowser/Private/Tiles/WorldTileDetailsCustomization.cpp
@@ -266,12 +266,8 @@ void FTileLODEntryDetailsCustomization::CustomizeChildren(TSharedRef<IPropertyHa
 		GET_MEMBER_NAME_CHECKED(FTileLODEntryDetails, Distance)
 		);
 	
-	TSharedPtr<IPropertyHandle> DetailsPercentage = StructPropertyHandle->GetChildHandle(
-		GET_MEMBER_NAME_CHECKED(FTileLODEntryDetails, DetailsPercentage)
-		);
-	
-	TSharedPtr<IPropertyHandle> MaxDeviation = StructPropertyHandle->GetChildHandle(
-		GET_MEMBER_NAME_CHECKED(FTileLODEntryDetails, MaxDeviation)
+	TSharedPtr<IPropertyHandle> SimplificationDetails = StructPropertyHandle->GetChildHandle(
+		GET_MEMBER_NAME_CHECKED(FTileLODEntryDetails, SimplificationDetails)
 		);
 
 	ChildBuilder.AddChildProperty(LODIndexHandle.ToSharedRef())
@@ -279,12 +275,9 @@ void FTileLODEntryDetailsCustomization::CustomizeChildren(TSharedRef<IPropertyHa
 
 	ChildBuilder.AddChildProperty(DistanceProperty.ToSharedRef())
 		.IsEnabled(TAttribute<bool>::Create(TAttribute<bool>::FGetter::CreateSP(this, &FTileLODEntryDetailsCustomization::IsLODDistanceEnabled)));
-	// Reduction settings available if editor supports LOD levels generation
-	ChildBuilder.AddChildProperty(DetailsPercentage.ToSharedRef())
+	
+	ChildBuilder.AddChildProperty(SimplificationDetails.ToSharedRef())
 		.IsEnabled(TAttribute<bool>::Create(TAttribute<bool>::FGetter::CreateSP(this, &FTileLODEntryDetailsCustomization::IsGenerateTileEnabled)));
-	// Useless now
-	ChildBuilder.AddChildProperty(MaxDeviation.ToSharedRef())
-		.Visibility(EVisibility::Hidden);
 }
 
 FReply FTileLODEntryDetailsCustomization::OnGenerateTile()
diff --git a/Engine/Source/Editor/WorldBrowser/Private/Tiles/WorldTileModel.cpp b/Engine/Source/Editor/WorldBrowser/Private/Tiles/WorldTileModel.cpp
--- a/Engine/Source/Editor/WorldBrowser/Private/Tiles/WorldTileModel.cpp
+++ b/Engine/Source/Editor/WorldBrowser/Private/Tiles/WorldTileModel.cpp
@@ -46,7 +46,6 @@ FWorldTileModel::FWorldTileModel(const TWeakObjectPtr<UEditorEngine>& InEditor,
 	{
 		FWorldCompositionTile& Tile = WorldComposition->GetTilesList()[TileIdx];
 		
-		TileDetails->SetInfo(Tile.Info);
 		TileDetails->PackageName = Tile.PackageName;
 		TileDetails->bPersistentLevel = false;
 				
@@ -72,6 +71,8 @@ FWorldTileModel::FWorldTileModel(const TWeakObjectPtr<UEditorEngine>& InEditor,
 				}
 			}
 		}
+
+		TileDetails->SetInfo(Tile.Info, LoadedLevel.Get());
 	}
 	else
 	{
@@ -468,11 +469,11 @@ void FWorldTileModel::Update()
 		LandscapeComponentSize = FVector2D(0.f, 0.f);
 		LandscapeComponentsRectXY = FIntRect(FIntPoint(MAX_int32, MAX_int32), FIntPoint(MIN_int32, MIN_int32));
 		
+		ULevel* Level = GetLevelObject();
 		// Receive tile info from world composition
 		FWorldTileInfo Info = LevelCollectionModel.GetWorld()->WorldComposition->GetTileInfo(TileDetails->PackageName);
-		TileDetails->SetInfo(Info);
-			
-		ULevel* Level = GetLevelObject();
+		TileDetails->SetInfo(Info, Level);
+						
 		if (Level != nullptr && Level->bIsVisible)
 		{
 			if (Level->LevelBoundsActor.IsValid())
@@ -646,6 +647,24 @@ void FWorldTileModel::OnLevelInfoUpdated()
 	if (!IsRootTile())
 	{
 		LevelCollectionModel.GetWorld()->WorldComposition->OnTileInfoUpdated(TileDetails->PackageName, TileDetails->GetInfo());
+		ULevel* Level = GetLevelObject();
+		if (Level)
+		{
+			bool bMarkDirty = false;
+			bMarkDirty|= !(Level->LevelSimplification[0] == TileDetails->LOD1.SimplificationDetails);
+			bMarkDirty|= !(Level->LevelSimplification[1] == TileDetails->LOD2.SimplificationDetails);
+			bMarkDirty|= !(Level->LevelSimplification[2] == TileDetails->LOD3.SimplificationDetails);
+			bMarkDirty|= !(Level->LevelSimplification[3] == TileDetails->LOD4.SimplificationDetails);
+			
+			if (bMarkDirty)
+			{
+				Level->LevelSimplification[0] = TileDetails->LOD1.SimplificationDetails;
+				Level->LevelSimplification[1] = TileDetails->LOD2.SimplificationDetails;
+				Level->LevelSimplification[2] = TileDetails->LOD3.SimplificationDetails;
+				Level->LevelSimplification[3] = TileDetails->LOD4.SimplificationDetails;
+				Level->MarkPackageDirty();
+			}
+		}
 	}
 }
 
diff --git a/Engine/Source/Runtime/CoreUObject/Private/Misc/WorldCompositionUtility.cpp b/Engine/Source/Runtime/CoreUObject/Private/Misc/WorldCompositionUtility.cpp
--- a/Engine/Source/Runtime/CoreUObject/Private/Misc/WorldCompositionUtility.cpp
+++ b/Engine/Source/Runtime/CoreUObject/Private/Misc/WorldCompositionUtility.cpp
@@ -27,10 +27,10 @@ FArchive& operator<<( FArchive& Ar, FWorldTileLODInfo& D )
 {
 	// Serialized with FPackageFileSummary
 	Ar << D.RelativeStreamingDistance 
-		<< D.GenDetailsPercentage
-		<< D.GenMaxDeviation
 		<< D.Reserved0
-		<< D.Reserved1;
+		<< D.Reserved1
+		<< D.Reserved2
+		<< D.Reserved3;
 	return Ar;
 }
 
diff --git a/Engine/Source/Runtime/CoreUObject/Public/Misc/WorldCompositionUtility.h b/Engine/Source/Runtime/CoreUObject/Public/Misc/WorldCompositionUtility.h
--- a/Engine/Source/Runtime/CoreUObject/Public/Misc/WorldCompositionUtility.h
+++ b/Engine/Source/Runtime/CoreUObject/Public/Misc/WorldCompositionUtility.h
@@ -53,10 +53,10 @@ class FWorldTileLODInfo
 public:
 	FWorldTileLODInfo()
 		: RelativeStreamingDistance(10000)
-		, GenDetailsPercentage(70.f)
-		, GenMaxDeviation(5.f)
 		, Reserved0(0)
 		, Reserved1(0)
+		, Reserved2(0)
+		, Reserved3(0)
 	{
 	}
 
@@ -64,22 +64,18 @@ public:
 
 	bool operator==(const FWorldTileLODInfo& OtherInfo) const
 	{
-		return (RelativeStreamingDistance == OtherInfo.RelativeStreamingDistance) &&
-				(GenDetailsPercentage == OtherInfo.GenDetailsPercentage) &&
-				(GenMaxDeviation == OtherInfo.GenMaxDeviation);
+		return RelativeStreamingDistance == OtherInfo.RelativeStreamingDistance;
 	}
 
 public:
 	/**  Relative to LOD0 streaming distance, absolute distance = LOD0 + StreamingDistanceDelta */
 	int32	RelativeStreamingDistance;
 	
-	/**  LOD generation settings */
-	float	GenDetailsPercentage;
-	float	GenMaxDeviation;
-	
 	/**  Reserved for additional options */
-	int32	Reserved0;
-	int32	Reserved1;
+	float	Reserved0;
+	float	Reserved1;
+	int32	Reserved2;
+	int32	Reserved3;
 };
 
 /** 
diff --git a/Engine/Source/Runtime/Engine/Classes/Engine/Level.h b/Engine/Source/Runtime/Engine/Classes/Engine/Level.h
--- a/Engine/Source/Runtime/Engine/Classes/Engine/Level.h
+++ b/Engine/Source/Runtime/Engine/Classes/Engine/Level.h
@@ -229,6 +229,40 @@ private:
 	friend class FLightmassProcessor;
 };
 
+USTRUCT()
+struct ENGINE_API FLevelSimplificationDetails
+{
+	GENERATED_USTRUCT_BODY()
+
+	// Percentage of details relative to main tile details
+	UPROPERTY(Category=ReductionSettings, EditAnywhere, meta=(ClampMin = "0", ClampMax = "100", UIMin = "0", UIMax = "100"))	
+	float DetailsPercentage;
+	
+	// Landscape LOD to use for static mesh generation
+	UPROPERTY(Category=Landscape, EditAnywhere, meta=(ClampMin = "0", ClampMax = "7", UIMin = "0", UIMax = "7"))
+	int32 LandscapeExportLOD;
+
+	// Whether to generate normal map for landscape static mesh
+	UPROPERTY(Category=Landscape, EditAnywhere)
+	bool bGenerateLandscapeNormalMap;
+
+	// Whether to generate roughness map for landscape static mesh
+	UPROPERTY(Category=Landscape, EditAnywhere)
+	bool bGenerateLandscapeRoughnessMap;
+	
+	// Whether to generate specular map for landscape static mesh
+	UPROPERTY(Category=Landscape, EditAnywhere)
+	bool bGenerateLandscapeSpecularMap;
+	
+	// Whether to bake foliage into landscape static mesh texture
+	UPROPERTY(Category=Landscape, EditAnywhere)
+	bool bBakeFoliageToLandscape;
+
+	FLevelSimplificationDetails();
+
+	bool operator == (const FLevelSimplificationDetails& Other) const;
+};
+
 //
 // The level object.  Contains the level's actor list, BSP information, and brush list.
 //
@@ -391,6 +425,12 @@ public:
 		return (OwningWorld && this == OwningWorld->CurrentLevelPendingVisibility);
 	}
 
+#if WITH_EDITORONLY_DATA
+	/** Level simplification settings for each LOD */
+	UPROPERTY()
+	FLevelSimplificationDetails LevelSimplification[WORLDTILE_LOD_MAX_INDEX];
+#endif //WITH_EDITORONLY_DATA
+
 #if PERF_TRACK_DETAILED_ASYNC_STATS
 	/** Mapping of how long each actor class takes to have UpdateComponents called on it */
 	TMap<const UClass*,struct FMapTimeEntry>		UpdateComponentsTimePerActorClass;
@@ -408,7 +448,6 @@ private:
 	FLevelBoundsActorUpdatedEvent LevelBoundsActorUpdatedEvent; 
 
 protected:
-
 	/** Array of all MovieSceneBindings that are used in this level.  These store the relationship between
 	    a MovieScene asset and possessed actors in this level. */
 	UPROPERTY()
@@ -424,9 +463,7 @@ protected:
 	UPROPERTY()
 	TArray<UAssetUserData*> AssetUserData;
 
-
 private:
-
 	// Actors awaiting input to be enabled once the appropriate PlayerController has been created
 	TArray<FPendingAutoReceiveInputActor> PendingAutoReceiveInputActors;
 
diff --git a/Engine/Source/Runtime/Engine/Private/Level.cpp b/Engine/Source/Runtime/Engine/Private/Level.cpp
--- a/Engine/Source/Runtime/Engine/Private/Level.cpp
+++ b/Engine/Source/Runtime/Engine/Private/Level.cpp
@@ -222,6 +222,27 @@ FArchive& operator<<( FArchive& Ar, FPrecomputedVolumeDistanceField& D )
 	return Ar;
 }
 
+FLevelSimplificationDetails::FLevelSimplificationDetails()
+ : DetailsPercentage(70.f)
+ , LandscapeExportLOD(7)
+ , bGenerateLandscapeNormalMap(true)
+ , bGenerateLandscapeRoughnessMap(false)
+ , bGenerateLandscapeSpecularMap(false)
+ , bBakeFoliageToLandscape(false)
+{
+}
+
+bool FLevelSimplificationDetails::operator == (const FLevelSimplificationDetails& Other) const
+{
+	return
+		DetailsPercentage == Other.DetailsPercentage &&
+		LandscapeExportLOD == Other.LandscapeExportLOD &&
+		bGenerateLandscapeNormalMap == Other.bGenerateLandscapeNormalMap &&
+		bGenerateLandscapeRoughnessMap == Other.bGenerateLandscapeRoughnessMap &&
+		bGenerateLandscapeSpecularMap == Other.bGenerateLandscapeSpecularMap &&
+		bBakeFoliageToLandscape == Other.bBakeFoliageToLandscape;
+}
+
 TMap<FName, UWorld*> ULevel::StreamedLevelsOwningWorld;
 
 ULevel::ULevel( const FObjectInitializer& ObjectInitializer )
