Performance fixes for content browser thumbnails
- Removed several of the border widgets from always being created if they are not used
- Added caching of expensive visibility attributes
- Cached WrapTextAt values for asset labels to all STextBlock to cache the data
- Removed several delegates from asset border colors

diff --git a/Engine/Source/Editor/ContentBrowser/Private/AssetViewWidgets.cpp b/Engine/Source/Editor/ContentBrowser/Private/AssetViewWidgets.cpp
--- a/Engine/Source/Editor/ContentBrowser/Private/AssetViewWidgets.cpp
+++ b/Engine/Source/Editor/ContentBrowser/Private/AssetViewWidgets.cpp
@@ -174,8 +174,17 @@ void SAssetViewItem::Tick( const FGeometry& AllottedGeometry, const double InCur
 {
 	SCompoundWidget::Tick(AllottedGeometry, InCurrentTime, InDeltaTime);
 
+	const float PrevSizeX = LastGeometry.Size.X;
+
 	LastGeometry = AllottedGeometry;
 
+	// Set cached wrap text width based on new "LastGeometry" value. 
+	// We set this only when changed because binding a delegate to text wrapping attributes is expensive
+	if( PrevSizeX != AllottedGeometry.Size.X && InlineRenameWidget.IsValid() )
+	{
+		InlineRenameWidget->SetWrapTextAt( GetNameTextWrapWidth() );
+	}
+
 	UpdatePackageDirtyState();
 
 	UpdateSourceControlState((float)InDeltaTime);
@@ -1296,7 +1305,6 @@ void SAssetTileItem::Construct( const FArguments& InArgs )
 					.HighlightText(InArgs._HighlightText)
 					.IsSelected(InArgs._IsSelected)
 					.IsReadOnly(ThumbnailEditMode)
-					.WrapTextAt(this, &SAssetTileItem::GetNameTextWrapWidth)
 					.Justification(ETextJustify::Center)
 					.LineBreakPolicy(FBreakIterator::CreateCamelCaseBreakIterator())
 			]
@@ -1352,11 +1360,7 @@ FSlateFontInfo SAssetTileItem::GetThumbnailFont() const
 	return FEditorStyle::GetFontStyle(RegularFont);
 }
 
-float SAssetTileItem::GetNameTextWrapWidth() const
-{
-	// Wrap to the entire size of the tile, minus some padding
-	return LastGeometry.Size.X - 2.f;
-}
+
 
 ///////////////////////////////
 // SAssetColumnItem
diff --git a/Engine/Source/Editor/ContentBrowser/Private/AssetViewWidgets.h b/Engine/Source/Editor/ContentBrowser/Private/AssetViewWidgets.h
--- a/Engine/Source/Editor/ContentBrowser/Private/AssetViewWidgets.h
+++ b/Engine/Source/Editor/ContentBrowser/Private/AssetViewWidgets.h
@@ -182,6 +182,8 @@ protected:
 	/** Delegate handler for when source control state changes */
 	void HandleSourceControlStateChanged();
 
+	/** Returns the width at which the name label will wrap the name */
+	virtual float GetNameTextWrapWidth() const { return 0.0f; }
 protected:
 
 	TSharedPtr< SInlineEditableTextBlock > InlineRenameWidget;
@@ -455,7 +457,10 @@ public:
 	/** Handles committing a name change */
 	virtual void OnAssetDataChanged() override;
 
-private:
+protected:
+	/** SAssetViewItem interface */
+	virtual float GetNameTextWrapWidth() const override { return LastGeometry.Size.X - 2.f; }
+
 	/** Returns the size of the thumbnail box widget */
 	FOptionalSize GetThumbnailBoxSize() const;
 
@@ -465,9 +470,6 @@ private:
 	/** Returns the size of the source control state box widget */
 	FOptionalSize GetSCCImageSize() const;
 
-	/** Returns the width at which the name label will wrap the name */
-	float GetNameTextWrapWidth() const;
-
 private:
 	/** The handle to the thumbnail that this item is rendering */
 	TSharedPtr<FAssetThumbnail> AssetThumbnail;
diff --git a/Engine/Source/Editor/UnrealEd/Private/AssetThumbnail.cpp b/Engine/Source/Editor/UnrealEd/Private/AssetThumbnail.cpp
--- a/Engine/Source/Editor/UnrealEd/Private/AssetThumbnail.cpp
+++ b/Engine/Source/Editor/UnrealEd/Private/AssetThumbnail.cpp
@@ -46,8 +46,8 @@ public:
 		HighlightedText = InArgs._HighlightedText;
 		Label = InArgs._Label;
 		HintColorAndOpacity = InArgs._HintColorAndOpacity;
-		AllowHintText = InArgs._AllowHintText;
-		ShowClassBackground = InArgs._ShowClassBackground;
+		bAllowHintText = InArgs._AllowHintText;
+		bShowClassBackground = InArgs._ShowClassBackground;
 
 		AssetThumbnail = InArgs._AssetThumbnail;
 		bHasRenderedThumbnail = false;
@@ -74,44 +74,42 @@ public:
 
 		TSharedRef<SOverlay> OverlayWidget = SNew(SOverlay);
 
-		FName Substyle;
-		if ( Class == UClass::StaticClass() )
+		if( Class == UClass::StaticClass() )
 		{
-			Substyle = FName(".ClassBackground");
 			ClassAssetClass = FindObject<UClass>(ANY_PACKAGE, *AssetData.AssetName.ToString());
 		}
-		else if(AssetTypeActions.IsValid())
-		{
-			Substyle = FName(".AssetBackground");
-		}
-		const FName BackgroundBrushName( *(Style.ToString() + Substyle.ToString()) );
 
 		ClassThumbnailBrushOverride = InArgs._ClassThumbnailBrushOverride;
 
-		OverlayWidget->AddSlot()
-		[
-			SNew(SBorder)
-			.BorderImage( this, &SAssetThumbnail::GetBackgroundBrush )
-			.BorderBackgroundColor( this, &SAssetThumbnail::GetAssetColor )
-			.Padding( GenericThumbnailBorderPadding )
-			.VAlign(VAlign_Center)
-			.HAlign(HAlign_Center)
-			.Visibility(this, &SAssetThumbnail::GetClassThumbnailVisibility)
+		const FSlateBrush* BackgroundBrush = GetBackgroundBrush();
+
+		if( ClassAssetClass != nullptr )
+		{
+			OverlayWidget->AddSlot()
 			[
-				SNew( SImage )
-				.Image( this, &SAssetThumbnail::GetClassThumbnailBrush )
-			]
-		];
+				SAssignNew(ClassThumbnailWidget, SBorder)
+				.BorderImage(BackgroundBrush)
+				.BorderBackgroundColor(AssetColor)
+				.Padding(GenericThumbnailBorderPadding)
+				.VAlign(VAlign_Center)
+				.HAlign(HAlign_Center)
+				.Visibility(this, &SAssetThumbnail::GetClassThumbnailVisibility)
+				[
+					SNew(SImage)
+					.Image(this, &SAssetThumbnail::GetClassThumbnailBrush)
+				]
+			];
+		}
 
 		// The generic representation of the thumbnail, for use before the rendered version, if it exists
 		OverlayWidget->AddSlot()
 		[
 			SNew(SBorder)
-			.BorderImage( this, &SAssetThumbnail::GetBackgroundBrush )
-			.BorderBackgroundColor( this, &SAssetThumbnail::GetAssetColor )
+			.BorderImage( BackgroundBrush )
+			.BorderBackgroundColor( AssetColor )
 			.Padding( GenericThumbnailBorderPadding )
 			.VAlign(VAlign_Center) .HAlign(HAlign_Center)
-			.Visibility(this, &SAssetThumbnail::GetGenericThumbnailVisibility)
+			.Visibility( this, &SAssetThumbnail::GetGenericThumbnailVisibility )
 			[
 				SAssignNew(LabelTextBlock, STextBlock)
 				.Text( GetLabelText() )
@@ -119,7 +117,6 @@ public:
 				.ColorAndOpacity( FEditorStyle::GetColor(Style, ".ColorAndOpacity") )
 				.ShadowOffset( FEditorStyle::GetVector(Style, ".ShadowOffset") )
 				.ShadowColorAndOpacity( FEditorStyle::GetColor(Style, ".ShadowColorAndOpacity") )
-				.WrapTextAt(this, &SAssetThumbnail::GetTextWrapWidth)
 				.HighlightText( HighlightedText )
 			]
 		];
@@ -148,49 +145,49 @@ public:
 			// The viewport for the rendered thumbnail, if it exists
 			OverlayWidget->AddSlot()
 			[
-				SNew(SBorder)
+				SAssignNew(RenderedThumbnailWidget, SBorder)
 				.Padding(0)
 				.BorderImage(FEditorStyle::GetBrush("NoBrush"))
 				.ColorAndOpacity(this, &SAssetThumbnail::GetViewportColorAndOpacity)
-				.Visibility(this, &SAssetThumbnail::GetViewportVisibility)
 				[
 					Viewport.ToSharedRef()
 				]
 			];
 		}
 
-		OverlayWidget->AddSlot()
-		.HAlign( HAlign_Center )
-		.VAlign( VAlign_Top )
-		.Padding(FMargin(2,2,2,2))
-		[
-			SNew(SBorder)
-			.BorderImage( FEditorStyle::GetBrush( Style, ".HintBackground" ) )
-			.BorderBackgroundColor( this, &SAssetThumbnail::GetHintBackgroundColor) //Adjust the opacity of the border itself
-			.ColorAndOpacity( HintColorAndOpacity ) //adjusts the opacity of the contents of the border
-			.Visibility( this, &SAssetThumbnail::GetHintTextVisibility )
-			.Padding(0)
-			[
-				SAssignNew( HintTextBlock, STextBlock )
-				.Text( GetLabelText() )
-				.Font( GetHintTextFont() )
-				.ColorAndOpacity( FEditorStyle::GetColor( Style, ".HintColorAndOpacity" ) )
-				.ShadowOffset( FEditorStyle::GetVector( Style, ".HintShadowOffset" ) )
-				.ShadowColorAndOpacity( FEditorStyle::GetColor( Style, ".HintShadowColorAndOpacity" ) )
-				.WrapTextAt(this, &SAssetThumbnail::GetTextWrapWidth)
-				.HighlightText( HighlightedText )
-			]
-		];
+		if( bAllowHintText )
+		{
+			OverlayWidget->AddSlot()
+				.HAlign(HAlign_Center)
+				.VAlign(VAlign_Top)
+				.Padding(FMargin(2, 2, 2, 2))
+				[
+					SNew(SBorder)
+					.BorderImage(FEditorStyle::GetBrush(Style, ".HintBackground"))
+					.BorderBackgroundColor(this, &SAssetThumbnail::GetHintBackgroundColor) //Adjust the opacity of the border itself
+					.ColorAndOpacity(HintColorAndOpacity) //adjusts the opacity of the contents of the border
+					.Visibility(this, &SAssetThumbnail::GetHintTextVisibility)
+					.Padding(0)
+					[
+						SAssignNew(HintTextBlock, STextBlock)
+						.Text(GetLabelText())
+						.Font(GetHintTextFont())
+						.ColorAndOpacity(FEditorStyle::GetColor(Style, ".HintColorAndOpacity"))
+						.ShadowOffset(FEditorStyle::GetVector(Style, ".HintShadowOffset"))
+						.ShadowColorAndOpacity(FEditorStyle::GetColor(Style, ".HintShadowColorAndOpacity"))
+						.HighlightText(HighlightedText)
+					]
+				];
+		}
 
 		// The asset color strip, only visible when the rendered viewport is
 		OverlayWidget->AddSlot()
 		.HAlign(HAlign_Fill)
 		.VAlign(VAlign_Bottom)
 		[
-			SNew(SBorder)
+			SAssignNew(AssetColorStripWidget, SBorder)
 			.BorderImage(FEditorStyle::GetBrush("WhiteBrush"))
 			.BorderBackgroundColor(AssetColor)
-			.Visibility(this, &SAssetThumbnail::GetViewportVisibility)
 			.Padding(this, &SAssetThumbnail::GetAssetColorStripPadding)
 		];
 
@@ -198,6 +195,9 @@ public:
 		[
 			OverlayWidget
 		];
+
+		UpdateThumbnailVisibilities();
+
 	}
 
 	FSlateColor GetHintBackgroundColor() const
@@ -230,11 +230,13 @@ public:
 			if ( LabelTextBlock.IsValid() )
 			{
 				LabelTextBlock->SetFont( GetTextFont() );
+				LabelTextBlock->SetWrapTextAt( GetTextWrapWidth() );
 			}
 
 			if ( HintTextBlock.IsValid() )
 			{
 				HintTextBlock->SetFont( GetHintTextFont() );
+				HintTextBlock->SetWrapTextAt( GetTextWrapWidth() );
 			}
 		}
 	}
@@ -288,6 +290,8 @@ private:
 		{
 			AssetColor = AssetTypeActions.Pin()->GetTypeColor();
 		}
+
+		UpdateThumbnailVisibilities();
 	}
 
 	FSlateFontInfo GetTextFont() const
@@ -322,11 +326,6 @@ private:
 		return FEditorStyle::GetBrush(BackgroundBrushName);
 	}
 
-	FSlateColor GetAssetColor() const
-	{
-		return AssetColor;
-	}
-
 	FSlateColor GetViewportBorderColorAndOpacity() const
 	{
 		return FLinearColor(AssetColor.R, AssetColor.G, AssetColor.B, ViewportFadeCurve.GetLerp());
@@ -366,15 +365,18 @@ private:
 
 	EVisibility GetClassThumbnailVisibility() const
 	{
-		const FSlateBrush* ClassThumbnailBrush = GetClassThumbnailBrush();
-		if(!bHasRenderedThumbnail && ClassThumbnailBrush)
+		if(!bHasRenderedThumbnail)
 		{
-			const FAssetData& AssetData = AssetThumbnail->GetAssetData();
-			FString AssetClass = AssetData.AssetClass.ToString();
-			UClass* Class = FindObjectSafe<UClass>(ANY_PACKAGE, *AssetClass);
-			if(Class == UClass::StaticClass())
+			const FSlateBrush* ClassThumbnailBrush = GetClassThumbnailBrush();
+			if( ClassThumbnailBrush )
 			{
-				return EVisibility::Visible;
+				const FAssetData& AssetData = AssetThumbnail->GetAssetData();
+				FString AssetClass = AssetData.AssetClass.ToString();
+				UClass* Class = FindObjectSafe<UClass>(ANY_PACKAGE, *AssetClass);
+				if(Class == UClass::StaticClass())
+				{
+					return EVisibility::Visible;
+				}
 			}
 		}
 
@@ -383,12 +385,12 @@ private:
 
 	EVisibility GetGenericThumbnailVisibility() const
 	{
-		return (GetClassThumbnailVisibility() == EVisibility::Visible || (bHasRenderedThumbnail && ViewportFadeAnimation.IsAtEnd())) ? EVisibility::Collapsed : EVisibility::Visible;
+		return ( CachedClassThumbnailVisibility == EVisibility::Visible ) || (bHasRenderedThumbnail && ViewportFadeAnimation.IsAtEnd()) ? EVisibility::Collapsed : EVisibility::Visible;
 	}
 
 	EVisibility GetHintTextVisibility() const
 	{
-		if ( AllowHintText && ( bHasRenderedThumbnail || !LabelTextBlock.IsValid() ) && HintColorAndOpacity.Get().A > 0 )
+		if ( bAllowHintText && ( bHasRenderedThumbnail || !LabelTextBlock.IsValid() ) && HintColorAndOpacity.Get().A > 0 )
 		{
 			return EVisibility::Visible;
 		}
@@ -400,7 +402,7 @@ private:
 	{
 		if ( !bHasRenderedThumbnail && AssetData == AssetThumbnail->GetAssetData() && ShouldRender() )
 		{
-			bHasRenderedThumbnail = true;
+			OnRenderedThumbnailChanged( true );
 			ViewportFadeAnimation.Play();
 		}
 	}
@@ -409,7 +411,7 @@ private:
 	{
 		if ( bHasRenderedThumbnail && AssetData == AssetThumbnail->GetAssetData() )
 		{
-			bHasRenderedThumbnail = false;
+			OnRenderedThumbnailChanged( false );
 		}
 	}
 
@@ -542,9 +544,41 @@ private:
 		return AssetData.AssetName.ToString();
 	}
 
+	void OnRenderedThumbnailChanged( bool bInHasRenderedThumbnail )
+	{
+		bHasRenderedThumbnail = bInHasRenderedThumbnail;
+
+		UpdateThumbnailVisibilities();
+	}
+
+	void UpdateThumbnailVisibilities()
+	{
+	
+		CachedClassThumbnailVisibility = GetClassThumbnailVisibility();
+	
+		if( ClassThumbnailWidget.IsValid() )
+		{
+			ClassThumbnailWidget->SetVisibility( CachedClassThumbnailVisibility );
+		}
+
+
+		if( RenderedThumbnailWidget.IsValid() )
+		{
+			RenderedThumbnailWidget->SetVisibility( GetViewportVisibility() );
+		}
+
+		if( AssetColorStripWidget.IsValid() )
+		{
+			AssetColorStripWidget->SetVisibility( GetViewportVisibility() );
+		}
+	}
+
 private:
 	TSharedPtr<STextBlock> LabelTextBlock;
 	TSharedPtr<STextBlock> HintTextBlock;
+	TSharedPtr<SBorder> ClassThumbnailWidget;
+	TSharedPtr<SBorder> RenderedThumbnailWidget;
+	TSharedPtr<SBorder> AssetColorStripWidget;
 	TSharedPtr<FAssetThumbnail> AssetThumbnail;
 	FCurveSequence ViewportFadeAnimation;
 	FCurveHandle ViewportFadeCurve;
@@ -557,10 +591,13 @@ private:
 	TAttribute< FText > HighlightedText;
 	EThumbnailLabel::Type Label;
 
+	/** Cached class thumbnail visibility. This only ever changes if the asset changes */ 
+	EVisibility CachedClassThumbnailVisibility;
+
 	TAttribute< FLinearColor > HintColorAndOpacity;
-	bool AllowHintText;
+	bool bAllowHintText;
 
-	bool ShowClassBackground;
+	bool bShowClassBackground;
 
 	/** The name of the thumbnail which should be used instead of the class thumbnail. */
 	FName ClassThumbnailBrushOverride;
diff --git a/Engine/Source/Runtime/Slate/Private/Widgets/Text/SInlineEditableTextBlock.cpp b/Engine/Source/Runtime/Slate/Private/Widgets/Text/SInlineEditableTextBlock.cpp
--- a/Engine/Source/Runtime/Slate/Private/Widgets/Text/SInlineEditableTextBlock.cpp
+++ b/Engine/Source/Runtime/Slate/Private/Widgets/Text/SInlineEditableTextBlock.cpp
@@ -120,6 +120,11 @@ void SInlineEditableTextBlock::SetText( const FString& InText )
 	TextBox->SetText( Text );
 }
 
+void SInlineEditableTextBlock::SetWrapTextAt( const TAttribute<float>& InWrapTextAt )
+{
+	TextBlock->SetWrapTextAt( InWrapTextAt );
+}
+
 FReply SInlineEditableTextBlock::OnMouseButtonDown( const FGeometry& MyGeometry, const FPointerEvent& MouseEvent )
 {
 	if( !MouseEvent.IsMouseButtonDown(EKeys::LeftMouseButton) || MouseEvent.IsControlDown() || MouseEvent.IsShiftDown())
diff --git a/Engine/Source/Runtime/Slate/Public/Widgets/Text/SInlineEditableTextBlock.h b/Engine/Source/Runtime/Slate/Public/Widgets/Text/SInlineEditableTextBlock.h
--- a/Engine/Source/Runtime/Slate/Public/Widgets/Text/SInlineEditableTextBlock.h
+++ b/Engine/Source/Runtime/Slate/Public/Widgets/Text/SInlineEditableTextBlock.h
@@ -110,6 +110,9 @@ class SLATE_API SInlineEditableTextBlock: public SCompoundWidget
 	void SetText( const TAttribute< FText >& InText );
 	void SetText( const FString& InText );
 
+	/** Sets the wrap text at attribute.  See WrapTextAt attribute */
+	void SetWrapTextAt(const TAttribute<float>& InWrapTextAt);
+
 protected:
 	/** Callback for the text box's OnTextChanged event */
 	void OnTextChanged(const FText& InText);
