[UE-5075/UE-5070] Adding accept source/target buttons to the merge tool.

diff --git a/Engine/Source/Developer/Merge/Private/Merge.cpp b/Engine/Source/Developer/Merge/Private/Merge.cpp
--- a/Engine/Source/Developer/Merge/Private/Merge.cpp
+++ b/Engine/Source/Developer/Merge/Private/Merge.cpp
@@ -120,21 +120,19 @@ static TSharedPtr<SWidget> GenerateMergeTabContents(TSharedRef<FBlueprintEditor>
                                                     const FRevisionInfo&    BaseRevInfo,
 													const UBlueprint*       RemoteBlueprint,
 													const FRevisionInfo&    RemoteRevInfo,
-													const UBlueprint&       LocalBlueprint,
+													const UBlueprint*       LocalBlueprint,
 													const FOnMergeResolved& MereResolutionCallback)
 {
 	TSharedPtr<SWidget> TabContent;
 
 	if (RemoteBlueprint && BaseBlueprint)
 	{
-		LocalBlueprint.bDuplicatingReadOnly = true;
 		FBlueprintMergeData Data(Editor
-			, static_cast<const UBlueprint*>(StaticDuplicateObject(&LocalBlueprint, GetTransientPackage(), TEXT("None")))
+			, LocalBlueprint
 			, BaseBlueprint
 			, RemoteRevInfo
 			, RemoteBlueprint
 			, BaseRevInfo);
-		LocalBlueprint.bDuplicatingReadOnly = false;
 
 		TabContent = SNew(SBlueprintMerge, Data)
 			.OnMergeResolved(MereResolutionCallback);
@@ -158,7 +156,7 @@ static TSharedPtr<SWidget> GenerateMergeTabContents(TSharedRef<FBlueprintEditor>
 		DisplayErrorMessage(
 			FText::Format(
 				LOCTEXT("MergeRevisionLoadFailed", "Aborted Merge of {0} because we could not load {1}")
-				, FText::FromString(LocalBlueprint.GetName())
+				, FText::FromString(LocalBlueprint->GetName())
 				, MissingFiles
 			)
 		);
@@ -246,7 +244,7 @@ TSharedRef<SDockTab> FMerge::GenerateMergeWidget(const UBlueprint& Object, TShar
 		FRevisionInfo BaseRevInfo = FRevisionInfo::InvalidRevision();
 		const UBlueprint* BaseBlueprint = Cast< UBlueprint >(LoadBaseRev(PackageName, AssetName, SourceControlStateRef, BaseRevInfo));
 
-		Contents = GenerateMergeTabContents(Editor, BaseBlueprint, BaseRevInfo, RemoteBlueprint, CurrentRevInfo, Object, FOnMergeResolved());
+		Contents = GenerateMergeTabContents(Editor, BaseBlueprint, BaseRevInfo, RemoteBlueprint, CurrentRevInfo, &Object, FOnMergeResolved());
 	}
 
 	TSharedRef<SDockTab> Tab =  FGlobalTabmanager::Get()->InvokeTab(MergeToolTabId);
@@ -268,7 +266,7 @@ TSharedRef<SDockTab> FMerge::GenerateMergeWidget(const UBlueprint* BaseBlueprint
 	}
 
 	// @TODO: pipe revision info through
-	TSharedPtr<SWidget> TabContents = GenerateMergeTabContents(Editor, BaseBlueprint, FRevisionInfo::InvalidRevision(), RemoteBlueprint, FRevisionInfo::InvalidRevision(), *LocalBlueprint, MergeResolutionCallback);
+	TSharedPtr<SWidget> TabContents = GenerateMergeTabContents(Editor, BaseBlueprint, FRevisionInfo::InvalidRevision(), RemoteBlueprint, FRevisionInfo::InvalidRevision(), LocalBlueprint, MergeResolutionCallback);
 
 	TSharedRef<SDockTab> Tab = FGlobalTabmanager::Get()->InvokeTab(MergeToolTabId);
 	Tab->SetContent(TabContents.ToSharedRef());
diff --git a/Engine/Source/Developer/Merge/Private/SBlueprintMerge.cpp b/Engine/Source/Developer/Merge/Private/SBlueprintMerge.cpp
--- a/Engine/Source/Developer/Merge/Private/SBlueprintMerge.cpp
+++ b/Engine/Source/Developer/Merge/Private/SBlueprintMerge.cpp
@@ -23,7 +23,7 @@ SBlueprintMerge::SBlueprintMerge()
 {
 }
 
-static void WriteBackup(UPackage& Package, const FString& Directory, const FString& Filename)
+static FString WriteBackup(UPackage& Package, const FString& Directory, const FString& Filename)
 {
 	if (GEditor)
 	{
@@ -31,9 +31,28 @@ static void WriteBackup(UPackage& Package, const FString& Directory, const FStri
 		FString OriginalFilename;
 		if (FPackageName::DoesPackageExist(Package.GetName(), nullptr, &OriginalFilename))
 		{
-			IFileManager::Get().Copy(*DestinationFilename, *OriginalFilename);
+			if (IFileManager::Get().Copy(*DestinationFilename, *OriginalFilename) == COPY_OK)
+			{
+				return DestinationFilename;
+			}
 		}
 	}
+	return FString();
+}
+
+static EAppReturnType::Type PromptUserIfUndoBufferToBeCleared(UBlueprint* MergeTarget)
+{
+	EAppReturnType::Type OkToMerge = EAppReturnType::Yes;
+	if (FKismetEditorUtilities::IsReferencedByUndoBuffer(MergeTarget))
+	{
+		FText TargetName = FText::FromName(MergeTarget->GetFName());
+
+		const FText WarnMessage = FText::Format(LOCTEXT("WarnOfUndoClear",
+			"{0} has undo actions associated with it. The undo buffer must be cleared to complete this merge. \n\n\
+			You will not be able to undo previous actions after this. Would you like to continue?"), TargetName);
+		OkToMerge = FMessageDialog::Open(EAppMsgType::YesNo, WarnMessage);
+	}
+	return OkToMerge;
 }
 
 void SBlueprintMerge::Construct(const FArguments InArgs, const FBlueprintMergeData& InData)
@@ -42,25 +61,42 @@ void SBlueprintMerge::Construct(const FArguments InArgs, const FBlueprintMergeDa
 
 	Data = InData;
 	OnMergeResolved = InArgs._OnMergeResolved;
+	LocalBackupPath.Empty();
+
+	// we want to make a read only copy of the local blueprint, since we allow 
+	// the user to modify/mutate the merge result (which starts as the local copy)
+	if (InData.BlueprintLocal == GetTargetBlueprint())
+	{
+		UBlueprint* LocalBlueprint = GetTargetBlueprint();
+		// we use this to suppress blueprint compilation during StaticDuplicateObject()
+		// @TODO: we don't remember exactly why this was needed, but I saw the 
+		//        the blueprint become unplacable and the file become unloadable
+		//        without this, but don't know exactly why
+		TGuardValue<bool> GuardDuplicatingReadOnly(LocalBlueprint->bDuplicatingReadOnly, true);
+
+		UPackage* TransientPackage = GetTransientPackage();
+		Data.BlueprintLocal = Cast<const UBlueprint>( StaticDuplicateObject(LocalBlueprint, TransientPackage,
+			*MakeUniqueObjectName(TransientPackage, LocalBlueprint->GetClass(), LocalBlueprint->GetFName()).ToString()) );
+	}
 
 	BackupSubDir = FPaths::GameSavedDir() / TEXT("Backup") / TEXT("Resolve_Backup[") + FDateTime::Now().ToString(TEXT("%Y-%m-%d-%H-%M-%S")) + TEXT("]");
 
 	// Because merge operations are so destructive and can be confusing, lets write backups of the files involved:
-	WriteBackup(*Data.BlueprintRemote->GetOutermost(), BackupSubDir, TEXT("RemoteAsset") + FPackageName::GetAssetPackageExtension());
-	WriteBackup(*Data.BlueprintBase->GetOutermost(), BackupSubDir, TEXT("CommonBaseAsset") + FPackageName::GetAssetPackageExtension());
-	WriteBackup(*Data.BlueprintLocal->GetOutermost(), BackupSubDir, TEXT("LocalAsset") + FPackageName::GetAssetPackageExtension());
+	WriteBackup(*InData.BlueprintRemote->GetOutermost(), BackupSubDir, TEXT("RemoteAsset") + FPackageName::GetAssetPackageExtension());
+	WriteBackup(*InData.BlueprintBase->GetOutermost(), BackupSubDir, TEXT("CommonBaseAsset") + FPackageName::GetAssetPackageExtension());
+	LocalBackupPath = WriteBackup(*InData.BlueprintLocal->GetOutermost(), BackupSubDir, TEXT("LocalAsset") + FPackageName::GetAssetPackageExtension());
 
-	auto GraphView = SNew(SMergeGraphView, InData);
+	auto GraphView = SNew(SMergeGraphView, Data);
 	GraphControl.Widget = GraphView;
 	GraphControl.DiffControl = &GraphView.Get();
 	GraphControl.MergeControl = &GraphView.Get();
 
-	auto TreeView = SNew(SMergeTreeView, InData);
+	auto TreeView = SNew(SMergeTreeView, Data);
 	TreeControl.Widget = TreeView;
 	TreeControl.DiffControl = &TreeView.Get();
 	TreeControl.MergeControl = &TreeView.Get();
 
-	auto DetailsView = SNew(SMergeDetailsView, InData);
+	auto DetailsView = SNew(SMergeDetailsView, Data);
 	DetailsControl.Widget = DetailsView;
 	DetailsControl.DiffControl = &DetailsView.Get();
 	DetailsControl.MergeControl = &DetailsView.Get();
@@ -99,7 +135,21 @@ void SBlueprintMerge::Construct(const FArguments InArgs, const FBlueprintMergeDa
 	);
 
 	// buttons for finishing the merge:
-	ToolbarBuilder.AddSeparator();
+	ToolbarBuilder.AddSeparator(); 
+	ToolbarBuilder.AddToolBarButton(
+		FUIAction(FExecuteAction::CreateSP(this, &SBlueprintMerge::OnAcceptRemote))
+		, NAME_None
+		, LOCTEXT("AcceptRemoteLabel", "Accept Source")
+		, LOCTEXT("AcceptRemoteTooltip", "Complete the merge operation - Replaces the Blueprint with a copy of the remote file.")
+		, FSlateIcon(FEditorStyle::GetStyleSetName(), "BlueprintMerge.AcceptSource")
+	);
+	ToolbarBuilder.AddToolBarButton(
+		FUIAction(FExecuteAction::CreateSP(this, &SBlueprintMerge::OnAcceptLocal))
+		, NAME_None
+		, LOCTEXT("AcceptLocalLabel", "Accept Target")
+		, LOCTEXT("AcceptLocalTooltip", "Complete the merge operation - Leaves the target Blueprint unchanged.")
+		, FSlateIcon(FEditorStyle::GetStyleSetName(), "BlueprintMerge.AcceptTarget")
+	);
 	ToolbarBuilder.AddToolBarButton(
 		FUIAction(FExecuteAction::CreateSP(this, &SBlueprintMerge::OnFinishMerge))
 		, NAME_None
@@ -134,8 +184,8 @@ void SBlueprintMerge::Construct(const FArguments InArgs, const FBlueprintMergeDa
 		]
 	];
 
-	InData.OwningEditor.Pin()->OnModeSet().AddSP( this, &SBlueprintMerge::OnModeChanged );
-	OnModeChanged( InData.OwningEditor.Pin()->GetCurrentMode() );
+	Data.OwningEditor.Pin()->OnModeSet().AddSP( this, &SBlueprintMerge::OnModeChanged );
+	OnModeChanged( Data.OwningEditor.Pin()->GetCurrentMode() );
 }
 
 UBlueprint* SBlueprintMerge::GetTargetBlueprint()
@@ -185,33 +235,7 @@ bool SBlueprintMerge::HasPrevConflict() const
 
 void SBlueprintMerge::OnFinishMerge()
 {
-	UPackage* Package = GetTargetBlueprint()->GetOutermost();
-	TArray<UPackage*> PackagesToSave;
-	PackagesToSave.Add(Package);
-
-	// Perform the resolve with the SCC plugin, we do this first so that the editor doesn't warn about writing to a file that is unresolved:
-	ISourceControlModule::Get().GetProvider().Execute(ISourceControlOperation::Create<FResolve>(), SourceControlHelpers::PackageFilenames(PackagesToSave), EConcurrency::Synchronous);
-
-	const auto Result = FEditorFileUtils::PromptForCheckoutAndSave(PackagesToSave, /*bCheckDirty=*/ false, /*bPromptToSave=*/ false);
-	if (Result != FEditorFileUtils::PR_Success)
-	{
-		static const float MessageTime = 5.0f;
-		const FText ErrorMessage = FText::Format( LOCTEXT("MergeWriteFailedError", "Failed to write merged files, please look for backups in {0}"), FText::FromString(BackupSubDir) );
-
-		FNotificationInfo ErrorNotification(ErrorMessage);
-		FSlateNotificationManager::Get().AddNotification( ErrorNotification );
-	}
-
-	// should come before CloseMergeTool(), because CloseMergeTool() makes its
-	// own call to OnMergeResolved (with an "Unknown" state).
-	OnMergeResolved.ExecuteIfBound(Package, EMergeResult::Completed);
-
-	// if we're using the merge command-line, it might close everything once
-	// a resolution is found (so the editor may be invalid now)
-	if (Data.OwningEditor.IsValid())
-	{
-		Data.OwningEditor.Pin()->CloseMergeTool();
-	}
+	ResolveMerge(GetTargetBlueprint());
 }
 
 void SBlueprintMerge::OnCancelClicked()
@@ -256,4 +280,91 @@ void SBlueprintMerge::OnModeChanged(FName NewMode)
 	}
 }
 
+void SBlueprintMerge::OnAcceptRemote()
+{
+	UBlueprint* TargetBlueprint = GetTargetBlueprint();
+	if (PromptUserIfUndoBufferToBeCleared(TargetBlueprint) == EAppReturnType::Yes)
+	{
+		const UBlueprint* RemoteBlueprint = Data.BlueprintRemote;
+		if (UBlueprint* NewBlueprint = FKismetEditorUtilities::ReplaceBlueprint(TargetBlueprint, RemoteBlueprint))
+		{
+			ResolveMerge(NewBlueprint);
+		}
+	}
+}
+
+void SBlueprintMerge::OnAcceptLocal()
+{
+	UBlueprint* TargetBlueprint = GetTargetBlueprint();
+	UPackage* TargetPackage = TargetBlueprint->GetOutermost();
+
+	FString PackageFilename;
+	if (FPackageName::DoesPackageExist(TargetPackage->GetName(), /*Guid =*/nullptr, &PackageFilename))
+	{
+		if (PromptUserIfUndoBufferToBeCleared(TargetBlueprint) == EAppReturnType::Yes)
+		{
+			FString SrcFilename = LocalBackupPath;
+			auto OverwriteBlueprintFile = [TargetBlueprint, &PackageFilename, &SrcFilename](UBlueprint* /*UnloadedBP*/)
+			{
+				if (IFileManager::Get().Copy(*PackageFilename, *SrcFilename) != COPY_OK)
+				{
+					const FText TargetName = FText::FromName(TargetBlueprint->GetFName());
+					const FText ErrorMessage = FText::Format(LOCTEXT("FailedMergeLocalCopy", "Failed to overwrite {0} target file."), TargetName);
+					FSlateNotificationManager::Get().AddNotification(FNotificationInfo(ErrorMessage));
+				}
+			};
+			// we have to wait until the package is unloaded (and the file is 
+			// unlocked) before we can overwrite the file; so we use defer the
+			// copy until then by adding it to the OnBlueprintUnloaded delegate
+			FKismetEditorUtilities::FOnBlueprintUnloaded::FDelegate OnPackageUnloaded = FKismetEditorUtilities::FOnBlueprintUnloaded::FDelegate::CreateLambda(OverwriteBlueprintFile);
+			FKismetEditorUtilities::OnBlueprintUnloaded.Add(OnPackageUnloaded);
+
+
+			if (UBlueprint* NewBlueprint = FKismetEditorUtilities::ReloadBlueprint(TargetBlueprint))
+			{
+				ResolveMerge(NewBlueprint);
+			}
+			FKismetEditorUtilities::OnBlueprintUnloaded.Remove(OnPackageUnloaded);
+		}
+	}
+}
+
+void SBlueprintMerge::ResolveMerge(UBlueprint* ResultantBlueprint)
+{
+	UPackage* Package = ResultantBlueprint->GetOutermost();
+	TArray<UPackage*> PackagesToSave;
+	PackagesToSave.Add(Package);
+
+	// Perform the resolve with the SCC plugin, we do this first so that the editor doesn't warn about writing to a file that is unresolved:
+	ISourceControlModule::Get().GetProvider().Execute(ISourceControlOperation::Create<FResolve>(), SourceControlHelpers::PackageFilenames(PackagesToSave), EConcurrency::Synchronous);
+
+	FEditorFileUtils::EPromptReturnCode const SaveResult = FEditorFileUtils::PromptForCheckoutAndSave(PackagesToSave, /*bCheckDirty=*/ false, /*bPromptToSave=*/ false);
+	if (SaveResult != FEditorFileUtils::PR_Success)
+	{
+		const FText ErrorMessage = FText::Format(LOCTEXT("MergeWriteFailedError", "Failed to write merged files, please look for backups in {0}"), FText::FromString(BackupSubDir));
+
+		FNotificationInfo ErrorNotification(ErrorMessage);
+		FSlateNotificationManager::Get().AddNotification(ErrorNotification);
+	}
+
+	// in the case where we have replaced/reloaded the blueprint 
+	// ("TargetBlueprint" is most likely, and should be, forcefully closing in 
+	// this scenario... we need to open the result to take its place)
+	if (ResultantBlueprint != GetTargetBlueprint())
+	{
+		FAssetEditorManager::Get().OpenEditorForAsset(ResultantBlueprint);
+	}
+	
+	// should come before CloseMergeTool(), because CloseMergeTool() makes its
+	// own call to OnMergeResolved (with an "Unknown" state).
+	OnMergeResolved.ExecuteIfBound(Package, EMergeResult::Completed);
+
+	// if we're using the merge command-line, it might close everything once
+	// a resolution is found (so the editor may be invalid now)
+	if (Data.OwningEditor.IsValid())
+	{
+		Data.OwningEditor.Pin()->CloseMergeTool();
+	}
+}
+
 #undef LOCTEXT_NAMESPACE
diff --git a/Engine/Source/Developer/Merge/Private/SBlueprintMerge.h b/Engine/Source/Developer/Merge/Private/SBlueprintMerge.h
--- a/Engine/Source/Developer/Merge/Private/SBlueprintMerge.h
+++ b/Engine/Source/Developer/Merge/Private/SBlueprintMerge.h
@@ -35,9 +35,13 @@ private:
 	void OnFinishMerge();
 	void OnCancelClicked();
 	void OnModeChanged(FName NewMode);
+	void OnAcceptRemote();
+	void OnAcceptLocal();
+	void ResolveMerge(UBlueprint* Result);
 
 	FBlueprintMergeData		Data; 
 	FString					BackupSubDir;
+	FString					LocalBackupPath;
 
 	TSharedPtr<SBorder>		MainView;
 
diff --git a/Engine/Source/Editor/EditorStyle/Private/SlateEditorStyle.cpp b/Engine/Source/Editor/EditorStyle/Private/SlateEditorStyle.cpp
--- a/Engine/Source/Editor/EditorStyle/Private/SlateEditorStyle.cpp
+++ b/Engine/Source/Editor/EditorStyle/Private/SlateEditorStyle.cpp
@@ -4666,6 +4666,8 @@ void FSlateEditorStyle::FStyle::SetupPersonaStyle()
 		Set("BlueprintMerge.PrevDiff", new IMAGE_BRUSH("/Icons/diff_prev_40x", Icon16x16));
 		Set("BlueprintMerge.Finish", new IMAGE_BRUSH("/Icons/LV_Save", Icon16x16));
 		Set("BlueprintMerge.Cancel", new IMAGE_BRUSH("/Icons/LV_Remove", Icon16x16));
+		Set("BlueprintMerge.AcceptSource", new IMAGE_BRUSH("/Icons/AcceptMergeSource_40x", Icon16x16));
+		Set("BlueprintMerge.AcceptTarget", new IMAGE_BRUSH("/Icons/AcceptMergeTarget_40x", Icon16x16));
 	}
 
 	// Play in editor / play in world
diff --git a/Engine/Source/Editor/Kismet/Private/DetailsDiff.cpp b/Engine/Source/Editor/Kismet/Private/DetailsDiff.cpp
--- a/Engine/Source/Editor/Kismet/Private/DetailsDiff.cpp
+++ b/Engine/Source/Editor/Kismet/Private/DetailsDiff.cpp
@@ -31,6 +31,11 @@ FDetailsDiff::FDetailsDiff(const UObject* InObject, const TArray< FPropertyPath
 	DifferingProperties = DetailsView->GetPropertiesInOrderDisplayed();
 }
 
+FDetailsDiff::~FDetailsDiff()
+{
+	DetailsView->SetOnDisplayedPropertiesChanged(::FOnDisplayedPropertiesChanged());
+}
+
 void FDetailsDiff::HighlightProperty(const FPropertySoftPath& PropertyName)
 {
 	// resolve the property soft path:
diff --git a/Engine/Source/Editor/Kismet/Public/DetailsDiff.h b/Engine/Source/Editor/Kismet/Public/DetailsDiff.h
--- a/Engine/Source/Editor/Kismet/Public/DetailsDiff.h
+++ b/Engine/Source/Editor/Kismet/Public/DetailsDiff.h
@@ -10,6 +10,7 @@ public:
 	DECLARE_DELEGATE( FOnDisplayedPropertiesChanged );
 
 	FDetailsDiff( const UObject* InObject, const TArray< FPropertyPath >& InDifferingProperties, FOnDisplayedPropertiesChanged InOnDisplayedPropertiesChanged );
+	~FDetailsDiff();
 
 	void HighlightProperty( const FPropertySoftPath& PropertyName );
 	TSharedRef< SWidget > DetailsWidget();
diff --git a/Engine/Source/Editor/UnrealEd/Private/Kismet2/Kismet2.cpp b/Engine/Source/Editor/UnrealEd/Private/Kismet2/Kismet2.cpp
--- a/Engine/Source/Editor/UnrealEd/Private/Kismet2/Kismet2.cpp
+++ b/Engine/Source/Editor/UnrealEd/Private/Kismet2/Kismet2.cpp
@@ -34,6 +34,9 @@
 #include "ComponentAssetBroker.h"
 #include "BlueprintEditorSettings.h"
 #include "Editor/UnrealEd/Classes/Editor/Transactor.h"
+#include "Editor/UnrealEd/Public/PackageTools.h"
+#include "NotificationManager.h"
+#include "SNotificationList.h" // for FNotificationInfo
 
 #define LOCTEXT_NAMESPACE "UnrealEd.Editor"
 
@@ -160,6 +163,181 @@ void FObjectsBeingDebuggedIterator::FindNextLevelScriptActor()
 	}
 }
 
+//////////////////////////////////////////////////////////////////////////
+// FBlueprintUnloader
+
+/** Utility struct, used to aid in unloading and replacing a specific blueprint. */
+struct FBlueprintUnloader
+{
+public:
+	FBlueprintUnloader(UBlueprint* OldBlueprint);
+
+	/** 
+	 * Unloads the specified Blueprint (marking it pending-kill, and removing it 
+	 * from its outer package). Optionally, will unload the package as well.
+	 *
+	 * @param  bResetPackage	Whether or not this should unload the entire package.
+	 */
+	void UnloadBlueprint(const bool bResetPackage);
+	
+	/** 
+	 * Replaces all old references to the original blueprints (its class/CDO/etc.)
+	 * @param  NewBlueprint	The blueprint to replace old references with
+	 */
+	void ReplaceStaleRefs(UBlueprint* NewBlueprint);
+
+private:
+	TWeakObjectPtr<UBlueprint> OldBlueprint;
+	UClass*  OldGeneratedClass;
+	UObject* OldCDO;
+	UClass*  OldSkeletonClass;
+	UObject* OldSkelCDO;
+};
+
+FBlueprintUnloader::FBlueprintUnloader(UBlueprint* OldBlueprintIn)
+	: OldBlueprint(OldBlueprintIn)
+	, OldGeneratedClass(OldBlueprint->GeneratedClass)
+	, OldCDO(nullptr)
+	, OldSkeletonClass(OldBlueprint->SkeletonGeneratedClass)
+	, OldSkelCDO(nullptr)
+{
+	if (OldGeneratedClass != nullptr)
+	{
+		OldCDO = OldGeneratedClass->GetDefaultObject(/*bCreateIfNeeded =*/false);
+	}
+	if (OldSkeletonClass != nullptr)
+	{
+		OldSkelCDO = OldSkeletonClass->GetDefaultObject(/*bCreateIfNeeded =*/false);
+	}
+	OldBlueprint = OldBlueprintIn;
+}
+
+void FBlueprintUnloader::UnloadBlueprint(const bool bResetPackage)
+{
+	if (OldBlueprint.IsValid())
+	{
+		UBlueprint* UnloadingBp = OldBlueprint.Get();
+
+		UPackage* const OldPackage = Cast<UPackage>(UnloadingBp->GetOutermost());
+		bool const bIsDirty = OldPackage->IsDirty();
+
+		UPackage* const TransientPackage = GetTransientPackage();
+		check(OldPackage != TransientPackage); // is the blueprint already unloaded?
+		
+		FName const BlueprintName = UnloadingBp->GetFName();
+		// move the blueprint to the transient package (to be picked up by garbage collection later)
+		FName UnloadedName = MakeUniqueObjectName(TransientPackage, UBlueprint::StaticClass(), BlueprintName);
+		UnloadingBp->Rename(*UnloadedName.ToString(), TransientPackage, REN_DontCreateRedirectors | REN_DoNotDirty);
+		// @TODO: currently, REN_DoNotDirty does not guarantee that the package 
+		//        will not be marked dirty
+		OldPackage->SetDirtyFlag(bIsDirty);
+
+		// make sure the blueprint is properly trashed (remove it from the package)
+		UnloadingBp->SetFlags(RF_Transient);
+		UnloadingBp->ClearFlags(RF_Standalone | RF_RootSet | RF_Transactional);
+		UnloadingBp->RemoveFromRoot();
+		UnloadingBp->MarkPendingKill();
+		// if it's in the undo buffer, then we have to clear that...
+		if (FKismetEditorUtilities::IsReferencedByUndoBuffer(UnloadingBp))
+		{
+			GEditor->Trans->Reset(LOCTEXT("UnloadedBlueprint", "Unloaded Blueprint"));
+		}
+
+		if (bResetPackage)
+		{
+			TArray<UPackage*> PackagesToUnload;
+			PackagesToUnload.Add(OldPackage);
+
+			FText PackageUnloadError;
+			PackageTools::UnloadPackages(PackagesToUnload, PackageUnloadError);
+
+			if (!PackageUnloadError.IsEmpty())
+			{
+				const FText ErrorMessage = FText::Format(LOCTEXT("UnloadBpPackageError", "Failed to unload Bluprint '{0}': {1}"),
+					FText::FromName(BlueprintName), PackageUnloadError);
+				FSlateNotificationManager::Get().AddNotification(FNotificationInfo(ErrorMessage));
+
+				// fallback to manually setting up the package so it can reload 
+				// the blueprint 
+				ResetLoaders(OldPackage);
+				OldPackage->ClearFlags(RF_WasLoaded);
+				OldPackage->bHasBeenFullyLoaded = false;
+				OldPackage->GetMetaData()->RemoveMetaDataOutsidePackage();
+			}
+		}
+
+		FKismetEditorUtilities::OnBlueprintUnloaded.Broadcast(UnloadingBp);
+
+		// handled in FBlueprintEditor (from the OnBlueprintUnloaded event)
+// 		IAssetEditorInstance* EditorInst = FAssetEditorManager::Get().FindEditorForAsset(UnloadingBp, /*bFocusIfOpen =*/false);
+// 		if (EditorInst != nullptr)
+// 		{
+// 			EditorInst->CloseWindow();
+// 		}
+	}
+}
+
+void FBlueprintUnloader::ReplaceStaleRefs(UBlueprint* NewBlueprint)
+{
+	//--------------------------------------
+	// Construct redirects
+	//--------------------------------------
+
+	TMap<UObject*, UObject*> Redirects;
+	TArray<UObject*> OldObjsNeedingReplacing;
+
+	if (OldBlueprint.IsValid(/*bEvenIfPendingKill =*/true))
+	{
+		UBlueprint* ToBeReplaced = OldBlueprint.Get(/*bEvenIfPendingKill =*/true);
+		if (OldGeneratedClass != nullptr)
+		{
+			OldObjsNeedingReplacing.Add(OldGeneratedClass);
+			Redirects.Add(OldGeneratedClass, NewBlueprint->GeneratedClass);
+		}
+		if (OldCDO != nullptr)
+		{
+			OldObjsNeedingReplacing.Add(OldCDO);
+			Redirects.Add(OldCDO, NewBlueprint->GeneratedClass->GetDefaultObject());
+		}
+		if (OldSkeletonClass != nullptr)
+		{
+			OldObjsNeedingReplacing.Add(OldSkeletonClass);
+			Redirects.Add(OldSkeletonClass, NewBlueprint->SkeletonGeneratedClass);
+		}
+		if (OldSkelCDO != nullptr)
+		{
+			OldObjsNeedingReplacing.Add(OldSkelCDO);
+			Redirects.Add(OldSkelCDO, NewBlueprint->SkeletonGeneratedClass->GetDefaultObject());
+		}
+
+		OldObjsNeedingReplacing.Add(ToBeReplaced);
+		Redirects.Add(ToBeReplaced, NewBlueprint);
+
+		// clear the object being debugged; otherwise ReplaceInstancesOfClass()  
+		// trys to reset it with a new level instance, and OldBlueprint won't 
+		// match the new instance's type (it's now a NewBlueprint)
+		ToBeReplaced->SetObjectBeingDebugged(nullptr);
+	}
+
+	//--------------------------------------
+	// Replace old references
+	//--------------------------------------
+
+	TArray<UObject*> Referencers;
+	// find all objects, still referencing the old blueprint/class/cdo/etc.
+	for (auto Referencer : TFindObjectReferencers<UObject>(OldObjsNeedingReplacing, /*PackageToCheck =*/nullptr, /*bIgnoreTemplates =*/false))
+	{
+		Referencers.Add(Referencer.Value);
+	}
+
+	FBlueprintCompileReinstancer::ReplaceInstancesOfClass(OldGeneratedClass, NewBlueprint->GeneratedClass);
+
+	for (UObject* Referencer : Referencers)
+	{
+		FArchiveReplaceObjectRef<UObject>(Referencer, Redirects, /*bNullPrivateRefs=*/false, /*bIgnoreOuterRef=*/false, /*bIgnoreArchetypeRef=*/false);
+	}
+}
+
 //////////////////////////////////////////////////////////////////////////
 
 // Static variable definition
@@ -338,150 +516,52 @@ UBlueprint* FKismetEditorUtilities::CreateBlueprint(UClass* ParentClass, UObject
 	return NewBP;
 }
 
-static bool IsReferencedByUndoBuffer(UBlueprint* Blueprint)
+UBlueprint* FKismetEditorUtilities::ReloadBlueprint(UBlueprint* StaleBlueprint)
 {
-	UObject* BlueprintObj = Blueprint;
-	FReferencerInformationList ReferencesIncludingUndo;
-	IsReferenced(BlueprintObj, GARBAGE_COLLECTION_KEEPFLAGS, /*bCheckSubObjects =*/true, &ReferencesIncludingUndo);
+	check(StaleBlueprint->IsAsset());
+	FStringAssetReference BlueprintAssetRef(StaleBlueprint);
 
-	FReferencerInformationList ReferencesExcludingUndo;
-	// Determine the in-memory references, *excluding* the undo buffer
-	GEditor->Trans->DisableObjectSerialization();
-	IsReferenced(BlueprintObj, GARBAGE_COLLECTION_KEEPFLAGS, /*bCheckSubObjects =*/true, &ReferencesExcludingUndo);
-	GEditor->Trans->EnableObjectSerialization();
+	FBlueprintUnloader Unloader(StaleBlueprint);
+	Unloader.UnloadBlueprint(/*bResetPackage =*/true);
 
-	// see if this object is the transaction buffer - set a flag so we know we need to clear the undo stack
-	const int32 TotalReferenceCount   = ReferencesIncludingUndo.ExternalReferences.Num() + ReferencesIncludingUndo.InternalReferences.Num();
-	const int32 NonUndoReferenceCount = ReferencesExcludingUndo.ExternalReferences.Num() + ReferencesExcludingUndo.InternalReferences.Num();
+	UBlueprint* ReloadedBlueprint = Cast<UBlueprint>(StaticLoadObject(UBlueprint::StaticClass(), /*Outer =*/nullptr, *BlueprintAssetRef.AssetLongPathname));
 
-	return (TotalReferenceCount > NonUndoReferenceCount);
+	Unloader.ReplaceStaleRefs(ReloadedBlueprint);
+	return ReloadedBlueprint;
 }
 
-static void UnloadBlueprint(UBlueprint* const Blueprint)
+UBlueprint* FKismetEditorUtilities::ReplaceBlueprint(UBlueprint* Target, UBlueprint const* ReplacementArchetype)
 {
-	UPackage* const OldPackage = Cast<UPackage>(Blueprint->GetOutermost());
+	UPackage* BlueprintPackage = Target->GetOutermost();
+	check(BlueprintPackage != GetTransientPackage());
+	const FString BlueprintName = Target->GetName();
 
-	UPackage* const TransientPackage = GetTransientPackage();
-	// move the blueprint to the transient package (to be picked up by garbage collection later)
-	FName UnloadedName = MakeUniqueObjectName(TransientPackage, UBlueprint::StaticClass(), Blueprint->GetFName());
-	Blueprint->Rename(*UnloadedName.ToString(), TransientPackage, REN_DontCreateRedirectors | REN_DoNotDirty);
+	FBlueprintUnloader Unloader(Target);
+	Unloader.UnloadBlueprint(/*bResetPackage =*/false);
 
-	// make sure the blueprint is properly trashed
-	Blueprint->SetFlags(RF_Transient);
-	Blueprint->ClearFlags(RF_Standalone | RF_RootSet | RF_Transactional);
-	Blueprint->RemoveFromRoot();
-	Blueprint->MarkPendingKill();
+	UBlueprint* Replacement = Cast<UBlueprint>(StaticDuplicateObject(ReplacementArchetype, BlueprintPackage, *BlueprintName));
 
-	FKismetEditorUtilities::OnBlueprintUnloaded.Broadcast(Blueprint);
-
-	// reset the blueprint's original package/linker so that we can get by
-	// any early returns (in the load code), and reload its exports as if new 
-	ResetLoaders(OldPackage);
-	OldPackage->ClearFlags(RF_WasLoaded);
-	OldPackage->bHasBeenFullyLoaded = false;
-	OldPackage->GetMetaData()->RemoveMetaDataOutsidePackage();
-
-	if (IsReferencedByUndoBuffer(Blueprint))
-	{
-		GEditor->Trans->Reset(LOCTEXT("ReloadedBlueprint", "Reloaded Blueprint"));
-	}
-
-	// handled in FBlueprintEditor (from the OnBlueprintUnloaded event)
-// 	IAssetEditorInstance* EditorInst = FAssetEditorManager::Get().FindEditorForAsset(Blueprint, /*bFocusIfOpen =*/false);
-// 	if (EditorInst != nullptr)
-// 	{
-// 		EditorInst->CloseWindow();
-// 	}
+	Unloader.ReplaceStaleRefs(Replacement);
+	return Replacement;
 }
 
-UBlueprint* FKismetEditorUtilities::ReloadBlueprint(UBlueprint* StaleBlueprint)
+bool FKismetEditorUtilities::IsReferencedByUndoBuffer(UBlueprint* Blueprint)
 {
-	check(StaleBlueprint->IsAsset());
-
-	UClass* OldGeneratedClass = StaleBlueprint->GeneratedClass;
-	UObject* OldCDO = nullptr; // will be filled out below
-	UClass* OldSkeletonClass = StaleBlueprint->SkeletonGeneratedClass;
-	UObject* OldSkelCDO = nullptr; // will be filled out below
-
-	//--------------------------------------
-	// Collect objects that need replacing
-	//--------------------------------------
-
-	TArray<UObject*> StaleObjects;
-	if (OldGeneratedClass != nullptr)
-	{
-		StaleObjects.Add(OldGeneratedClass);
-
-		OldCDO = OldGeneratedClass->GetDefaultObject(/*bCreateIfNeeded =*/false);
-		if (OldCDO != nullptr)
-		{
-			StaleObjects.Add(OldCDO);
-		}		
-	}
-	if (OldSkeletonClass != nullptr)
-	{
-		StaleObjects.Add(OldSkeletonClass);
-
-		OldSkelCDO = OldSkeletonClass->GetDefaultObject(/*bCreateIfNeeded =*/false);
-		if (OldSkelCDO != nullptr)
-		{
-			StaleObjects.Add(OldSkelCDO);
-		}
-	}
-	StaleObjects.Add(StaleBlueprint);
-
-	//--------------------------------------
-	// Unload/Load Blueprint
-	//--------------------------------------
-
-	FStringAssetReference BlueprintAssetRef(StaleBlueprint);
-	UnloadBlueprint(StaleBlueprint);
-	
-	UBlueprint* ReloadedBlueprint = Cast<UBlueprint>(StaticLoadObject(UBlueprint::StaticClass(), /*Outer =*/nullptr, *BlueprintAssetRef.AssetLongPathname));
-
-	//--------------------------------------
-	// Construct redirects
-	//--------------------------------------
-
-	TMap<UObject*, UObject*> Redirects;
-	if (OldGeneratedClass != nullptr)
-	{
-		Redirects.Add(OldGeneratedClass, ReloadedBlueprint->GeneratedClass);
-	}
-	if (OldCDO != nullptr)
-	{
-		Redirects.Add(OldCDO, ReloadedBlueprint->GeneratedClass->GetDefaultObject());
-	}
-	if (OldSkeletonClass != nullptr)
-	{
-		Redirects.Add(OldSkeletonClass, ReloadedBlueprint->SkeletonGeneratedClass);
-	}
-	if (OldSkelCDO != nullptr)
-	{
-		Redirects.Add(OldSkelCDO, ReloadedBlueprint->SkeletonGeneratedClass->GetDefaultObject());
-	}
-	Redirects.Add(StaleBlueprint, ReloadedBlueprint);
-
-	//--------------------------------------
-	// Replace old references
-	//--------------------------------------
+	UObject* BlueprintObj = Blueprint;
+	FReferencerInformationList ReferencesIncludingUndo;
+	IsReferenced(BlueprintObj, GARBAGE_COLLECTION_KEEPFLAGS, /*bCheckSubObjects =*/true, &ReferencesIncludingUndo);
 
-	TArray<UObject*> Referencers;
-	// find all objects, still referencing the old blueprint/class/cdo/etc.
-	for (auto Referencer : TFindObjectReferencers<UObject>(StaleObjects, /*PackageToCheck =*/nullptr, /*bIgnoreTemplates =*/false))
-	{
-		Referencers.Add(Referencer.Value);
-	}
+	FReferencerInformationList ReferencesExcludingUndo;
+	// Determine the in-memory references, *excluding* the undo buffer
+	GEditor->Trans->DisableObjectSerialization();
+	IsReferenced(BlueprintObj, GARBAGE_COLLECTION_KEEPFLAGS, /*bCheckSubObjects =*/true, &ReferencesExcludingUndo);
+	GEditor->Trans->EnableObjectSerialization();
 
-	StaleBlueprint->SetObjectBeingDebugged(nullptr);
-	FBlueprintCompileReinstancer::ReplaceInstancesOfClass(OldGeneratedClass, ReloadedBlueprint->GeneratedClass);
+	// see if this object is the transaction buffer - set a flag so we know we need to clear the undo stack
+	const int32 TotalReferenceCount   = ReferencesIncludingUndo.ExternalReferences.Num() + ReferencesIncludingUndo.InternalReferences.Num();
+	const int32 NonUndoReferenceCount = ReferencesExcludingUndo.ExternalReferences.Num() + ReferencesExcludingUndo.InternalReferences.Num();
 
-	for (UObject* Referencer : Referencers)
-	{
-		FArchiveReplaceObjectRef<UObject>(Referencer, Redirects, /*bNullPrivateRefs=*/false, /*bIgnoreOuterRef=*/false, /*bIgnoreArchetypeRef=*/false);
-	}
-	
-	return ReloadedBlueprint;
+	return (TotalReferenceCount > NonUndoReferenceCount);
 }
 
 void FKismetEditorUtilities::CompileBlueprint(UBlueprint* BlueprintObj, bool bIsRegeneratingOnLoad, bool bSkipGarbageCollection, bool bSaveIntermediateProducts, FCompilerResultsLog* pResults)
diff --git a/Engine/Source/Editor/UnrealEd/Public/Kismet2/KismetEditorUtilities.h b/Engine/Source/Editor/UnrealEd/Public/Kismet2/KismetEditorUtilities.h
--- a/Engine/Source/Editor/UnrealEd/Public/Kismet2/KismetEditorUtilities.h
+++ b/Engine/Source/Editor/UnrealEd/Public/Kismet2/KismetEditorUtilities.h
@@ -34,18 +34,38 @@ public:
 	 * Event that's broadcast anytime a blueprint is unloaded, and becomes 
 	 * invalid (with calls to ReloadBlueprint(), for example).
 	 */
-	DECLARE_EVENT_OneParam(FKismetEditorUtilities, FOnBlueprintUnloaded, UBlueprint*);
+	DECLARE_MULTICAST_DELEGATE_OneParam(FOnBlueprintUnloaded, UBlueprint*);
 	static FOnBlueprintUnloaded OnBlueprintUnloaded;
 
 	/** 
 	 * Unloads the supplied Blueprint (marking it pending-kill, and removing it 
 	 * from its outer package). Then proceeds to reload from disk.
-
+	 *
 	 * @param  TargetBlueprint	The Blueprint you want to unload and replace.
 	 * @return The freshly loaded Blueprint (replacing the, now invalid, input).
 	 */
 	static UBlueprint* ReloadBlueprint(UBlueprint* TargetBlueprint);
 
+	/** 
+	 * Unloads the specified Blueprint (marking it pending-kill, and removing it 
+	 * from its outer package). Then proceeds to replace all references with a
+	 * copy of the one passed.
+	 *
+	 * @param  Target		The Blueprint you want to unload and replace.
+	 * @param  Replacement	The Blueprint you cloned and used to replace Target.
+	 * @return The duplicated replacement Blueprint.
+	 */
+	static UBlueprint* ReplaceBlueprint(UBlueprint* Target, UBlueprint const* Replacement);
+
+	/** 
+	 * Determines if the specified blueprint is referenced currently in the undo 
+	 * buffer.
+	 *
+	 * @param  Blueprint	The Blueprint you want to query about.
+	 * @return True if the Blueprint is saved in the undo buffer, false if not.
+	 */
+	static bool IsReferencedByUndoBuffer(UBlueprint* Blueprint);
+
 	/** Create the correct event graphs for this blueprint */
 	static void CreateDefaultEventGraphs(UBlueprint* Blueprint);
 
