[UE-5291] Preventing a crash that occurs when referencing another instance of the same BP in the construction script.

diff --git a/Engine/Source/Editor/UnrealEd/Private/Kismet2/KismetReinstanceUtilities.cpp b/Engine/Source/Editor/UnrealEd/Private/Kismet2/KismetReinstanceUtilities.cpp
--- a/Engine/Source/Editor/UnrealEd/Private/Kismet2/KismetReinstanceUtilities.cpp
+++ b/Engine/Source/Editor/UnrealEd/Private/Kismet2/KismetReinstanceUtilities.cpp
@@ -12,6 +12,7 @@
 DEFINE_STAT(EKismetReinstancerStats_ReplaceInstancesOfClass);
 DEFINE_STAT(EKismetReinstancerStats_FindReferencers);
 DEFINE_STAT(EKismetReinstancerStats_ReplaceReferences);
+DEFINE_STAT(EKismetReinstancerStats_ReplacementConstruction);
 DEFINE_STAT(EKismetReinstancerStats_UpdateBytecodeReferences);
 DEFINE_STAT(EKismetReinstancerStats_RecompileChildClasses);
 DEFINE_STAT(EKismetReinstancerStats_ReplaceClassNoReinsancing);
@@ -280,6 +281,209 @@ void FBlueprintCompileReinstancer::UpdateBytecodeReferences()
 	}
 }
 
+/** 
+ * Utility struct that represents a single replacement actor. Used to cache off
+ * attachment info for the old actor (the one being replaced), that will be
+ * used later for the new actor (after all instances have been replaced).
+ */
+struct FActorReplacementHelper
+{
+	/** NOTE: this detaches OldActor from all child/parent attachments. */
+	FActorReplacementHelper(AActor* InNewActor, AActor* OldActor)
+		: NewActor(InNewActor)
+		, TargetWorldTransform(FTransform::Identity)
+		, TargetAttachParent(nullptr)
+		, TargetParentComponent(nullptr)
+		, bSelectNewActor(OldActor->IsSelected())
+	{
+		CachedActorData = StaticCastSharedPtr<AActor::FActorTransactionAnnotation>(OldActor->GetTransactionAnnotation());
+		CacheAttachInfo(OldActor);
+	}
+
+	/**
+	 * Runs construction scripts on the new actor and then finishes it off by
+	 * attaching it to the same attachments that its predecessor was set with. 
+	 */
+	void Finalize();
+
+
+private:
+	/**
+	 * Stores off the old actor's children, and its AttachParent; so that we can
+	 * apply them to the new actor in Finalize().
+	 *
+	 * NOTE: this detaches OldActor from all child/parent attachments.
+	 *
+	 * @param  OldActor		The actor whose attachment setup you want reproduced.
+	 */
+	void CacheAttachInfo(const AActor* OldActor);
+
+	/**
+	 * Stores off the old actor's children; so that in Finalize(), we may 
+	 * reattach them under the new actor.
+	 *
+	 * @param  OldActor		The actor whose child actors you want moved.
+	 */
+	void CacheChildAttachments(const AActor* OldActor);
+
+	/**
+	 * Takes the cached child actors, as well as the old AttachParent, and sets
+	 * up the new actor so that its attachment hierarchy reflects the old actor
+	 * that it is replacing.
+	 */
+	void ApplyAttachments();
+	
+	/**
+	 * Takes the cached child actors, and attaches them under the new actor.
+	 *
+	 * @param  RootComponent	The new actor's root, which the child actors should attach to.
+	 */
+	void AttachChildActors(USceneComponent* RootComponent);
+
+private:
+	AActor*          NewActor;
+	FTransform       TargetWorldTransform;
+	AActor*          TargetAttachParent;
+	USceneComponent* TargetParentComponent; // should belong to TargetAttachParent
+	FName            TargetAttachSocket;
+	bool             bSelectNewActor;
+
+	/** Info store about attached actors */
+	struct FAttachedActorInfo
+	{
+		AActor* AttachedActor;
+		FName   AttachedToSocket;
+	};
+	TArray<FAttachedActorInfo> PendingChildAttachments;
+
+	/** Holds actor component data, etc. that we use to apply */
+	TSharedPtr<AActor::FActorTransactionAnnotation> CachedActorData;
+};
+
+void FActorReplacementHelper::Finalize()
+{
+	
+	// because this is an editor context it's important to use this execution guard
+	FEditorScriptExecutionGuard ScriptGuard;
+
+	// run the construction script, which will use the properties we just copied over
+	if (CachedActorData.IsValid())
+	{
+		NewActor->ExecuteConstruction(TargetWorldTransform, &CachedActorData->ComponentInstanceData);
+	}
+	else
+	{
+		FComponentInstanceDataCache DummyComponentData;
+		NewActor->ExecuteConstruction(TargetWorldTransform, &DummyComponentData);
+	}	
+	ApplyAttachments();
+
+	if (bSelectNewActor)
+	{
+		GEditor->SelectActor(NewActor, /*bInSelected =*/true, /*bNotify =*/false);
+	}
+	
+	// Destroy actor and clear references.
+	NewActor->Modify();
+	if (GEditor->Layers.IsValid()) // ensure(NULL != GEditor->Layers) ?? While cooking the Layers is NULL.
+	{
+		GEditor->Layers->InitializeNewActorLayers(NewActor);
+	}
+}
+
+void FActorReplacementHelper::CacheAttachInfo(const AActor* OldActor)
+{
+	CacheChildAttachments(OldActor);
+
+	if (USceneComponent* OldRootComponent = OldActor->GetRootComponent())
+	{
+		if (OldRootComponent->AttachParent != nullptr)
+		{
+			TargetAttachParent = OldRootComponent->AttachParent->GetOwner();
+			// Root component should never be attached to another component in the same actor!
+			if (TargetAttachParent == OldActor)
+			{
+				UE_LOG(LogBlueprint, Warning, TEXT("ReplaceInstancesOfClass: RootComponent (%s) attached to another component in this Actor (%s)."), *OldRootComponent->GetPathName(), *TargetAttachParent->GetPathName());
+				TargetAttachParent = nullptr;
+			}
+
+			TargetAttachSocket    = OldRootComponent->AttachSocketName;
+			TargetParentComponent = OldRootComponent->AttachParent;
+			
+			// detach it to remove any scaling
+			OldRootComponent->DetachFromParent(/*bMaintainWorldPosition =*/true);
+		}
+
+		// Save off transform
+		TargetWorldTransform = OldRootComponent->ComponentToWorld;
+		TargetWorldTransform.SetTranslation(OldRootComponent->GetComponentLocation()); // take into account any custom location
+	}
+}
+
+void FActorReplacementHelper::CacheChildAttachments(const AActor* OldActor)
+{
+	TArray<AActor*> AttachedActors;
+	OldActor->GetAttachedActors(AttachedActors);
+
+	// if there are attached objects detach them and store the socket names
+	for (AActor* AttachedActor : AttachedActors)
+	{
+		USceneComponent* AttachedActorRoot = AttachedActor->GetRootComponent();
+		if (AttachedActorRoot && AttachedActorRoot->AttachParent)
+		{
+			// Save info about actor to reattach
+			FAttachedActorInfo Info;
+			Info.AttachedActor = AttachedActor;
+			Info.AttachedToSocket = AttachedActorRoot->AttachSocketName;
+			PendingChildAttachments.Add(Info);
+
+			// Now detach it
+			AttachedActorRoot->DetachFromParent(/*bMaintainWorldPosition =*/true);
+		}
+	}
+}
+
+void FActorReplacementHelper::ApplyAttachments()
+{
+	USceneComponent* NewRootComponent = NewActor->GetRootComponent();
+	if (NewRootComponent == nullptr)
+	{
+		return;
+	}
+	// attach the new instance to original parent
+	if (TargetAttachParent != nullptr)
+	{
+		if (TargetParentComponent == nullptr)
+		{
+			TargetParentComponent = TargetAttachParent->GetRootComponent();
+		}
+		else
+		{
+			NewRootComponent->AttachTo(TargetParentComponent, TargetAttachSocket, EAttachLocation::KeepWorldPosition);
+		}
+	}
+
+	AttachChildActors(NewRootComponent);
+}
+
+void FActorReplacementHelper::AttachChildActors(USceneComponent* RootComponent)
+{
+	// if we had attached children reattach them now - unless they are already attached
+	for (FAttachedActorInfo& Info : PendingChildAttachments)
+	{
+		// If this actor is no longer attached to anything, reattach
+		if (!Info.AttachedActor->IsPendingKill() && Info.AttachedActor->GetAttachParentActor() == nullptr)
+		{
+			USceneComponent* ChildRoot = Info.AttachedActor->GetRootComponent();
+			if (ChildRoot && ChildRoot->AttachParent != RootComponent)
+			{
+				ChildRoot->AttachTo(RootComponent, Info.AttachedToSocket, EAttachLocation::KeepWorldPosition);
+				ChildRoot->UpdateComponentToWorld();
+			}
+		}
+	}
+}
+
 void FBlueprintCompileReinstancer::ReplaceInstancesOfClass(UClass* OldClass, UClass* NewClass, UObject*	OriginalCDO)
 {
 	USelection* SelectedActors;
@@ -287,13 +491,16 @@ void FBlueprintCompileReinstancer::ReplaceInstancesOfClass(UClass* OldClass, UCl
 	TArray<UObject*> ObjectsToReplace;
 	const bool bLogConversions = false; // for debugging
 
-	TMap<FStringAssetReference, UObject*> ReinstancedObjectsWeakReferenceMap;
-
 	// Map of old objects to new objects
 	TMap<UObject*, UObject*> OldToNewInstanceMap;
-	TMap<UClass*, UClass*> OldToNewClassMap;
+	TMap<UClass*, UClass*>   OldToNewClassMap;
 	OldToNewClassMap.Add(OldClass, NewClass);
 
+	TMap<FStringAssetReference, UObject*> ReinstancedObjectsWeakReferenceMap;
+
+	// actors being replace
+	TArray<FActorReplacementHelper> ReplacementActors;
+
 	// Set global flag to let system know we are reconstructing blueprint instances
 	TGuardValue<bool> GuardTemplateNameFlag(GIsReconstructingBlueprintInstances, true);
 
@@ -305,247 +512,140 @@ void FBlueprintCompileReinstancer::ReplaceInstancesOfClass(UClass* OldClass, UCl
 		SelectedActors = GEditor->GetSelectedActors();
 		SelectedActors->BeginBatchSelectOperation();
 		SelectedActors->Modify();
-
+		
 		// Then fix 'real' (non archetype) instances of the class
-		for(int32 i=0; i<ObjectsToReplace.Num(); i++)
+		for (UObject* OldObject : ObjectsToReplace)
 		{
-			UObject* OldObject = ObjectsToReplace[i];
-			if (!OldObject->IsTemplate() && !OldObject->IsPendingKill())
+			if (OldObject->IsTemplate() || OldObject->IsPendingKill())
 			{
-				AActor* OldActor = Cast<AActor>(OldObject);
-				UObject* NewObject = NULL;
-				UObject* OldBlueprintDebugObject = NULL;
-				UBlueprint* CorrespondingBlueprint = NULL;
-				// If this object is being debugged, cache it off so we can preserve the 'object being debugged' association
-				CorrespondingBlueprint = Cast<UBlueprint>(OldObject->GetClass()->ClassGeneratedBy);
-				if ( CorrespondingBlueprint )
-				{
-					if (CorrespondingBlueprint->GetObjectBeingDebugged() == OldObject)
-					{
-						OldBlueprintDebugObject = OldObject;
-					}
-				}
-				if (OldActor != NULL)
-				{
-					FTransform WorldTransform = FTransform::Identity;
-					FVector* LocationPtr = NULL;
-					FVector Location = FVector::ZeroVector;
-					FRotator* RotationPtr = NULL;
-					FRotator Rotation = FRotator::ZeroRotator;
-					AActor* AttachParent = NULL;
-					USceneComponent* AttachParentComp = NULL;
-					FName   AttachSocketName;
-
-					// Struct to store info about attached actors
-					struct FAttachedActorInfo
-					{
-						AActor* AttachedActor;
-						FName AttachedToSocket;
-					};
-
-					// Save info about attached actors
-					TArray<FAttachedActorInfo> AttachedActorInfos;
-
-					// Create cache to store component data across rerunning construction scripts
-					FComponentInstanceDataCache InstanceDataCache(OldActor);
-
-					// If there are attached objects detach them and store the socket names
-					TArray<AActor*> AttachedActors;
-					OldActor->GetAttachedActors(AttachedActors);
-					for (AActor* AttachedActor : AttachedActors)
-					{
-						USceneComponent* AttachedActorRoot = AttachedActor->GetRootComponent();
-						if (AttachedActorRoot && AttachedActorRoot->AttachParent)
-						{
-							// Save info about actor to reattach
-							FAttachedActorInfo Info;
-							Info.AttachedActor = AttachedActor;
-							Info.AttachedToSocket = AttachedActorRoot->AttachSocketName;
-							AttachedActorInfos.Add(Info);
-
-							// Now detach it
-							AttachedActorRoot->DetachFromParent(true);
-						}
-					}
-
-					if (USceneComponent* OldRootComponent = OldActor->GetRootComponent())
-					{
-						if(OldRootComponent->AttachParent != NULL)
-						{
-							AttachParent = OldRootComponent->AttachParent->GetOwner();
-							// Root component should never be attached to another component in the same actor!
-							if (AttachParent == OldActor)
-							{
-								UE_LOG(LogBlueprint, Warning, TEXT("ReplaceInstancesOfClass: RootComponent (%s) attached to another component in this Actor (%s)."), *OldRootComponent->GetPathName(), *AttachParent->GetPathName());
-								AttachParent = NULL;
-							}
-
-							AttachParentComp = OldRootComponent->AttachParent;
-							AttachSocketName = OldRootComponent->AttachSocketName;
-
-							//detach it to remove any scaling
-							OldRootComponent->DetachFromParent(true);
-						}
-
-						// Save off transform
-						WorldTransform = OldRootComponent->ComponentToWorld;
-						WorldTransform.SetTranslation(OldRootComponent->GetComponentLocation()); // take into account any custom location
-						Location = OldActor->GetActorLocation();
-						LocationPtr = &Location;
-						Rotation = OldActor->GetActorRotation();
-						RotationPtr = &Rotation;
-					}
-
-					// If this actor was spawned from an Archetype, we spawn the new actor from the new version of that archetype
-					UObject* OldArchetype = OldActor->GetArchetype();
-					UWorld* World = OldActor->GetWorld();
-					AActor* NewArchetype = Cast<AActor>(OldToNewInstanceMap.FindRef(OldArchetype));
-
-					// Check that either this was an instance of the class directly, or we found a new archetype for it
-					check(OldArchetype == OldClass->GetDefaultObject() || NewArchetype);
-
-					// Spawn the new actor instance, in the same level as the original, but deferring running the construction script until we have transferred modified properties
-					ULevel* ActorLevel = OldActor->GetLevel();				
-					UClass** MappedClass = OldToNewClassMap.Find(OldActor->GetClass());
-					UClass* SpawnClass = MappedClass ? *MappedClass : NewClass;
-
-					FActorSpawnParameters SpawnInfo;
-					SpawnInfo.OverrideLevel = ActorLevel;
-					SpawnInfo.Template = NewArchetype;
-					SpawnInfo.bNoCollisionFail = true;
-					SpawnInfo.bDeferConstruction = true;
-
-					// Temporarily remove the deprecated flag so we can respawn the Blueprint in the level
-					bool bIsClassDeprecated = SpawnClass->HasAnyClassFlags(CLASS_Deprecated);
-					SpawnClass->ClassFlags &= ~CLASS_Deprecated;
-
-					AActor* NewActor = World->SpawnActor(SpawnClass, LocationPtr, RotationPtr, SpawnInfo );
-
-					// Reassign the deprecated flag if it was previously assigned
-					if(bIsClassDeprecated)
-					{
-						SpawnClass->ClassFlags |= CLASS_Deprecated;
-					}
-
-					check(NewActor);
-					NewObject = NewActor;
-
-					ReinstancedObjectsWeakReferenceMap.Add(OldObject, NewObject);
-
-					OldActor->DestroyConstructedComponents(); // don't want to serialize components from the old actor
-					// Unregister native components so we don't copy any sub-components they generate for themselves (like UCameraComponent does)
-					OldActor->UnregisterAllComponents(); 
-
-					NewActor->UnregisterAllComponents(); // Unregister any native components, might have cached state based on properties we are going to overwrite
-
-					UEditorEngine::CopyPropertiesForUnrelatedObjects(OldActor, NewActor);
-
-					NewActor->ResetPropertiesForConstruction(); // reset properties/streams
-
-					NewActor->RegisterAllComponents(); // Register native components
-
-					// Because this is an editor context it's important to use this execution guard:
-					FEditorScriptExecutionGuard ScriptGuard;
+				continue;
+			}
 
-					// Run the construction script, which will use the properties we just copied over
-					NewActor->ExecuteConstruction(WorldTransform, &InstanceDataCache);
+			UBlueprint* CorrespondingBlueprint  = Cast<UBlueprint>(OldObject->GetClass()->ClassGeneratedBy);
+			UObject*    OldBlueprintDebugObject = nullptr;
+			// If this object is being debugged, cache it off so we can preserve the 'object being debugged' association
+			if ((CorrespondingBlueprint != nullptr) && (CorrespondingBlueprint->GetObjectBeingDebugged() == OldObject))
+			{
+				OldBlueprintDebugObject = OldObject;
+			}
 
-					if (USceneComponent* NewRootComponent = NewActor->GetRootComponent())
-					{
-						//Attach the new instance to original parent
-						if(AttachParent != nullptr)
-						{
-							if(AttachParentComp == nullptr)
-							{
-								AttachParentComp = AttachParent->GetRootComponent();
-							}
+			AActor*  OldActor = Cast<AActor>(OldObject);
+			UObject* NewObject = nullptr;
+			// if the object to replace is an actor...
+			if (OldActor != nullptr)
+			{
+				FVector  Location = FVector::ZeroVector;
+				FRotator Rotation = FRotator::ZeroRotator;
+				if (USceneComponent* OldRootComponent = OldActor->GetRootComponent())
+				{
+					Location = OldActor->GetActorLocation();
+					Rotation = OldActor->GetActorRotation();
+				}
 
-							if(AttachParentComp != nullptr)
-							{
-								NewRootComponent->AttachTo(AttachParentComp, AttachSocketName, EAttachLocation::KeepWorldPosition);
-							}
-						}
+				// If this actor was spawned from an Archetype, we spawn the new actor from the new version of that archetype
+				UObject* OldArchetype = OldActor->GetArchetype();
+				UWorld*  World        = OldActor->GetWorld();
+				AActor*  NewArchetype = Cast<AActor>(OldToNewInstanceMap.FindRef(OldArchetype));
+				// Check that either this was an instance of the class directly, or we found a new archetype for it
+				check(OldArchetype == OldClass->GetDefaultObject() || NewArchetype);
+
+				// Spawn the new actor instance, in the same level as the original, but deferring running the construction script until we have transferred modified properties
+				ULevel*  ActorLevel  = OldActor->GetLevel();
+				UClass** MappedClass = OldToNewClassMap.Find(OldActor->GetClass());
+				UClass*  SpawnClass  = MappedClass ? *MappedClass : NewClass;
+
+				FActorSpawnParameters SpawnInfo;
+				SpawnInfo.OverrideLevel      = ActorLevel;
+				SpawnInfo.Template           = NewArchetype;
+				SpawnInfo.bNoCollisionFail   = true;
+				SpawnInfo.bDeferConstruction = true;
+
+				// Temporarily remove the deprecated flag so we can respawn the Blueprint in the level
+				const bool bIsClassDeprecated = SpawnClass->HasAnyClassFlags(CLASS_Deprecated);
+				SpawnClass->ClassFlags &= ~CLASS_Deprecated;
+
+				AActor* NewActor = World->SpawnActor(SpawnClass, &Location, &Rotation, SpawnInfo);
+				// Reassign the deprecated flag if it was previously assigned
+				if (bIsClassDeprecated)
+				{
+					SpawnClass->ClassFlags |= CLASS_Deprecated;
+				}
 
-						// If we had attached children reattach them now - unless they are already attached
-						for(FAttachedActorInfo& Info : AttachedActorInfos)
-						{
-							// If this actor is no longer attached to anything, reattach
-							if (!Info.AttachedActor->IsPendingKill() && Info.AttachedActor->GetAttachParentActor() == nullptr)
-							{
-								USceneComponent* ChildRoot = Info.AttachedActor->GetRootComponent();
-								if (ChildRoot && ChildRoot->AttachParent != NewRootComponent)
-								{
-									ChildRoot->AttachTo(NewRootComponent, Info.AttachedToSocket, EAttachLocation::KeepWorldPosition);
-									ChildRoot->UpdateComponentToWorld();
-								}
-							}
-						}
-					}
+				check(NewActor != nullptr);
+				NewObject = NewActor;
+				// store the new actor for the second pass (NOTE: this detaches 
+				// OldActor from all child/parent attachments)
+				//
+				// running the NewActor's construction-script is saved for that 
+				// second pass (because the construction-script may reference 
+				// another instance that hasn't been replaced yet).
+				ReplacementActors.Add(FActorReplacementHelper(NewActor, OldActor));
 
-					// Determine if the actor used to be selected.
-					if (OldActor->IsSelected())
-					{
-						// Deselect the old actor and select the new actor.
-						const bool bNotify = false;
-						GEditor->SelectActor(OldActor, /*bInSelected=*/ false, /*bNotify=*/ bNotify);
-						GEditor->SelectActor(NewActor, /*bInSelected=*/ true, /*bNotify=*/ bNotify);
+				ReinstancedObjectsWeakReferenceMap.Add(OldObject, NewObject);
 
-						bSelectionChanged = true;
-					}
+				OldActor->DestroyConstructedComponents(); // don't want to serialize components from the old actor
+				// Unregister native components so we don't copy any sub-components they generate for themselves (like UCameraComponent does)
+				OldActor->UnregisterAllComponents();
 
-					// Destroy actor and clear references.
-					NewActor->Modify();
-					if(GEditor->Layers.IsValid()) // ensure(NULL != GEditor->Layers) ?? While cooking the Layers is NULL.
-					{
-						GEditor->Layers->InitializeNewActorLayers( NewActor );
+				// Unregister any native components, might have cached state based on properties we are going to overwrite
+				NewActor->UnregisterAllComponents(); 
 
-						GEditor->Layers->DisassociateActorFromLayers( OldActor );
-					}
-				
-					World->EditorDestroyActor( OldActor, true );
+				UEditorEngine::CopyPropertiesForUnrelatedObjects(OldActor, NewActor);
+				// reset properties/streams
+				NewActor->ResetPropertiesForConstruction(); 
+				// register native components
+				NewActor->RegisterAllComponents(); 
 
-					OldToNewInstanceMap.Add(OldActor, NewActor);
+				// 
+				// clean up the old actor (unselect it, remove it from the world, etc.)...
 
-					if(bLogConversions)
-					{
-						UE_LOG(LogBlueprint, Log, TEXT("Converted instance '%s' to '%s'"), *OldActor->GetPathName(), *NewActor->GetPathName() );
-					}
+				if (OldActor->IsSelected())
+				{
+					GEditor->SelectActor(OldActor, /*bInSelected =*/false, /*bNotify =*/false);
+					bSelectionChanged = true;
 				}
-				else
+				if (GEditor->Layers.IsValid()) // ensure(NULL != GEditor->Layers) ?? While cooking the Layers is NULL.
 				{
-					FName OldName(OldObject->GetFName());
-					OldObject->Rename(NULL, OldObject->GetOuter(), REN_DoNotDirty | REN_DontCreateRedirectors);
-					NewObject = ConstructObject<UObject>(NewClass, OldObject->GetOuter(), OldName);
-					check(NewObject);
+					GEditor->Layers->DisassociateActorFromLayers(OldActor);
+				}
 
-					UEditorEngine::CopyPropertiesForUnrelatedObjects(OldObject, NewObject);
+ 				World->EditorDestroyActor(OldActor, /*bShouldModifyLevel =*/true);
+ 				OldToNewInstanceMap.Add(OldActor, NewActor);
+			}
+			else
+			{
+				FName OldName(OldObject->GetFName());
+				OldObject->Rename(NULL, OldObject->GetOuter(), REN_DoNotDirty | REN_DontCreateRedirectors);
+				NewObject = ConstructObject<UObject>(NewClass, OldObject->GetOuter(), OldName);
+				check(NewObject != nullptr);
+
+				UEditorEngine::CopyPropertiesForUnrelatedObjects(OldObject, NewObject);
 
-					if (UAnimInstance* AnimTree = Cast<UAnimInstance>(NewObject))
+				if (UAnimInstance* AnimTree = Cast<UAnimInstance>(NewObject))
+				{
+					// Initialising the anim instance isn't enough to correctly set up the skeletal mesh again in a
+					// paused world, need to initialise the skeletal mesh component that contains the anim instance.
+					if (USkeletalMeshComponent* SkelComponent = Cast<USkeletalMeshComponent>(AnimTree->GetOuter()))
 					{
-						// Initialising the anim instance isn't enough to correctly set up the skeletal mesh again in a
-						// paused world, need to initialise the skeletal mesh component that contains the anim instance.
-						if(USkeletalMeshComponent* SkelComponent = Cast<USkeletalMeshComponent>(AnimTree->GetOuter()))
-						{
-							SkelComponent->InitAnim(true);
-						}
+						SkelComponent->InitAnim(true);
 					}
+				}
 
-					OldObject->RemoveFromRoot();
-					OldObject->MarkPendingKill();
+				OldObject->RemoveFromRoot();
+				OldObject->MarkPendingKill();
 
-					OldToNewInstanceMap.Add(OldObject, NewObject);
+				OldToNewInstanceMap.Add(OldObject, NewObject);
+			}
 
-					if(bLogConversions)
-					{
-						UE_LOG(LogBlueprint, Log, TEXT("Converted instance '%s' to '%s'"), *OldObject->GetPathName(), *NewObject->GetPathName() );
-					}
-				}
+			// If this original object came from a blueprint and it was in the selected debug set, change the debugging to the new object.
+			if ((CorrespondingBlueprint) && (OldBlueprintDebugObject) && (NewObject))
+			{
+				CorrespondingBlueprint->SetObjectBeingDebugged(NewObject);
+			}
 
-				// If this original object came from a blueprint and it was in the selected debug set, change the debugging to the new object.
-				if( ( CorrespondingBlueprint) && ( OldBlueprintDebugObject ) && ( NewObject ) )
-				{
-					CorrespondingBlueprint->SetObjectBeingDebugged(NewObject);
-				}
+			if (bLogConversions)
+			{
+				UE_LOG(LogBlueprint, Log, TEXT("Converted instance '%s' to '%s'"), *OldObject->GetPathName(), *NewObject->GetPathName());
 			}
 		}
 	}
@@ -639,6 +739,26 @@ void FBlueprintCompileReinstancer::ReplaceInstancesOfClass(UClass* OldClass, UCl
 		}
 	}
 
+	{ BP_SCOPED_COMPILER_EVENT_STAT(EKismetReinstancerStats_ReplacementConstruction);
+
+		// the process of setting up new replacement actors is split into two 
+		// steps (this here, is the second)...
+		// 
+		// the "finalization" here runs the replacement actor's construction-
+		// script and is left until late to account for a scenario where the 
+		// construction-script attempts to modify another instance of the 
+		// same class... if this were to happen above, in the ObjectsToReplace 
+		// loop, then accessing that other instance would cause an assert in 
+		// UProperty::ContainerPtrToValuePtrInternal() (which appropriatly 
+		// complains that the other instance's type doesn't match because it 
+		// hasn't been replaced yet... that's why we wait until after 
+		// FArchiveReplaceObjectRef to run construction-scripts).
+		for (FActorReplacementHelper& ReplacementActor : ReplacementActors)
+		{
+			ReplacementActor.Finalize();
+		}
+	}
+
 	SelectedActors->EndBatchSelectOperation();
 	if (bSelectionChanged)
 	{
diff --git a/Engine/Source/Editor/UnrealEd/Public/Kismet2/KismetReinstanceUtilities.h b/Engine/Source/Editor/UnrealEd/Public/Kismet2/KismetReinstanceUtilities.h
--- a/Engine/Source/Editor/UnrealEd/Public/Kismet2/KismetReinstanceUtilities.h
+++ b/Engine/Source/Editor/UnrealEd/Public/Kismet2/KismetReinstanceUtilities.h
@@ -13,6 +13,7 @@ DECLARE_STATS_GROUP(TEXT("Kismet Reinstancer"), STATGROUP_KismetReinstancer, STA
 DECLARE_CYCLE_STAT_EXTERN(TEXT("Replace Instances"), EKismetReinstancerStats_ReplaceInstancesOfClass, STATGROUP_KismetReinstancer, );
 DECLARE_CYCLE_STAT_EXTERN(TEXT("Find Referencers"), EKismetReinstancerStats_FindReferencers, STATGROUP_KismetReinstancer, );
 DECLARE_CYCLE_STAT_EXTERN(TEXT("Replace References"), EKismetReinstancerStats_ReplaceReferences, STATGROUP_KismetReinstancer, );
+DECLARE_CYCLE_STAT_EXTERN(TEXT("Construct Replacements"), EKismetReinstancerStats_ReplacementConstruction, STATGROUP_KismetReinstancer, );
 DECLARE_CYCLE_STAT_EXTERN(TEXT("Update Bytecode References"), EKismetReinstancerStats_UpdateBytecodeReferences, STATGROUP_KismetReinstancer, );
 DECLARE_CYCLE_STAT_EXTERN(TEXT("Recompile Child Classes"), EKismetReinstancerStats_RecompileChildClasses, STATGROUP_KismetReinstancer, );
 DECLARE_CYCLE_STAT_EXTERN(TEXT("Replace Classes Without Reinstancing"), EKismetReinstancerStats_ReplaceClassNoReinsancing, STATGROUP_KismetReinstancer, );
diff --git a/Engine/Source/Runtime/CoreUObject/Private/UObject/PropertyBaseObject.cpp b/Engine/Source/Runtime/CoreUObject/Private/UObject/PropertyBaseObject.cpp
--- a/Engine/Source/Runtime/CoreUObject/Private/UObject/PropertyBaseObject.cpp
+++ b/Engine/Source/Runtime/CoreUObject/Private/UObject/PropertyBaseObject.cpp
@@ -380,7 +380,7 @@ void UObjectPropertyBase::CheckValidObject(void* Value) const
 		// PropertyClass itself (in the middle of an FArchiveReplaceObjectRef 
 		// pass)... if this is the case, then we might have already replaced 
 		// the object's class, but not the PropertyClass yet (or vise-versa)... 
-		// so we use this to ensure in that situation that we don't clear the 
+		// so we use this to ensure, in that situation, that we don't clear the 
 		// object value (if CLASS_NewerVersionExists is set, then we are likely 
 		// in the middle of an FArchiveReplaceObjectRef pass)
 		bool bIsReplacingClassRefs = PropertyClass->HasAnyClassFlags(CLASS_NewerVersionExists) != ObjectClass->HasAnyClassFlags(CLASS_NewerVersionExists);
