[UE-6282] Now, also looking to inherited debug data for parent blueprint property values (when debugging a child instance in a parent graph)

diff --git a/Engine/Source/Editor/UnrealEd/Private/Kismet2/KismetDebugUtilities.cpp b/Engine/Source/Editor/UnrealEd/Private/Kismet2/KismetDebugUtilities.cpp
--- a/Engine/Source/Editor/UnrealEd/Private/Kismet2/KismetDebugUtilities.cpp
+++ b/Engine/Source/Editor/UnrealEd/Private/Kismet2/KismetDebugUtilities.cpp
@@ -698,12 +698,21 @@ bool FKismetDebugUtilities::HasDebuggingData(const UBlueprint* Blueprint)
 // Looks thru the debugging data for any class variables associated with the node
 UProperty* FKismetDebugUtilities::FindClassPropertyForPin(UBlueprint* Blueprint, const UEdGraphPin* Pin)
 {
-	if (UBlueprintGeneratedClass* Class = Cast<UBlueprintGeneratedClass>(*Blueprint->GeneratedClass))
+	UProperty* FoundProperty = nullptr;
+
+	UClass* Class = Blueprint->GeneratedClass;
+	while (UBlueprintGeneratedClass* BlueprintClass = Cast<UBlueprintGeneratedClass>(Class))
 	{
-		return Class->GetDebugData().FindClassPropertyForPin(Pin);
+		FoundProperty = BlueprintClass->GetDebugData().FindClassPropertyForPin(Pin);
+		if (FoundProperty != nullptr)
+		{
+			break;
+		}
+
+		Class = BlueprintClass->GetSuperClass();
 	}
 
-	return NULL;
+	return FoundProperty;
 }
 
 // Looks thru the debugging data for any class variables associated with the node (e.g., temporary variables or timelines)
