Fix crash in editor when playing in standalone mode.

diff --git a/Engine/Source/Editor/UnrealEd/Private/PlayLevel.cpp b/Engine/Source/Editor/UnrealEd/Private/PlayLevel.cpp
--- a/Engine/Source/Editor/UnrealEd/Private/PlayLevel.cpp
+++ b/Engine/Source/Editor/UnrealEd/Private/PlayLevel.cpp
@@ -979,7 +979,8 @@ public:
 					const FString WorldAssetPackageName = StreamingLevel->GetWorldAssetPackageName();
 					const FName WorldAssetPackageFName = StreamingLevel->GetWorldAssetPackageFName();
 					PreviousStreamingPackageNames.Add( WorldAssetPackageFName );
-					const FString StreamingLevelPackageName = FString::Printf(TEXT("%s%s/%s%s"), *AutosavePackagePrefix, *FPackageName::GetLongPackagePath( WorldAssetPackageName ), *MapnamePrefix, *FPackageName::GetLongPackageAssetName( WorldAssetPackageName ));
+					FString StreamingLevelPackageName = FString::Printf(TEXT("%s%s/%s%s"), *AutosavePackagePrefix, *FPackageName::GetLongPackagePath( WorldAssetPackageName ), *MapnamePrefix, *FPackageName::GetLongPackageAssetName( WorldAssetPackageName ));
+					StreamingLevelPackageName.ReplaceInline(TEXT("//"), TEXT("/"));
 					StreamingLevel->SetWorldAssetByPackageName(FName(*StreamingLevelPackageName));
 				}
 			}
