Fixed duplicate entries in the CookingManifest.

diff --git a/Engine/Source/Editor/UnrealEd/Private/Commandlets/ChunkManifestGenerator.cpp b/Engine/Source/Editor/UnrealEd/Private/Commandlets/ChunkManifestGenerator.cpp
--- a/Engine/Source/Editor/UnrealEd/Private/Commandlets/ChunkManifestGenerator.cpp
+++ b/Engine/Source/Editor/UnrealEd/Private/Commandlets/ChunkManifestGenerator.cpp
@@ -495,7 +495,7 @@ bool FChunkManifestGenerator::SaveCookedPackageAssetRegistry( const FString& San
 							// get the package name and see if we have already written it out this run
 							
 							FString CookedPackageName;
-							check( JsonPackage->TryGetStringField(TEXT("CookedPackageName"), CookedPackageName) );
+							check( JsonPackage->TryGetStringField(TEXT("SourcePackageName"), CookedPackageName) );
 
 							const FName CookedPackageFName(*CookedPackageName);
 							if ( CookedPackages.Contains(CookedPackageFName))
