Editor: Initializing some additional variables

diff --git a/Engine/Source/Editor/UnrealEd/Private/EdMode.cpp b/Engine/Source/Editor/UnrealEd/Private/EdMode.cpp
--- a/Engine/Source/Editor/UnrealEd/Private/EdMode.cpp
+++ b/Engine/Source/Editor/UnrealEd/Private/EdMode.cpp
@@ -176,11 +176,12 @@ const FName FEdMode::MD_MakeEditWidget(TEXT("MakeEditWidget"));
 const FName FEdMode::MD_ValidateWidgetUsing(TEXT("ValidateWidgetUsing"));
 
 FEdMode::FEdMode()
-	: bPendingDeletion( false )
-	, CurrentTool( NULL )
-	, EditedPropertyName(TEXT(""))
-	, EditedPropertyIndex( INDEX_NONE )
-	, bEditedPropertyIsTransform( false )
+	: bPendingDeletion(false)
+	, CurrentWidgetAxis(EAxisList::None)
+	, CurrentTool(nullptr)
+	, Owner(nullptr)
+	, EditedPropertyIndex(INDEX_NONE)
+	, bEditedPropertyIsTransform(false)
 {
 	bDrawKillZ = true;
 }
@@ -607,13 +608,13 @@ void FEdMode::Render(const FSceneView* View,FViewport* Viewport,FPrimitiveDrawIn
 						const FVector Vertex0 = Poly->Vertices[0];
 						FVector Vertex1 = Poly->Vertices[1];
 
-						MeshBuilder.AddVertex(Vertex0, FVector2D::ZeroVector, FVector(1,0,0), FVector(0,1,0), FVector(0,0,1), FColor(255,255,255));
-						MeshBuilder.AddVertex(Vertex1, FVector2D::ZeroVector, FVector(1,0,0), FVector(0,1,0), FVector(0,0,1), FColor(255,255,255));
+						MeshBuilder.AddVertex(Vertex0, FVector2D::ZeroVector, FVector(1,0,0), FVector(0,1,0), FVector(0,0,1), FColor::White);
+						MeshBuilder.AddVertex(Vertex1, FVector2D::ZeroVector, FVector(1,0,0), FVector(0,1,0), FVector(0,0,1), FColor::White);
 
 						for( int32 VertexIdx = 2 ; VertexIdx < Poly->Vertices.Num() ; ++VertexIdx )
 						{
 							const FVector Vertex2 = Poly->Vertices[VertexIdx];
-							MeshBuilder.AddVertex(Vertex2, FVector2D::ZeroVector, FVector(1,0,0), FVector(0,1,0), FVector(0,0,1), FColor(255,255,255));
+							MeshBuilder.AddVertex(Vertex2, FVector2D::ZeroVector, FVector(1,0,0), FVector(0,1,0), FVector(0,0,1), FColor::White);
 							MeshBuilder.AddTriangle(VertexOffset,VertexOffset + VertexIdx,VertexOffset+VertexIdx-1);
 							Vertex1 = Vertex2;
 						}
@@ -700,7 +701,7 @@ void FEdMode::Render(const FSceneView* View,FViewport* Viewport,FPrimitiveDrawIn
 				bool bIsTransform = WidgetInfos[i].bIsTransform;
 
 				bool bSelected = (WidgetName == EditedPropertyName) && (WidgetIndex == EditedPropertyIndex);
-				FColor WidgetColor = bSelected ? FColor(255,255,255) : FColor(128,128,255);
+				FColor WidgetColor = bSelected ? FColor::White : FColor(128, 128, 255);
 
 				FVector LocalPos = FVector::ZeroVector;
 				if(bIsTransform)
diff --git a/Engine/Source/Editor/UnrealEd/Private/EditorComponents.cpp b/Engine/Source/Editor/UnrealEd/Private/EditorComponents.cpp
--- a/Engine/Source/Editor/UnrealEd/Private/EditorComponents.cpp
+++ b/Engine/Source/Editor/UnrealEd/Private/EditorComponents.cpp
@@ -309,31 +309,24 @@ void FGridWidget::DrawNewGrid(const FSceneView* View, FPrimitiveDrawInterface* P
 FEditorCommonDrawHelper.
 ------------------------------------------------------------------------------*/
 FEditorCommonDrawHelper::FEditorCommonDrawHelper()
-	: GridWidget(0)
+	: bDrawGrid(true)
+	, bDrawPivot(true)
+	, bDrawBaseInfo(true)
+	, bDrawWorldBox(false)
+	, bDrawKillZ(false)
+	, AxesLineThickness(0.0f)
+	, GridColorAxis(70, 70, 70)
+	, GridColorMajor(40, 40, 40)
+	, GridColorMinor(20, 20, 20)
+	, PerspectiveGridSize(HALF_WORLD_MAX1)
+	, PivotColor(FColor::Red)
+	, PivotSize(0.02f)
+	, NumCells(64)
+	, BaseBoxColor(FColor::Green)
+	, DepthPriorityGroup(SDPG_World)
+	, GridDepthBias(0.000001f)
+	, GridWidget(nullptr)
 {
-	bDrawGrid = true;
-	bDrawPivot = true;
-	bDrawBaseInfo = true;
-	AxesLineThickness = 0.f;
-
-	GridColorAxis = FColor(70, 70, 70);
-	GridColorMajor = FColor(40, 40, 40);
-	GridColorMinor = FColor(20, 20, 20);
-
-	PerspectiveGridSize = HALF_WORLD_MAX1;
-	NumCells = 64;
-
-	bDrawWorldBox = false;
-	bDrawKillZ = false;
-
-	PivotColor = FColor(255,0,0);
-	PivotSize = 0.02f;
-
-	BaseBoxColor = FColor(0,255,0);
-
-	DepthPriorityGroup=SDPG_World;
-
-	GridDepthBias = 0.000001;
 }
 
 FEditorCommonDrawHelper::~FEditorCommonDrawHelper()
@@ -487,8 +480,8 @@ void FEditorCommonDrawHelper::DrawOldGrid(const FSceneView* View,FPrimitiveDrawI
 		{
 			float KillZ = GWorld->GetWorldSettings()->KillZ;
 
-			PDI->DrawLine( FVector(-HALF_WORLD_MAX,0,KillZ), FVector(HALF_WORLD_MAX,0,KillZ), FColor(255,0,0), SDPG_Foreground );
-			PDI->DrawLine( FVector(0,-HALF_WORLD_MAX,KillZ), FVector(0,HALF_WORLD_MAX,KillZ), FColor(255,0,0), SDPG_Foreground );
+			PDI->DrawLine(FVector(-HALF_WORLD_MAX, 0, KillZ), FVector(HALF_WORLD_MAX, 0, KillZ), FColor::Red, SDPG_Foreground);
+			PDI->DrawLine(FVector(0, -HALF_WORLD_MAX, KillZ), FVector(0, HALF_WORLD_MAX, KillZ), FColor::Red, SDPG_Foreground);
 		}
 	}
 
