Editor: Pushed basic ModeTools support into FEditorViewportClient

diff --git a/Engine/Plugins/2D/Paper2D/Source/Paper2DEditor/Private/PaperEditorViewportClient.cpp b/Engine/Plugins/2D/Paper2D/Source/Paper2DEditor/Private/PaperEditorViewportClient.cpp
--- a/Engine/Plugins/2D/Paper2D/Source/Paper2DEditor/Private/PaperEditorViewportClient.cpp
+++ b/Engine/Plugins/2D/Paper2D/Source/Paper2DEditor/Private/PaperEditorViewportClient.cpp
@@ -41,9 +41,10 @@ USelection* FAssetEditorModeTools::GetSelectedObjects() const
 // FPaperEditorViewportClient
 
 FPaperEditorViewportClient::FPaperEditorViewportClient()
-	: FEditorViewportClient(*( new FAssetEditorModeTools() ))
+	: FEditorViewportClient(new FAssetEditorModeTools())
 	, CheckerboardTexture(NULL)
 {
+	bOwnsModeTools = true;
 	ZoomPos = FVector2D::ZeroVector;
 	ZoomAmount = 1.0f;
 
diff --git a/Engine/Source/Editor/Cascade/Private/CascadeEmitterCanvasClient.cpp b/Engine/Source/Editor/Cascade/Private/CascadeEmitterCanvasClient.cpp
--- a/Engine/Source/Editor/Cascade/Private/CascadeEmitterCanvasClient.cpp
+++ b/Engine/Source/Editor/Cascade/Private/CascadeEmitterCanvasClient.cpp
@@ -23,7 +23,7 @@
 
 
 FCascadeEmitterCanvasClient::FCascadeEmitterCanvasClient(TWeakPtr<FCascade> InCascade, TWeakPtr<SCascadeEmitterCanvas> InCascadeViewport)
-	: FEditorViewportClient(GLevelEditorModeTools())
+	: FEditorViewportClient(nullptr)
 	, CascadePtr(InCascade)
 	, CascadeViewportPtr(InCascadeViewport)
 	, EmitterWidth(180)
diff --git a/Engine/Source/Editor/Cascade/Private/CascadePreviewViewportClient.cpp b/Engine/Source/Editor/Cascade/Private/CascadePreviewViewportClient.cpp
--- a/Engine/Source/Editor/Cascade/Private/CascadePreviewViewportClient.cpp
+++ b/Engine/Source/Editor/Cascade/Private/CascadePreviewViewportClient.cpp
@@ -22,7 +22,7 @@
 #define LOCTEXT_NAMESPACE "CascadeViewportClient"
 
 FCascadeEdPreviewViewportClient::FCascadeEdPreviewViewportClient(TWeakPtr<FCascade> InCascade, TWeakPtr<SCascadePreviewViewport> InCascadeViewport)
-	: FEditorViewportClient(GLevelEditorModeTools())
+	: FEditorViewportClient(nullptr)
 	, CascadePtr(InCascade)
 	, CascadeViewportPtr(InCascadeViewport)
 	, CascadePreviewScene(FPreviewScene::ConstructionValues()
@@ -211,19 +211,19 @@ void FCascadeEdPreviewViewportClient::Draw(FViewport* Viewport, FCanvas* Canvas)
 	if (GetDrawElement(OriginAxis))
 	{
 		FMatrix ArrowMatrix = FMatrix(XAxis, YAxis, ZAxis, FVector::ZeroVector);
-		LineBatcher->DrawDirectionalArrow(ArrowMatrix, FColor(255,0,0), 10.f, 1.0f, SDPG_World);
+		LineBatcher->DrawDirectionalArrow(ArrowMatrix, FColorList::Red, 10.f, 1.0f, SDPG_World);
 
 		ArrowMatrix = FMatrix(YAxis, ZAxis, XAxis, FVector::ZeroVector);
-		LineBatcher->DrawDirectionalArrow(ArrowMatrix, FColor(0,255,0), 10.f, 1.0f, SDPG_World);
+		LineBatcher->DrawDirectionalArrow(ArrowMatrix, FColorList::Green, 10.f, 1.0f, SDPG_World);
 
 		ArrowMatrix = FMatrix(ZAxis, XAxis, YAxis, FVector::ZeroVector);
-		LineBatcher->DrawDirectionalArrow(ArrowMatrix, FColor(0,0,255), 10.f, 1.0f, SDPG_World);
+		LineBatcher->DrawDirectionalArrow(ArrowMatrix, FColorList::Blue, 10.f, 1.0f, SDPG_World);
 	}
 
 	if (GetDrawElement(WireSphere))
 	{
 		FVector Base(0.f);
-		FColor WireColor = FColor(255, 0, 0);
+		FColor WireColor = FColor::Red;
 		const int32 NumRings = 16;
 		const float RotatorMultiplier = 360.f / NumRings;
 		const int32 NumSides = 32;
diff --git a/Engine/Source/Editor/DestructibleMeshEditor/Private/SDestructibleMeshEditorViewport.cpp b/Engine/Source/Editor/DestructibleMeshEditor/Private/SDestructibleMeshEditorViewport.cpp
--- a/Engine/Source/Editor/DestructibleMeshEditor/Private/SDestructibleMeshEditorViewport.cpp
+++ b/Engine/Source/Editor/DestructibleMeshEditor/Private/SDestructibleMeshEditorViewport.cpp
@@ -81,7 +81,7 @@ private:
 };
 
 FDestructibleMeshEditorViewportClient::FDestructibleMeshEditorViewportClient(TWeakPtr<IDestructibleMeshEditor> InDestructibleMeshEditor, FPreviewScene& InPreviewScene)
-	: FEditorViewportClient(GLevelEditorModeTools(), &InPreviewScene)
+	: FEditorViewportClient(nullptr, &InPreviewScene)
 	, DestructibleMeshEditorPtr(InDestructibleMeshEditor)
 {
 	SetViewMode(VMI_Lit);
diff --git a/Engine/Source/Editor/DetailCustomizations/Private/AnimMontageSegmentDetails.cpp b/Engine/Source/Editor/DetailCustomizations/Private/AnimMontageSegmentDetails.cpp
--- a/Engine/Source/Editor/DetailCustomizations/Private/AnimMontageSegmentDetails.cpp
+++ b/Engine/Source/Editor/DetailCustomizations/Private/AnimMontageSegmentDetails.cpp
@@ -9,7 +9,7 @@
 
 /////////////////////////////////////////////////////////////////////////
 FAnimationSegmentViewportClient::FAnimationSegmentViewportClient(FPreviewScene& InPreviewScene)
-	: FEditorViewportClient(GLevelEditorModeTools(), &InPreviewScene)
+	: FEditorViewportClient(nullptr, &InPreviewScene)
 {
 	SetViewMode(VMI_Lit);
 
diff --git a/Engine/Source/Editor/Kismet/Private/SCSEditorViewportClient.cpp b/Engine/Source/Editor/Kismet/Private/SCSEditorViewportClient.cpp
--- a/Engine/Source/Editor/Kismet/Private/SCSEditorViewportClient.cpp
+++ b/Engine/Source/Editor/Kismet/Private/SCSEditorViewportClient.cpp
@@ -106,7 +106,7 @@ namespace
 // FSCSEditorViewportClient
 
 FSCSEditorViewportClient::FSCSEditorViewportClient(TWeakPtr<FBlueprintEditor>& InBlueprintEditorPtr, FPreviewScene& InPreviewScene)
-	: FEditorViewportClient(GLevelEditorModeTools(), &InPreviewScene)
+	: FEditorViewportClient(nullptr, &InPreviewScene)
 	,BlueprintEditorPtr(InBlueprintEditorPtr)
 	,PreviewBlueprint(NULL)
 	,PreviewActorBounds(ForceInitToZero)
diff --git a/Engine/Source/Editor/MaterialEditor/Private/SMaterialEditorViewport.cpp b/Engine/Source/Editor/MaterialEditor/Private/SMaterialEditorViewport.cpp
--- a/Engine/Source/Editor/MaterialEditor/Private/SMaterialEditorViewport.cpp
+++ b/Engine/Source/Editor/MaterialEditor/Private/SMaterialEditorViewport.cpp
@@ -45,7 +45,7 @@ private:
 };
 
 FMaterialEditorViewportClient::FMaterialEditorViewportClient(TWeakPtr<IMaterialEditor> InMaterialEditor, FPreviewScene& InPreviewScene)
-	: FEditorViewportClient( GLevelEditorModeTools(), &InPreviewScene )
+	: FEditorViewportClient( nullptr, &InPreviewScene )
 	, MaterialEditorPtr(InMaterialEditor)
 {
 	// Setup defaults for the common draw helper.
diff --git a/Engine/Source/Editor/Matinee/Private/MatineeViewportClient.cpp b/Engine/Source/Editor/Matinee/Private/MatineeViewportClient.cpp
--- a/Engine/Source/Editor/Matinee/Private/MatineeViewportClient.cpp
+++ b/Engine/Source/Editor/Matinee/Private/MatineeViewportClient.cpp
@@ -15,7 +15,7 @@
 -----------------------------------------------------------------------------*/
 
 FMatineeViewportClient::FMatineeViewportClient( class FMatinee* InMatinee )
-	: FEditorViewportClient(GLevelEditorModeTools())
+	: FEditorViewportClient(nullptr)
 {
 	InterpEd = InMatinee;
 
diff --git a/Engine/Source/Editor/NiagaraEditor/Private/SNiagaraEffectEditorViewport.cpp b/Engine/Source/Editor/NiagaraEditor/Private/SNiagaraEffectEditorViewport.cpp
--- a/Engine/Source/Editor/NiagaraEditor/Private/SNiagaraEffectEditorViewport.cpp
+++ b/Engine/Source/Editor/NiagaraEditor/Private/SNiagaraEffectEditorViewport.cpp
@@ -35,7 +35,7 @@ private:
 };
 
 FNiagaraEffectEditorViewportClient::FNiagaraEffectEditorViewportClient(TWeakPtr<INiagaraEffectEditor> InEffectEditor, FPreviewScene& InPreviewScene)
-: FEditorViewportClient( GLevelEditorModeTools(), &InPreviewScene )
+: FEditorViewportClient( nullptr, &InPreviewScene )
 , EffectEditorPtr(InEffectEditor)
 {
 	// Setup defaults for the common draw helper.
diff --git a/Engine/Source/Editor/Persona/Private/AnimationEditorViewportClient.cpp b/Engine/Source/Editor/Persona/Private/AnimationEditorViewportClient.cpp
--- a/Engine/Source/Editor/Persona/Private/AnimationEditorViewportClient.cpp
+++ b/Engine/Source/Editor/Persona/Private/AnimationEditorViewportClient.cpp
@@ -79,7 +79,7 @@ IMPLEMENT_HIT_PROXY( HPersonaBoneProxy, HHitProxy );
 // FAnimationViewportClient
 
 FAnimationViewportClient::FAnimationViewportClient( FPreviewScene& InPreviewScene, TWeakPtr<FPersona> InPersonaPtr )
-	: FEditorViewportClient(GLevelEditorModeTools(), &InPreviewScene)
+	: FEditorViewportClient(nullptr, &InPreviewScene)
 	, PersonaPtr( InPersonaPtr )
 	, bManipulating(false)
 	, bInTransaction(false)
diff --git a/Engine/Source/Editor/Persona/Private/SAnimationSequenceBrowser.cpp b/Engine/Source/Editor/Persona/Private/SAnimationSequenceBrowser.cpp
--- a/Engine/Source/Editor/Persona/Private/SAnimationSequenceBrowser.cpp
+++ b/Engine/Source/Editor/Persona/Private/SAnimationSequenceBrowser.cpp
@@ -962,7 +962,8 @@ bool SAnimationSequenceBrowser::IsToolTipPreviewVisible()
 	return bVisible;
 }
 
-FAnimationAssetViewportClient::FAnimationAssetViewportClient(FPreviewScene& InPreviewScene) : FEditorViewportClient(GLevelEditorModeTools(), &InPreviewScene)
+FAnimationAssetViewportClient::FAnimationAssetViewportClient(FPreviewScene& InPreviewScene)
+	: FEditorViewportClient(nullptr, &InPreviewScene)
 {
 	SetViewMode(VMI_Lit);
 
diff --git a/Engine/Source/Editor/PhAT/Private/PhATPreviewViewportClient.cpp b/Engine/Source/Editor/PhAT/Private/PhATPreviewViewportClient.cpp
--- a/Engine/Source/Editor/PhAT/Private/PhATPreviewViewportClient.cpp
+++ b/Engine/Source/Editor/PhAT/Private/PhATPreviewViewportClient.cpp
@@ -20,7 +20,7 @@
 #include "DrawDebugHelpers.h"
 
 FPhATEdPreviewViewportClient::FPhATEdPreviewViewportClient(TWeakPtr<FPhAT> InPhAT, TSharedPtr<FPhATSharedData> Data)
-	: FEditorViewportClient(GLevelEditorModeTools(), &Data->PreviewScene)
+	: FEditorViewportClient(nullptr, &Data->PreviewScene)
 	, PhATPtr(InPhAT)
 	, SharedData(Data)
 	, MinPrimSize(0.5f)
@@ -176,7 +176,7 @@ void FPhATEdPreviewViewportClient::DrawCanvas( FViewport& InViewport, FSceneView
 
 				FName BoneName = SharedData->EditorSkelMesh->RefSkeleton.GetBoneName(i);
 
-				FColor BoneNameColor = FColor(255,255,255);
+				FColor BoneNameColor = FColor::White;
 				//iterate through selected bones and see if any match
 				for(int32 j=0; j< SharedData->SelectedBodies.Num(); ++j)
 				{
@@ -184,7 +184,7 @@ void FPhATEdPreviewViewportClient::DrawCanvas( FViewport& InViewport, FSceneView
 					FName SelectedBoneName = SharedData->PhysicsAsset->BodySetup[SelectedBodyIndex]->BoneName;
 					if(SelectedBoneName == BoneName)
 					{
-						BoneNameColor = FColor(0,255,0);
+						BoneNameColor = FColor::Green;
 						break;
 					}
 
@@ -892,7 +892,7 @@ void FPhATEdPreviewViewportClient::SimMouseMove(float DeltaX, float DeltaY)
 	float QuickRadius = 5 - SimGrabPush / SimHoldDistanceChangeDelta;
 	QuickRadius = QuickRadius < 2 ? 2 : QuickRadius;
 
-	DrawDebugPoint(GetWorld(), NewLocation, QuickRadius, FColor(255,0,0), false, 0.3);
+	DrawDebugPoint(GetWorld(), NewLocation, QuickRadius, FColorList::Red, false, 0.3);
 
 	SharedData->MouseHandle->SetTargetLocation(NewLocation);
 	SharedData->MouseHandle->GrabbedComponent->WakeRigidBody(SharedData->MouseHandle->GrabbedBoneName);
diff --git a/Engine/Source/Editor/StaticMeshEditor/Private/StaticMeshEditorViewportClient.cpp b/Engine/Source/Editor/StaticMeshEditor/Private/StaticMeshEditorViewportClient.cpp
--- a/Engine/Source/Editor/StaticMeshEditor/Private/StaticMeshEditorViewportClient.cpp
+++ b/Engine/Source/Editor/StaticMeshEditor/Private/StaticMeshEditorViewportClient.cpp
@@ -46,7 +46,7 @@ namespace {
 }
 
 FStaticMeshEditorViewportClient::FStaticMeshEditorViewportClient(TWeakPtr<IStaticMeshEditor> InStaticMeshEditor, TWeakPtr<SStaticMeshEditorViewport> InStaticMeshEditorViewport, FPreviewScene& InPreviewScene, UStaticMesh* InPreviewStaticMesh, UStaticMeshComponent* InPreviewStaticMeshComponent)
-	: FEditorViewportClient(GLevelEditorModeTools(), &InPreviewScene)
+	: FEditorViewportClient(nullptr, &InPreviewScene)
 	, StaticMeshEditorPtr(InStaticMeshEditor)
 	, StaticMeshEditorViewportPtr(InStaticMeshEditorViewport)
 {
diff --git a/Engine/Source/Editor/UnrealEd/Private/EditorModeManager.cpp b/Engine/Source/Editor/UnrealEd/Private/EditorModeManager.cpp
--- a/Engine/Source/Editor/UnrealEd/Private/EditorModeManager.cpp
+++ b/Engine/Source/Editor/UnrealEd/Private/EditorModeManager.cpp
@@ -38,7 +38,7 @@ FEditorModeTools::FEditorModeTools()
 	,	SnappedActor( 0 )
 	,	TranslateRotateXAxisAngle(0)
 	,	DefaultID(FBuiltinEditorModes::EM_Default)
-	,	WidgetMode( FWidget::WM_Translate )
+	,	WidgetMode( FWidget::WM_None )
 	,	OverrideWidgetMode( FWidget::WM_None )
 	,	bShowWidget( 1 )
 	,	bHideViewportUI(false)
diff --git a/Engine/Source/Editor/UnrealEd/Private/EditorModes.cpp b/Engine/Source/Editor/UnrealEd/Private/EditorModes.cpp
--- a/Engine/Source/Editor/UnrealEd/Private/EditorModes.cpp
+++ b/Engine/Source/Editor/UnrealEd/Private/EditorModes.cpp
@@ -50,6 +50,10 @@ const FEditorModeID FBuiltinEditorModes::EM_ActorPicker(TEXT("EM_ActorPicker"));
 
 FEdModeDefault::FEdModeDefault()
 {
-	bDrawKillZ = true;
+	bDrawGrid = false;
+	bDrawPivot = false;
+	bDrawBaseInfo = false;
+	bDrawWorldBox = false;
+	bDrawKillZ = false;
 }
 
diff --git a/Engine/Source/Editor/UnrealEd/Private/EditorViewportClient.cpp b/Engine/Source/Editor/UnrealEd/Private/EditorViewportClient.cpp
--- a/Engine/Source/Editor/UnrealEd/Private/EditorViewportClient.cpp
+++ b/Engine/Source/Editor/UnrealEd/Private/EditorViewportClient.cpp
@@ -192,8 +192,9 @@ int32 FEditorViewportClient::GetCameraSpeedSetting() const
 
 float const FEditorViewportClient::SafePadding = 0.075f;
 
-FEditorViewportClient::FEditorViewportClient(FEditorModeTools& InModeTools, FPreviewScene* InPreviewScene)
-	: CameraSpeedSetting(4)
+FEditorViewportClient::FEditorViewportClient(FEditorModeTools* InModeTools, FPreviewScene* InPreviewScene)
+	: bAllowMatineePreview(false)
+	, CameraSpeedSetting(4)
 	, ImmersiveDelegate()
 	, VisibilityDelegate()
 	, Viewport(NULL)
@@ -218,7 +219,9 @@ FEditorViewportClient::FEditorViewportClient(FEditorModeTools& InModeTools, FPre
 	, bDrawAxes(true)
 	, bSetListenerPosition(false)
 	, LandscapeLODOverride(-1)
-	, ModeTools(&InModeTools)
+	, bDrawVertices(false)
+	, bOwnsModeTools(false)
+	, ModeTools(InModeTools)
 	, Widget(new FWidget)
 	, MouseDeltaTracker(new FMouseDeltaTracker)
 	, RecordingInterpEd(NULL)
@@ -255,6 +258,14 @@ FEditorViewportClient::FEditorViewportClient(FEditorModeTools& InModeTools, FPre
 	, FarPlane(0.0f)
 	, bInGameViewMode(false)
 {
+	if (ModeTools == nullptr)
+	{
+		ModeTools = new FEditorModeTools();
+		bOwnsModeTools = true;
+	}
+
+	//@TODO: MODETOOLS: Would like to make this the default, and have specific editors opt-out, but for now opt-in is the safer choice
+	//Widget->SetUsesEditorModeTools(ModeTools);
 
 	ViewState.Allocate();
 
@@ -324,6 +335,11 @@ FEditorViewportClient::~FEditorViewportClient()
 	FCoreDelegates::StatEnabled.RemoveAll(this);
 	FCoreDelegates::StatDisabled.RemoveAll(this);
 	FCoreDelegates::StatDisableAll.RemoveAll(this);
+
+	if (bOwnsModeTools)
+	{
+		delete ModeTools;
+	}
 }
 
 void FEditorViewportClient::RedrawRequested(FViewport* InViewport)
@@ -750,11 +766,14 @@ void FEditorViewportClient::ReceivedFocus(FViewport* InViewport)
 		// Allow the joystick to be used for matinee capture
 		InViewport->SetUserFocus( true );
 	}
+
+	ModeTools->ReceivedFocus(this, Viewport);
 }
 
 void FEditorViewportClient::LostFocus(FViewport* InViewport)
 {
 	StopTracking();
+	ModeTools->LostFocus(this, Viewport);
 }
 
 void FEditorViewportClient::Tick(float DeltaTime)
@@ -806,6 +825,9 @@ void FEditorViewportClient::Tick(float DeltaTime)
 	{
 		Invalidate();
 	}
+
+	// Tick the editor modes
+	ModeTools->Tick(this, DeltaTime);
 }
 
 namespace ViewportDeadZoneConstants
@@ -1517,6 +1539,20 @@ void FEditorViewportClient::HandleViewportStatDisableAll(const bool bInAnyViewpo
 
 void FEditorViewportClient::UpdateMouseDelta()
 {
+	// Do nothing if a drag tool is being used.
+	if (MouseDeltaTracker->UsingDragTool() || ModeTools->DisallowMouseDeltaTracking())
+	{
+		return;
+	}
+
+	// Stop tracking and do nothing else if we're tracking and the widget mode has changed mid-track.
+	// It can confuse the widget code that handles the mouse movements.
+	if (bIsTracking && MouseDeltaTracker->GetTrackingWidgetMode() != GetWidgetMode())
+	{
+		StopTracking();
+		return;
+	}
+
 	FVector DragDelta = FVector::ZeroVector;
 
 	bool bNeedMovementSnapping = true;
@@ -1821,6 +1857,44 @@ bool FEditorViewportClient::InputKey(FViewport* InViewport, int32 ControllerId,
 		return true;
 	}
 
+	// Let the current mode have a look at the input before reacting to it.
+	if (ModeTools->InputKey(this, Viewport, Key, Event))
+	{
+		return true;
+	}
+
+	FInputEventState InputState(InViewport, Key, Event);
+
+	bool bHandled = false;
+
+	if ((IsOrtho() || InputState.IsAltButtonPressed()) && (Key == EKeys::Left || Key == EKeys::Right || Key == EKeys::Up || Key == EKeys::Down))
+	{
+		NudgeSelectedObjects(InputState);
+
+		bHandled = true;
+	}
+	else if (Key == EKeys::Escape && Event == IE_Pressed/* && TrackingTransaction.IsActive()*/) //@TODO: MODETOOLS: Now consuming Escape regardless of if we were previously tracking
+	{
+		// Pressing Escape cancels the current operation
+		AbortTracking();
+		bHandled = true;
+	}
+
+	// If in ortho and right mouse button and ctrl is pressed
+	if (!InputState.IsAltButtonPressed()
+		&& InputState.IsCtrlButtonPressed()
+		&& !InputState.IsButtonPressed(EKeys::LeftMouseButton)
+		&& !InputState.IsButtonPressed(EKeys::MiddleMouseButton)
+		&& InputState.IsButtonPressed(EKeys::RightMouseButton)
+		&& IsOrtho())
+	{
+		ModeTools->SetWidgetModeOverride(FWidget::WM_Rotate);
+	}
+	else
+	{
+		ModeTools->SetWidgetModeOverride(FWidget::WM_None);
+	}
+
 	const int32	HitX = InViewport->GetMouseX();
 	const int32	HitY = InViewport->GetMouseY();
 
@@ -1830,8 +1904,6 @@ bool FEditorViewportClient::InputKey(FViewport* InViewport, int32 ControllerId,
 		JoystickState->KeyEventValues.Add(Key, Event);
 	}
 
-	FInputEventState InputState( InViewport, Key, Event );
-
 	const bool bWasCursorVisible = InViewport->IsCursorVisible();
 	const bool bWasSoftwareCursorVisible = InViewport->IsSoftwareCursorVisible();
 
@@ -1842,8 +1914,6 @@ bool FEditorViewportClient::InputKey(FViewport* InViewport, int32 ControllerId,
 	RequiredCursorVisibiltyAndAppearance.bDontResetCursor = false;
 	UpdateRequiredCursorVisibility();
 
-	bool bHandled = false;
-
 	if( bWasCursorVisible != RequiredCursorVisibiltyAndAppearance.bHardwareCursorVisible || bWasSoftwareCursorVisible != RequiredCursorVisibiltyAndAppearance.bSoftwareCursorVisible )
 	{
 		bHandled = true;	
@@ -1968,9 +2038,14 @@ void FEditorViewportClient::StopTracking()
 	bHasMouseMovedSinceClick = false;
 }
 
+void FEditorViewportClient::AbortTracking()
+{
+	StopTracking();
+}
+
 bool FEditorViewportClient::IsInImmersiveViewport() const
 {
-	return ImmersiveDelegate.IsBound() ? ImmersiveDelegate.Execute() : false ;
+	return ImmersiveDelegate.IsBound() ? ImmersiveDelegate.Execute() : false;
 }
 
 void FEditorViewportClient::StartTrackingDueToInput( const struct FInputEventState& InputState, FSceneView& View )
@@ -2517,6 +2592,97 @@ void FEditorViewportClient::AddReferencedObjects( FReferenceCollector& Collector
 	}
 }
 
+void FEditorViewportClient::ProcessClick(class FSceneView& View, class HHitProxy* HitProxy, FKey Key, EInputEvent Event, uint32 HitX, uint32 HitY)
+{
+	const FViewportClick Click(&View, this, Key, Event, HitX, HitY);
+	ModeTools->HandleClick(this, HitProxy, Click);
+}
+
+bool FEditorViewportClient::InputWidgetDelta(FViewport* InViewport, EAxisList::Type CurrentAxis, FVector& Drag, FRotator& Rot, FVector& Scale)
+{
+	if (ModeTools->InputDelta(this, Viewport, Drag, Rot, Scale))
+	{
+		if (ModeTools->AllowWidgetMove())
+		{
+			ModeTools->PivotLocation += Drag;
+			ModeTools->SnappedLocation += Drag;
+		}
+
+		// Update visuals of the rotate widget 
+		ApplyDeltaToRotateWidget(Rot);
+		return true;
+	}
+	else
+	{
+		return false;
+	}
+}
+
+void FEditorViewportClient::SetWidgetMode(FWidget::EWidgetMode NewMode)
+{
+	if (!ModeTools->IsTracking() && !IsFlightCameraActive())
+	{
+		ModeTools->SetWidgetMode(NewMode);
+
+		// force an invalidation (non-deferred) of the hit proxy here, otherwise we will
+		// end up checking against an incorrect hit proxy if the cursor is not moved
+		Viewport->InvalidateHitProxy();
+		bShouldCheckHitProxy = true;
+
+		// Fire event delegate
+		ModeTools->BroadcastWidgetModeChanged(NewMode);
+	}
+
+	RedrawAllViewportsIntoThisScene();
+}
+
+bool FEditorViewportClient::CanSetWidgetMode(FWidget::EWidgetMode NewMode) const
+{
+	return ModeTools->GetShowWidget() == true;
+}
+
+FWidget::EWidgetMode FEditorViewportClient::GetWidgetMode() const
+{
+	return ModeTools->GetWidgetMode();
+}
+
+FVector FEditorViewportClient::GetWidgetLocation() const
+{
+	return ModeTools->GetWidgetLocation();
+}
+
+FMatrix FEditorViewportClient::GetWidgetCoordSystem() const
+{
+	return ModeTools->GetCustomInputCoordinateSystem();
+}
+
+void FEditorViewportClient::SetWidgetCoordSystemSpace(ECoordSystem NewCoordSystem)
+{
+	ModeTools->SetCoordSystem(NewCoordSystem);
+	RedrawAllViewportsIntoThisScene();
+}
+
+ECoordSystem FEditorViewportClient::GetWidgetCoordSystemSpace() const
+{
+	return ModeTools->GetCoordSystem();
+}
+
+void FEditorViewportClient::ApplyDeltaToRotateWidget(const FRotator& InRot)
+{
+	//apply rotation to translate rotate widget
+	if (!InRot.IsZero())
+	{
+		FRotator TranslateRotateWidgetRotation(0, ModeTools->TranslateRotateXAxisAngle, 0);
+		TranslateRotateWidgetRotation += InRot;
+		ModeTools->TranslateRotateXAxisAngle = TranslateRotateWidgetRotation.Yaw;
+	}
+}
+
+void FEditorViewportClient::RedrawAllViewportsIntoThisScene()
+{
+	Invalidate();
+}
+
 FSceneInterface* FEditorViewportClient::GetScene() const
 {
 	UWorld* World = GetWorld();
@@ -2543,6 +2709,13 @@ UWorld* FEditorViewportClient::GetWorld() const
 	return OutWorldPtr;
 }
 
+void FEditorViewportClient::DrawCanvas(FViewport& InViewport, FSceneView& View, FCanvas& Canvas)
+{
+	// Information string
+	Canvas.DrawShadowedString(4, 4, *ModeTools->InfoString, GEngine->GetSmallFont(), FColor::White);
+
+	ModeTools->DrawHUD(this, &InViewport, &View, &Canvas);
+}
 
 void FEditorViewportClient::SetupViewForRendering( FSceneViewFamily& ViewFamily, FSceneView& View )
 {
@@ -2732,6 +2905,11 @@ void FEditorViewportClient::Draw(const FSceneView* View,FPrimitiveDrawInterface*
 		DrawHelper.Draw( View, PDI );
 	}
 
+	ModeTools->DrawActiveModes(View, PDI);
+
+	// Draw the current editor mode.
+	ModeTools->Render(View, Viewport, PDI);
+
 	// This viewport was just rendered, reset this value.
 	FramesSinceLastDraw = 0;
 }
@@ -2883,6 +3061,12 @@ bool FEditorViewportClient::InputAxis(FViewport* InViewport, int32 ControllerId,
 		return true;
 	}
 
+	// Let the current mode have a look at the input before reacting to it.
+	if (ModeTools->InputAxis(this, Viewport, ControllerId, Key, Delta, DeltaTime))
+	{
+		return true;
+	}
+
 	const bool bMouseButtonDown = InViewport->KeyState( EKeys::LeftMouseButton ) || InViewport->KeyState( EKeys::MiddleMouseButton ) || InViewport->KeyState( EKeys::RightMouseButton );
 	const bool bLightMoveDown = InViewport->KeyState(EKeys::L);
 
@@ -3429,9 +3613,10 @@ EAxisList::Type FEditorViewportClient::GetCurrentWidgetAxis() const
 	return Widget->GetCurrentAxis();
 }
 
-void FEditorViewportClient::SetCurrentWidgetAxis( EAxisList::Type InAxis )
+void FEditorViewportClient::SetCurrentWidgetAxis(EAxisList::Type InAxis)
 {
-	Widget->SetCurrentAxis( InAxis );
+	Widget->SetCurrentAxis(InAxis);
+	ModeTools->SetCurrentWidgetAxis(InAxis);
 }
 
 void FEditorViewportClient::AdjustTransformWidgetSize(const int32 SizeDelta)
@@ -3832,6 +4017,8 @@ void FEditorViewportClient::OnJoystickPlugged(const uint32 InControllerID, const
 
 void FEditorViewportClient::MouseEnter(FViewport* InViewport,int32 x, int32 y)
 {
+	ModeTools->MouseEnter(this, Viewport, x, y);
+
 	MouseMove(InViewport, x, y);
 }
 
@@ -3840,12 +4027,17 @@ void FEditorViewportClient::MouseMove(FViewport* InViewport,int32 x, int32 y)
 	check(IsInGameThread());
 
 	CurrentMousePos = FIntPoint(x, y);
+
+	// Let the current editor mode know about the mouse movement.
+	ModeTools->MouseMove(this, Viewport, x, y);
 }
 
 void FEditorViewportClient::MouseLeave(FViewport* InViewport)
 {
 	check(IsInGameThread());
 
+	ModeTools->MouseLeave(this, Viewport);
+
 	CurrentMousePos = FIntPoint(-1, -1);
 
 	FCommonViewportClient::MouseLeave(InViewport);
diff --git a/Engine/Source/Editor/UnrealEd/Private/LevelEditorViewport.cpp b/Engine/Source/Editor/UnrealEd/Private/LevelEditorViewport.cpp
--- a/Engine/Source/Editor/UnrealEd/Private/LevelEditorViewport.cpp
+++ b/Engine/Source/Editor/UnrealEd/Private/LevelEditorViewport.cpp
@@ -1423,7 +1423,7 @@ void FTrackingTransaction::PromotePendingToActive()
 }
 
 FLevelEditorViewportClient::FLevelEditorViewportClient()
-	: FEditorViewportClient(GLevelEditorModeTools(), NULL)
+	: FEditorViewportClient(&GLevelEditorModeTools(), NULL)
 	, ViewHiddenLayers()
 	, VolumeActorVisibility()
 	, LastEditorViewLocation( FVector::ZeroVector )
@@ -1453,7 +1453,9 @@ FLevelEditorViewportClient::FLevelEditorViewportClient()
 
 	GEditor->LevelViewportClients.Add(this);
 
-	Widget->SetUsesEditorModeTools( ModeTools );
+	// The level editor fully supports mode tools and isn't doing any incompatible stuff with the Widget
+	ModeTools->SetWidgetMode(FWidget::WM_Translate);
+	Widget->SetUsesEditorModeTools(ModeTools);
 
 	// Register for editor cleanse events so we can release references to hovered actors
 	FEditorSupportDelegates::CleanseEditor.AddRaw(this, &FLevelEditorViewportClient::OnEditorCleanse);
@@ -1834,8 +1836,6 @@ void FLevelEditorViewportClient::Tick(float DeltaTime)
 		GUnrealEd->UpdatePivotLocationForSelection();
 	}
 
-	ModeTools->Tick(this, DeltaTime);
-
 	// Update the preview mesh for the preview mesh mode. 
 	GEditor->UpdatePreviewMesh();
 
@@ -2057,7 +2057,7 @@ void FLevelEditorViewportClient::ProjectActorsIntoWorld(const TArray<AActor*>& A
 	}
 }
 
-bool FLevelEditorViewportClient::InputWidgetDelta( FViewport* Viewport, EAxisList::Type CurrentAxis, FVector& Drag, FRotator& Rot, FVector& Scale )
+bool FLevelEditorViewportClient::InputWidgetDelta(FViewport* Viewport, EAxisList::Type CurrentAxis, FVector& Drag, FRotator& Rot, FVector& Scale)
 {
 	if (GUnrealEd->ComponentVisManager.HandleInputDelta(this, Viewport, Drag, Rot, Scale))
 	{
@@ -2067,20 +2067,13 @@ bool FLevelEditorViewportClient::InputWidgetDelta( FViewport* Viewport, EAxisLis
 	bool bHandled = false;
 
 	// Give the current editor mode a chance to use the input first.  If it does, don't apply it to anything else.
-	if (ModeTools->InputDelta(this, Viewport, Drag, Rot, Scale))
+	if (FEditorViewportClient::InputWidgetDelta(Viewport, CurrentAxis, Drag, Rot, Scale))
 	{
-		if (ModeTools->AllowWidgetMove())
-		{
-			ModeTools->PivotLocation += Drag;
-			ModeTools->SnappedLocation += Drag;
-		}
-
-		// Update visuals of the rotate widget 
-		ApplyDeltaToRotateWidget( Rot );
 		bHandled = true;
 	}
 	else
 	{
+		//@TODO: MODETOOLS: Much of this needs to get pushed to Super, but not all of it can be...
 		if( CurrentAxis != EAxisList::None )
 		{
 			// Skip actors transformation routine in case if any of the selected actors locked
@@ -2193,38 +2186,6 @@ TSharedPtr<FDragTool> FLevelEditorViewportClient::MakeDragTool( EDragTool::Type
 	return DragTool;
 }
 
-void FLevelEditorViewportClient::UpdateMouseDelta()
-{
-	// Do nothing if a drag tool is being used.
-	if (MouseDeltaTracker->UsingDragTool() || ModeTools->DisallowMouseDeltaTracking())
-	{
-		return;
-	}
-
-	// Stop tracking and do nothing else if we're tracking and the widget mode has changed mid-track.
-	// It can confuse the widget code that handles the mouse movements.
-	if (bIsTracking && MouseDeltaTracker->GetTrackingWidgetMode() != ModeTools->GetWidgetMode())
-	{
-		StopTracking();
-		return;
-	}
-
-	// If any actor in the selection requires snapping, they all need to be snapped.
-	bool bNeedMovementSnapping = false;
-
-	for ( FSelectionIterator It( GEditor->GetSelectedActorIterator() ) ; It ; ++It )
-	{
-		AActor* Actor = static_cast<AActor*>( *It );
-		checkSlow( Actor->IsA(AActor::StaticClass()) );
-
-		bNeedMovementSnapping = true;
-		break;
-	}
-
-	FEditorViewportClient::UpdateMouseDelta();
-}
-
-
 static bool CommandAcceptsInput( FLevelEditorViewportClient& ViewportClient, FKey Key, const TSharedPtr<FUICommandInfo> Command )
 {
 	const FInputGesture& Gesture = *Command->GetActiveGesture();
@@ -2316,77 +2277,36 @@ bool FLevelEditorViewportClient::InputKey(FViewport* Viewport, int32 ControllerI
 		return true;
 	}
 
-	
-	// Let the current mode have a look at the input before reacting to it.
-	if (ModeTools->InputKey(this, Viewport, Key, Event))
-	{
-		return true;
-	}
-
+	bool bHandled = FEditorViewportClient::InputKey(Viewport,ControllerId,Key,Event,AmountDepressed,bGamepad);
 
-	bool bHandled = false;
 	// Handle input for the player height preview mode. 
-	if( !InputState.IsMouseButtonEvent() && CommandAcceptsInput( *this, Key, GetLevelViewportCommands().EnablePreviewMesh ) )
+	if (!InputState.IsMouseButtonEvent() && CommandAcceptsInput(*this, Key, GetLevelViewportCommands().EnablePreviewMesh))
 	{
 		// Holding down the backslash buttons turns on the mode. 
-		if( Event == IE_Pressed )
+		if (Event == IE_Pressed)
 		{
 			GEditor->SetPreviewMeshMode(true);
 
 			// If shift down, cycle between the preview meshes
-			if( CommandAcceptsInput( *this, Key, GetLevelViewportCommands().CyclePreviewMesh ) )
+			if (CommandAcceptsInput(*this, Key, GetLevelViewportCommands().CyclePreviewMesh))
 			{
 				GEditor->CyclePreviewMesh();
 			}
 		}
 		// Releasing backslash turns off the mode. 
-		else if( Event == IE_Released )
+		else if (Event == IE_Released)
 		{
 			GEditor->SetPreviewMeshMode(false);
 		}
 
 		bHandled = true;
 	}
-	else if( (IsOrtho() || InputState.IsAltButtonPressed()) && (Key == EKeys::Left || Key == EKeys::Right || Key == EKeys::Up || Key == EKeys::Down) )
-	{
-		NudgeSelectedObjects( InputState );
-
-		bHandled = true;
-	}
-	else if (Key == EKeys::Escape && Event == IE_Pressed && TrackingTransaction.IsActive())
-	{
-		// Pressing Escape cancels the current operation
-
-		// Applying the global undo here will reset the drag operation
-		GUndo->Apply();
-		TrackingTransaction.Cancel();
-		StopTracking();
-		bHandled = true;
-	}
-
-	// If in ortho and right mouse button and ctrl is pressed
-	if ( !InputState.IsAltButtonPressed()
-		&& InputState.IsCtrlButtonPressed() 
-		&& !InputState.IsButtonPressed(EKeys::LeftMouseButton) 
-		&& !InputState.IsButtonPressed(EKeys::MiddleMouseButton) 
-		&& InputState.IsButtonPressed(EKeys::RightMouseButton) 
-		&& IsOrtho() )
-	{
-		ModeTools->SetWidgetModeOverride(FWidget::WM_Rotate);
-	}
-	else
-	{
-		ModeTools->SetWidgetModeOverride(FWidget::WM_None);
-	}
-
-	bHandled |= FEditorViewportClient::InputKey(Viewport,ControllerId,Key,Event,AmountDepressed,bGamepad);
 
 	// Clear Duplicate Actors mode when ALT and all mouse buttons are released
 	if ( !InputState.IsAltButtonPressed() && !InputState.IsAnyMouseButtonDown() )
 	{
 		bDuplicateActorsInProgress = false;
 	}
-
 	
 	return bHandled;
 }
@@ -2427,7 +2347,7 @@ void FLevelEditorViewportClient::TrackingStarted( const FInputEventState& InInpu
 
 		if( bIsDraggingWidget )
 		{
-			// Notify that this actor is begining to move
+			// Notify that this actor is beginning to move
 			GEditor->BroadcastBeginObjectMovement( *Actor );
 		}
 
@@ -2621,15 +2541,26 @@ void FLevelEditorViewportClient::TrackingStopped()
 		FScopedLevelDirtied LevelDirtyCallback;
 		LevelDirtyCallback.Request();
 
-		GEditor->RedrawLevelEditingViewports();
+		RedrawAllViewportsIntoThisScene();
 	}
 
 	PreDragActorTransforms.Empty();
 }
 
+void FLevelEditorViewportClient::AbortTracking()
+{
+	if (TrackingTransaction.IsActive())
+	{
+		// Applying the global undo here will reset the drag operation
+		GUndo->Apply();
+		TrackingTransaction.Cancel();
+		StopTracking();
+	}
+}
+
 void FLevelEditorViewportClient::HandleViewportSettingChanged(FName PropertyName)
 {
-	if (PropertyName == TEXT("bUseSelectionOutline"))
+	if (PropertyName == GET_MEMBER_NAME_CHECKED(ULevelEditorViewportSettings, bUseSelectionOutline))
 	{
 		EngineShowFlags.SelectionOutline = GetDefault<ULevelEditorViewportSettings>()->bUseSelectionOutline;
 	}
@@ -2690,7 +2621,7 @@ void FLevelEditorViewportClient::NudgeSelectedObjects( const struct FInputEventS
 		Widget->SetCurrentAxis( EAxisList::None );
 	}
 
-	GEditor->RedrawLevelEditingViewports();
+	RedrawAllViewportsIntoThisScene();
 }
 
 
@@ -2765,15 +2696,8 @@ bool FLevelEditorViewportClient::InputAxis(FViewport* Viewport, int32 Controller
 	// @todo Slate: GCurrentLevelEditingViewportClient is switched multiple times per frame and since we draw the border in slate this effectively causes border to always draw on the last viewport
 
 	FScopedSetCurrentViewportClient( this );
-	// Let the current mode have a look at the input before reacting to it.
-	if (ModeTools->InputAxis(this, Viewport, ControllerId, Key, Delta, DeltaTime))
-	{
-		return true;
-	}
 
-	FEditorViewportClient::InputAxis(Viewport,ControllerId,Key,Delta,DeltaTime,NumSamples,bGamepad);
-
-	return true;
+	return FEditorViewportClient::InputAxis(Viewport, ControllerId, Key, Delta, DeltaTime, NumSamples, bGamepad);
 }
 
 
@@ -2817,45 +2741,10 @@ bool FLevelEditorViewportClient::IsVolumeVisibleInViewport( const AActor& Volume
 	return VolumeActorVisibility[ VolumeId ];
 }
 
-void FLevelEditorViewportClient::SetWidgetMode( FWidget::EWidgetMode ActivatedMode )
+void FLevelEditorViewportClient::RedrawAllViewportsIntoThisScene()
 {
-	if (!ModeTools->IsTracking() && !IsFlightCameraActive())
-	{
-		ModeTools->SetWidgetMode(ActivatedMode);
-
-		// force an invalidation (non-deferred) of the hit proxy here, otherwise we will
-		// end up checking against an incorrect hit proxy if the cursor is not moved
-		Viewport->InvalidateHitProxy();
-		bShouldCheckHitProxy = true;
-
-		// Fire event delegate
-		ModeTools->BroadcastWidgetModeChanged(ActivatedMode);
-	}
-	
 	// Invalidate all viewports, so the new gizmo is rendered in each one
-	for (auto ViewportClient : GEditor->LevelViewportClients)
-	{
-		if (ViewportClient)
-		{
-			ViewportClient->Invalidate();
-		}
-	}
-}
-
-bool FLevelEditorViewportClient::CanSetWidgetMode( FWidget::EWidgetMode NewMode ) const
-{
-	return ModeTools->GetShowWidget() == true;
-}
-
-void FLevelEditorViewportClient::SetWidgetCoordSystemSpace( ECoordSystem NewCoordSystem )
-{
-	ModeTools->SetCoordSystem(NewCoordSystem);
-	Invalidate();
-}
-
-FWidget::EWidgetMode FLevelEditorViewportClient::GetWidgetMode() const
-{
-	return ModeTools->GetWidgetMode();
+	GEditor->RedrawLevelEditingViewports();
 }
 
 FVector FLevelEditorViewportClient::GetWidgetLocation() const
@@ -2866,7 +2755,7 @@ FVector FLevelEditorViewportClient::GetWidgetLocation() const
 		return ComponentVisWidgetLocation;
 	}
 
-	return ModeTools->GetWidgetLocation();
+	return FEditorViewportClient::GetWidgetLocation();
 }
 
 FMatrix FLevelEditorViewportClient::GetWidgetCoordSystem() const 
@@ -2877,16 +2766,10 @@ FMatrix FLevelEditorViewportClient::GetWidgetCoordSystem() const
 		return ComponentVisWidgetCoordSystem;
 	}
 
-	return ModeTools->GetCustomInputCoordinateSystem();
-}
-
-ECoordSystem FLevelEditorViewportClient::GetWidgetCoordSystemSpace() const
-{
-	return ModeTools->GetCoordSystem();
+	return FEditorViewportClient::GetWidgetCoordSystem();
 }
 
-
-void FLevelEditorViewportClient::MoveLockedActorToCamera() const
+void FLevelEditorViewportClient::MoveLockedActorToCamera()
 {
 	// If turned on, move any selected actors to the cameras location/rotation
 	TWeakObjectPtr<AActor> ActiveActorLock = GetActiveActorLock();
@@ -2907,7 +2790,7 @@ void FLevelEditorViewportClient::MoveLockedActorToCamera() const
 		FScopedLevelDirtied LevelDirtyCallback;
 		LevelDirtyCallback.Request();
 
-		GEditor->RedrawLevelEditingViewports();
+		RedrawAllViewportsIntoThisScene();
 	}
 }
 
@@ -3275,43 +3158,6 @@ void FLevelEditorViewportClient::ApplyDeltaToActor( AActor* InActor, const FVect
 	UpdateLockedActorViewports(InActor, true);
 }
 
-/** Updates the rotate widget with the passed in delta rotation. */
-void FLevelEditorViewportClient::ApplyDeltaToRotateWidget( const FRotator& InRot )
-{
-	//apply rotation to translate rotate widget
-	if (!InRot.IsZero())
-	{
-		FRotator TranslateRotateWidgetRotation(0, ModeTools->TranslateRotateXAxisAngle, 0);
-		TranslateRotateWidgetRotation += InRot;
-		ModeTools->TranslateRotateXAxisAngle = TranslateRotateWidgetRotation.Yaw;
-	}
-
-}
-
-void FLevelEditorViewportClient::MouseEnter( FViewport* Viewport,int32 x, int32 y )
-{
-	ModeTools->MouseEnter(this, Viewport, x, y);
-}
-
-void FLevelEditorViewportClient::MouseLeave( FViewport* Viewport ) 
-{
-	ModeTools->MouseLeave(this, Viewport);
-
-	FEditorViewportClient::MouseLeave(Viewport);
-}
-
-void FLevelEditorViewportClient::MouseMove(FViewport* Viewport,int32 x, int32 y)
-{
-	FEditorViewportClient::MouseMove(Viewport,x,y);
-
-	// Let the current editor mode know about the mouse movement.
-	if (IsLevelEditorClient() && ModeTools->MouseMove(this, Viewport, x, y))
-	{
-		return;
-	}
-}
-
-
 EMouseCursor::Type FLevelEditorViewportClient::GetCursor(FViewport* Viewport,int32 X,int32 Y)
 {
 	EMouseCursor::Type CursorType = FEditorViewportClient::GetCursor(Viewport,X,Y);
@@ -3512,12 +3358,6 @@ bool FLevelEditorViewportClient::GetActiveSafeFrame(float& OutAspectRatio) const
 	return false;
 }
 
-void FLevelEditorViewportClient::SetCurrentWidgetAxis( EAxisList::Type NewAxis )
-{
-	FEditorViewportClient::SetCurrentWidgetAxis( NewAxis );
-	ModeTools->SetCurrentWidgetAxis(NewAxis);
-}
-
 /** 
  * Renders a view frustum specified by the provided frustum parameters
  *
@@ -3596,12 +3436,6 @@ void FLevelEditorViewportClient::Draw(const FSceneView* View,FPrimitiveDrawInter
 		DrawTextureStreamingBounds(View, PDI);
 	}
 
-
-	ModeTools->DrawActiveModes(View, PDI);
-
-	// Draw the current editor mode.
-	ModeTools->Render(View, Viewport, PDI);
-
 	// Determine if a view frustum should be rendered in the viewport.
 	// The frustum should definitely be rendered if the viewport has a view parent.
 	bool bRenderViewFrustum = ViewState.GetReference()->HasViewParent();
@@ -3794,10 +3628,7 @@ void FLevelEditorViewportClient::DrawCanvas( FViewport& InViewport, FSceneView&
 		GUnrealEd->DrawComponentVisualizersHUD(&InViewport, &View, &Canvas);
 	}
 
-	// Information string
-	Canvas.DrawShadowedString(4, 4, *ModeTools->InfoString, GEngine->GetSmallFont(), FColor(255, 255, 255));
-
-	ModeTools->DrawHUD(this, &InViewport, &View, &Canvas);
+	FEditorViewportClient::DrawCanvas(InViewport, View, Canvas);
 
 	// Testbed
 	FCanvasItemTestbed TestBed;
@@ -3830,14 +3661,14 @@ void FLevelEditorViewportClient::DrawTextureStreamingBounds(const FSceneView* Vi
 				InMin -= FVector(Max);
 				InMax += FVector(Max);
 				FBox Box = FBox(InMin, InMax);
-				DrawWireBox(PDI, Box, FColor(255,255,0), SDPG_World);
+				DrawWireBox(PDI, Box, FColorList::Yellow, SDPG_World);
 #else	//#if defined(_STREAMING_BOUNDS_DRAW_BOX_)
 				// Draw bounding spheres
 				FVector Origin = STI.BoundingSphere.Center;
 				float Radius = STI.BoundingSphere.W;
-				DrawCircle(PDI, Origin, FVector(1,0,0), FVector(0,1,0), FColor(255,255,0), Radius, 32, SDPG_World);
-				DrawCircle(PDI, Origin, FVector(1,0,0), FVector(0,0,1), FColor(255,255,0), Radius, 32, SDPG_World);
-				DrawCircle(PDI, Origin, FVector(0,1,0), FVector(0,0,1), FColor(255,255,0), Radius, 32, SDPG_World);
+				DrawCircle(PDI, Origin, FVector(1, 0, 0), FVector(0, 1, 0), FColorList::Yellow, Radius, 32, SDPG_World);
+				DrawCircle(PDI, Origin, FVector(1, 0, 0), FVector(0, 0, 1), FColorList::Yellow, Radius, 32, SDPG_World);
+				DrawCircle(PDI, Origin, FVector(0, 1, 0), FVector(0, 0, 1), FColorList::Yellow, Radius, 32, SDPG_World);
 #endif	//#if defined(_STREAMING_BOUNDS_DRAW_BOX_)
 			}
 		}
@@ -3977,17 +3808,6 @@ void FLevelEditorViewportClient::SetCameraSpeedSetting(int32 SpeedSetting)
 	GetMutableDefault<ULevelEditorViewportSettings>()->CameraSpeed = SpeedSetting;
 }
 
-void FLevelEditorViewportClient::ReceivedFocus(FViewport* Viewport)
-{
-	FEditorViewportClient::ReceivedFocus(Viewport);
-	ModeTools->ReceivedFocus(this, Viewport);
-}
-
-void FLevelEditorViewportClient::LostFocus(FViewport* Viewport)
-{
-	FEditorViewportClient::LostFocus(Viewport);
-	ModeTools->LostFocus(this, Viewport);
-}
 
 bool FLevelEditorViewportClient::OverrideHighResScreenshotCaptureRegion(FIntRect& OutCaptureRegion)
 {
diff --git a/Engine/Source/Editor/UnrealEd/Private/SSkeletonWidget.cpp b/Engine/Source/Editor/UnrealEd/Private/SSkeletonWidget.cpp
--- a/Engine/Source/Editor/UnrealEd/Private/SSkeletonWidget.cpp
+++ b/Engine/Source/Editor/UnrealEd/Private/SSkeletonWidget.cpp
@@ -1069,7 +1069,7 @@ class FBasePoseViewportClient: public FEditorViewportClient
 {
 public:
 	FBasePoseViewportClient(FPreviewScene& InPreviewScene)
-			: FEditorViewportClient(GLevelEditorModeTools(), &InPreviewScene)
+			: FEditorViewportClient(nullptr, &InPreviewScene)
 	{
 		SetViewMode(VMI_Lit);
 
diff --git a/Engine/Source/Editor/UnrealEd/Public/EditorViewportClient.h b/Engine/Source/Editor/UnrealEd/Public/EditorViewportClient.h
--- a/Engine/Source/Editor/UnrealEd/Public/EditorViewportClient.h
+++ b/Engine/Source/Editor/UnrealEd/Public/EditorViewportClient.h
@@ -240,7 +240,7 @@ class UNREALED_API FEditorViewportClient : public FCommonViewportClient, public
 public:
 	friend class FMouseDeltaTracker;
 
-	FEditorViewportClient(FEditorModeTools& InModeTools, FPreviewScene* InPreviewScene = nullptr);
+	FEditorViewportClient(FEditorModeTools* InModeTools, FPreviewScene* InPreviewScene = nullptr);
 	virtual ~FEditorViewportClient();
 
 	/**
@@ -475,7 +475,7 @@ public:
 	 * @param HitX			The X location of the mouse
 	 * @param HitY			The Y location of the mouse
 	 */
-	virtual void ProcessClick(class FSceneView& View, class HHitProxy* HitProxy, FKey Key, EInputEvent Event, uint32 HitX, uint32 HitY) {}
+	virtual void ProcessClick(class FSceneView& View, class HHitProxy* HitProxy, FKey Key, EInputEvent Event, uint32 HitX, uint32 HitY);
 	
 	/**
 	 * Called when mouse movement tracking begins
@@ -500,17 +500,17 @@ public:
 	 * @param Rot			The amount the widget was rotated  (the value depends on the coordinate system of the widget.  See GetWidgetCoordSystem )
 	 * @param Scale			The amount the widget was scaled (the value depends on the coordinate system of the widget.  See GetWidgetCoordSystem )
 	 */
-	virtual bool InputWidgetDelta( FViewport* InViewport, EAxisList::Type CurrentAxis, FVector& Drag, FRotator& Rot, FVector& Scale ) { return false; }
+	virtual bool InputWidgetDelta(FViewport* InViewport, EAxisList::Type CurrentAxis, FVector& Drag, FRotator& Rot, FVector& Scale);
 
 	/**
 	 * Sets the current widget mode
 	 */
-	virtual void SetWidgetMode( FWidget::EWidgetMode NewMode ) {};
+	virtual void SetWidgetMode(FWidget::EWidgetMode NewMode);
 
 	/**
 	 * Whether or not the new widget mode can be set in this viewport
 	 */
-	virtual bool CanSetWidgetMode( FWidget::EWidgetMode NewMode ) const { return true; }
+	virtual bool CanSetWidgetMode(FWidget::EWidgetMode NewMode) const;
 
 	/**
 	 * Whether or not the widget mode can be cycled
@@ -520,28 +520,28 @@ public:
 	/**
 	 * @return The current display mode for transform widget 
 	 */
-	virtual FWidget::EWidgetMode GetWidgetMode() const { return FWidget::WM_None; }
+	virtual FWidget::EWidgetMode GetWidgetMode() const;
 
 	/**
 	 * @return The world space location of the transform widget
 	 */
-	virtual FVector GetWidgetLocation() const { return FVector::ZeroVector; }
+	virtual FVector GetWidgetLocation() const;
 
 	/**
 	 * @return The current coordinate system for drawing and input of the transform widget.  
 	 * For world coordiante system return the identity matrix
 	 */
-	virtual FMatrix GetWidgetCoordSystem() const { return FMatrix::Identity; }
+	virtual FMatrix GetWidgetCoordSystem() const;
 
 	/**
 	 * Sets the coordinate system space to use
 	 */
-	virtual void SetWidgetCoordSystemSpace( ECoordSystem NewCoordSystem ) {};
+	virtual void SetWidgetCoordSystemSpace( ECoordSystem NewCoordSystem );
 
 	/**
 	 * @return The coordinate system space (world or local) to display the widget in
 	 */
-	virtual ECoordSystem GetWidgetCoordSystemSpace() const { return COORD_World; }
+	virtual ECoordSystem GetWidgetCoordSystemSpace() const;
 
 	/**
 	 * Sets the current axis being manipulated by the transform widget
@@ -568,7 +568,7 @@ public:
 	 * @param View			The view of the scene to be rendered
 	 * @param Canvas		The canvas to draw on
 	 */
-	virtual void DrawCanvas( FViewport& InViewport, FSceneView& View, FCanvas& Canvas ) {};
+	virtual void DrawCanvas(FViewport& InViewport, FSceneView& View, FCanvas& Canvas);
 
 	/**
 	 * Render the drag tool in the viewport
@@ -1018,6 +1018,12 @@ protected:
 	 */
 	virtual void PerspectiveCameraMoved() {}
 
+	/** Updates the rotate widget with the passed in delta rotation. */
+	void ApplyDeltaToRotateWidget(const FRotator& InRot);
+
+	/** Invalidates this and other linked viewports (anything viewing the same scene) */
+	virtual void RedrawAllViewportsIntoThisScene();
+
 	/**
 	 * Used to store the required cursor visibility states and override cursor appearance
 	 */
@@ -1133,6 +1139,11 @@ protected:
 	 */
 	void StopTracking();
 
+	/**
+	 * Aborts mouse tracking (stop and cancel)
+	 */
+	virtual void AbortTracking();
+
 	/** Enables or disables camera lock **/
 	void EnableCameraLock(bool bEnable);
 
@@ -1147,6 +1158,8 @@ protected:
 	/** Helper function to calculate the safe frame rectangle on the current viewport */
 	bool CalculateEditorConstrainedViewRect(FSlateRect& OutSafeFrameRect, FViewport* InViewport);
 
+	virtual void NudgeSelectedObjects(const struct FInputEventState& InputState) {}
+
 private:
 	/** @return Whether or not the camera should be panned or dollied */
 	bool ShouldPanOrDollyCamera() const;
@@ -1290,8 +1303,10 @@ public:
 	bool bDrawVertices;
 
 protected:
+	/** Does this viewport client own the mode tools instance pointed at by ModeTools control the lifetime of it? */
+	bool bOwnsModeTools;
 
-	/** Editor mode tools provided to this instance. Assumed to be managed externally */
+	/** Editor mode tools provided to this instance. Assumed to be managed externally if bOwnsModeTools is false */
 	FEditorModeTools*		ModeTools;
 
 	FWidget*				Widget;
diff --git a/Engine/Source/Editor/UnrealEd/Public/LevelEditorViewport.h b/Engine/Source/Editor/UnrealEd/Public/LevelEditorViewport.h
--- a/Engine/Source/Editor/UnrealEd/Public/LevelEditorViewport.h
+++ b/Engine/Source/Editor/UnrealEd/Public/LevelEditorViewport.h
@@ -138,9 +138,6 @@ public:
 	virtual void DrawCanvas( FViewport& InViewport, FSceneView& View, FCanvas& Canvas ) override;
 	virtual bool InputKey(FViewport* Viewport, int32 ControllerId, FKey Key, EInputEvent Event, float AmountDepressed = 1.f, bool bGamepad=false) override;
 	virtual bool InputAxis(FViewport* Viewport, int32 ControllerId, FKey Key, float Delta, float DeltaTime, int32 NumSamples=1, bool bGamepad=false) override;
-	virtual void MouseEnter( FViewport* Viewport,int32 x, int32 y ) override;
-	virtual void MouseLeave( FViewport* Viewport ) override;
-	virtual void MouseMove(FViewport* Viewport,int32 x, int32 y) override;
 	virtual EMouseCursor::Type GetCursor(FViewport* Viewport,int32 X,int32 Y) override;
 	virtual void CapturedMouseMove( FViewport* InViewport, int32 InMouseX, int32 InMouseY ) override;
 	virtual void Tick(float DeltaSeconds) override;
@@ -149,19 +146,13 @@ public:
 	virtual bool IsLevelEditorClient() const { return ParentLevelEditor.IsValid(); }
 	virtual void TrackingStarted( const struct FInputEventState& InInputState, bool bIsDraggingWidget, bool bNudge ) override;
 	virtual void TrackingStopped() override;
-	virtual void SetWidgetMode( FWidget::EWidgetMode NewMode ) override;
-	virtual bool CanSetWidgetMode( FWidget::EWidgetMode NewMode ) const override;
-	virtual void SetWidgetCoordSystemSpace( ECoordSystem NewCoordSystem ) override;
-	virtual FWidget::EWidgetMode GetWidgetMode() const override;
+	virtual void AbortTracking() override;
 	virtual FVector GetWidgetLocation() const override;
 	virtual FMatrix GetWidgetCoordSystem() const;
-	virtual ECoordSystem GetWidgetCoordSystemSpace() const;
 	virtual void SetupViewForRendering( FSceneViewFamily& ViewFamily, FSceneView& View ) override;
 	virtual FLinearColor GetBackgroundColor() const override;
 	virtual int32 GetCameraSpeedSetting() const override;
 	virtual void SetCameraSpeedSetting(int32 SpeedSetting) override;
-	virtual void ReceivedFocus(FViewport* Viewport) override;
-	virtual void LostFocus(FViewport* Viewport) override;
 
 	virtual bool OverrideHighResScreenshotCaptureRegion(FIntRect& OutCaptureRegion) override;
 
@@ -219,7 +210,7 @@ public:
 	 */
 	EAxisList::Type GetVertAxis() const;
 
-	void NudgeSelectedObjects( const struct FInputEventState& InputState );
+	virtual void NudgeSelectedObjects( const struct FInputEventState& InputState ) override;
 
 	/**
 	 * Moves the viewport camera according to the locked actors location and rotation
@@ -239,9 +230,6 @@ public:
 	void ApplyDeltaToActors( const FVector& InDrag, const FRotator& InRot, const FVector& InScale );
 	void ApplyDeltaToActor( AActor* InActor, const FVector& InDeltaDrag, const FRotator& InDeltaRot, const FVector& InDeltaScale );
 
-	/** Updates the rotate widget with the passed in delta rotation. */
-	void ApplyDeltaToRotateWidget( const FRotator& InRot );
-
 	virtual void SetIsSimulateInEditorViewport( bool bInIsSimulateInEditorViewport ) override;
 
 	/**
@@ -301,9 +289,7 @@ public:
 	void SetAllSpriteCategoryVisibility( bool bVisible );
 
 	/** FEditorViewportClient Interface*/
-	virtual void UpdateMouseDelta() override;
 	virtual void ProcessClick(FSceneView& View, HHitProxy* HitProxy, FKey Key, EInputEvent Event, uint32 HitX, uint32 HitY);
-	virtual void SetCurrentWidgetAxis( EAxisList::Type NewAxis ) override;
 	virtual UWorld* GetWorld() const override;
 
 	void SetReferenceToWorldContext(FWorldContext& WorldContext);
@@ -520,6 +506,7 @@ protected:
 	virtual bool ShouldLockPitch() const override;
 	virtual void CheckHoveredHitProxy( HHitProxy* HoveredHitProxy ) override;
 	virtual bool GetActiveSafeFrame(float& OutAspectRatio) const override;
+	virtual void RedrawAllViewportsIntoThisScene() override;
 
 private:
 	/**
@@ -531,7 +518,7 @@ private:
 	/**
 	 * Moves the locked actor according to the viewport cameras location and rotation
 	 */
-	void MoveLockedActorToCamera() const;
+	void MoveLockedActorToCamera();
 
 	
 	/** @return	Returns true if the delta tracker was used to modify any selected actors or BSP.  Must be called before EndTracking(). */
