Replaced bad IsRelative calls in automation tests for getting asset paths from package paths.

diff --git a/Engine/Source/Editor/UnrealEd/Private/Tests/AutomationEditorCommon.cpp b/Engine/Source/Editor/UnrealEd/Private/Tests/AutomationEditorCommon.cpp
--- a/Engine/Source/Editor/UnrealEd/Private/Tests/AutomationEditorCommon.cpp
+++ b/Engine/Source/Editor/UnrealEd/Private/Tests/AutomationEditorCommon.cpp
@@ -77,14 +77,14 @@ namespace AutomationEditorCommonUtils
 		const FString Filename = FPaths::ConvertRelativePathToFull(PackagePath);
 		FString EngineFileName = Filename;
 		FString GameFileName = Filename;
-		if (FPaths::MakePathRelativeTo(EngineFileName, *FPaths::EngineContentDir()) && !FPaths::IsRelative(EngineFileName))
+		if (FPaths::MakePathRelativeTo(EngineFileName, *FPaths::EngineContentDir()) && !EngineFileName.Contains(TEXT("../")))
 		{
 			const FString ShortName = FPaths::GetBaseFilename(EngineFileName);
 			const FString PathName = FPaths::GetPath(EngineFileName);
 			const FString AssetName = FString::Printf(TEXT("/Engine/%s/%s.%s"), *PathName, *ShortName, *ShortName);
 			return AssetName;
 		}
-		else if (FPaths::MakePathRelativeTo(GameFileName, *FPaths::GameContentDir()) && !FPaths::IsRelative(GameFileName))
+		else if (FPaths::MakePathRelativeTo(GameFileName, *FPaths::GameContentDir()) && !GameFileName.Contains(TEXT("../")))
 		{
 			const FString ShortName = FPaths::GetBaseFilename(GameFileName);
 			const FString PathName = FPaths::GetPath(GameFileName);
