A missing file in CL# 2380072

diff --git a/Engine/Source/Editor/UnrealEd/Private/ApexClothingUtils.cpp b/Engine/Source/Editor/UnrealEd/Private/ApexClothingUtils.cpp
--- a/Engine/Source/Editor/UnrealEd/Private/ApexClothingUtils.cpp
+++ b/Engine/Source/Editor/UnrealEd/Private/ApexClothingUtils.cpp
@@ -1098,9 +1098,9 @@ NxClothingAsset* ApplyTransform(NxClothingAsset* ApexClothingAsset)
 
 		ApexClothingAssetAuthoring->updateBindPoses(NewBindPoses.GetData(), NewBindPoses.Num(), true, true);
 
-		// destroy and create a new asset based off the authoring version
-	
-		for (uint32 i=0; i<256; i++)
+		int32 NumLODs = ApexClothingAssetAuthoring->getNumLods();
+		// destroy and create a new asset based off the authoring version	
+		for (int32 i = 0; i<NumLODs; i++)
 		{
 			NxParameterized::Interface *ApexRenderMeshAssetAuthoring = ApexClothingAssetAuthoring->getRenderMeshAssetAuthoring(i);
 			if ( ApexRenderMeshAssetAuthoring )
@@ -1113,7 +1113,6 @@ NxClothingAsset* ApplyTransform(NxClothingAsset* ApexClothingAsset)
 				{
 					uint32 TextureUVOrigin;
 					AuthorHandle.getParamU32(TextureUVOrigin);
-					TextureUVOrigin = NxTextureUVOrigin::ORIGIN_BOTTOM_LEFT;
 					switch ( TextureUVOrigin )
 					{
 					case NxTextureUVOrigin::ORIGIN_TOP_LEFT:
@@ -1209,7 +1208,10 @@ void RemoveAssetFromSkeletalMesh(USkeletalMesh* SkelMesh, uint32 AssetIndex, boo
 		}
 	}
 
-	SkelMesh->ClothingAssets[AssetIndex].ApexClothingAsset->SetValid(false);
+	// release and make it invalid
+	GPhysCommandHandler->DeferredRelease(SkelMesh->ClothingAssets[AssetIndex].ApexClothingAsset);
+	SkelMesh->ClothingAssets[AssetIndex].ApexClothingAsset = NULL;
+
 	//this requires to refresh UI layout
 	SkelMesh->ClothingAssets.RemoveAt(AssetIndex);
 
@@ -1363,7 +1365,7 @@ bool ImportClothingSectionFromClothingAsset( USkeletalMesh* SkelMesh, uint32 LOD
 
 	FClothingAssetData& AssetData = SkelMesh->ClothingAssets[AssetIndex];
 
-	NxClothingAsset* ApexClothingAsset = AssetData.ApexClothingAsset->GetAsset();
+	NxClothingAsset* ApexClothingAsset = AssetData.ApexClothingAsset;
 
 	// The APEX Clothing Asset contains an APEX Render Mesh Asset, get a pointer to this
 	physx::PxU32 NumLODLevels = ApexClothingAsset->getNumGraphicalLodLevels();
@@ -1527,13 +1529,18 @@ EClothUtilRetType ImportApexAssetFromApexFile(FString& ApexFile, USkeletalMesh*
 		Data->ApexFileName = ApexFile;
 		Data->AssetName = AssetName;
 		Data->bClothPropertiesChanged = false;
-		Data->ApexClothingAsset = MakeShareable(new FClothingAssetWrapper(ApexClothingAsset));
+		Data->ApexClothingAsset = ApexClothingAsset;
 	}
 	else
 	{
 		// re-import
 		FClothingAssetData& AssetData = SkelMesh->ClothingAssets[AssetIndex];
-		AssetData.ApexClothingAsset = MakeShareable(new FClothingAssetWrapper(ApexClothingAsset));
+		if (AssetData.ApexClothingAsset)
+		{
+			GPhysCommandHandler->DeferredRelease(AssetData.ApexClothingAsset);
+			AssetData.ApexClothingAsset = NULL;
+		}
+		AssetData.ApexClothingAsset = ApexClothingAsset;
 		AssetData.ApexFileName = ApexFile;
 
 		FSkeletalMeshResource* ImportedResource= SkelMesh->GetImportedResource();
