Set to Default option in Viewport now resets volume, sprite and layer visibliity flags.

diff --git a/Engine/Source/Editor/LevelEditor/Private/SLevelViewport.cpp b/Engine/Source/Editor/LevelEditor/Private/SLevelViewport.cpp
--- a/Engine/Source/Editor/LevelEditor/Private/SLevelViewport.cpp
+++ b/Engine/Source/Editor/LevelEditor/Private/SLevelViewport.cpp
@@ -1887,6 +1887,14 @@ void SLevelViewport::OnUseDefaultShowFlags(bool bUseSavedDefaults)
 	ApplyViewMode(CachedViewMode, LevelViewportClient->IsPerspective(), LevelViewportClient->EngineShowFlags);
 	ApplyViewMode(CachedViewMode, LevelViewportClient->IsPerspective(), LevelViewportClient->LastEngineShowFlags);
 
+	// set volume / layer / sprite visibility defaults
+	if (!bUseSavedDefaults)
+	{
+		LevelViewportClient->InitializeVisibilityFlags();
+		GUnrealEd->UpdateVolumeActorVisibility(NULL, LevelViewportClient.Get());
+		GEditor->Layers->UpdatePerViewVisibility(LevelViewportClient.Get());
+	}
+
 	LevelViewportClient->Invalidate();
 }
 
diff --git a/Engine/Source/Editor/UnrealEd/Private/LevelEditorViewport.cpp b/Engine/Source/Editor/UnrealEd/Private/LevelEditorViewport.cpp
--- a/Engine/Source/Editor/UnrealEd/Private/LevelEditorViewport.cpp
+++ b/Engine/Source/Editor/UnrealEd/Private/LevelEditorViewport.cpp
@@ -1466,16 +1466,7 @@ FLevelEditorViewportClient::FLevelEditorViewportClient()
 	// GEditorModeTools serves as our draw helper
 	bUsesDrawHelper = false;
 
-	// make sure all actors know about this view for per-view layer vis
-	GEditor->Layers->UpdatePerViewVisibility( this );
-
-	// Get the number of volume classes so we can initialize our bit array
-	TArray<UClass*> VolumeClasses;
-	UUnrealEdEngine::GetSortedVolumeClasses(&VolumeClasses);
-	VolumeActorVisibility.Init(true, VolumeClasses.Num() );
-
-	// Initialize all sprite categories to visible
-	SpriteCategoryVisibility.Init( true, GUnrealEd->SpriteIDToIndexMap.Num() );
+	InitializeVisibilityFlags();
 
 	// Sign up for notifications about users changing settings.
 	GetMutableDefault<ULevelEditorViewportSettings>()->OnSettingChanged().AddRaw(this, &FLevelEditorViewportClient::HandleViewportSettingChanged);
@@ -1513,6 +1504,20 @@ FLevelEditorViewportClient::~FLevelEditorViewportClient()
 	RemoveReferenceToWorldContext(GEditor->GetEditorWorldContext());
 }
 
+void FLevelEditorViewportClient::InitializeVisibilityFlags()
+{
+	// make sure all actors know about this view for per-view layer vis
+	GEditor->Layers->UpdatePerViewVisibility(this);
+
+	// Get the number of volume classes so we can initialize our bit array
+	TArray<UClass*> VolumeClasses;
+	UUnrealEdEngine::GetSortedVolumeClasses(&VolumeClasses);
+	VolumeActorVisibility.Init(true, VolumeClasses.Num());
+
+	// Initialize all sprite categories to visible
+	SpriteCategoryVisibility.Init(true, GUnrealEd->SpriteIDToIndexMap.Num());
+}
+
 FSceneView* FLevelEditorViewportClient::CalcSceneView(FSceneViewFamily* ViewFamily)
 {
 	bWasControlledByOtherViewport = false;
diff --git a/Engine/Source/Editor/UnrealEd/Public/LevelEditorViewport.h b/Engine/Source/Editor/UnrealEd/Public/LevelEditorViewport.h
--- a/Engine/Source/Editor/UnrealEd/Public/LevelEditorViewport.h
+++ b/Engine/Source/Editor/UnrealEd/Public/LevelEditorViewport.h
@@ -158,6 +158,11 @@ public:
 
 	void SetIsCameraCut( bool bInIsCameraCut ) { bEditorCameraCut = bInIsCameraCut; }
 
+	/** 
+	 * Initialize visibility flags
+	 */
+	void InitializeVisibilityFlags();
+
 	/**
 	 * Reset the camera position and rotation.  Used when creating a new level.
 	 */
