Reinstated some code that updated the build progress dialog during status updates.

diff --git a/Engine/Source/Editor/UnrealEd/Private/FeedbackContextEditor.cpp b/Engine/Source/Editor/UnrealEd/Private/FeedbackContextEditor.cpp
--- a/Engine/Source/Editor/UnrealEd/Private/FeedbackContextEditor.cpp
+++ b/Engine/Source/Editor/UnrealEd/Private/FeedbackContextEditor.cpp
@@ -432,9 +432,36 @@ void FFeedbackContextEditor::FinalizeSlowTask()
 
 void FFeedbackContextEditor::ProgressReported( const float TotalProgressInterp, FText DisplayMessage )
 {
-	if (SlowTaskWidget.IsValid() && FSlateApplication::Get().CanDisplayWindows())
+	// Clean up deferred cleanup objects from rendering thread every once in a while.
+	static double LastTimePendingCleanupObjectsWhereDeleted;
+	if( FPlatformTime::Seconds() - LastTimePendingCleanupObjectsWhereDeleted > 1 )
 	{
-		TickSlate();
+		// Get list of objects that are pending cleanup.
+		FPendingCleanupObjects* PendingCleanupObjects = GetPendingCleanupObjects();
+		// Flush rendering commands in the queue.
+		FlushRenderingCommands();
+		// It is now safe to delete the pending clean objects.
+		delete PendingCleanupObjects;
+		// Keep track of time this operation was performed so we don't do it too often.
+		LastTimePendingCleanupObjectsWhereDeleted = FPlatformTime::Seconds();
+	}
+
+	if (FSlateApplication::Get().CanDisplayWindows())
+	{
+		if (BuildProgressWidget.IsValid())
+		{
+			if (!DisplayMessage.IsEmpty())
+			{
+				BuildProgressWidget->SetBuildStatusText(DisplayMessage);
+			}
+
+			BuildProgressWidget->SetBuildProgressPercent(TotalProgressInterp * 100, 100);
+			TickSlate();
+		}
+		else if (SlowTaskWidget.IsValid())
+		{
+			TickSlate();
+		}
 	}
 	else
 	{
