Add user setting for using Ctrl key in landscape editor

diff --git a/Engine/Config/BaseEditorUserSettings.ini b/Engine/Config/BaseEditorUserSettings.ini
--- a/Engine/Config/BaseEditorUserSettings.ini
+++ b/Engine/Config/BaseEditorUserSettings.ini
@@ -239,6 +239,9 @@ LastExecutedLaunchPlatform=
 [/Script/UnrealEd.LevelEditorViewportSettings]
 ; Use WASD flight camera controls by default
 FlightCameraControlType=WASD_RMBOnly
+; Ignore Ctrl key by default in the landscape/foliage editors
+LandscapeEditorControlType=IgnoreCtrl
+FoliageEditorControlType=IgnoreCtrl
 ; If true, moves the canvas and shows the mouse.  If false, uses original camera movement
 bPanMovesCanvas=True
 ; If true, zooms centering on the mouse position.  If false, the zoom is around the center of the viewport
diff --git a/Engine/Source/Editor/FoliageEdit/Private/FoliageEdMode.cpp b/Engine/Source/Editor/FoliageEdit/Private/FoliageEdMode.cpp
--- a/Engine/Source/Editor/FoliageEdit/Private/FoliageEdMode.cpp
+++ b/Engine/Source/Editor/FoliageEdit/Private/FoliageEdMode.cpp
@@ -1824,11 +1824,20 @@ bool FEdModeFoliage::InputKey(FEditorViewportClient* ViewportClient, FViewport*
 {
 	if (UISettings.GetPaintToolSelected() || UISettings.GetReapplyToolSelected() || UISettings.GetLassoSelectToolSelected())
 	{
+		// Require Ctrl or not as per user preference
+		ELandscapeFoliageEditorControlType FoliageEditorControlType = GetDefault<ULevelEditorViewportSettings>()->FoliageEditorControlType;
+
 		if (Key == EKeys::LeftMouseButton && Event == IE_Pressed)
 		{
 			// Only activate tool if we're not already moving the camera and we're not trying to drag a transform widget
 			// Not using "if (!ViewportClient->IsMovingCamera())" because it's wrong in ortho viewports :D
-			if (!Viewport->KeyState(EKeys::MiddleMouseButton) && !Viewport->KeyState(EKeys::RightMouseButton) && !IsAltDown(Viewport) && ViewportClient->GetCurrentWidgetAxis() == EAxisList::None)
+			bool bMovingCamera = Viewport->KeyState(EKeys::MiddleMouseButton) || Viewport->KeyState(EKeys::RightMouseButton) || IsAltDown(Viewport);
+
+			if ((Viewport->IsPenActive() && Viewport->GetTabletPressure() > 0.f) ||
+				(!bMovingCamera && ViewportClient->GetCurrentWidgetAxis() == EAxisList::None &&
+					(FoliageEditorControlType == ELandscapeFoliageEditorControlType::IgnoreCtrl ||
+					 FoliageEditorControlType == ELandscapeFoliageEditorControlType::RequireCtrl   && IsCtrlDown(Viewport) ||
+					 FoliageEditorControlType == ELandscapeFoliageEditorControlType::RequireNoCtrl && !IsCtrlDown(Viewport))))
 			{
 				if (!bToolActive)
 				{
@@ -1863,7 +1872,8 @@ bool FEdModeFoliage::InputKey(FEditorViewportClient* ViewportClient, FViewport*
 			}
 		}
 
-		if (bToolActive && Key == EKeys::LeftMouseButton && Event == IE_Released)
+		if (bToolActive && Event == IE_Released &&
+			(Key == EKeys::LeftMouseButton || (FoliageEditorControlType == ELandscapeFoliageEditorControlType::RequireCtrl && (Key == EKeys::LeftControl || Key == EKeys::RightControl))))
 		{
 			//Set the cursor position to that of the slate cursor so it wont snap back
 			Viewport->SetPreCaptureMousePosFromSlateCursor();
diff --git a/Engine/Source/Editor/LandscapeEditor/Private/LandscapeEdMode.cpp b/Engine/Source/Editor/LandscapeEditor/Private/LandscapeEdMode.cpp
--- a/Engine/Source/Editor/LandscapeEditor/Private/LandscapeEdMode.cpp
+++ b/Engine/Source/Editor/LandscapeEditor/Private/LandscapeEdMode.cpp
@@ -1159,11 +1159,26 @@ bool FEdModeLandscape::InputKey(FEditorViewportClient* ViewportClient, FViewport
 			return true;
 		}
 
+		// Require Ctrl or not as per user preference
+		ELandscapeFoliageEditorControlType LandscapeEditorControlType = GetDefault<ULevelEditorViewportSettings>()->LandscapeEditorControlType;
+
+		// HACK - Splines tool has not yet been updated to support not using ctrl
+		if (CurrentBrush->GetBrushType() == FLandscapeBrush::BT_Splines)
+		{
+			LandscapeEditorControlType = ELandscapeFoliageEditorControlType::RequireCtrl;
+		}
+
 		if (Key == EKeys::LeftMouseButton && Event == IE_Pressed)
 		{
 			// Only activate tool if we're not already moving the camera and we're not trying to drag a transform widget
 			// Not using "if (!ViewportClient->IsMovingCamera())" because it's wrong in ortho viewports :D
-			if (!Viewport->KeyState(EKeys::MiddleMouseButton) && !Viewport->KeyState(EKeys::RightMouseButton) && !IsAltDown(Viewport) && ViewportClient->GetCurrentWidgetAxis() == EAxisList::None)
+			bool bMovingCamera = Viewport->KeyState(EKeys::MiddleMouseButton) || Viewport->KeyState(EKeys::RightMouseButton) || IsAltDown(Viewport);
+
+			if ((Viewport->IsPenActive() && Viewport->GetTabletPressure() > 0.f) ||
+				(!bMovingCamera && ViewportClient->GetCurrentWidgetAxis() == EAxisList::None &&
+					(LandscapeEditorControlType == ELandscapeFoliageEditorControlType::IgnoreCtrl ||
+					 LandscapeEditorControlType == ELandscapeFoliageEditorControlType::RequireCtrl   && IsCtrlDown(Viewport) ||
+					 LandscapeEditorControlType == ELandscapeFoliageEditorControlType::RequireNoCtrl && !IsCtrlDown(Viewport))))
 			{
 				if (CurrentTool && (CurrentTool->GetSupportedTargetTypes() == ELandscapeToolTargetTypeMask::NA || CurrentToolTarget.TargetType != ELandscapeToolTargetType::Invalid))
 				{
@@ -1191,7 +1206,8 @@ bool FEdModeLandscape::InputKey(FEditorViewportClient* ViewportClient, FViewport
 			}
 		}
 
-		if (Key == EKeys::LeftMouseButton)
+		if (Key == EKeys::LeftMouseButton ||
+			(LandscapeEditorControlType == ELandscapeFoliageEditorControlType::RequireCtrl && (Key == EKeys::LeftControl || Key == EKeys::RightControl)))
 		{
 			if (Event == IE_Released && CurrentTool && bToolActive)
 			{
diff --git a/Engine/Source/Editor/UnrealEd/Classes/Settings/LevelEditorViewportSettings.h b/Engine/Source/Editor/UnrealEd/Classes/Settings/LevelEditorViewportSettings.h
--- a/Engine/Source/Editor/UnrealEd/Classes/Settings/LevelEditorViewportSettings.h
+++ b/Engine/Source/Editor/UnrealEd/Classes/Settings/LevelEditorViewportSettings.h
@@ -28,12 +28,22 @@ enum ERotationGridMode
 UENUM()
 enum EWASDType
 {
-	WASD_Always UMETA(DisplayName="Use WASD for Camera Controls"),
+	WASD_Always  UMETA(DisplayName="Use WASD for Camera Controls"),
 	WASD_RMBOnly UMETA(DisplayName="Use WASD only when a Mouse Button is Pressed"),
-	WASD_Never UMETA(DisplayName="Never use WASD for Camera Controls"),
+	WASD_Never   UMETA(DisplayName="Never use WASD for Camera Controls"),
 	WASD_MAX,
 };
 
+/**
+ * Is Ctrl key required for editing landscape/foliage?
+ */
+UENUM()
+enum class ELandscapeFoliageEditorControlType : uint8
+{
+	IgnoreCtrl    UMETA(DisplayName = "Ignore Ctrl key (allow but don't require Ctrl held)"),
+	RequireCtrl   UMETA(DisplayName = "Require Ctrl held for tools"),
+	RequireNoCtrl UMETA(DisplayName = "Require Ctrl is not held"),
+};
 
 /**
  * Units used by measuring tool
@@ -42,8 +52,8 @@ UENUM()
 enum EMeasuringToolUnits
 {
 	MeasureUnits_Centimeters UMETA(DisplayName="Centimeters"),
-	MeasureUnits_Meters UMETA(DisplayName="Meters"),
-	MeasureUnits_Kilometers UMETA(DisplayName="Kilometers")
+	MeasureUnits_Meters      UMETA(DisplayName="Meters"),
+	MeasureUnits_Kilometers  UMETA(DisplayName="Kilometers")
 };
 
 
@@ -182,7 +192,15 @@ class UNREALED_API ULevelEditorViewportSettings
 
 	/** Enable the use of flight camera controls under various circumstances. */
 	UPROPERTY(EditAnywhere, config, Category=Controls)
-	TEnumAsByte<enum EWASDType> FlightCameraControlType;
+	TEnumAsByte<EWASDType> FlightCameraControlType;
+
+	/** Choose the control scheme for landscape tools (ignored for pen input) */
+	UPROPERTY(EditAnywhere, config, Category=Controls)
+	ELandscapeFoliageEditorControlType LandscapeEditorControlType;
+
+	/** Choose the control scheme for foliage tools */
+	UPROPERTY(EditAnywhere, config, Category=Controls)
+	ELandscapeFoliageEditorControlType FoliageEditorControlType;
 
 	/** If true, moves the canvas and shows the mouse.  If false, uses original camera movement. */
 	UPROPERTY(EditAnywhere, config, Category=Controls, meta=(DisplayName = "Grab and Drag to Move Orthographic Cameras"), AdvancedDisplay)
@@ -221,7 +239,7 @@ class UNREALED_API ULevelEditorViewportSettings
 	uint32 bUseAbsoluteTranslation:1;
 
 	/** If enabled, the viewport will stream in levels automatically when the camera is moved. */
-	UPROPERTY(EditAnywhere, config, Category=Controls, meta=(DisplayName = "Stream in Levels Automatically when Camera is Moved"), aDvancedDisplay)
+	UPROPERTY(EditAnywhere, config, Category=Controls, meta=(DisplayName = "Stream in Levels Automatically when Camera is Moved"), AdvancedDisplay)
 	bool bLevelStreamingVolumePrevis;
 
 	/** When checked, orbit the camera by using the L or U keys when unchecked, Alt and Left Mouse Drag will orbit around the look at point */
