Fixed bug so that materials and textures are found when reimporting a static mesh.

diff --git a/Engine/Source/Editor/UnrealEd/Private/Fbx/FbxMaterialImport.cpp b/Engine/Source/Editor/UnrealEd/Private/Fbx/FbxMaterialImport.cpp
--- a/Engine/Source/Editor/UnrealEd/Private/Fbx/FbxMaterialImport.cpp
+++ b/Engine/Source/Editor/UnrealEd/Private/Fbx/FbxMaterialImport.cpp
@@ -358,7 +358,7 @@ void UnFbx::FFbxImporter::FixupMaterial( FbxSurfaceMaterial& FbxMaterial, UMater
 //
 //-------------------------------------------------------------------------
 
-void UnFbx::FFbxImporter::CreateUnrealMaterial(FbxSurfaceMaterial& FbxMaterial, TArray<UMaterialInterface*>& OutMaterials, TArray<FString>& UVSets)
+FString UnFbx::FFbxImporter::GetMaterialFullName(FbxSurfaceMaterial& FbxMaterial)
 {
 	FString MaterialFullName = ANSI_TO_TCHAR(MakeName(FbxMaterial.GetName()));
 
@@ -373,19 +373,25 @@ void UnFbx::FFbxImporter::CreateUnrealMaterial(FbxSurfaceMaterial& FbxMaterial,
 			if (SkinXXNumber.IsNumeric())
 			{
 				// remove the '_skinXX' suffix from the material name					
-				MaterialFullName = MaterialFullName.LeftChop(Offset+1);
+				MaterialFullName = MaterialFullName.LeftChop(Offset + 1);
 			}
 		}
 	}
 
 	MaterialFullName = ObjectTools::SanitizeObjectName(MaterialFullName);
 
+	return MaterialFullName;
+}
+
+void UnFbx::FFbxImporter::CreateUnrealMaterial(FbxSurfaceMaterial& FbxMaterial, TArray<UMaterialInterface*>& OutMaterials, TArray<FString>& UVSets)
+{
 	// Make sure we have a parent
 	if ( !ensure(Parent) )
 	{
 		return;
 	}
 
+	FString MaterialFullName = GetMaterialFullName(FbxMaterial);
 	FString BasePackageName = FPackageName::GetLongPackagePath(Parent->GetOutermost()->GetName()) / MaterialFullName;
 	BasePackageName = PackageTools::SanitizePackageName(BasePackageName);
 
diff --git a/Engine/Source/Editor/UnrealEd/Private/Fbx/FbxStaticMeshImport.cpp b/Engine/Source/Editor/UnrealEd/Private/Fbx/FbxStaticMeshImport.cpp
--- a/Engine/Source/Editor/UnrealEd/Private/Fbx/FbxStaticMeshImport.cpp
+++ b/Engine/Source/Editor/UnrealEd/Private/Fbx/FbxStaticMeshImport.cpp
@@ -204,7 +204,6 @@ bool UnFbx::FFbxImporter::BuildStaticMeshFromGeometry(FbxMesh* Mesh, UStaticMesh
 	{
 		FFbxMaterial* NewMaterial = new(MeshMaterials) FFbxMaterial;
 		FbxSurfaceMaterial *FbxMaterial = Node->GetMaterial(MaterialIndex);
-		FString MaterialName = ANSI_TO_TCHAR(MakeName(FbxMaterial->GetName()));
 		NewMaterial->FbxMaterial = FbxMaterial;
 		if (ImportOptions->bImportMaterials)
 		{
@@ -212,7 +211,9 @@ bool UnFbx::FFbxImporter::BuildStaticMeshFromGeometry(FbxMesh* Mesh, UStaticMesh
 		}
 		else
 		{
-			UMaterialInterface* UnrealMaterialInterface = FindObject<UMaterialInterface>(NULL,*MaterialName);
+			FString MaterialFullName = GetMaterialFullName(*FbxMaterial);
+			FString BasePackageName = PackageTools::SanitizePackageName(FPackageName::GetLongPackagePath(StaticMesh->GetOutermost()->GetName()) / MaterialFullName);
+			UMaterialInterface* UnrealMaterialInterface = FindObject<UMaterialInterface>(NULL, *(BasePackageName + TEXT(".") + MaterialFullName));
 			if (UnrealMaterialInterface == NULL)
 			{
 				UnrealMaterialInterface = UMaterial::GetDefaultMaterial(MD_Surface);
diff --git a/Engine/Source/Editor/UnrealEd/Private/FbxImporter.h b/Engine/Source/Editor/UnrealEd/Private/FbxImporter.h
--- a/Engine/Source/Editor/UnrealEd/Private/FbxImporter.h
+++ b/Engine/Source/Editor/UnrealEd/Private/FbxImporter.h
@@ -851,15 +851,23 @@ protected:
 	 * @return int32 material count that created from the Fbx node
 	 */
 	int32 CreateNodeMaterials(FbxNode* FbxNode, TArray<UMaterialInterface*>& outMaterials, TArray<FString>& UVSets);
-	
+
 	/**
-	* Create Unreal material from Fbx material.
-	* Only setup channels that connect to texture, and setup the UV coordinate of texture.
-	* If diffuse channel has no texture, one default node will be created with constant.
-	*
-	* @param KFbxSurfaceMaterial*  Fbx material
-	* @param outMaterials Unreal Materials we created
-	* @param outUVSets
+	 * Make material Unreal asset name from the Fbx material
+	 *
+	 * @param FbxMaterial Material from the Fbx node
+	 * @return Sanitized asset name
+	 */
+	FString GetMaterialFullName(FbxSurfaceMaterial& FbxMaterial);
+
+	/**
+	 * Create Unreal material from Fbx material.
+	 * Only setup channels that connect to texture, and setup the UV coordinate of texture.
+	 * If diffuse channel has no texture, one default node will be created with constant.
+	 *
+	 * @param KFbxSurfaceMaterial*  Fbx material
+	 * @param outMaterials Unreal Materials we created
+	 * @param outUVSets
 	 */
 	void CreateUnrealMaterial(FbxSurfaceMaterial& FbxMaterial, TArray<UMaterialInterface*>& OutMaterials, TArray<FString>& UVSets);
 	
