Added support for creating releases which can be used to base dlc off.
Added support for basing dlc cooks of releases.
Stopped auto save while cooking in editor.
Removed unneeded package file stamp checks from SaveCookedPackage.
Rewrote Ini version settings checking so that it works from cook in editor.
Fixed up cleaning sandboxes so that it runs each time you start cook by the book (a little slower but more correct).

diff --git a/Engine/Source/Editor/UnrealEd/Classes/CookOnTheSide/CookOnTheFlyServer.h b/Engine/Source/Editor/UnrealEd/Classes/CookOnTheSide/CookOnTheFlyServer.h
--- a/Engine/Source/Editor/UnrealEd/Classes/CookOnTheSide/CookOnTheFlyServer.h
+++ b/Engine/Source/Editor/UnrealEd/Classes/CookOnTheSide/CookOnTheFlyServer.h
@@ -349,7 +349,7 @@ private:
 			return true;
 		}
 
-		void RemoveAllCookedFilesForPlatform( const FName& PlatformName )
+		void RemoveAllFilesForPlatform( const FName& PlatformName )
 		{
 			FScopeLock ScopeLock(&SynchronizationObject);
 
@@ -371,11 +371,25 @@ private:
 			return false;
 			
 		}
-		int RemoveAll( const FName& Filename )
+		int RemoveFile( const FName& Filename )
 		{
 			FScopeLock ScopeLock( &SynchronizationObject );
 			return FilesProcessed.Remove( Filename );
 		}
+
+		bool RemoveFileForPlatform( const FName& Filename, const FName& PlatformName )
+		{
+			FScopeLock ScopeLock( &SynchronizationObject );
+			FFilePlatformRequest* ProcessedFile = FilesProcessed.Find(Filename);
+			if( ProcessedFile )
+			{
+				ProcessedFile->RemovePlatform(PlatformName);
+				return true; 
+			}
+			return false;
+		}
+
+
 		void Empty(int32 ExpectedNumElements = 0)
 		{
 			FScopeLock ScopeLock( &SynchronizationObject );
@@ -554,7 +568,8 @@ private:
 private:
 	/** Current cook mode the cook on the fly server is running in */
 	ECookMode::Type CurrentCookMode;
-	
+	/** Directory to output to instead of the default should be empty in the case of DLC cooking */ 
+	FString OutputDirectoryOverride;
 	//////////////////////////////////////////////////////////////////////////
 	// Cook by the book options
 	struct FCookByTheBookOptions
@@ -576,6 +591,8 @@ private:
 		bool bCancel;
 		/** DlcName setup if we are cooking dlc will be used as the directory to save cooked files to */
 		FString DlcName;
+		/** Create a release from this manifest and store it in the releases directory for this cgame */
+		FString CreateReleaseVersion;
 		/** Leak test: last gc items (only valid when running from commandlet requires gc between each cooked package) */
 		TSet<FWeakObjectPtr> LastGCItems;
 		/** Map of platform name to manifest generator, manifest is only used in cook by the book however it needs to be maintained across multiple cook by the books. */
@@ -605,24 +622,28 @@ private:
 	FThreadSafeUnsolicitedPackagesList UnsolicitedCookedPackages;
 	FThreadSafeFilenameSet CookedPackages; // set of files which have been cooked when needing to recook a file the entry will need to be removed from here
 
-	TMap<FName, FCachedPackageFilename> PackageFilenameCache; // filename cache (only process the string operations once)
+	FString GetCachedPackageFilename( const FName& PackageName ) const;
+	FString GetCachedStandardPackageFilename( const FName& PackageName ) const;
+	FName GetCachedStandardPackageFileFName( const FName& PackageName ) const;
+	FString GetCachedPackageFilename( UPackage* Package ) const;
+	FString GetCachedStandardPackageFilename( UPackage* Package ) const;
+	FName GetCachedStandardPackageFileFName( UPackage* Package ) const;
+	const FString& GetCachedSandboxFilename( UPackage* Package, TAutoPtr<class FSandboxPlatformFile>& SandboxFile ) const;
+	const FCachedPackageFilename& Cache(const FName& PackageName) const;
+	void ClearPackageFilenameCache() const;
+
+	// declared mutable as it's used to cache package filename strings and I don't want to declare all functions using it as non const
+	// used by GetCached * Filename functions
+	mutable TMap<FName, FCachedPackageFilename> PackageFilenameCache; // filename cache (only process the string operations once)
 
-	FString GetCachedPackageFilename( const FName& PackageName );
-	FString GetCachedStandardPackageFilename( const FName& PackageName );
-	FName GetCachedStandardPackageFileFName( const FName& PackageName );
-	FString GetCachedPackageFilename( UPackage* Package );
-	FString GetCachedStandardPackageFilename( UPackage* Package );
-	FName GetCachedStandardPackageFileFName( UPackage* Package );
-	const FString& GetCachedSandboxFilename( UPackage* Package, TAutoPtr<class FSandboxPlatformFile>& SandboxFile );
-	const FCachedPackageFilename& Cache(const FName& PackageName);
-	void ClearPackageFilenameCache();
 
 	// declared mutable as it's used purely as a cache and don't want to have to declare all the functions as non const just because of this cache
+	// used by IniSettingsOutOfDate and GetCurrentIniStrings 
 	mutable TMap<FName, TArray<FString>> CachedIniVersionStringsMap;
 
 public:
 
-	void WarmCookedPackages(const FString& AssetRegistryPath, const TArray<FName>& TargetPlatformNames);
+	// void WarmCookedPackages(const FString& AssetRegistryPath, const TArray<FName>& TargetPlatformNames);
 
 	
 
@@ -669,14 +690,30 @@ public:
 	void EndNetworkFileServer();
 
 
+	struct FCookByTheBookStartupOptions
+	{
+	public:
+		TArray<ITargetPlatform*> TargetPlatforms;
+		TArray<FString> CookMaps;
+		TArray<FString> CookDirectories;
+		TArray<FString> CookCultures; 
+		TArray<FString> IniMapSections;
+		ECookByTheBookOptions CookOptions;
+		FString DLCName;
+		FString CreateReleaseVersion;
+		FString BasedOnReleaseVersion;
+
+		FCookByTheBookStartupOptions() :
+			CookOptions(ECookByTheBookOptions::None),
+			DLCName(FString())
+		{ }
+	};
+
 	/**
 	 * Start a cook by the book session
 	 * Cook on the fly can't run at the same time as cook by the book
 	 */
-	void StartCookByTheBook(const TArray<ITargetPlatform*>& TargetPlatforms, 
-		const TArray<FString>& CookMaps, const TArray<FString>& CookDirectories, 
-		const TArray<FString>& CookCultures, const TArray<FString>& IniMapSections, 
-		ECookByTheBookOptions CookOptions = ECookByTheBookOptions::None, const FString &DLCName = FString() );
+	void StartCookByTheBook( const FCookByTheBookStartupOptions& CookByTheBookStartupOptions );
 
 	/**
 	 * Queue a cook by the book cancel (you might want to do this instead of calling cancel directly so that you don't have to be in the game thread when canceling
@@ -796,6 +833,15 @@ private:
 	 */
 	void CookByTheBookFinished();
 
+	/**
+	 * Get all the packages which are listed in asset registry passed in.  
+	 *
+	 * @param AssetRegistryPath path of the assetregistry.bin file to read
+	 * @param OutPackageNames out list of packages which were contained in the asset registry file
+	 * @return true if successfully read false otherwise
+	 */
+	bool GetAllPackagesFromAssetRegistry( const FString& AssetRegistryPath, TArray<FName>& OutPackageNames ) const;
+
 
 	//////////////////////////////////////////////////////////////////////////
 	// cook on the fly specific functions
@@ -863,9 +909,17 @@ private:
 
 	/**
 	 * Convert a path to a full sandbox path
-	 * This function should be used instead of calling the 
+	 * is effected by the cooking dlc settings
+	 * This function should be used instead of calling the FSandbox Sandbox->ConvertToSandboxPath functions
 	 */
 	FString ConvertToFullSandboxPath( const FString &FileName, bool bForWrite = false ) const;
+	FString ConvertToFullSandboxPath( const FString &FileName, bool bForWrite, const FString& PlatformName ) const;
+	/**
+	 * Get the sandbox root directory for that platform
+	 * is effected by the CookingDlc settings
+	 * This should be used instead of calling the Sandbox function
+	 */
+	FString GetSandboxDirectory( const FString& PlatformName ) const;
 
 	inline bool IsCookingDLC() const
 	{
@@ -878,13 +932,37 @@ private:
 		return false;
 	}
 
+	inline bool IsCreatingReleaseVersion()
+	{
+		if ( CookByTheBookOptions )
+		{
+			return !CookByTheBookOptions->CreateReleaseVersion.IsEmpty();
+		}
+		return false;
+	}
+
+	/**
+	 * Loads the cooked ini version settings maps into the Ini settings cache
+	 * 
+	 * @param TargetPlatforms to look for ini settings for
+	 */
+	bool CacheIniVersionStringsMap( const ITargetPlatform* TargetPlatform ) const;
+
 	/**
 	 * Checks if important ini settings have changed since last cook for each target platform 
 	 * 
 	 * @param TargetPlatforms to check if out of date
 	 * @param OutOfDateTargetPlatforms return list of out of date target platforms which should be cleaned
 	 */
-	bool IniSettingsOutOfDate( const TArray<ITargetPlatform*>& TargetPlatforms, TArray<ITargetPlatform*>& OutOfDateTargetPlatforms, bool bShouldSaveCookedVersions = true ) const;
+	bool IniSettingsOutOfDate( const ITargetPlatform* TargetPlatform ) const;
+
+	/**
+	 * Saves ini settings which are in the memory cache to the hard drive in ini files
+	 *
+	 * @param TargetPlatforms to save
+	 */
+	bool SaveCurrentIniSettings( const ITargetPlatform* TargetPlatform ) const;
+
 
 	/**
 	 * IsCookFlagSet
@@ -897,20 +975,6 @@ private:
 		return (CookFlags & InCookFlags) != ECookInitializationFlags::None;
 	}
 
-	/**
-	 * Returns cooker output directory.
-	 *
-	 * @param PlatformName Target platform name.
-	 * @return Cooker output directory for the specified platform.
-	 */
-	FString GetOutputDirectory( const FString& PlatformName ) const;
-
-
-	/**
-	 * Gets the asset registry path and name
-	 */
-	FString GetAssetRegistryPath() const;
-
 	/**
 	 *	Get the given packages 'cooked' timestamp (i.e. account for dependencies)
 	 *
@@ -953,8 +1017,10 @@ private:
 	void SaveGlobalShaderMapFiles(const TArray<ITargetPlatform*>& Platforms);
 
 
+	/** Create sandbox file in directory using current settings supplied */
+	void CreateSandboxFile();
 	/** Gets the output directory respecting any command line overrides */
-	FString GetOutputDirectoryOverride( const FString& OutputDirectoryOverride ) const;
+	FString GetOutputDirectoryOverride() const;
 
 	/** Cleans sandbox folders for all target platforms */
 	void CleanSandbox(const TArray<ITargetPlatform*>& Platforms);
diff --git a/Engine/Source/Editor/UnrealEd/Private/Commandlets/CookCommandlet.cpp b/Engine/Source/Editor/UnrealEd/Private/Commandlets/CookCommandlet.cpp
--- a/Engine/Source/Editor/UnrealEd/Private/Commandlets/CookCommandlet.cpp
+++ b/Engine/Source/Editor/UnrealEd/Private/Commandlets/CookCommandlet.cpp
@@ -1038,7 +1038,13 @@ bool UCookCommandlet::NewCook( const TArray<ITargetPlatform*>& Platforms, TArray
 	// parse commandline options 
 
 	FString DLCName;
-	FParse::Value( *Params, TEXT("DLCNAME"), DLCName);
+	FParse::Value( *Params, TEXT("DLCNAME="), DLCName);
+
+	FString BasedOnReleaseVersion;
+	FParse::Value( *Params, TEXT("BasedOnReleaseVersion="), BasedOnReleaseVersion);
+
+	FString CreateReleaseVersion;
+	FParse::Value( *Params, TEXT("CreateReleaseVersion="), CreateReleaseVersion);
 
 	/*FString AssetRegistry;
 	if (FParse::Value(*Params, TEXT("SHIPPEDASSETREGISTRY="), AssetRegistry))
@@ -1139,7 +1145,18 @@ bool UCookCommandlet::NewCook( const TArray<ITargetPlatform*>& Platforms, TArray
 		MapList.Add( MapName );
 	}
 
-	CookOnTheFlyServer->StartCookByTheBook(Platforms, MapList, CmdLineDirEntries, CmdLineCultEntries, CmdLineIniSections, CookOptions, DLCName );
+	UCookOnTheFlyServer::FCookByTheBookStartupOptions StartupOptions;
+
+	StartupOptions.TargetPlatforms = Platforms;
+	Swap( StartupOptions.CookMaps, MapList );
+	Swap( StartupOptions.CookDirectories, CmdLineDirEntries );
+	Swap( StartupOptions.CookCultures, CmdLineCultEntries );
+	Swap( StartupOptions.DLCName, DLCName );
+	Swap( StartupOptions.BasedOnReleaseVersion, BasedOnReleaseVersion );
+	Swap( StartupOptions.CreateReleaseVersion, CreateReleaseVersion );
+	StartupOptions.CookOptions = CookOptions;
+
+	CookOnTheFlyServer->StartCookByTheBook( StartupOptions );
 
 	// Garbage collection should happen when either
 	//	1. We have cooked a map
diff --git a/Engine/Source/Editor/UnrealEd/Private/CookOnTheFlyServer.cpp b/Engine/Source/Editor/UnrealEd/Private/CookOnTheFlyServer.cpp
--- a/Engine/Source/Editor/UnrealEd/Private/CookOnTheFlyServer.cpp
+++ b/Engine/Source/Editor/UnrealEd/Private/CookOnTheFlyServer.cpp
@@ -260,6 +260,30 @@ struct FRecompileRequest
 	bool bComplete;
 };
 
+
+
+const FString& GetAssetRegistryPath()
+{
+	static const FString AssetRegistryPath = FPaths::GameDir();
+	return AssetRegistryPath;
+}
+
+/**
+ * Return the release asset registry filename for the release version supplied
+ */
+FString GetReleaseVersionAssetRegistryPath(const FString& ReleaseVersion, const FName& PlatformName )
+{
+	// cache the part of the path which is static because getting the GameDir is really slow and also string manipulation
+	const static FString GameDirectory = FPaths::GameDir() / FString(TEXT("Releases"));
+	return  GameDirectory / ReleaseVersion / PlatformName.ToString();
+}
+
+const FString& GetAssetRegistryFilename()
+{
+	static const FString AssetRegistryFilename = FString(TEXT("AssetRegistry.bin"));
+	return AssetRegistryFilename;
+}
+
 /**
  * Uses the FMessageLog to log a message
  * 
@@ -285,41 +309,41 @@ void LogCookerMessage( const FString& MessageText, EMessageSeverity::Type Severi
 /* FIlename caching functions
  *****************************************************************************/
 
-FString UCookOnTheFlyServer::GetCachedPackageFilename( const FName& PackageName )
+FString UCookOnTheFlyServer::GetCachedPackageFilename( const FName& PackageName ) const 
 {
 	return Cache( PackageName ).PackageFilename;
 }
 
-FString UCookOnTheFlyServer::GetCachedStandardPackageFilename( const FName& PackageName )
+FString UCookOnTheFlyServer::GetCachedStandardPackageFilename( const FName& PackageName ) const 
 {
 	return Cache( PackageName ).StandardFilename;
 }
 
-FName UCookOnTheFlyServer::GetCachedStandardPackageFileFName( const FName& PackageName )
+FName UCookOnTheFlyServer::GetCachedStandardPackageFileFName( const FName& PackageName ) const 
 {
 	return Cache( PackageName ).StandardFileFName;
 }
 
 
-FString UCookOnTheFlyServer::GetCachedPackageFilename( UPackage* Package )
+FString UCookOnTheFlyServer::GetCachedPackageFilename( UPackage* Package ) const 
 {
 	check( Package->GetName() == Package->GetFName().ToString() );
 	return Cache( Package->GetFName() ).PackageFilename;
 }
 
-FString UCookOnTheFlyServer::GetCachedStandardPackageFilename( UPackage* Package )
+FString UCookOnTheFlyServer::GetCachedStandardPackageFilename( UPackage* Package ) const 
 {
 	check( Package->GetName() == Package->GetFName().ToString() );
 	return Cache( Package->GetFName() ).StandardFilename;
 }
 
-FName UCookOnTheFlyServer::GetCachedStandardPackageFileFName( UPackage* Package )
+FName UCookOnTheFlyServer::GetCachedStandardPackageFileFName( UPackage* Package ) const 
 {
 	check( Package->GetName() == Package->GetFName().ToString() );
 	return Cache( Package->GetFName() ).StandardFileFName;
 }
 
-const FString& UCookOnTheFlyServer::GetCachedSandboxFilename( UPackage* Package, TAutoPtr<class FSandboxPlatformFile>& SandboxFile )
+const FString& UCookOnTheFlyServer::GetCachedSandboxFilename( UPackage* Package, TAutoPtr<class FSandboxPlatformFile>& SandboxFile ) const 
 {
 	FName PackageFName = Package->GetFName();
 	static TMap<FName, FString> CachedSandboxFilenames;
@@ -333,7 +357,7 @@ const FString& UCookOnTheFlyServer::GetCachedSandboxFilename( UPackage* Package,
 	return CachedSandboxFilenames.Add( PackageFName, MoveTemp(SandboxFilename) );
 }
 
-const UCookOnTheFlyServer::FCachedPackageFilename& UCookOnTheFlyServer::Cache(const FName& PackageName)
+const UCookOnTheFlyServer::FCachedPackageFilename& UCookOnTheFlyServer::Cache(const FName& PackageName) const 
 {
 	FCachedPackageFilename *Cached = PackageFilenameCache.Find( PackageName );
 	if ( Cached != NULL )
@@ -358,9 +382,11 @@ const UCookOnTheFlyServer::FCachedPackageFilename& UCookOnTheFlyServer::Cache(co
 	return PackageFilenameCache.Emplace( PackageName, MoveTemp(FCachedPackageFilename(MoveTemp(PackageFilename),MoveTemp(StandardFilename), StandardFileFName)) );
 }
 
-void UCookOnTheFlyServer::ClearPackageFilenameCache()
-	{
+void UCookOnTheFlyServer::ClearPackageFilenameCache() const 
+{
 	PackageFilenameCache.Empty();
+	// need to clear the IniVersionStringsMap too
+	CachedIniVersionStringsMap.Empty();
 }
 
 
@@ -408,6 +434,34 @@ bool UCookOnTheFlyServer::StartNetworkFileServer( const bool BindAnyPort )
 	check( IsCookOnTheFlyMode() );
 	//GetDerivedDataCacheRef().WaitForQuiescence(false);
 
+	CreateSandboxFile();
+
+	ITargetPlatformManagerModule& TPM = GetTargetPlatformManagerRef();
+	const TArray<ITargetPlatform*>& Platforms = TPM.GetCookingTargetPlatforms();
+
+	{
+		// When cooking on the fly the full registry is saved at the beginning
+		// in cook by the book asset registry is saved after the cook is finished
+		FAssetRegistryModule& AssetRegistryModule = FModuleManager::LoadModuleChecked<FAssetRegistryModule>(TEXT("AssetRegistry"));
+		IAssetRegistry& AssetRegistry = AssetRegistryModule.Get();
+
+		// write it out to a memory archive
+		FArrayWriter SerializedAssetRegistry;
+		AssetRegistry.Serialize(SerializedAssetRegistry);
+		UE_LOG(LogCook, Display, TEXT("Generated asset registry size is %5.2fkb"), (float)SerializedAssetRegistry.Num() / 1024.f);
+
+		// now save it in each cooked directory
+		FString RegistryFilename = FPaths::GameDir() / TEXT("AssetRegistry.bin");
+		// Use SandboxFile to do path conversion to properly handle sandbox paths (outside of standard paths in particular).
+		FString SandboxFilename = ConvertToFullSandboxPath(*RegistryFilename, true);
+
+		for (int32 Index = 0; Index < Platforms.Num(); Index++)
+		{
+			FString PlatFilename = SandboxFilename.Replace(TEXT("[Platform]"), *Platforms[Index]->PlatformName());
+			FFileHelper::SaveArrayToFile(SerializedAssetRegistry, *PlatFilename);
+		}
+	}
+
 	// start the listening thread
 	FFileRequestDelegate FileRequestDelegate(FFileRequestDelegate::CreateUObject(this, &UCookOnTheFlyServer::HandleNetworkFileServerFileRequest));
 	FRecompileShadersDelegate RecompileShadersDelegate(FRecompileShadersDelegate::CreateUObject(this, &UCookOnTheFlyServer::HandleNetworkFileServerRecompileShaders));
@@ -425,6 +479,10 @@ bool UCookOnTheFlyServer::StartNetworkFileServer( const bool BindAnyPort )
 	{
 		NetworkFileServers.Add( HttpFileServer );
 	}
+
+	// clean sandbox when we start the network file server
+	CleanSandbox(Platforms);
+
 	// loop while waiting for requests
 	GIsRequestingExit = false;
 	return true;
@@ -834,6 +892,14 @@ uint32 UCookOnTheFlyServer::TickCookOnTheSide( const float TimeSlice, uint32 &Co
 			break;
 		}
 
+		// prevent autosave from happening until we are finished cooking
+		// causes really bad hitches
+		if ( GUnrealEd && (&GUnrealEd->GetPackageAutoSaver()) )
+		{
+			const static float SecondsWarningTillAutosave = 10.0f;
+			GUnrealEd->GetPackageAutoSaver().ForceMinimumTimeTillAutoSave(SecondsWarningTillAutosave);
+		}
+
 		if (CookedPackages.Exists(ToBuild))
 		{
 #if DEBUG_COOKONTHEFLY
@@ -1359,7 +1425,7 @@ void UCookOnTheFlyServer::MarkPackageDirtyForCooker( UPackage *Package )
 			}
 		}
 		}
-		CookedPackages.RemoveAll( PackageFFileName );
+		CookedPackages.RemoveFile( PackageFFileName );
 	}
 }
 
@@ -1404,14 +1470,6 @@ void UCookOnTheFlyServer::TickRecompileShaderRequests()
 }
 
 
-FString UCookOnTheFlyServer::GetOutputDirectory( const FString& PlatformName ) const
-{
-	// Use SandboxFile to get the correct sandbox directory.
-	FString OutputDirectory = SandboxFile->GetSandboxDirectory();
-
-	return OutputDirectory.Replace(TEXT("[Platform]"), *PlatformName);
-}
-
 
 bool UCookOnTheFlyServer::GetPackageTimestamp( const FString& InFilename, FDateTime& OutDateTime )
 {
@@ -1435,7 +1493,6 @@ bool UCookOnTheFlyServer::SaveCookedPackage( UPackage* Package, uint32 SaveFlags
 	return SaveCookedPackage( Package, SaveFlags, bOutWasUpToDate, TargetPlatformNames );
 }
 
-
 bool UCookOnTheFlyServer::ShouldCook(const FString& InFileName, const FName &InPlatformName)
 {
 	bool bDoCook = false;
@@ -1502,13 +1559,15 @@ bool UCookOnTheFlyServer::SaveCookedPackage( UPackage* Package, uint32 SaveFlags
 
 	if (Filename.Len())
 	{
-		FString PkgFilename;
+		FString Name = Package->GetPathName();
+
+		/*FString PkgFilename;
 		FDateTime DependentTimeStamp = FDateTime::MinValue();
 
 		// We always want to use the dependent time stamp when saving a cooked package...
 		// Iterative or not!
 		FString PkgFile;
-		FString Name = Package->GetPathName();
+		
 
 		if (IsCookFlagSet(ECookInitializationFlags::Iterative) && FPackageName::DoesPackageExist(Name, NULL, &PkgFile))
 		{
@@ -1516,11 +1575,9 @@ bool UCookOnTheFlyServer::SaveCookedPackage( UPackage* Package, uint32 SaveFlags
 
 			if (GetPackageTimestamp(FPaths::GetBaseFilename(PkgFilename, false), DependentTimeStamp) == false)
 			{
-#if DEBUG_COOKONTHEFLY
 				UE_LOG(LogCook, Display, TEXT("Failed to find depedency timestamp for: %s"), *PkgFilename);
-#endif
 			}
-		}
+		}*/
 
 		// Use SandboxFile to do path conversion to properly handle sandbox paths (outside of standard paths in particular).
 		Filename = ConvertToFullSandboxPath(*Filename, true);
@@ -1573,8 +1630,11 @@ bool UCookOnTheFlyServer::SaveCookedPackage( UPackage* Package, uint32 SaveFlags
 		{
 			FString PlatFilename = Filename.Replace(TEXT("[Platform]"), *Target->PlatformName());
 
+			bool bCookPackage = true;
+
+
 			// If we are not iterative cooking, then cook the package
-			bool bCookPackage = (IsCookFlagSet(ECookInitializationFlags::Iterative) == false);
+			/*bool bCookPackage = (IsCookFlagSet(ECookInitializationFlags::Iterative) == false);
 
 			if (bCookPackage == false)
 			{
@@ -1582,7 +1642,7 @@ bool UCookOnTheFlyServer::SaveCookedPackage( UPackage* Package, uint32 SaveFlags
 				FDateTime CookedTimeStamp = IFileManager::Get().GetTimeStamp(*PlatFilename);
 				int32 CookedTimespanSeconds = (CookedTimeStamp - DependentTimeStamp).GetTotalSeconds();
 				bCookPackage = (CookedTimeStamp == FDateTime::MinValue()) || (CookedTimespanSeconds < 0);
-			}
+			}*/
 
 			// don't save Editor resources from the Engine if the target doesn't have editoronly data
 			if (IsCookFlagSet(ECookInitializationFlags::SkipEditorContent) && Name.StartsWith(TEXT("/Engine/Editor")) && !Target->HasEditorOnlyData())
@@ -1689,9 +1749,9 @@ bool UCookOnTheFlyServer::SaveCookedPackage( UPackage* Package, uint32 SaveFlags
 	return bSavedCorrectly;
 }
 
-void UCookOnTheFlyServer::Initialize( ECookMode::Type DesiredCookMode, ECookInitializationFlags InCookFlags, const FString &OutputDirectoryOverride )
+void UCookOnTheFlyServer::Initialize( ECookMode::Type DesiredCookMode, ECookInitializationFlags InCookFlags, const FString &InOutputDirectoryOverride )
 {
-
+	OutputDirectoryOverride = InOutputDirectoryOverride;
 	CurrentCookMode = DesiredCookMode;
 	CookFlags = InCookFlags;
 	
@@ -1705,23 +1765,9 @@ void UCookOnTheFlyServer::Initialize( ECookMode::Type DesiredCookMode, ECookInit
 	ITargetPlatformManagerModule& TPM = GetTargetPlatformManagerRef();
 	const TArray<ITargetPlatform*>& Platforms = TPM.GetCookingTargetPlatforms();
 
-	// Local sandbox file wrapper. This will be used to handle path conversions,
-	// but will not be used to actually write/read files so we can safely
-	// use [Platform] token in the sandbox directory name and then replace it
-	// with the actual platform name.
-	SandboxFile = new FSandboxPlatformFile(false);
-
-	// Output directory override.	
-	FString OutputDirectory = GetOutputDirectoryOverride(OutputDirectoryOverride);
-
-	// Use SandboxFile to do path conversion to properly handle sandbox paths (outside of standard paths in particular).
-	SandboxFile->Initialize(&FPlatformFileManager::Get().GetPlatformFile(), *FString::Printf(TEXT("-sandbox=\"%s\""), *OutputDirectory));
-
-	CleanSandbox(Platforms);
-
 	// always generate the asset registry before starting to cook, for either method
 	GenerateAssetRegistry(Platforms);
-
+	
 }
 
 uint32 UCookOnTheFlyServer::NumConnections() const
@@ -1738,14 +1784,23 @@ uint32 UCookOnTheFlyServer::NumConnections() const
 	return Result;
 }
 
-FString UCookOnTheFlyServer::GetOutputDirectoryOverride( const FString &OutputDirectoryOverride ) const
+FString UCookOnTheFlyServer::GetOutputDirectoryOverride() const
 {
 	FString OutputDirectory = OutputDirectoryOverride;
 	// Output directory override.	
 	if (OutputDirectory.Len() <= 0)
 	{
-		// Full path so that the sandbox wrapper doesn't try to re-base it under Sandboxes
-		OutputDirectory = FPaths::Combine(*FPaths::GameDir(), TEXT("Saved"), TEXT("Cooked"), TEXT("[Platform]"));
+		if ( IsCookingDLC() )
+		{
+			check( IsCookByTheBookMode() );
+			OutputDirectory = FPaths::Combine(*FPaths::GamePluginsDir(), *CookByTheBookOptions->DlcName, TEXT("Saved"), TEXT("Cooked"), TEXT("[Platform]"));
+		}
+		else
+		{
+			// Full path so that the sandbox wrapper doesn't try to re-base it under Sandboxes
+			OutputDirectory = FPaths::Combine(*FPaths::GameDir(), TEXT("Saved"), TEXT("Cooked"), TEXT("[Platform]"));
+		}
+		
 		OutputDirectory = FPaths::ConvertRelativePathToFull(OutputDirectory);
 	}
 	else if (!OutputDirectory.Contains(TEXT("[Platform]"), ESearchCase::IgnoreCase, ESearchDir::FromEnd) )
@@ -1879,68 +1934,100 @@ bool UCookOnTheFlyServer::GetCookedIniVersionStrings( const ITargetPlatform* Tar
 	return GConfig->GetArray( TEXT("CookSettings"), TEXT("VersionedIniParams"), IniVersionStrings, PlatformSandboxEditorIni ) > 0;
 }
 
-
-bool UCookOnTheFlyServer::IniSettingsOutOfDate( const TArray<ITargetPlatform*>& TargetPlatforms, TArray<ITargetPlatform*>& OutOfDateTargetPlatforms, bool bShouldSaveCookedVersions ) const
+bool UCookOnTheFlyServer::CacheIniVersionStringsMap( const ITargetPlatform* TargetPlatform ) const
 {
-	const FString CookedIni = FPaths::GameDir() / TEXT("CookedIniVersion.txt");
-	const FString SandboxCookedIni = SandboxFile->ConvertToAbsolutePathForExternalAppForWrite(*CookedIni);
+	// check if the cached ones are filled out
+	const FName TargetPlatformName = FName(*TargetPlatform->PlatformName());
+	TArray<FString>* FoundCookedIniVersionStrings = CachedIniVersionStringsMap.Find( TargetPlatformName );
 
-	for ( const auto& TargetPlatform : TargetPlatforms )
+	if ( FoundCookedIniVersionStrings == NULL )
 	{
-		TArray<FString> CurrentIniVersionStrings;
-		if ( GetCurrentIniVersionStrings(TargetPlatform, CurrentIniVersionStrings) == false )
+		TArray<FString> CookedIniVersionStrings;
+		if ( GetCookedIniVersionStrings( TargetPlatform, CookedIniVersionStrings ) )
 		{
-			// previous cook seems half baked... bomb out and recook all the things
-			OutOfDateTargetPlatforms.Add( TargetPlatform);
-			continue;
+			FoundCookedIniVersionStrings = &CachedIniVersionStringsMap.Emplace( TargetPlatformName, MoveTemp(CookedIniVersionStrings) );
 		}
+	}
+	return true;
+}
 
-		// check if the cached ones are filled out
-		const FName TargetPlatformName = FName(*TargetPlatform->PlatformName());
-		TArray<FString>* FoundCookedIniVersionStrings = CachedIniVersionStringsMap.Find( TargetPlatformName );
+bool UCookOnTheFlyServer::IniSettingsOutOfDate( const ITargetPlatform* TargetPlatform ) const
+{
+	TArray<FString> CurrentIniVersionStrings;
+	if ( GetCurrentIniVersionStrings(TargetPlatform, CurrentIniVersionStrings) == false )
+	{
+		// previous cook seems half baked... bomb out and recook all the things
+		return true;
+	}
 
-		if ( FoundCookedIniVersionStrings == NULL )
-		{
+	// check if the cached ones are filled out
+	const FName TargetPlatformName = FName(*TargetPlatform->PlatformName());
+	TArray<FString>* FoundCookedIniVersionStrings = CachedIniVersionStringsMap.Find( TargetPlatformName );
+
+	if ( FoundCookedIniVersionStrings == NULL )
+	{
 		TArray<FString> CookedIniVersionStrings;
-			GetCookedIniVersionStrings( TargetPlatform, CookedIniVersionStrings );
-			FoundCookedIniVersionStrings = &CachedIniVersionStringsMap.Emplace( TargetPlatformName, MoveTemp(CookedIniVersionStrings) );
-		}
+		GetCookedIniVersionStrings( TargetPlatform, CookedIniVersionStrings );
+		FoundCookedIniVersionStrings = &CachedIniVersionStringsMap.Emplace( TargetPlatformName, MoveTemp(CookedIniVersionStrings) );
+	}
 
-		check( FoundCookedIniVersionStrings );
-		bool bCurrentIniSettingsChanged = false;
-		for ( const auto& CurrentVersionString : CurrentIniVersionStrings ) 
+	check( FoundCookedIniVersionStrings );
+	bool bCurrentIniSettingsChanged = false;
+	for ( const auto& CurrentVersionString : CurrentIniVersionStrings ) 
+	{
+		if ( FoundCookedIniVersionStrings->Contains(CurrentVersionString) == false )
 		{
-			if ( FoundCookedIniVersionStrings->Contains(CurrentVersionString) == false )
-			{
-				
-				bCurrentIniSettingsChanged = true;
-				break;
-			}
+			bCurrentIniSettingsChanged = true;
+			break;
 		}
+	}
 
-		if ( bCurrentIniSettingsChanged )
-{
-		OutOfDateTargetPlatforms.Add(TargetPlatform);
-		*FoundCookedIniVersionStrings = CurrentIniVersionStrings;
-
-			// save the iniversion strings
-			if ( bShouldSaveCookedVersions )
+	if ( bCurrentIniSettingsChanged )
 	{
-		const FString PlatformSandboxCookedIni = SandboxCookedIni.Replace(TEXT("[Platform]"), *TargetPlatform->PlatformName());
-				GConfig->SetArray(TEXT("CookSettings"), TEXT("VersionedIniParams"), CurrentIniVersionStrings, PlatformSandboxCookedIni);
-		GConfig->Flush(false, PlatformSandboxCookedIni);
+		*FoundCookedIniVersionStrings = CurrentIniVersionStrings;
 	}
+
+	return bCurrentIniSettingsChanged;
 }
-	}
+
+bool UCookOnTheFlyServer::SaveCurrentIniSettings( const ITargetPlatform* TargetPlatform ) const
+{
+	const FString CookedIni = FPaths::GameDir() / TEXT("CookedIniVersion.txt");
+	const FString SandboxCookedIni = ConvertToFullSandboxPath(CookedIni, true);
+
+
+	const FName TargetPlatformName = FName(*TargetPlatform->PlatformName());
+	TArray<FString>* FoundCookedIniVersionStrings = CachedIniVersionStringsMap.Find( TargetPlatformName );
+		
+	// need to call IniSettingsOutOfDate before calling this function
+	check( FoundCookedIniVersionStrings );
+
+	// save the iniversion strings
+	const FString PlatformSandboxCookedIni = SandboxCookedIni.Replace(TEXT("[Platform]"), *TargetPlatform->PlatformName());
+	GConfig->SetArray(TEXT("CookSettings"), TEXT("VersionedIniParams"), *FoundCookedIniVersionStrings, PlatformSandboxCookedIni);
+	GConfig->Flush(false, PlatformSandboxCookedIni);
 
 	return true;
 }
 
 void UCookOnTheFlyServer::CleanSandbox(const TArray<ITargetPlatform*>& Platforms)
 {
+	// before we can delete any cooked files we need to make sure that we have finished writing them
+	UPackage::WaitForAsyncFileWrites();
+
+#if OUTPUT_TIMING
 	double SandboxCleanTime = 0.0;
+#endif
 	{
+#if OUTPUT_TIMING
 		SCOPE_SECONDS_COUNTER(SandboxCleanTime);
+#endif
+		
+		for ( const auto& Platform : Platforms )
+		{
+			CacheIniVersionStringsMap(Platform);
+		}
+
 
 		if (IsCookFlagSet(ECookInitializationFlags::Iterative) == false)
 		{
@@ -1948,19 +2035,31 @@ void UCookOnTheFlyServer::CleanSandbox(const TArray<ITargetPlatform*>& Platforms
 			for (int32 Index = 0; Index < Platforms.Num(); Index++)
 			{
 				ITargetPlatform* Target = Platforms[Index];
-				FString SandboxDirectory = GetOutputDirectory(Target->PlatformName());
+				
+				FString SandboxDirectory = GetSandboxDirectory( Target->PlatformName() ); // GetOutputDirectory(Target->PlatformName());
 				IFileManager::Get().DeleteDirectory(*SandboxDirectory, false, true);
+
+				ClearPlatformCookedData( FName( *Target->PlatformName() ) );
+
+				IniSettingsOutOfDate( Target );
+				SaveCurrentIniSettings( Target );
 			}
+
 		}
 		else
 		{
-			TArray<ITargetPlatform*> OutOfDatePlatforms;
-			IniSettingsOutOfDate(Platforms, OutOfDatePlatforms);
-			// wipe entire directory for any platforms which are out of date
-			for ( const auto& Target : OutOfDatePlatforms )
+			for ( const auto& Target : Platforms )
 			{
-				FString SandboxDirectory = GetOutputDirectory(Target->PlatformName());
-				IFileManager::Get().DeleteDirectory(*SandboxDirectory, false, true);
+				if ( IniSettingsOutOfDate(Target) )
+				{
+					ClearPlatformCookedData( FName( *Target->PlatformName() ) );
+
+					FString SandboxDirectory = GetSandboxDirectory(Target->PlatformName());
+					IFileManager::Get().DeleteDirectory(*SandboxDirectory, false, true);
+
+
+					SaveCurrentIniSettings(Target);
+				}
 			}
 
 			// check each package to see if it's out of date
@@ -1974,8 +2073,8 @@ void UCookOnTheFlyServer::CleanSandbox(const TArray<ITargetPlatform*>& Platforms
 			for (int32 Index = 0; Index < Platforms.Num(); Index++)
 			{
 				ITargetPlatform* Target = Platforms[Index];
-				FString SandboxDirectory = GetOutputDirectory(Target->PlatformName());
-
+				FString SandboxDirectory = GetSandboxDirectory(Target->PlatformName());
+				FName PlatformFName( *Target->PlatformName() );
 				// use the timestamp grabbing visitor
 				IPlatformFile& PlatformFile = FPlatformFileManager::Get().GetPlatformFile();
 				FLocalTimestampDirectoryVisitor Visitor(PlatformFile, DirectoriesToSkip, DirectoriesToNotRecurse, false);
@@ -1999,6 +2098,8 @@ void UCookOnTheFlyServer::CleanSandbox(const TArray<ITargetPlatform*>& Platforms
 							UE_LOG(LogCook, Display, TEXT("Deleting out of date cooked file: %s"), *CookedFilename);
 #endif
 							IFileManager::Get().Delete(*CookedFilename);
+
+							CookedPackages.RemoveFileForPlatform(FName(*StandardCookedFilename), PlatformFName);
 						}
 					}
 				}
@@ -2008,7 +2109,7 @@ void UCookOnTheFlyServer::CleanSandbox(const TArray<ITargetPlatform*>& Platforms
 			CollectGarbage(RF_Native);
 		}
 	}
-#if DEBUG_COOKONTHEFLY
+#if OUTPUT_TIMING
 	UE_LOG(LogCook, Display, TEXT("Sandbox cleanup took %5.3f seconds"), SandboxCleanTime);
 #endif
 }
@@ -2043,26 +2144,6 @@ void UCookOnTheFlyServer::GenerateAssetRegistry(const TArray<ITargetPlatform*>&
 				AssetRegistry.SearchAllAssets(true);
 			}
 		}
-
-		// When not cooking on the fly the registry will be saved after the cooker has finished
-		if (CurrentCookMode == ECookMode::CookOnTheFly)
-		{
-			// write it out to a memory archive
-			FArrayWriter SerializedAssetRegistry;
-			AssetRegistry.Serialize(SerializedAssetRegistry);
-			UE_LOG(LogCook, Display, TEXT("Generated asset registry size is %5.2fkb"), (float)SerializedAssetRegistry.Num() / 1024.f);
-
-			// now save it in each cooked directory
-			FString RegistryFilename = FPaths::GameDir() / TEXT("AssetRegistry.bin");
-			// Use SandboxFile to do path conversion to properly handle sandbox paths (outside of standard paths in particular).
-			FString SandboxFilename = ConvertToFullSandboxPath(*RegistryFilename, true);
-
-			for (int32 Index = 0; Index < Platforms.Num(); Index++)
-			{
-				FString PlatFilename = SandboxFilename.Replace(TEXT("[Platform]"), *Platforms[Index]->PlatformName());
-				FFileHelper::SaveArrayToFile(SerializedAssetRegistry, *PlatFilename);
-			}
-		}
 	}
 	UE_LOG(LogCook, Display, TEXT("Done creating registry. It took %5.2fs."), GenerateAssetRegistryTime);
 }
@@ -2177,7 +2258,7 @@ void UCookOnTheFlyServer::CollectFilesToCook(TArray<FString>& FilesInPath, const
 	if ( IsCookingDLC() )
 	{
 		// get the dlc and make sure we cook that directory 
-		FString DLCPath = FPaths::GamePluginsDir() / CookByTheBookOptions->DlcName;
+		FString DLCPath = FPaths::GamePluginsDir() / CookByTheBookOptions->DlcName / FString(TEXT("Content"));;
 
 		TArray<FString> Files;
 		IFileManager::Get().FindFilesRecursive(Files, *DLCPath, *(FString(TEXT("*")) + FPackageName::GetAssetPackageExtension()), true, false);
@@ -2316,7 +2397,6 @@ void UCookOnTheFlyServer::CollectFilesToCook(TArray<FString>& FilesInPath, const
 	}
 }
 
-
 bool UCookOnTheFlyServer::IsCookByTheBookRunning() const
 {
 	return CookByTheBookOptions && CookByTheBookOptions->bRunning;
@@ -2325,6 +2405,8 @@ bool UCookOnTheFlyServer::IsCookByTheBookRunning() const
 
 void UCookOnTheFlyServer::SaveGlobalShaderMapFiles(const TArray<ITargetPlatform*>& Platforms)
 {
+	// we don't support this behavior
+	check( !IsCookingDLC() );
 	for (int32 Index = 0; Index < Platforms.Num(); Index++)
 	{
 		// make sure global shaders are up to date!
@@ -2338,7 +2420,7 @@ void UCookOnTheFlyServer::SaveGlobalShaderMapFiles(const TArray<ITargetPlatform*
 
 		check( IsInGameThread() );
 
-		FString OutputDir = GetOutputDirectory(RecompileData.PlatformName);
+		FString OutputDir = GetSandboxDirectory(RecompileData.PlatformName);
 
 		RecompileShadersForRemote
 			(RecompileData.PlatformName, 
@@ -2351,9 +2433,25 @@ void UCookOnTheFlyServer::SaveGlobalShaderMapFiles(const TArray<ITargetPlatform*
 	}
 }
 
+FString UCookOnTheFlyServer::GetSandboxDirectory( const FString& PlatformName ) const
+{
+	FString Result;
+	Result = SandboxFile->GetSandboxDirectory();
+
+	Result.ReplaceInline(TEXT("[Platform]"), *PlatformName);
+
+	/*if ( IsCookingDLC() )
+	{
+		check( IsCookByTheBookRunning() );
+		Result.ReplaceInline(TEXT("/Cooked/"), *FString::Printf(TEXT("/CookedDLC_%s/"), *CookByTheBookOptions->DlcName) );
+	}*/
+	return Result;
+}
 
 FString UCookOnTheFlyServer::ConvertToFullSandboxPath( const FString &FileName, bool bForWrite ) const
 {
+	check( SandboxFile );
+
 	FString Result;
 	if ( bForWrite)
 	{
@@ -2364,14 +2462,22 @@ FString UCookOnTheFlyServer::ConvertToFullSandboxPath( const FString &FileName,
 		Result = SandboxFile->ConvertToAbsolutePathForExternalAppForRead(*FileName);
 	}
 
-	if ( IsCookByTheBookRunning() && IsCookingDLC() )
+	/*if ( IsCookingDLC() )
 	{
+		check( IsCookByTheBookRunning() );
 		Result.ReplaceInline(TEXT("/Cooked/"), *FString::Printf(TEXT("/CookedDLC_%s/"), *CookByTheBookOptions->DlcName) );
-	}
+	}*/
 	return Result;
 }
 
 
+FString UCookOnTheFlyServer::ConvertToFullSandboxPath( const FString &FileName, bool bForWrite, const FString& PlatformName ) const
+{
+	FString Result = ConvertToFullSandboxPath( FileName, bForWrite );
+	Result.ReplaceInline(TEXT("[Platform]"), *PlatformName);
+	return Result;
+}
+
 void UCookOnTheFlyServer::CookByTheBookFinished()
 {
 	check( IsInGameThread() );
@@ -2384,12 +2490,12 @@ void UCookOnTheFlyServer::CookByTheBookFinished()
 
 	{
 		// Save modified asset registry with all streaming chunk info generated during cook
-		FString RegistryFilename = FPaths::GameDir() / TEXT("AssetRegistry.bin");
-		FString SandboxRegistryFilename = ConvertToFullSandboxPath(*RegistryFilename, true);
+		const FString RegistryFilename = FPaths::GameDir() / TEXT("AssetRegistry.bin");
+		const FString SandboxRegistryFilename = ConvertToFullSandboxPath(*RegistryFilename, true);
 		// the registry filename will be modified when we call save asset registry
 		
-		FString CookedAssetRegistry = FPaths::GameDir() / TEXT("CookedAssetRegistry.json");
-		FString SandboxCookedAssetRegistryFilename = ConvertToFullSandboxPath(*CookedAssetRegistry, true);
+		const FString CookedAssetRegistry = FPaths::GameDir() / TEXT("CookedAssetRegistry.json");
+		const FString SandboxCookedAssetRegistryFilename = ConvertToFullSandboxPath(*CookedAssetRegistry, true);
 
 		for ( auto& Manifest : CookByTheBookOptions->ManifestGenerators )
 		{
@@ -2398,6 +2504,20 @@ void UCookOnTheFlyServer::CookByTheBookFinished()
 			Manifest.Value->SaveAssetRegistry(SandboxRegistryFilename);
 
 			Manifest.Value->SaveCookedPackageAssetRegistry(SandboxCookedAssetRegistryFilename, true);
+
+			if ( IsCreatingReleaseVersion() )
+			{
+				const FString VersionedRegistryPath = GetReleaseVersionAssetRegistryPath( CookByTheBookOptions->CreateReleaseVersion, Manifest.Key ); 
+
+				IFileManager::Get().MakeDirectory( *VersionedRegistryPath ,true );
+
+				const FString VersionedRegistryFilename = VersionedRegistryPath / GetAssetRegistryFilename();
+
+				const FString CookedAssetRegistryFilename = SandboxRegistryFilename.Replace(TEXT("[Platform]"), *Manifest.Key.ToString());
+
+				IFileManager::Get().Copy(*VersionedRegistryFilename, *CookedAssetRegistryFilename, true, true);
+				// Manifest.Value->SaveManifests( VersionedRegistryFilename );
+			}
 		}
 	}
 
@@ -2405,6 +2525,9 @@ void UCookOnTheFlyServer::CookByTheBookFinished()
 	const float TotalCookTime = (float)(FPlatformTime::Seconds() - CookByTheBookOptions->CookStartTime);
 	UE_LOG(LogCook, Display, TEXT("Cook by the book total time in tick %fs total time %f"), CookByTheBookOptions->CookTime, TotalCookTime);
 
+	check( SandboxFile );
+	SandboxFile = NULL;
+
 	CookByTheBookOptions->bRunning = false;
 }
 
@@ -2446,11 +2569,6 @@ void UCookOnTheFlyServer::StopAndClearCookedData()
 	CookedPackages.Empty(); // set of files which have been cooked when needing to recook a file the entry will need to be removed from here
 }
 
-FString UCookOnTheFlyServer::GetAssetRegistryPath() const
-{
-	return FPaths::GameDir() / TEXT("AssetRegistry.bin");
-}
-
 void UCookOnTheFlyServer::ClearAllCookedData()
 {
 	// if we are going to clear the cooked packages it is conceivable that we will recook the packages which we just cooked 
@@ -2467,16 +2585,43 @@ void UCookOnTheFlyServer::ClearPlatformCookedData( const FName& PlatformName )
 	// that means it's also conceivable that we will recook the same package which currently has an outstanding async write request
 	UPackage::WaitForAsyncFileWrites();
 
-	CookedPackages.RemoveAllCookedFilesForPlatform( PlatformName );
+	CookedPackages.RemoveAllFilesForPlatform( PlatformName );
 	TArray<FName> PackageNames;
 	UnsolicitedCookedPackages.GetPackagesForPlatformAndRemove(PlatformName, PackageNames);
 }
 
-void UCookOnTheFlyServer::StartCookByTheBook(const TArray<ITargetPlatform*>& TargetPlatforms, 
-											 const TArray<FString> &CookMaps, const TArray<FString> &CookDirectories, 
-											 const TArray<FString> &CookCultures, const TArray<FString> &IniMapSections, 
-											 ECookByTheBookOptions CookOptions, const FString &DLCName )
+void UCookOnTheFlyServer::CreateSandboxFile()
 {
+
+	// initialize the sandbox file after determining if we are cooking dlc
+	// Local sandbox file wrapper. This will be used to handle path conversions,
+	// but will not be used to actually write/read files so we can safely
+	// use [Platform] token in the sandbox directory name and then replace it
+	// with the actual platform name.
+	check( SandboxFile == NULL );
+	SandboxFile = new FSandboxPlatformFile(false);
+
+	// Output directory override.	
+	FString OutputDirectory = GetOutputDirectoryOverride();
+
+	// Use SandboxFile to do path conversion to properly handle sandbox paths (outside of standard paths in particular).
+	SandboxFile->Initialize(&FPlatformFileManager::Get().GetPlatformFile(), *FString::Printf(TEXT("-sandbox=\"%s\""), *OutputDirectory));
+}
+
+void UCookOnTheFlyServer::StartCookByTheBook( const FCookByTheBookStartupOptions& CookByTheBookStartupOptions )
+{
+
+	const TArray<ITargetPlatform*>& TargetPlatforms = CookByTheBookStartupOptions.TargetPlatforms;
+	const TArray<FString>& CookMaps = CookByTheBookStartupOptions.CookMaps;
+	const TArray<FString>& CookDirectories = CookByTheBookStartupOptions.CookDirectories;
+	const TArray<FString>& CookCultures = CookByTheBookStartupOptions.CookCultures;
+	const TArray<FString>& IniMapSections = CookByTheBookStartupOptions.IniMapSections;
+	const ECookByTheBookOptions& CookOptions = CookByTheBookStartupOptions.CookOptions;
+	const FString& DLCName = CookByTheBookStartupOptions.DLCName;
+
+	const FString& CreateReleaseVersion = CookByTheBookStartupOptions.CreateReleaseVersion;
+	const FString& BasedOnReleaseVersion = CookByTheBookStartupOptions.BasedOnReleaseVersion;
+
 	check( IsInGameThread() );
 
 	// can transition from cook on the fly to cook by the book from the editor
@@ -2510,18 +2655,41 @@ void UCookOnTheFlyServer::StartCookByTheBook(const TArray<ITargetPlatform*>& Tar
 
 	*/
 
-	TArray<ITargetPlatform*> ModifiedTargetPlatforms;
+
+
+	if ( CookByTheBookOptions->DlcName != DLCName )
+	{
+		// we are going to change the state of dlc we need to clean out our package filename cache (the generated filename cache is dependent on this key)
+		ClearPackageFilenameCache();
+		
+		CookByTheBookOptions->DlcName = DLCName;
+		check( OutputDirectoryOverride.Len() == 0 );
+	}
+
+	CreateSandboxFile();
+
+	// if we are cooking dlc then we don't need to clean sandbox
+	CleanSandbox(TargetPlatforms);
+	// deprecated is handled in the clean sandbox function above
+	/*TArray<ITargetPlatform*> ModifiedTargetPlatforms;
 	IniSettingsOutOfDate(TargetPlatforms, ModifiedTargetPlatforms);
 	for ( const auto& ModifiedTargetPlatform : ModifiedTargetPlatforms )
 	{
 		ClearPlatformCookedData( FName( *ModifiedTargetPlatform->PlatformName() ) );
+	}*/
+
+	if ( IsRealtimeMode() )
+	{
+		// from this point on (after we call CleanSandbox) if we are in the editor we can act iterative
+		CookFlags |= ECookInitializationFlags::Iterative;
 	}
 
 	CookByTheBookOptions->bRunning = true;
 	CookByTheBookOptions->bCancel = false;
 	CookByTheBookOptions->CookTime = 0.0f;
 	CookByTheBookOptions->CookStartTime = FPlatformTime::Seconds();
-	CookByTheBookOptions->DlcName = DLCName;
+	
+	CookByTheBookOptions->CreateReleaseVersion = CreateReleaseVersion;
 
 	bool bCookAll = (CookOptions & ECookByTheBookOptions::CookAll) != ECookByTheBookOptions::None;
 	bool bMapsOnly = (CookOptions & ECookByTheBookOptions::MapsOnly) != ECookByTheBookOptions::None;
@@ -2563,21 +2731,40 @@ void UCookOnTheFlyServer::StartCookByTheBook(const TArray<ITargetPlatform*>& Tar
 
 	if ( IsCookingDLC() )
 	{
-		ClearPackageFilenameCache();
+		// if we are cooking dlc we must be based of a release version cook
+		check( !BasedOnReleaseVersion.IsEmpty() );
 
-		FString OriginalSandboxRegistryFilename = SandboxFile->ConvertToAbsolutePathForExternalAppForRead(*GetAssetRegistryPath());
-		// load the original asset registry 
-		// don't recook any of the already cooked assets
-		// warm the cooked packages
-		WarmCookedPackages( OriginalSandboxRegistryFilename, TargetPlatformNames );
-	}
+		for ( const auto& PlatformName : TargetPlatformNames )
+		{
+			const FString OriginalSandboxRegistryFilename = GetReleaseVersionAssetRegistryPath(BasedOnReleaseVersion, PlatformName ) / GetAssetRegistryFilename();
+
+			TArray<FName> PackageFiles;
+			// if this check fails probably because the asset registry can't be found or read
+			bool bSucceeded = GetAllPackagesFromAssetRegistry(OriginalSandboxRegistryFilename, PackageFiles);
+			check( bSucceeded );
 
+			if ( bSucceeded )
+			{
+				for ( const auto& PackageFilename : PackageFiles )
+				{
+					CookedPackages.Add( MoveTemp( FFilePlatformRequest( PackageFilename, TargetPlatformNames ) ) );
+				}
+			}
+		}
+	}
+	
+	
 	
 	// allow the game to fill out the asset registry, as well as get a list of objects to always cook
 	TArray<FString> FilesInPath;
 	FGameDelegates::Get().GetCookModificationDelegate().ExecuteIfBound(FilesInPath);
 
-	SaveGlobalShaderMapFiles(TargetPlatforms);
+	// don't resave the global shader map files in dlc
+	if ( !IsCookingDLC() )
+	{
+		SaveGlobalShaderMapFiles(TargetPlatforms);
+	}
+	
 
 	CollectFilesToCook(FilesInPath, CookMaps, CookDirectories, CookCultures, IniMapSections, bCookAll, bMapsOnly, bNoDev );
 	if (FilesInPath.Num() == 0)
@@ -2606,6 +2793,37 @@ void UCookOnTheFlyServer::StartCookByTheBook(const TArray<ITargetPlatform*>& Tar
 		}	
 	}
 
+
+	if ( !IsCookingDLC() )
+	{
+		// if we are not cooking dlc then basedOfRelease version just needs to make sure that we cook all the packages which are in the previous release (as well as the new ones)
+		if ( !BasedOnReleaseVersion.IsEmpty() )
+		{
+			// if we are based of a release and we are not cooking dlc then we should always be creating a new one (note that we could be creating the same one we are based of).
+			check( !CreateReleaseVersion.IsEmpty() );
+
+			for ( const auto& PlatformName : TargetPlatformNames )
+			{
+				TArray<FName> PlatformArray;
+				PlatformArray.Add( PlatformName );
+
+				// if we are based of a cook and we are creating a new one we need to make sure that at least all the old packages are cooked as well as the new ones
+				FString OriginalAssetRegistryPath = GetReleaseVersionAssetRegistryPath( BasedOnReleaseVersion, PlatformName ) / GetAssetRegistryFilename();
+
+				TArray<FName> PackageFiles;
+				verify( !GetAllPackagesFromAssetRegistry(OriginalAssetRegistryPath, PackageFiles) );
+
+				for ( const auto& PackageFilename : PackageFiles )
+				{
+					CookRequests.EnqueueUnique( MoveTemp( FFilePlatformRequest( PackageFilename, PlatformArray) ) );
+				}
+
+			}
+
+		}
+	}
+
+
 	// this is to support canceling cooks from the editor
 	// this is required to make sure that the cooker is in a good state after cancel occurs
 	// if too many packages are being recooked after resume then we may need to figure out a different way to do this
@@ -2690,7 +2908,8 @@ void UCookOnTheFlyServer::HandleNetworkFileServerFileRequest( const FString& Fil
 void UCookOnTheFlyServer::HandleNetworkFileServerRecompileShaders(const FShaderRecompileData& RecompileData)
 {
 	// shouldn't receive network requests unless we are in cook on the fly mode
-	check( CurrentCookMode == ECookMode::CookOnTheFly );
+	check( IsCookOnTheFlyMode() );
+	check( !IsCookingDLC() );
 	// if we aren't in the game thread, we need to push this over to the game thread and wait for it to finish
 	if (!IsInGameThread())
 	{
@@ -2716,7 +2935,7 @@ void UCookOnTheFlyServer::HandleNetworkFileServerRecompileShaders(const FShaderR
 		return;
 	}
 
-	FString OutputDir = GetOutputDirectory(RecompileData.PlatformName);
+	FString OutputDir = GetSandboxDirectory(RecompileData.PlatformName);
 
 	RecompileShadersForRemote
 		(RecompileData.PlatformName, 
@@ -2729,7 +2948,7 @@ void UCookOnTheFlyServer::HandleNetworkFileServerRecompileShaders(const FShaderR
 		RecompileData.bCompileChangedShaders);
 }
 
-void UCookOnTheFlyServer::WarmCookedPackages(const FString& AssetRegistryPath, const TArray<FName>& TargetPlatformNames)
+/*void UCookOnTheFlyServer::WarmCookedPackages(const FString& AssetRegistryPath, const TArray<FName>& TargetPlatformNames)
 {
 	FArrayReader SerializedAssetData;
 	if (FFileHelper::LoadFileToArray(SerializedAssetData, *AssetRegistryPath))
@@ -2766,4 +2985,36 @@ void UCookOnTheFlyServer::WarmCookedPackages(const FString& AssetRegistryPath, c
 
 		delete[] PreallocatedAssetDataBuffer;
 	}
+}*/
+
+
+bool UCookOnTheFlyServer::GetAllPackagesFromAssetRegistry( const FString& AssetRegistryPath, TArray<FName>& OutPackageNames ) const
+{
+	FArrayReader SerializedAssetData;
+	if (FFileHelper::LoadFileToArray(SerializedAssetData, *AssetRegistryPath))
+	{
+		int32 LocalNumAssets = 0;
+		SerializedAssetData << LocalNumAssets;
+
+		// allocate one single block for all asset data structs (to reduce tens of thousands of heap allocations)
+		FAssetData* PreallocatedAssetDataBuffer = new FAssetData[LocalNumAssets];
+		
+		for (int32 AssetIndex = 0; AssetIndex < LocalNumAssets; AssetIndex++)
+		{
+			// make a new asset data object
+			FAssetData* NewAssetData = &PreallocatedAssetDataBuffer[AssetIndex];
+
+			// load it
+			SerializedAssetData << *NewAssetData;
+
+			UE_LOG(LogCook, Verbose, TEXT("Read package %s from %s"), *GetCachedStandardPackageFilename(NewAssetData->ObjectPath), *AssetRegistryPath);
+
+			OutPackageNames.Add( GetCachedStandardPackageFileFName(NewAssetData->ObjectPath) );
+		}
+
+		delete[] PreallocatedAssetDataBuffer;
+
+		return true;
+	}
+	return false;
 }
\ No newline at end of file
diff --git a/Engine/Source/Editor/UnrealEd/Private/UnrealEdEngine.cpp b/Engine/Source/Editor/UnrealEd/Private/UnrealEdEngine.cpp
--- a/Engine/Source/Editor/UnrealEd/Private/UnrealEdEngine.cpp
+++ b/Engine/Source/Editor/UnrealEd/Private/UnrealEdEngine.cpp
@@ -147,7 +147,14 @@ bool UUnrealEdEngine::CanCookByTheBookInEditor() const
 
 void UUnrealEdEngine::StartCookByTheBookInEditor( const TArray<ITargetPlatform*> &TargetPlatforms, const TArray<FString> &CookMaps, const TArray<FString> &CookDirectories, const TArray<FString> &CookCultures, const TArray<FString> &IniMapSections )
 {
-	CookServer->StartCookByTheBook( TargetPlatforms, CookMaps, CookDirectories, CookCultures, IniMapSections );
+	UCookOnTheFlyServer::FCookByTheBookStartupOptions StartupOptions;
+	StartupOptions.CookMaps = CookMaps;
+	StartupOptions.TargetPlatforms = TargetPlatforms;
+	StartupOptions.CookDirectories = CookDirectories;
+	StartupOptions.CookCultures = CookCultures;
+	StartupOptions.IniMapSections = IniMapSections;
+
+	CookServer->StartCookByTheBook( StartupOptions );
 }
 
 bool UUnrealEdEngine::IsCookByTheBookInEditorFinished() const 
