Can set default values on local function variables.

Details view supports displaying UProperties from external UStruct's that are not a property of a UObject.

diff --git a/Engine/Plugins/Messaging/UdpMessaging/Source/UdpMessaging/Private/Transport/UdpDeserializedMessage.cpp b/Engine/Plugins/Messaging/UdpMessaging/Source/UdpMessaging/Private/Transport/UdpDeserializedMessage.cpp
--- a/Engine/Plugins/Messaging/UdpMessaging/Source/UdpMessaging/Private/Transport/UdpDeserializedMessage.cpp
+++ b/Engine/Plugins/Messaging/UdpMessaging/Source/UdpMessaging/Private/Transport/UdpDeserializedMessage.cpp
@@ -92,7 +92,7 @@ bool FUdpDeserializedMessage::Deserialize( const FUdpReassembledMessageRef& Reas
 
 	// create message body
 	MessageData = FMemory::Malloc(TypeInfo->PropertiesSize);
-	TypeInfo->InitializeScriptStruct(MessageData);
+	TypeInfo->InitializeStruct(MessageData);
 
 	// deserialize message body
 	FJsonStructDeserializerBackend Backend(MessageReader);
diff --git a/Engine/Source/Editor/BlueprintGraph/Classes/K2Node_FunctionEntry.h b/Engine/Source/Editor/BlueprintGraph/Classes/K2Node_FunctionEntry.h
--- a/Engine/Source/Editor/BlueprintGraph/Classes/K2Node_FunctionEntry.h
+++ b/Engine/Source/Editor/BlueprintGraph/Classes/K2Node_FunctionEntry.h
@@ -45,6 +45,7 @@ class UK2Node_FunctionEntry : public UK2Node_FunctionTerminator
 	virtual bool DrawNodeAsEntry() const override { return true; }
 	virtual class FNodeHandlingFunctor* CreateNodeHandler(class FKismetCompilerContext& CompilerContext) const override;
 	virtual void GetRedirectPinNames(const UEdGraphPin& Pin, TArray<FString>& RedirectPinNames) const override;
+	virtual void ExpandNode(class FKismetCompilerContext& CompilerContext, UEdGraph* SourceGraph) override;
 	// End UK2Node interface
 
 	// Begin UK2Node_EditablePinBase interface
diff --git a/Engine/Source/Editor/BlueprintGraph/Private/EdGraphSchema_K2.cpp b/Engine/Source/Editor/BlueprintGraph/Private/EdGraphSchema_K2.cpp
--- a/Engine/Source/Editor/BlueprintGraph/Private/EdGraphSchema_K2.cpp
+++ b/Engine/Source/Editor/BlueprintGraph/Private/EdGraphSchema_K2.cpp
@@ -3613,6 +3613,14 @@ void UEdGraphSchema_K2::TrySetDefaultValue(UEdGraphPin& Pin, const FString& NewD
 		UseDefaultObject = FindObject<UObject>(ANY_PACKAGE, *NewDefaultValue);
 		UseDefaultValue = NULL;
 	}
+	else if(Pin.PinType.PinCategory == PC_Text)
+	{
+		// Set Text pins by string as if it were text.
+		TrySetDefaultText(Pin, FText::FromString(NewDefaultValue));
+		UseDefaultObject = nullptr;
+		UseDefaultValue.Empty();
+		return;
+	}
 	else
 	{
 		UseDefaultObject = NULL;
diff --git a/Engine/Source/Editor/BlueprintGraph/Private/K2Node_FunctionEntry.cpp b/Engine/Source/Editor/BlueprintGraph/Private/K2Node_FunctionEntry.cpp
--- a/Engine/Source/Editor/BlueprintGraph/Private/K2Node_FunctionEntry.cpp
+++ b/Engine/Source/Editor/BlueprintGraph/Private/K2Node_FunctionEntry.cpp
@@ -253,4 +253,114 @@ FString UK2Node_FunctionEntry::GetDeprecationMessage() const
 	return Super::GetDeprecationMessage();
 }
 
+void UK2Node_FunctionEntry::ExpandNode(class FKismetCompilerContext& CompilerContext, UEdGraph* SourceGraph)
+{
+	Super::ExpandNode(CompilerContext, SourceGraph);
+
+	const UEdGraphSchema_K2* Schema = CompilerContext.GetSchema();
+	
+	UEdGraphPin* OldStartExecPin = nullptr;
+
+	if(Pins[0]->LinkedTo.Num())
+	{
+		OldStartExecPin = Pins[0]->LinkedTo[0];
+	}
+	
+	UEdGraphPin* LastActiveOutputPin = Pins[0];
+
+	// Only look for FunctionEntry nodes who were duplicated and have a source object
+	if ( UK2Node_FunctionEntry* OriginalNode = Cast<UK2Node_FunctionEntry>(CompilerContext.MessageLog.FindSourceObject(this)) )
+	{
+		check(OriginalNode->GetOuter());
+
+		// Find the associated UFunction
+		UFunction* Function = FindField<UFunction>(CompilerContext.Blueprint->SkeletonGeneratedClass, *OriginalNode->GetOuter()->GetName());
+		for (TFieldIterator<UProperty> It(Function); It; ++It)
+		{
+			if (const UProperty* Property = *It)
+			{
+				for (auto& LocalVar : LocalVariables)
+				{
+					if (LocalVar.VarName == Property->GetFName() && !LocalVar.DefaultValue.IsEmpty())
+					{
+						// Add a variable set node for the local variable and hook it up immediately following the entry node or the last added local variable
+						UK2Node_VariableSet* VariableSetNode = CompilerContext.SpawnIntermediateNode<UK2Node_VariableSet>(this, SourceGraph);
+						VariableSetNode->SetFromProperty(Property, false);
+						Schema->ConfigureVarNode(VariableSetNode, LocalVar.VarName, Function, CompilerContext.Blueprint);
+						VariableSetNode->AllocateDefaultPins();
+						CompilerContext.MessageLog.NotifyIntermediateObjectCreation(VariableSetNode, this);
+
+						if(UEdGraphPin* SetPin = VariableSetNode->FindPin(Property->GetName()))
+						{
+							if(LocalVar.VarType.bIsArray)
+							{
+								TSharedPtr<FStructOnScope> StructData = MakeShareable(new FStructOnScope(Function));
+								FBlueprintEditorUtils::PropertyValueFromString(Property, LocalVar.DefaultValue, StructData->GetStructMemory());
+
+								// Create a Make Array node to setup the array's defaults
+								UK2Node_MakeArray* MakeArray = CompilerContext.SpawnIntermediateNode<UK2Node_MakeArray>(this, SourceGraph);
+								MakeArray->AllocateDefaultPins();
+								MakeArray->GetOutputPin()->MakeLinkTo(SetPin);
+								MakeArray->PostReconstructNode();
+
+								const UArrayProperty* ArrayProperty = Cast<UArrayProperty>(Property);
+								check(ArrayProperty);
+
+								FScriptArrayHelper ArrayHelper(ArrayProperty, StructData->GetStructMemory());
+								FScriptArrayHelper DefaultArrayHelper(ArrayProperty, StructData->GetStructMemory());
+
+								uint8* StructDefaults = NULL;
+								UStructProperty* StructProperty = dynamic_cast<UStructProperty*>(ArrayProperty->Inner);
+								if ( StructProperty != NULL )
+								{
+									checkSlow(StructProperty->Struct);
+									StructDefaults = (uint8*)FMemory::Malloc(StructProperty->Struct->GetStructureSize());
+									StructProperty->InitializeValue(StructDefaults);
+								}
+
+								// Go through each element in the array to set the default value
+								for( int32 ArrayIndex = 0 ; ArrayIndex < ArrayHelper.Num() ; ArrayIndex++ )
+								{
+									uint8* PropData = ArrayHelper.GetRawPtr(ArrayIndex);
+
+									// Always use struct defaults if the inner is a struct, for symmetry with the import of array inner struct defaults
+									uint8* PropDefault = ( StructProperty != NULL ) ? StructDefaults :
+										( ( StructData->GetStructMemory() && DefaultArrayHelper.Num() > ArrayIndex ) ? DefaultArrayHelper.GetRawPtr(ArrayIndex) : NULL );
+
+									// Retrieve the element's default value
+									FString DefaultValue;
+									ArrayProperty->Inner->ExportTextItem( DefaultValue, PropData, PropDefault, NULL, PPF_Delimited, NULL );
+
+									if(ArrayIndex > 0)
+									{
+										MakeArray->AddInputPin();
+									}
+
+									// Add one to the index for the pin to set the default on to skip the output pin
+									Schema->TrySetDefaultValue(*MakeArray->Pins[ArrayIndex + 1], DefaultValue);
+								}
+							}
+							else
+							{
+								// Set the default value
+								Schema->TrySetDefaultValue(*SetPin, LocalVar.DefaultValue);
+							}
+						}
+
+						LastActiveOutputPin->BreakAllPinLinks();
+						LastActiveOutputPin->MakeLinkTo(VariableSetNode->Pins[0]);
+						LastActiveOutputPin = VariableSetNode->Pins[1];
+					}
+				}
+			}
+		}
+
+		// Finally, hook up the last node to the old node the function entry node was connected to
+		if(OldStartExecPin)
+		{
+			LastActiveOutputPin->MakeLinkTo(OldStartExecPin);
+		}
+	}
+}
+
 #undef LOCTEXT_NAMESPACE
\ No newline at end of file
diff --git a/Engine/Source/Editor/BlueprintGraph/Private/K2Node_MakeArray.cpp b/Engine/Source/Editor/BlueprintGraph/Private/K2Node_MakeArray.cpp
--- a/Engine/Source/Editor/BlueprintGraph/Private/K2Node_MakeArray.cpp
+++ b/Engine/Source/Editor/BlueprintGraph/Private/K2Node_MakeArray.cpp
@@ -308,8 +308,7 @@ void UK2Node_MakeArray::AddInputPin()
 	Modify();
 
 	++NumInputs;
-	CreatePin(EGPD_Input, UEdGraphSchema_K2::PC_Wildcard, TEXT(""), NULL, false, false, *FString::Printf(TEXT("[%d]"), (NumInputs-1)));
-	PostReconstructNode();
+	CreatePin(EGPD_Input, GetOutputPin()->PinType.PinCategory, TEXT(""), NULL, false, false, *FString::Printf(TEXT("[%d]"), (NumInputs-1)));
 	
 	const bool bIsCompiling = GetBlueprint()->bBeingCompiled;
 	if( !bIsCompiling )
diff --git a/Engine/Source/Editor/Kismet/Private/BlueprintDetailsCustomization.cpp b/Engine/Source/Editor/Kismet/Private/BlueprintDetailsCustomization.cpp
--- a/Engine/Source/Editor/Kismet/Private/BlueprintDetailsCustomization.cpp
+++ b/Engine/Source/Editor/Kismet/Private/BlueprintDetailsCustomization.cpp
@@ -444,11 +444,21 @@ void FBlueprintVarActionDetails::CustomizeDetails( IDetailLayoutBuilder& DetailL
 
 	// Add in default value editing for properties that can be edited, local properties cannot be edited
 	UBlueprint* Blueprint = GetBlueprintObj();
-	if ((Blueprint != NULL) && (Blueprint->GeneratedClass != NULL) && !IsALocalVariable(VariableProperty))
+	if ((Blueprint != NULL) && (Blueprint->GeneratedClass != NULL))
 	{
 		if (VariableProperty != NULL)
 		{
-			const UProperty* OriginalProperty = FindField<UProperty>(Blueprint->GeneratedClass, VariableProperty->GetFName());
+			const UProperty* OriginalProperty = nullptr;
+			
+			if(!IsALocalVariable(VariableProperty))
+			{
+				OriginalProperty = FindField<UProperty>(Blueprint->GeneratedClass, VariableProperty->GetFName());
+			}
+			else
+			{
+				OriginalProperty = VariableProperty;
+			}
+
 			if (OriginalProperty == NULL)
 			{
 				// Prevent editing the default value of a skeleton property
@@ -518,10 +528,54 @@ void FBlueprintVarActionDetails::CustomizeDetails( IDetailLayoutBuilder& DetailL
 		}
 		else 
 		{
-			// Things are in order, show the property and allow it to be edited
-			TArray<UObject*> ObjectList;
-			ObjectList.Add(Blueprint->GeneratedClass->GetDefaultObject());
-			IDetailPropertyRow* Row = DefaultValueCategory.AddExternalProperty(ObjectList, VariableProperty->GetFName());
+			if(IsALocalVariable(VariableProperty))
+			{
+				UFunction* StructScope = Cast<UFunction>(VariableProperty->GetOuter());
+				check(StructScope);
+
+				TSharedPtr<FStructOnScope> StructData = MakeShareable(new FStructOnScope((UFunction*)StructScope));
+				UEdGraph* Graph = FBlueprintEditorUtils::FindScopeGraph(GetBlueprintObj(), (UFunction*)StructScope);
+				TWeakObjectPtr<UK2Node_EditablePinBase> EntryNode;
+				TWeakObjectPtr<UK2Node_EditablePinBase> ResultNode;
+				FBlueprintEditorUtils::GetEntryAndResultNodes(Graph, EntryNode, ResultNode);
+
+				UK2Node_FunctionEntry* FuncEntry = Cast<UK2Node_FunctionEntry>(EntryNode.Get());
+
+				UClass* ActorClass = FindObject<UClass>(ANY_PACKAGE, TEXT("Actor"));
+				UClass* BPActorClass = FindObject<UClass>(ANY_PACKAGE, TEXT("BP_Actor_C"));
+				for(auto& LocalVar : FuncEntry->LocalVariables)
+				{
+					if(LocalVar.VarName == VariableProperty->GetFName()) //Property->GetFName())
+					{
+						// Only set the default value if there is one
+						if(!LocalVar.DefaultValue.IsEmpty())
+						{
+							FBlueprintEditorUtils::PropertyValueFromString(VariableProperty, LocalVar.DefaultValue, StructData->GetStructMemory());
+						}
+						break;
+					}
+				}
+
+				TWeakPtr<FBlueprintEditor> BlueprintEditor = MyBlueprint.Pin()->GetBlueprintEditor();
+				if(BlueprintEditor.IsValid())
+				{
+					TSharedPtr< IDetailsView > DetailsView  = BlueprintEditor.Pin()->GetInspector()->GetPropertyView();
+
+					if(DetailsView.IsValid())
+					{
+						DetailsView->OnFinishedChangingProperties().AddSP(this, &FBlueprintVarActionDetails::OnFinishedChangingProperties, StructData, EntryNode);
+					}
+				}
+
+				IDetailPropertyRow* Row = DefaultValueCategory.AddExternalProperty(StructData, VariableProperty->GetFName());
+			}
+			else
+			{
+				// Things are in order, show the property and allow it to be edited
+				TArray<UObject*> ObjectList;
+				ObjectList.Add(Blueprint->GeneratedClass->GetDefaultObject());
+				IDetailPropertyRow* Row = DefaultValueCategory.AddExternalProperty(ObjectList, VariableProperty->GetFName());
+			}
 		}
 
 		TSharedPtr<SToolTip> TransientTooltip = IDocumentation::Get()->CreateToolTip(LOCTEXT("VariableTransient_Tooltip", "Should this variable not serialize and be zero-filled at load?"), NULL, DocLink, TEXT("Transient"));
@@ -1613,6 +1667,44 @@ EVisibility FBlueprintVarActionDetails::IsTooltipEditVisible() const
 	return EVisibility::Collapsed;
 }
 
+void FBlueprintVarActionDetails::OnFinishedChangingProperties(const FPropertyChangedEvent& InPropertyChangedEvent, TSharedPtr<FStructOnScope> InStructData, TWeakObjectPtr<UK2Node_EditablePinBase> InEntryNode)
+{
+	check(InPropertyChangedEvent.MemberProperty
+		&& InPropertyChangedEvent.MemberProperty->GetOwnerStruct()
+		&& InPropertyChangedEvent.MemberProperty->GetOwnerStruct()->IsA<UFunction>());
+
+	// Find the top level property that was modified within the UFunction
+	const UProperty* DirectProperty = InPropertyChangedEvent.MemberProperty;
+	while (!Cast<const UFunction>(DirectProperty->GetOuter()))
+	{
+		DirectProperty = CastChecked<const UProperty>(DirectProperty->GetOuter());
+	}
+
+	FString DefaultValueString;
+	bool bDefaultValueSet = false;
+
+	if (InStructData.IsValid())
+	{
+		bDefaultValueSet = FBlueprintEditorUtils::PropertyValueToString(DirectProperty, InStructData->GetStructMemory(), DefaultValueString);
+
+		if(bDefaultValueSet)
+		{
+			UK2Node_FunctionEntry* FuncEntry = Cast<UK2Node_FunctionEntry>(InEntryNode.Get());
+
+			// Search out the correct local variable in the Function Entry Node and set the default value
+			for(auto& LocalVar : FuncEntry->LocalVariables)
+			{
+				if(LocalVar.VarName == DirectProperty->GetFName())
+				{
+					LocalVar.DefaultValue = DefaultValueString;
+					FBlueprintEditorUtils::MarkBlueprintAsModified(GetBlueprintObj());
+					break;
+				}
+			}
+		}
+	}
+}
+
 static FDetailWidgetRow& AddRow( TArray<TSharedRef<FDetailWidgetRow> >& OutChildRows )
 {
 	TSharedRef<FDetailWidgetRow> NewRow( new FDetailWidgetRow );
diff --git a/Engine/Source/Editor/Kismet/Private/BlueprintDetailsCustomization.h b/Engine/Source/Editor/Kismet/Private/BlueprintDetailsCustomization.h
--- a/Engine/Source/Editor/Kismet/Private/BlueprintDetailsCustomization.h
+++ b/Engine/Source/Editor/Kismet/Private/BlueprintDetailsCustomization.h
@@ -79,6 +79,15 @@ private:
 	void OnVarTypeChanged(const FEdGraphPinType& NewPinType);
 	EVisibility IsTooltipEditVisible() const;
 
+	/**
+	 * Callback when changing a local variable property
+	 *
+	 * @param InPropertyChangedEvent	Information on the property changed
+	 * @param InStructData				The struct data where the value of the properties are stored
+	 * @param InEntryNode				Entry node where the default values of local variables are stored
+	 */
+	void OnFinishedChangingProperties(const FPropertyChangedEvent& InPropertyChangedEvent, TSharedPtr<FStructOnScope> InStructData, TWeakObjectPtr<UK2Node_EditablePinBase> InEntryNode);
+
 	/** Callback to decide if the category drop down menu should be enabled */
 	bool GetVariableCategoryChangeEnabled() const;
 
diff --git a/Engine/Source/Editor/Kismet/Private/SBlueprintPalette.cpp b/Engine/Source/Editor/Kismet/Private/SBlueprintPalette.cpp
--- a/Engine/Source/Editor/Kismet/Private/SBlueprintPalette.cpp
+++ b/Engine/Source/Editor/Kismet/Private/SBlueprintPalette.cpp
@@ -1373,18 +1373,22 @@ FText SBlueprintPaletteItem::GetToolTipText() const
 		else if (PaletteAction->GetTypeId() == FEdGraphSchemaAction_K2LocalVar::StaticGetTypeId())
 		{
 			FEdGraphSchemaAction_K2LocalVar* LocalVarAction = (FEdGraphSchemaAction_K2LocalVar*)PaletteAction.Get();
-			UClass* VarClass = CastChecked<UClass>(LocalVarAction->GetVariableScope()->GetOuter());
-			if (bShowClassInTooltip && (VarClass != NULL))
-			{
-				UBlueprint* Blueprint = UBlueprint::GetBlueprintFromClass(VarClass);
-				ClassDisplayName = (Blueprint != NULL) ? Blueprint->GetName() : VarClass->GetName();
-			}
-			else
+			// The variable scope can not be found in intermediate graphs
+			if(LocalVarAction->GetVariableScope())
 			{
-				FString Result;
-				FBlueprintEditorUtils::GetBlueprintVariableMetaData(Blueprint, LocalVarAction->GetVariableName(), LocalVarAction->GetVariableScope(), TEXT("tooltip"), Result);
-				// Only use the variable tooltip if it has been filled out.
-				ToolTipText = !Result.IsEmpty() ? Result : GetVarType(LocalVarAction->GetVariableScope(), LocalVarAction->GetVariableName(), true, true);
+				UClass* VarClass = CastChecked<UClass>(LocalVarAction->GetVariableScope()->GetOuter());
+				if (bShowClassInTooltip && (VarClass != NULL))
+				{
+					UBlueprint* Blueprint = UBlueprint::GetBlueprintFromClass(VarClass);
+					ClassDisplayName = (Blueprint != NULL) ? Blueprint->GetName() : VarClass->GetName();
+				}
+				else
+				{
+					FString Result;
+					FBlueprintEditorUtils::GetBlueprintVariableMetaData(Blueprint, LocalVarAction->GetVariableName(), LocalVarAction->GetVariableScope(), TEXT("tooltip"), Result);
+					// Only use the variable tooltip if it has been filled out.
+					ToolTipText = !Result.IsEmpty() ? Result : GetVarType(LocalVarAction->GetVariableScope(), LocalVarAction->GetVariableName(), true, true);
+				}
 			}
 		}
 		else if (PaletteAction->GetTypeId() == FEdGraphSchemaAction_K2Delegate::StaticGetTypeId())
diff --git a/Engine/Source/Editor/Kismet/Private/SKismetInspector.cpp b/Engine/Source/Editor/Kismet/Private/SKismetInspector.cpp
--- a/Engine/Source/Editor/Kismet/Private/SKismetInspector.cpp
+++ b/Engine/Source/Editor/Kismet/Private/SKismetInspector.cpp
@@ -210,7 +210,7 @@ void SKismetInspector::Construct(const FArguments& InArgs)
 	//@TODO: .IsEnabled( FSlateApplication::Get().GetNormalExecutionAttribute() );
 	PropertyView->SetIsPropertyVisibleDelegate( FIsPropertyVisible::CreateSP(this, &SKismetInspector::IsPropertyVisible) );
 	PropertyView->SetIsPropertyEditingEnabledDelegate(InArgs._IsPropertyEditingEnabledDelegate);
-	PropertyView->OnFinishedChangingProperties().Add( InArgs._OnFinishedChangingProperties );
+	UserOnFinishedChangingProperties = InArgs._OnFinishedChangingProperties;
 
 	TWeakPtr<SMyBlueprint> MyBlueprint = Kismet2.IsValid() ? Kismet2->GetMyBlueprintWidget() : InArgs._MyBlueprintWidget;
 	
@@ -342,6 +342,9 @@ void SKismetInspector::AddPropertiesRecursive(UProperty* Property)
 
 void SKismetInspector::UpdateFromObjects(const TArray<UObject*>& PropertyObjects, struct FKismetSelectionInfo& SelectionInfo, const FShowDetailsOptions& Options)
 {
+	PropertyView->OnFinishedChangingProperties().Clear();
+	PropertyView->OnFinishedChangingProperties().Add( UserOnFinishedChangingProperties );
+
 	if (!Options.bForceRefresh)
 	{
 		// Early out if the PropertyObjects and the SelectedObjects are the same
diff --git a/Engine/Source/Editor/Kismet/Public/SKismetInspector.h b/Engine/Source/Editor/Kismet/Public/SKismetInspector.h
--- a/Engine/Source/Editor/Kismet/Public/SKismetInspector.h
+++ b/Engine/Source/Editor/Kismet/Public/SKismetInspector.h
@@ -103,6 +103,9 @@ protected:
 	/** Set of object properties that should be visible */
 	TSet<TWeakObjectPtr<UProperty> > SelectedObjectProperties;
 	
+	/** User defined delegate for OnFinishedChangingProperties */
+	FOnFinishedChangingProperties::FDelegate UserOnFinishedChangingProperties;
+
 protected:
 	/** Show properties of the selected object */
 	void SetPropertyWindowContents(TArray<UObject*> Objects);
diff --git a/Engine/Source/Editor/KismetCompiler/Private/KismetCompiler.cpp b/Engine/Source/Editor/KismetCompiler/Private/KismetCompiler.cpp
--- a/Engine/Source/Editor/KismetCompiler/Private/KismetCompiler.cpp
+++ b/Engine/Source/Editor/KismetCompiler/Private/KismetCompiler.cpp
@@ -818,6 +818,7 @@ void FKismetCompilerContext::CreateUserDefinedLocalVariablesForFunction(FKismetF
 			NewProperty->SetMetaData(TEXT("FriendlyName"), *Variable.FriendlyName);
 			NewProperty->SetMetaData(TEXT("Category"), *Variable.Category.ToString());
 			NewProperty->RepNotifyFunc = Variable.RepNotifyFunc;
+			NewProperty->SetPropertyFlags(Variable.PropertyFlags);
 
 			if(!Variable.DefaultValue.IsEmpty())
 			{
diff --git a/Engine/Source/Editor/PropertyEditor/Private/DetailCategoryBuilderImpl.cpp b/Engine/Source/Editor/PropertyEditor/Private/DetailCategoryBuilderImpl.cpp
--- a/Engine/Source/Editor/PropertyEditor/Private/DetailCategoryBuilderImpl.cpp
+++ b/Engine/Source/Editor/PropertyEditor/Private/DetailCategoryBuilderImpl.cpp
@@ -18,6 +18,7 @@
 #include "DetailPropertyRow.h"
 #include "DetailGroup.h"
 #include "DetailCustomBuilderRow.h"
+#include "StructurePropertyNode.h"
 
 namespace DetailLayoutConstants
 {
@@ -309,6 +310,60 @@ IDetailPropertyRow* FDetailCategoryImpl::AddExternalProperty( const TArray<UObje
 	return NULL;
 }
 
+IDetailPropertyRow* FDetailCategoryImpl::AddExternalProperty( TSharedPtr<FStructOnScope> StructData, FName PropertyName, EPropertyLocation::Type Location/* = EPropertyLocation::Default*/ )
+{
+	TSharedPtr<FStructurePropertyNode> RootPropertyNode( new FStructurePropertyNode );
+
+	//SET
+	RootPropertyNode->SetStructure(StructData);
+
+	FPropertyNodeInitParams InitParams;
+	InitParams.ParentNode = NULL;
+	InitParams.Property = NULL;
+	InitParams.ArrayOffset = 0;
+	InitParams.ArrayIndex = INDEX_NONE;
+	InitParams.bAllowChildren = true;
+	InitParams.bForceHiddenPropertyVisibility = FPropertySettings::Get().ShowHiddenProperties();
+	InitParams.bCreateCategoryNodes = false;
+
+	RootPropertyNode->InitNode(InitParams);
+
+	if( RootPropertyNode.IsValid() )
+	{
+		RootPropertyNode->RebuildChildren();
+
+		FDetailLayoutCustomization NewCustomization;
+
+		for(int32 ChildIdx = 0; ChildIdx < RootPropertyNode->GetNumChildNodes(); ++ChildIdx)
+		{
+			TSharedPtr< FPropertyNode > PropertyNode = RootPropertyNode->GetChildNode(ChildIdx);
+			if(UProperty* Property = PropertyNode->GetProperty())
+			{
+				if(Property->GetFName() == PropertyName)
+				{
+					NewCustomization.PropertyRow = MakeShareable( new FDetailPropertyRow( PropertyNode, AsShared(), RootPropertyNode ) );
+					break;
+				}
+			}
+		}
+
+		bool bForAdvanced = false;
+		if( Location == EPropertyLocation::Advanced )
+		{
+			// Force advanced
+			bForAdvanced = true;
+		}
+
+		DetailLayoutBuilder.Pin()->AddExternalRootPropertyNode( RootPropertyNode.ToSharedRef() );
+
+		AddCustomLayout( NewCustomization, bForAdvanced );
+
+		return NewCustomization.PropertyRow.Get();
+	}
+
+	return NULL;
+}
+
 void FDetailCategoryImpl::AddPropertyNode( TSharedRef<FPropertyNode> PropertyNode, FName InstanceName )
 {
 	FDetailLayoutCustomization NewCustomization;
diff --git a/Engine/Source/Editor/PropertyEditor/Private/DetailCategoryBuilderImpl.h b/Engine/Source/Editor/PropertyEditor/Private/DetailCategoryBuilderImpl.h
--- a/Engine/Source/Editor/PropertyEditor/Private/DetailCategoryBuilderImpl.h
+++ b/Engine/Source/Editor/PropertyEditor/Private/DetailCategoryBuilderImpl.h
@@ -137,6 +137,7 @@ public:
 	virtual IDetailPropertyRow& AddProperty( FName PropertyPath, UClass* ClassOuter = NULL, FName InstanceName = NAME_None,  EPropertyLocation::Type Location = EPropertyLocation::Default ) override;
 	virtual IDetailPropertyRow& AddProperty( TSharedPtr<IPropertyHandle> PropertyHandle,  EPropertyLocation::Type Location = EPropertyLocation::Default ) override;
 	virtual IDetailPropertyRow* AddExternalProperty( const TArray<UObject*>& Objects, FName PropertyName, EPropertyLocation::Type Location = EPropertyLocation::Default ) override;
+	virtual IDetailPropertyRow* AddExternalProperty( TSharedPtr<FStructOnScope> StructData, FName PropertyName, EPropertyLocation::Type Location = EPropertyLocation::Default ) override;
 	virtual IDetailLayoutBuilder& GetParentLayout() const override { return *DetailLayoutBuilder.Pin(); }
 	virtual FDetailWidgetRow& AddCustomRow( const FText& FilterString, bool bForAdvanced = false ) override;
 	virtual void AddCustomBuilder( TSharedRef<IDetailCustomNodeBuilder> InCustomBuilder, bool bForAdvanced = false ) override;
diff --git a/Engine/Source/Editor/PropertyEditor/Private/DetailLayoutBuilderImpl.cpp b/Engine/Source/Editor/PropertyEditor/Private/DetailLayoutBuilderImpl.cpp
--- a/Engine/Source/Editor/PropertyEditor/Private/DetailLayoutBuilderImpl.cpp
+++ b/Engine/Source/Editor/PropertyEditor/Private/DetailLayoutBuilderImpl.cpp
@@ -402,7 +402,7 @@ TSharedRef<IPropertyHandle> FDetailLayoutBuilderImpl::GetPropertyHandle( TShared
 	return PropertyHandle.ToSharedRef();
 }
 
-void FDetailLayoutBuilderImpl::AddExternalRootPropertyNode( TSharedRef<FObjectPropertyNode> InExternalRootNode )
+void FDetailLayoutBuilderImpl::AddExternalRootPropertyNode( TSharedRef<FPropertyNode> InExternalRootNode )
 {
 	DetailsView.AddExternalRootPropertyNode( InExternalRootNode );
 }
diff --git a/Engine/Source/Editor/PropertyEditor/Private/DetailLayoutBuilderImpl.h b/Engine/Source/Editor/PropertyEditor/Private/DetailLayoutBuilderImpl.h
--- a/Engine/Source/Editor/PropertyEditor/Private/DetailLayoutBuilderImpl.h
+++ b/Engine/Source/Editor/PropertyEditor/Private/DetailLayoutBuilderImpl.h
@@ -138,7 +138,7 @@ public:
 	 *
 	 * @param InExternalRootNode	The node to add
 	 */
-	void AddExternalRootPropertyNode( TSharedRef<FObjectPropertyNode> InExternalRootNode );
+	void AddExternalRootPropertyNode( TSharedRef<FPropertyNode> InExternalRootNode );
 
 	/** @return The details view that owns this layout */
 	IDetailsViewPrivate& GetDetailsView() { return DetailsView; }
diff --git a/Engine/Source/Editor/PropertyEditor/Private/DetailPropertyRow.cpp b/Engine/Source/Editor/PropertyEditor/Private/DetailPropertyRow.cpp
--- a/Engine/Source/Editor/PropertyEditor/Private/DetailPropertyRow.cpp
+++ b/Engine/Source/Editor/PropertyEditor/Private/DetailPropertyRow.cpp
@@ -9,7 +9,7 @@
 #include "DetailGroup.h"
 #include "CustomChildBuilder.h"
 
-FDetailPropertyRow::FDetailPropertyRow(TSharedPtr<FPropertyNode> InPropertyNode, TSharedRef<FDetailCategoryImpl> InParentCategory, TSharedPtr<FObjectPropertyNode> InExternalRootNode )
+FDetailPropertyRow::FDetailPropertyRow(TSharedPtr<FPropertyNode> InPropertyNode, TSharedRef<FDetailCategoryImpl> InParentCategory, TSharedPtr<FPropertyNode> InExternalRootNode )
 	: CustomIsEnabledAttrib( true )
 	, PropertyNode( InPropertyNode )
 	, ParentCategory( InParentCategory )
diff --git a/Engine/Source/Editor/PropertyEditor/Private/DetailPropertyRow.h b/Engine/Source/Editor/PropertyEditor/Private/DetailPropertyRow.h
--- a/Engine/Source/Editor/PropertyEditor/Private/DetailPropertyRow.h
+++ b/Engine/Source/Editor/PropertyEditor/Private/DetailPropertyRow.h
@@ -5,7 +5,7 @@
 class FDetailPropertyRow : public IDetailPropertyRow, public IPropertyTypeCustomizationUtils, public TSharedFromThis<FDetailPropertyRow>
 {
 public:
-	FDetailPropertyRow( TSharedPtr<FPropertyNode> InPropertyNode, TSharedRef<FDetailCategoryImpl> InParentCategory, TSharedPtr<FObjectPropertyNode> InExternalRootNode = NULL);
+	FDetailPropertyRow( TSharedPtr<FPropertyNode> InPropertyNode, TSharedRef<FDetailCategoryImpl> InParentCategory, TSharedPtr<FPropertyNode> InExternalRootNode = NULL);
 
 	/** IDetailPropertyRow interface */
 	virtual IDetailPropertyRow& DisplayName( const FText& InDisplayName ) override;
@@ -120,7 +120,7 @@ private:
 	/** The category this row resides in */
 	TWeakPtr<FDetailCategoryImpl> ParentCategory;
 	/** Root of the property node if this node comes from an external tree */
-	TSharedPtr<FObjectPropertyNode> ExternalRootNode;
+	TSharedPtr<FPropertyNode> ExternalRootNode;
 	/** Whether or not to show standard property buttons */
 	bool bShowPropertyButtons;
 	/** True to show custom property children */
diff --git a/Engine/Source/Editor/PropertyEditor/Private/IDetailsViewPrivate.h b/Engine/Source/Editor/PropertyEditor/Private/IDetailsViewPrivate.h
--- a/Engine/Source/Editor/PropertyEditor/Private/IDetailsViewPrivate.h
+++ b/Engine/Source/Editor/PropertyEditor/Private/IDetailsViewPrivate.h
@@ -91,7 +91,7 @@ public:
 	 *
 	 * @param InExternalRootNode	The node to add
 	 */
-	virtual void AddExternalRootPropertyNode(TSharedRef<FObjectPropertyNode> ExternalRootNode) = 0;
+	virtual void AddExternalRootPropertyNode(TSharedRef<FPropertyNode> ExternalRootNode) = 0;
 
 	/**
 	 * Creates the color picker window for this property view.
diff --git a/Engine/Source/Editor/PropertyEditor/Private/SDetailsView.cpp b/Engine/Source/Editor/PropertyEditor/Private/SDetailsView.cpp
--- a/Engine/Source/Editor/PropertyEditor/Private/SDetailsView.cpp
+++ b/Engine/Source/Editor/PropertyEditor/Private/SDetailsView.cpp
@@ -635,7 +635,7 @@ void SDetailsView::UnregisterInstancedCustomPropertyLayout( UClass* Class )
 	InstancedClassToDetailLayoutMap.Remove( Class );
 }
 
-void SDetailsView::AddExternalRootPropertyNode( TSharedRef<FObjectPropertyNode> ExternalRootNode )
+void SDetailsView::AddExternalRootPropertyNode( TSharedRef<FPropertyNode> ExternalRootNode )
 {
 	ExternalRootPropertyNodes.Add( ExternalRootNode );
 }
diff --git a/Engine/Source/Editor/PropertyEditor/Private/SDetailsView.h b/Engine/Source/Editor/PropertyEditor/Private/SDetailsView.h
--- a/Engine/Source/Editor/PropertyEditor/Private/SDetailsView.h
+++ b/Engine/Source/Editor/PropertyEditor/Private/SDetailsView.h
@@ -88,7 +88,7 @@ public:
 	 *
 	 * @param InExternalRootNode	The node to add
 	 */
-	void AddExternalRootPropertyNode( TSharedRef<FObjectPropertyNode> ExternalRootNode );
+	void AddExternalRootPropertyNode( TSharedRef<FPropertyNode> ExternalRootNode );
 	
 	/**
 	 * @return True if a category is hidden by any of the uobject classes currently in view by this details panel
diff --git a/Engine/Source/Editor/PropertyEditor/Private/SDetailsViewBase.cpp b/Engine/Source/Editor/PropertyEditor/Private/SDetailsViewBase.cpp
--- a/Engine/Source/Editor/PropertyEditor/Private/SDetailsViewBase.cpp
+++ b/Engine/Source/Editor/PropertyEditor/Private/SDetailsViewBase.cpp
@@ -714,14 +714,14 @@ void SDetailsViewBase::Tick(const FGeometry& AllottedGeometry, const double InCu
 	{
 		for (int32 NodeIndex = 0; NodeIndex < ExternalRootPropertyNodes.Num(); ++NodeIndex)
 		{
-			TSharedPtr<FObjectPropertyNode> ObjectNode = ExternalRootPropertyNodes[NodeIndex].Pin();
+			TSharedPtr<FPropertyNode> PropertyNode = ExternalRootPropertyNodes[NodeIndex].Pin();
 
-			if (ObjectNode.IsValid())
+			if (PropertyNode.IsValid())
 			{
-				Result = ObjectNode->EnsureDataIsValid();
+				Result = PropertyNode->EnsureDataIsValid();
 				if (Result == FPropertyNode::PropertiesChanged || Result == FPropertyNode::EditInlineNewValueChanged)
 				{
-					RestoreExpandedItems(ObjectNode);
+					RestoreExpandedItems(PropertyNode);
 					UpdatePropertyMap();
 					// Note this will invalidate all the external root nodes so there is no need to continue
 					ExternalRootPropertyNodes.Empty();
@@ -729,7 +729,7 @@ void SDetailsViewBase::Tick(const FGeometry& AllottedGeometry, const double InCu
 				}
 				else if (Result == FPropertyNode::ArraySizeChanged)
 				{
-					RestoreExpandedItems(ObjectNode);
+					RestoreExpandedItems(PropertyNode);
 					UpdateFilteredDetails();
 				}
 			}
@@ -933,12 +933,12 @@ void SDetailsViewBase::UpdateFilteredDetails()
 
 		for (int32 NodeIndex = 0; NodeIndex < ExternalRootPropertyNodes.Num(); ++NodeIndex)
 		{
-			TSharedPtr<FObjectPropertyNode> ObjectNode = ExternalRootPropertyNodes[NodeIndex].Pin();
+			TSharedPtr<FPropertyNode> PropertyNode = ExternalRootPropertyNodes[NodeIndex].Pin();
 
-			if (ObjectNode.IsValid())
+			if (PropertyNode.IsValid())
 			{
-				ObjectNode->FilterNodes(CurrentFilter.FilterStrings);
-				ObjectNode->ProcessSeenFlags(true);
+				PropertyNode->FilterNodes(CurrentFilter.FilterStrings);
+				PropertyNode->ProcessSeenFlags(true);
 			}
 		}
 
diff --git a/Engine/Source/Editor/PropertyEditor/Private/SDetailsViewBase.h b/Engine/Source/Editor/PropertyEditor/Private/SDetailsViewBase.h
--- a/Engine/Source/Editor/PropertyEditor/Private/SDetailsViewBase.h
+++ b/Engine/Source/Editor/PropertyEditor/Private/SDetailsViewBase.h
@@ -454,7 +454,7 @@ protected:
 	TSharedPtr<IDetailPropertyExtensionHandler> ExtensionHandler;
 
 	/** External property nodes which need to validated each tick */
-	TArray< TWeakPtr<FObjectPropertyNode> > ExternalRootPropertyNodes;
+	TArray< TWeakPtr<FPropertyNode> > ExternalRootPropertyNodes;
 
 	/** The tree node that is currently highlighted, may be none: */
 	TWeakPtr< IDetailTreeNode > CurrentlyHighlightedNode;
diff --git a/Engine/Source/Editor/PropertyEditor/Private/SStructureDetailsView.cpp b/Engine/Source/Editor/PropertyEditor/Private/SStructureDetailsView.cpp
--- a/Engine/Source/Editor/PropertyEditor/Private/SStructureDetailsView.cpp
+++ b/Engine/Source/Editor/PropertyEditor/Private/SStructureDetailsView.cpp
@@ -22,10 +22,10 @@ SStructureDetailsView::~SStructureDetailsView()
 	SaveExpandedItems();
 }
 
-UScriptStruct* SStructureDetailsView::GetBaseScriptStruct() const
+UStruct* SStructureDetailsView::GetBaseScriptStruct() const
 {
-	const UScriptStruct* Struct = StructData.IsValid() ? StructData->GetStruct() : NULL;
-	return const_cast<UScriptStruct*>(Struct);
+	const UStruct* Struct = StructData.IsValid() ? StructData->GetStruct() : NULL;
+	return const_cast<UStruct*>(Struct);
 }
 
 void SStructureDetailsView::Construct(const FArguments& InArgs)
@@ -178,7 +178,7 @@ void SStructureDetailsView::SetStructureData(TSharedPtr<FStructOnScope> InStruct
 	//SET
 	StructData = InStructData;
 	RootNode->SetStructure(StructData);
-	if (!StructData.IsValid() || !StructData->IsValid())
+	if (!StructData.IsValid())
 	{
 		bIsLocked = false;
 	}
@@ -242,4 +242,4 @@ void SStructureDetailsView::CustomUpdatePropertyMap()
 	DetailLayout->DefaultCategory(NAME_None).SetDisplayName(NAME_None, CustomName);
 }
 
-#undef LOCTEXT_NAMESPACE
\ No newline at end of file
+#undef LOCTEXT_NAMESPACE
diff --git a/Engine/Source/Editor/PropertyEditor/Private/SStructureDetailsView.h b/Engine/Source/Editor/PropertyEditor/Private/SStructureDetailsView.h
--- a/Engine/Source/Editor/PropertyEditor/Private/SStructureDetailsView.h
+++ b/Engine/Source/Editor/PropertyEditor/Private/SStructureDetailsView.h
@@ -26,7 +26,7 @@ public:
 	 */
 	void Construct(const FArguments& InArgs);
 
-	UScriptStruct* GetBaseScriptStruct() const;
+	UStruct* GetBaseScriptStruct() const;
 
 	virtual bool IsConnected() const override;
 
@@ -63,7 +63,7 @@ public:
 	}
 
 	virtual void ForceRefresh() override;
-	virtual void AddExternalRootPropertyNode(TSharedRef<FObjectPropertyNode> ExternalRootNode) override {}
+	virtual void AddExternalRootPropertyNode(TSharedRef<FPropertyNode> ExternalRootNode) override {}
 
 	/** IDetailsView interface */
 	virtual UStruct* GetBaseStruct() const override
diff --git a/Engine/Source/Editor/PropertyEditor/Private/StructurePropertyNode.cpp b/Engine/Source/Editor/PropertyEditor/Private/StructurePropertyNode.cpp
--- a/Engine/Source/Editor/PropertyEditor/Private/StructurePropertyNode.cpp
+++ b/Engine/Source/Editor/PropertyEditor/Private/StructurePropertyNode.cpp
@@ -11,7 +11,7 @@
 void FStructurePropertyNode::InitChildNodes()
 {
 	const bool bShouldShowHiddenProperties = !!HasNodeFlags(EPropertyNodeFlags::ShouldShowHiddenProperties);
-	const UScriptStruct* Struct = StructData.IsValid() ? StructData->GetStruct() : NULL;
+	const UStruct* Struct = StructData.IsValid() ? StructData->GetStruct() : NULL;
 
 	for (TFieldIterator<UProperty> It(Struct); It; ++It)
 	{
diff --git a/Engine/Source/Editor/PropertyEditor/Public/DetailCategoryBuilder.h b/Engine/Source/Editor/PropertyEditor/Public/DetailCategoryBuilder.h
--- a/Engine/Source/Editor/PropertyEditor/Public/DetailCategoryBuilder.h
+++ b/Engine/Source/Editor/PropertyEditor/Public/DetailCategoryBuilder.h
@@ -80,6 +80,16 @@ public:
 	 */
 	virtual IDetailPropertyRow* AddExternalProperty( const TArray<UObject*>& Objects, FName PropertyName, EPropertyLocation::Type Location = EPropertyLocation::Default ) = 0;
 
+	/**
+	 * Adds an external property, that is contained within a UStruct, that is not a property on the object(s) being customized 
+	 *
+	 * @param StructData		Struct data to find the property within
+	 * @param PropertyName		The name of the property to view
+	 * @param Location			The location within the category where the property is shown
+	 * @return A property row for customizing the property or NULL if the property could not be found
+	 */
+	virtual IDetailPropertyRow* AddExternalProperty( TSharedPtr<FStructOnScope> StructData, FName PropertyName, EPropertyLocation::Type Location = EPropertyLocation::Default ) = 0;
+
 	/**
 	 * Adds a custom widget row to the category
 	 *
diff --git a/Engine/Source/Editor/UnrealEd/Private/DataTableEditorUtils.cpp b/Engine/Source/Editor/UnrealEd/Private/DataTableEditorUtils.cpp
--- a/Engine/Source/Editor/UnrealEd/Private/DataTableEditorUtils.cpp
+++ b/Engine/Source/Editor/UnrealEd/Private/DataTableEditorUtils.cpp
@@ -21,7 +21,7 @@ bool FDataTableEditorUtils::RemoveRow(UDataTable* DataTable, FName Name)
 		const bool bRemoved = DataTable->RowMap.RemoveAndCopyValue(Name, RowData);
 		if (bRemoved && RowData)
 		{
-			DataTable->RowStruct->DestroyScriptStruct(RowData);
+			DataTable->RowStruct->DestroyStruct(RowData);
 			FMemory::Free(RowData);
 			bResult = true;
 		}
@@ -40,7 +40,7 @@ uint8* FDataTableEditorUtils::AddRow(UDataTable* DataTable, FName RowName)
 	DataTable->Modify();
 	// Allocate data to store information, using UScriptStruct to know its size
 	uint8* RowData = (uint8*)FMemory::Malloc(DataTable->RowStruct->PropertiesSize);
-	DataTable->RowStruct->InitializeScriptStruct(RowData);
+	DataTable->RowStruct->InitializeStruct(RowData);
 	// And be sure to call DestroyScriptStruct later
 
 	if (auto UDStruct = Cast<const UUserDefinedStruct>(DataTable->RowStruct))
diff --git a/Engine/Source/Editor/UnrealEd/Private/Kismet2/BlueprintEditorUtils.cpp b/Engine/Source/Editor/UnrealEd/Private/Kismet2/BlueprintEditorUtils.cpp
--- a/Engine/Source/Editor/UnrealEd/Private/Kismet2/BlueprintEditorUtils.cpp
+++ b/Engine/Source/Editor/UnrealEd/Private/Kismet2/BlueprintEditorUtils.cpp
@@ -3348,33 +3348,9 @@ bool FBlueprintEditorUtils::AddMemberVariable(UBlueprint* Blueprint, const FName
 	{
 		NewVar.PropertyFlags |= CPF_BlueprintAssignable | CPF_BlueprintCallable;
 	}
-	else if ((NewVarType.PinCategory == K2Schema->PC_Object) || (NewVarType.PinCategory == K2Schema->PC_Interface))
+	else
 	{
-		if (NewVar.VarType.PinSubCategory == K2Schema->PSC_Self)
-		{
-			NewVar.VarType.PinSubCategory.Empty();
-			NewVar.VarType.PinSubCategoryObject = *Blueprint->GeneratedClass;
-		}
-		else if (!NewVar.VarType.PinSubCategoryObject.IsValid())
-		{
-			// Fall back to UObject if the given type is not valid. This can happen for example if a variable is removed from
-			// a Blueprint parent class along with the variable's type and the user then attempts to recreate the missing variable
-			// through a stale variable node's context menu in a child Blueprint graph.
-			NewVar.VarType.PinSubCategory.Empty();
-			NewVar.VarType.PinSubCategoryObject = UObject::StaticClass();
-		}
-
-		// if it's a PC_Object, then it should have an associated UClass object
-		check(NewVar.VarType.PinSubCategoryObject.IsValid());
-		const UClass* ClassObject = Cast<UClass>(NewVar.VarType.PinSubCategoryObject.Get());
-		check(ClassObject != NULL);
-
-		if (ClassObject->IsChildOf(AActor::StaticClass()))
-		{
-			// prevent Actor variables from having default values (because Blueprint templates are library elements that can 
-			// bridge multiple levels and different levels might not have the actor that the default is referencing).
-			NewVar.PropertyFlags |= CPF_DisableEditOnTemplate;
-		}
+		PostSetupObjectPinType(Blueprint, NewVar);
 	}
 	NewVar.Category = K2Schema->VR_DefaultCategory;
 	NewVar.DefaultValue = DefaultValue;
@@ -3895,6 +3871,8 @@ bool FBlueprintEditorUtils::AddLocalVariable(UBlueprint* Blueprint, UEdGraph* In
 		NewVar.FriendlyName = FName::NameToDisplayString( NewVar.VarName.ToString(), (NewVar.VarType.PinCategory == K2Schema->PC_Boolean) ? true : false );
 		NewVar.Category = K2Schema->VR_DefaultCategory;
 
+		PostSetupObjectPinType(Blueprint, NewVar);
+
 		FunctionEntryNodes[0]->Modify();
 		FunctionEntryNodes[0]->LocalVariables.Add(NewVar);
 
@@ -4160,6 +4138,28 @@ void FBlueprintEditorUtils::ChangeLocalVariableType(UBlueprint* InBlueprint, con
 				// Mark the Blueprint as structurally modified so we can reconstruct the node successfully
 				FBlueprintEditorUtils::MarkBlueprintAsStructurallyModified(InBlueprint);
 
+				if ((NewPinType.PinCategory == K2Schema->PC_Object) || (NewPinType.PinCategory == K2Schema->PC_Interface))
+				{
+					// if it's a PC_Object, then it should have an associated UClass object
+					if(NewPinType.PinSubCategoryObject.IsValid())
+					{
+						const UClass* ClassObject = Cast<UClass>(NewPinType.PinSubCategoryObject.Get());
+						check(ClassObject != NULL);
+
+						if (ClassObject->IsChildOf(AActor::StaticClass()))
+						{
+							// prevent Actor variables from having default values (because Blueprint templates are library elements that can 
+							// bridge multiple levels and different levels might not have the actor that the default is referencing).
+							Variable.PropertyFlags |= CPF_DisableEditOnTemplate;
+						}
+						else 
+						{
+							// clear the disable-default-value flag that might have been present (if this was an AActor variable before)
+							Variable.PropertyFlags &= ~(CPF_DisableEditOnTemplate);
+						}
+					}
+				}
+
 				// Reconstruct all local variables referencing the modified one
 				for(UK2Node_Variable* VariableNode : VariableNodes)
 				{
@@ -6428,8 +6428,9 @@ bool FBlueprintEditorUtils::PropertyValueFromString(const UProperty* Property, c
 		}
 		else if( Property->IsA(UClassProperty::StaticClass()) )
 		{
-			UClass* FoundClass = FindObject<UClass>(ANY_PACKAGE, *Value);
-			CastChecked<UClassProperty>(Property)->SetPropertyValue_InContainer(DefaultObject, FoundClass);
+			FStringOutputDevice ImportError;
+			const auto EndOfParsedBuff = Property->ImportText(Value.IsEmpty() ? TEXT("()") : *Value, Property->ContainerPtrToValuePtr<uint8>(DefaultObject), 0, NULL, &ImportError);
+			bParseSucceeded = EndOfParsedBuff && ImportError.IsEmpty();
 		}
 		else if( Property->IsA(UObjectPropertyBase::StaticClass()) )
 		{
@@ -6985,5 +6986,38 @@ bool FBlueprintEditorUtils::CheckIfGraphHasLatentFunctions(UEdGraph* InGraph)
 	return Local::CheckIfGraphHasLatentFunctions(InGraph, InspectedGraphList);
 }
 
+void FBlueprintEditorUtils::PostSetupObjectPinType(UBlueprint* InBlueprint, FBPVariableDescription& InOutVarDesc)
+{
+	UEdGraphSchema_K2 const* K2Schema = GetDefault<UEdGraphSchema_K2>();
+	if ((InOutVarDesc.VarType.PinCategory == K2Schema->PC_Object) || (InOutVarDesc.VarType.PinCategory == K2Schema->PC_Interface))
+	{
+		if (InOutVarDesc.VarType.PinSubCategory == K2Schema->PSC_Self)
+		{
+			InOutVarDesc.VarType.PinSubCategory.Empty();
+			InOutVarDesc.VarType.PinSubCategoryObject = *InBlueprint->GeneratedClass;
+		}
+		else if (!InOutVarDesc.VarType.PinSubCategoryObject.IsValid())
+		{
+			// Fall back to UObject if the given type is not valid. This can happen for example if a variable is removed from
+			// a Blueprint parent class along with the variable's type and the user then attempts to recreate the missing variable
+			// through a stale variable node's context menu in a child Blueprint graph.
+			InOutVarDesc.VarType.PinSubCategory.Empty();
+			InOutVarDesc.VarType.PinSubCategoryObject = UObject::StaticClass();
+		}
+
+		// if it's a PC_Object, then it should have an associated UClass object
+		check(InOutVarDesc.VarType.PinSubCategoryObject.IsValid());
+		const UClass* ClassObject = Cast<UClass>(InOutVarDesc.VarType.PinSubCategoryObject.Get());
+		check(ClassObject != NULL);
+
+		if (ClassObject->IsChildOf(AActor::StaticClass()))
+		{
+			// prevent Actor variables from having default values (because Blueprint templates are library elements that can 
+			// bridge multiple levels and different levels might not have the actor that the default is referencing).
+			InOutVarDesc.PropertyFlags |= CPF_DisableEditOnTemplate;
+		}
+	}
+}
+
 #undef LOCTEXT_NAMESPACE
 
diff --git a/Engine/Source/Editor/UnrealEd/Public/Kismet2/BlueprintEditorUtils.h b/Engine/Source/Editor/UnrealEd/Public/Kismet2/BlueprintEditorUtils.h
--- a/Engine/Source/Editor/UnrealEd/Public/Kismet2/BlueprintEditorUtils.h
+++ b/Engine/Source/Editor/UnrealEd/Public/Kismet2/BlueprintEditorUtils.h
@@ -1112,6 +1112,12 @@ protected:
 	 */
 	static void GetLoadedChildBlueprints(UBlueprint* InBlueprint, TArray<UBlueprint*>& OutBlueprints);
 
+	/**
+	 * Validates flags and settings on object pins, keeping them from being given default values and from being in invalid states
+	 *
+	 * @param InOutVarDesc		The variable description to validate
+	 */
+	static void PostSetupObjectPinType(UBlueprint* InBlueprint, FBPVariableDescription& InOutVarDesc);
 public:
 	static FName GetFunctionNameFromClassByGuid(const UClass* InClass, const FGuid FunctionGuid);
 	static bool GetFunctionGuidFromClassByFieldName(const UClass* InClass, const FName FunctionName, FGuid& FunctionGuid);
diff --git a/Engine/Source/Runtime/CoreUObject/Private/UObject/Class.cpp b/Engine/Source/Runtime/CoreUObject/Private/UObject/Class.cpp
--- a/Engine/Source/Runtime/CoreUObject/Private/UObject/Class.cpp
+++ b/Engine/Source/Runtime/CoreUObject/Private/UObject/Class.cpp
@@ -670,6 +670,58 @@ void UStruct::Link(FArchive& Ar, bool bRelinkExistingProperties)
 	*RefLinkPtr = NULL;
 }
 
+void UStruct::InitializeStruct(void* InDest, int32 ArrayDim/* = 1*/) const
+{
+	uint8 *Dest = (uint8*)InDest;
+	check(Dest);
+
+	int32 Stride = GetStructureSize();
+
+	//@todo UE4 optimize
+	FMemory::Memzero(Dest, 1 * Stride);
+
+	bool bHitBase = false;
+	for (UProperty* Property = PropertyLink; Property && !bHitBase; Property = Property->PropertyLinkNext)
+	{
+		if (!Property->IsInContainer(0))
+		{
+			for (int32 ArrayIndex = 0; ArrayIndex < 1; ArrayIndex++)
+			{
+				Property->InitializeValue_InContainer(Dest + ArrayIndex * Stride);
+			}
+		}
+		else
+		{
+			bHitBase = true;
+		}
+	}
+}
+
+void UStruct::DestroyStruct(void* Dest, int32 ArrayDim) const
+{
+	uint8 *Data = (uint8*)Dest;
+	int32 Stride = GetStructureSize();
+
+	bool bHitBase = false;
+	for (UProperty* P = DestructorLink; P  && !bHitBase; P = P->DestructorLinkNext)
+	{
+		if (!P->IsInContainer(0))
+		{
+			if (!P->HasAnyPropertyFlags(CPF_NoDestructor))
+			{
+				for ( int32 ArrayIndex = 0; ArrayIndex < ArrayDim; ArrayIndex++ )
+				{
+					P->DestroyValue_InContainer(Data + ArrayIndex * Stride);
+				}
+			}
+		}
+		else
+		{
+			bHitBase = true;
+		}
+	}
+}
+
 //
 // Serialize all of the class's data that belongs in a particular
 // bin and resides in Data.
@@ -2138,8 +2190,7 @@ void UScriptStruct::CopyScriptStruct(void* InDest, void const* InSrc, int32 Arra
 	}
 }
 
-
-void UScriptStruct::InitializeScriptStruct(void* InDest, int32 ArrayDim) const
+void UScriptStruct::InitializeStruct(void* InDest, int32 ArrayDim) const
 {
 	uint8 *Dest = (uint8*)InDest;
 	check(Dest);
@@ -2232,7 +2283,7 @@ void UScriptStruct::ClearScriptStruct(void* Dest, int32 ArrayDim) const
 
 }
 
-void UScriptStruct::DestroyScriptStruct(void* Dest, int32 ArrayDim) const
+void UScriptStruct::DestroyStruct(void* Dest, int32 ArrayDim) const
 {
 	if (StructFlags & (STRUCT_IsPlainOldData | STRUCT_NoDestructor))
 	{
diff --git a/Engine/Source/Runtime/CoreUObject/Private/UObject/PropertyStruct.cpp b/Engine/Source/Runtime/CoreUObject/Private/UObject/PropertyStruct.cpp
--- a/Engine/Source/Runtime/CoreUObject/Private/UObject/PropertyStruct.cpp
+++ b/Engine/Source/Runtime/CoreUObject/Private/UObject/PropertyStruct.cpp
@@ -395,7 +395,7 @@ void UStructProperty::CopyValuesInternal( void* Dest, void const* Src, int32 Cou
 
 void UStructProperty::InitializeValueInternal( void* InDest ) const
 {
-	Struct->InitializeScriptStruct(InDest, ArrayDim);
+	Struct->InitializeStruct(InDest, ArrayDim);
 }
 
 void UStructProperty::ClearValueInternal( void* Data ) const
@@ -405,7 +405,7 @@ void UStructProperty::ClearValueInternal( void* Data ) const
 
 void UStructProperty::DestroyValueInternal( void* Dest ) const
 {
-	Struct->DestroyScriptStruct(Dest, ArrayDim);
+	Struct->DestroyStruct(Dest, ArrayDim);
 }
 
 /**
diff --git a/Engine/Source/Runtime/CoreUObject/Public/UObject/Class.h b/Engine/Source/Runtime/CoreUObject/Public/UObject/Class.h
--- a/Engine/Source/Runtime/CoreUObject/Public/UObject/Class.h
+++ b/Engine/Source/Runtime/CoreUObject/Public/UObject/Class.h
@@ -303,6 +303,23 @@ public:
 
 	virtual void SerializeTaggedProperties( FArchive& Ar, uint8* Data, UStruct* DefaultsStruct, uint8* Defaults, const UObject* BreakRecursionIfFullyLoad=NULL) const;
 
+	/**
+	 * Initialize a struct over uninitialized memory. This may be done by calling the native constructor or individually initializing properties
+	 *
+	 * @param	Dest		Pointer to memory to initialize
+	 * @param	ArrayDim	Number of elements in the array
+	 * @param	Stride		Stride of the array, If this default (0), then we will pull the size from the struct
+	 */
+	virtual void InitializeStruct(void* Dest, int32 ArrayDim = 1) const;
+	/**
+	 * Destroy a struct in memory. This may be done by calling the native destructor and then the constructor or individually reinitializing properties
+	 *
+	 * @param	Dest		Pointer to memory to destory
+	 * @param	ArrayDim	Number of elements in the array
+	 * @param	Stride		Stride of the array. If this default (0), then we will pull the size from the struct
+	 */
+	virtual void DestroyStruct(void* Dest, int32 ArrayDim = 1) const;
+
 #if WITH_EDITOR
 private:
 	virtual UProperty* CustomFindProperty(const FName Name) const { return NULL; };
@@ -986,6 +1003,8 @@ public:
 
 	// UStruct interface.
 	virtual COREUOBJECT_API void Link(FArchive& Ar, bool bRelinkExistingProperties) override;
+	virtual COREUOBJECT_API void InitializeStruct(void* Dest, int32 ArrayDim = 1) const override;
+	virtual COREUOBJECT_API void DestroyStruct(void* Dest, int32 ArrayDim = 1) const override;
 	// End of UStruct interface.
 
 	/** Stash a CppStructOps for future use 
@@ -1062,14 +1081,6 @@ public:
 	 * @param	Stride		Stride of the array, If this default (0), then we will pull the size from the struct
 	 */
 	COREUOBJECT_API void CopyScriptStruct(void* Dest, void const* Src, int32 ArrayDim = 1) const;
-	/**
-	 * Initialize a struct over uninitialized memory. This may be done by calling the native constructor or individually initializing properties
-	 *
-	 * @param	Dest		Pointer to memory to initialize
-	 * @param	ArrayDim	Number of elements in the array
-	 * @param	Stride		Stride of the array, If this default (0), then we will pull the size from the struct
-	 */
-	COREUOBJECT_API void InitializeScriptStruct(void* Dest, int32 ArrayDim = 1) const;
 	/**
 	 * Reintialize a struct in memory. This may be done by calling the native destructor and then the constructor or individually reinitializing properties
 	 *
@@ -1078,14 +1089,6 @@ public:
 	 * @param	Stride		Stride of the array, only relevant if there more than one element. If this default (0), then we will pull the size from the struct
 	 */
 	void ClearScriptStruct(void* Dest, int32 ArrayDim = 1) const;
-	/**
-	 * Destroy a struct in memory. This may be done by calling the native destructor and then the constructor or individually reinitializing properties
-	 *
-	 * @param	Dest		Pointer to memory to destory
-	 * @param	ArrayDim	Number of elements in the array
-	 * @param	Stride		Stride of the array. If this default (0), then we will pull the size from the struct
-	 */
-	COREUOBJECT_API void DestroyScriptStruct(void* Dest, int32 ArrayDim = 1) const;
 
 	virtual COREUOBJECT_API void RecursivelyPreload();
 };
@@ -1093,19 +1096,20 @@ public:
 class FStructOnScope
 {
 protected:
-	TWeakObjectPtr<const UScriptStruct> ScriptStruct;
+	TWeakObjectPtr<const UStruct> ScriptStruct;
 	uint8* SampleStructMemory;
 
 	FStructOnScope() : SampleStructMemory(NULL) {}
+
 public:
-	FStructOnScope(const UScriptStruct* InScriptStruct)
+	FStructOnScope(const UStruct* InScriptStruct)
 		: ScriptStruct(InScriptStruct)
 		, SampleStructMemory(NULL)
 	{
 		if (ScriptStruct.IsValid())
 		{
 			SampleStructMemory = (uint8*)FMemory::Malloc(ScriptStruct->GetStructureSize());
-			ScriptStruct.Get()->InitializeScriptStruct(SampleStructMemory);
+			ScriptStruct.Get()->InitializeStruct(SampleStructMemory);
 		}
 	}
 
@@ -1113,7 +1117,7 @@ public:
 
 	virtual const uint8* GetStructMemory() const { return SampleStructMemory; }
 
-	virtual const UScriptStruct* GetStruct() const { return ScriptStruct.Get(); }
+	virtual const UStruct* GetStruct() const { return ScriptStruct.Get(); }
 
 	virtual bool IsValid() const { return ScriptStruct.IsValid() && SampleStructMemory; }
 
@@ -1121,7 +1125,7 @@ public:
 	{
 		if (ScriptStruct.IsValid() && SampleStructMemory)
 		{
-			ScriptStruct.Get()->DestroyScriptStruct(SampleStructMemory);
+			ScriptStruct.Get()->DestroyStruct(SampleStructMemory);
 			ScriptStruct = NULL;
 		}
 
@@ -2203,7 +2207,6 @@ protected:
 	virtual UObject* CreateDefaultObject();
 };
 
-
 /**
  * Helper template to call the default constructor for a class
  */
diff --git a/Engine/Source/Runtime/Engine/Private/DataTable.cpp b/Engine/Source/Runtime/Engine/Private/DataTable.cpp
--- a/Engine/Source/Runtime/Engine/Private/DataTable.cpp
+++ b/Engine/Source/Runtime/Engine/Private/DataTable.cpp
@@ -116,7 +116,7 @@ void UDataTable::LoadStructData(FArchive& Ar)
 
 		// Load row data
 		uint8* RowData = (uint8*)FMemory::Malloc(LoadUsingStruct->PropertiesSize);
-		LoadUsingStruct->InitializeScriptStruct(RowData);
+		LoadUsingStruct->InitializeStruct(RowData);
 		// And be sure to call DestroyScriptStruct later
 		LoadUsingStruct->SerializeTaggedProperties(Ar, RowData, LoadUsingStruct, NULL);
 
@@ -220,7 +220,7 @@ void UDataTable::EmptyTable()
 	for (auto RowIt = RowMap.CreateIterator(); RowIt; ++RowIt)
 	{
 		uint8* RowData = RowIt.Value();
-		LoadUsingStruct->DestroyScriptStruct(RowData);
+		LoadUsingStruct->DestroyStruct(RowData);
 		FMemory::Free(RowData);
 	}
 
@@ -579,7 +579,7 @@ TArray<FString> UDataTable::CreateTableFromCSVString(const FString& InString)
 
 		// Allocate data to store information, using UScriptStruct to know its size
 		uint8* RowData = (uint8*)FMemory::Malloc(RowStruct->PropertiesSize);
-		RowStruct->InitializeScriptStruct(RowData);
+		RowStruct->InitializeStruct(RowData);
 		// And be sure to call DestroyScriptStruct later
 
 #if WITH_EDITOR
diff --git a/Engine/Source/Runtime/GameplayAbilities/Private/AbilitySystemBlueprintLibrary.cpp b/Engine/Source/Runtime/GameplayAbilities/Private/AbilitySystemBlueprintLibrary.cpp
--- a/Engine/Source/Runtime/GameplayAbilities/Private/AbilitySystemBlueprintLibrary.cpp
+++ b/Engine/Source/Runtime/GameplayAbilities/Private/AbilitySystemBlueprintLibrary.cpp
@@ -88,7 +88,7 @@ FGameplayAbilityTargetDataHandle UAbilitySystemBlueprintLibrary::FilterTargetDat
 				//Copy the data first, since we don't understand the internals of it
 				UScriptStruct* ScriptStruct = UnfilteredData->GetScriptStruct();
 				FGameplayAbilityTargetData* NewData = (FGameplayAbilityTargetData*)FMemory::Malloc(ScriptStruct->GetCppStructOps()->GetSize());
-				ScriptStruct->InitializeScriptStruct(NewData);
+				ScriptStruct->InitializeStruct(NewData);
 				ScriptStruct->CopyScriptStruct(NewData, UnfilteredData);
 				ReturnDataHandle.Data.Add(TSharedPtr<FGameplayAbilityTargetData>(NewData));
 				if (FilteredActors.Num() < UnfilteredData->GetActors().Num())
@@ -417,4 +417,4 @@ FGameplayEffectSpecHandle UAbilitySystemBlueprintLibrary::AddLinkedGameplayEffec
 	}
 
 	return SpecHandle;
-}
\ No newline at end of file
+}
diff --git a/Engine/Source/Runtime/GameplayAbilities/Private/GameplayAbilityTargetTypes.cpp b/Engine/Source/Runtime/GameplayAbilities/Private/GameplayAbilityTargetTypes.cpp
--- a/Engine/Source/Runtime/GameplayAbilities/Private/GameplayAbilityTargetTypes.cpp
+++ b/Engine/Source/Runtime/GameplayAbilities/Private/GameplayAbilityTargetTypes.cpp
@@ -173,7 +173,7 @@ bool FGameplayAbilityTargetDataHandle::NetSerialize(FArchive& Ar, class UPackage
 				check(!Data[i].IsValid());
 
 				FGameplayAbilityTargetData * NewData = (FGameplayAbilityTargetData*)FMemory::Malloc(ScriptStruct->GetCppStructOps()->GetSize());
-				ScriptStruct->InitializeScriptStruct(NewData);
+				ScriptStruct->InitializeStruct(NewData);
 
 				Data[i] = TSharedPtr<FGameplayAbilityTargetData>(NewData);
 			}
diff --git a/Engine/Source/Runtime/GameplayAbilities/Private/GameplayEffectTypes.cpp b/Engine/Source/Runtime/GameplayAbilities/Private/GameplayEffectTypes.cpp
--- a/Engine/Source/Runtime/GameplayAbilities/Private/GameplayEffectTypes.cpp
+++ b/Engine/Source/Runtime/GameplayAbilities/Private/GameplayEffectTypes.cpp
@@ -147,7 +147,7 @@ bool FGameplayEffectContextHandle::NetSerialize(FArchive& Ar, class UPackageMap*
 			check(!Data.IsValid());
 
 			FGameplayEffectContext * NewData = (FGameplayEffectContext*)FMemory::Malloc(ScriptStruct->GetCppStructOps()->GetSize());
-			ScriptStruct->InitializeScriptStruct(NewData);
+			ScriptStruct->InitializeStruct(NewData);
 
 			Data = TSharedPtr<FGameplayEffectContext>(NewData);
 		}
diff --git a/Engine/Source/Runtime/Messaging/Private/Bus/MessageContext.cpp b/Engine/Source/Runtime/Messaging/Private/Bus/MessageContext.cpp
--- a/Engine/Source/Runtime/Messaging/Private/Bus/MessageContext.cpp
+++ b/Engine/Source/Runtime/Messaging/Private/Bus/MessageContext.cpp
@@ -12,7 +12,7 @@ FMessageContext::~FMessageContext()
 	{
 		if (TypeInfo.IsValid())
 		{
-			TypeInfo->DestroyScriptStruct(Message);
+			TypeInfo->DestroyStruct(Message);
 		}
 
 		FMemory::Free(Message);
