Editor: Fix issue where mode tools outside of the level editor were not being routed CapturedMouseMove

diff --git a/Engine/Source/Editor/UnrealEd/Private/EditorViewportClient.cpp b/Engine/Source/Editor/UnrealEd/Private/EditorViewportClient.cpp
--- a/Engine/Source/Editor/UnrealEd/Private/EditorViewportClient.cpp
+++ b/Engine/Source/Editor/UnrealEd/Private/EditorViewportClient.cpp
@@ -4047,6 +4047,12 @@ void FEditorViewportClient::CapturedMouseMove( FViewport* InViewport, int32 InMo
 {
 	UpdateRequiredCursorVisibility();
 	ApplyRequiredCursorVisibility();
+
+	// Let the current editor mode know about the mouse movement.
+	if (ModeTools->CapturedMouseMove(this, InViewport, InMouseX, InMouseY))
+	{
+		return;
+	}
 }
 
 void FEditorViewportClient::OpenScreenshot( FString SourceFilePath )
diff --git a/Engine/Source/Editor/UnrealEd/Private/LevelEditorViewport.cpp b/Engine/Source/Editor/UnrealEd/Private/LevelEditorViewport.cpp
--- a/Engine/Source/Editor/UnrealEd/Private/LevelEditorViewport.cpp
+++ b/Engine/Source/Editor/UnrealEd/Private/LevelEditorViewport.cpp
@@ -3214,16 +3214,10 @@ FViewportCursorLocation FLevelEditorViewportClient::GetCursorWorldLocationFromMo
  */
 void FLevelEditorViewportClient::CapturedMouseMove( FViewport* InViewport, int32 InMouseX, int32 InMouseY )
 {
-	FEditorViewportClient::CapturedMouseMove(InViewport,InMouseX,InMouseY);
-
 	// Commit to any pending transactions now
 	TrackingTransaction.PromotePendingToActive();
 
-	// Let the current editor mode know about the mouse movement.
-	if (IsLevelEditorClient() && ModeTools->CapturedMouseMove(this, InViewport, InMouseX, InMouseY))
-	{
-		return;
-	}
+	FEditorViewportClient::CapturedMouseMove(InViewport, InMouseX, InMouseY);
 }
 
 
