Prompt to open code files in IDE after being installed by feature pack

diff --git a/Engine/Source/Editor/UnrealEd/Private/Factories/PackFactory.cpp b/Engine/Source/Editor/UnrealEd/Private/Factories/PackFactory.cpp
--- a/Engine/Source/Editor/UnrealEd/Private/Factories/PackFactory.cpp
+++ b/Engine/Source/Editor/UnrealEd/Private/Factories/PackFactory.cpp
@@ -8,6 +8,7 @@
 #include "Runtime/PakFile/Public/IPlatformFilePak.h"
 #include "ISourceControlModule.h"
 #include "SourceControlHelpers.h"
+#include "SourceCodeNavigation.h"
 #include "HotReloadInterface.h"
 #include "AES.h"
 #include "ModuleManager.h"
@@ -274,6 +275,7 @@ UObject* UPackFactory::FactoryCreateBinary
 		TMap<FString, FString> ConfigParameters;
 
 		TArray<FString> WrittenFiles;
+		TArray<FString> WrittenSourceFiles;
 
 		// Process the config files and identify if we have source files
 		for (FPakFile::FFileIterator It(PakFile); It; ++It, ++FileCount)
@@ -422,6 +424,7 @@ UObject* UPackFactory::FactoryCreateBinary
 					if (FFileHelper::SaveStringToFile(SourceContents, *DestFilename))
 					{
 						WrittenFiles.Add(*DestFilename);
+						WrittenSourceFiles.Add(*DestFilename);
 					}
 					else
 					{
@@ -474,13 +477,24 @@ UObject* UPackFactory::FactoryCreateBinary
 		{
 			// If we wrote out source files, kick off the hot reload process
 			// TODO: Ideally we would block on finishing the factory until the hot reload completed
-			if (bContainsSource)
+			if (WrittenSourceFiles.Num() > 0)
 			{
 				IHotReloadInterface& HotReloadSupport = FModuleManager::LoadModuleChecked<IHotReloadInterface>("HotReload");
 				if( !HotReloadSupport.IsCurrentlyCompiling() )
 				{
 					HotReloadSupport.DoHotReloadFromEditor();
 				}
+
+				if (FSlateApplication::Get().SupportsSourceAccess() )
+				{
+					// Code successfully added, notify the user and ask about opening the IDE now
+					const FText Message = NSLOCTEXT("PackFactory", "CodeAdded", "Added source file(s). Would you like to edit the code now?");
+					if ( FMessageDialog::Open(EAppMsgType::YesNo, Message) == EAppReturnType::Yes )
+					{
+						FSourceCodeNavigation::OpenSourceFiles(WrittenSourceFiles);
+					}
+				}
+
 			}
 
 			// Find an asset to return as the file to be selected in the content browser
