Fixed up merged cook cooking too much when AlwaysCookMaps is specified.

diff --git a/Engine/Source/Editor/UnrealEd/Private/CookOnTheFlyServer.cpp b/Engine/Source/Editor/UnrealEd/Private/CookOnTheFlyServer.cpp
--- a/Engine/Source/Editor/UnrealEd/Private/CookOnTheFlyServer.cpp
+++ b/Engine/Source/Editor/UnrealEd/Private/CookOnTheFlyServer.cpp
@@ -2186,22 +2186,30 @@ void UCookOnTheFlyServer::AddFileToCook( TArray<FString>& InOutFilesToCook, cons
 
 void UCookOnTheFlyServer::CollectFilesToCook(TArray<FString>& FilesInPath, const TArray<FString> &CookMaps, const TArray<FString> &InCookDirectories, const TArray<FString> &CookCultures, const TArray<FString> &IniMapSections, bool bCookAll, bool bMapsOnly, bool bNoDev)
 {
-	TArray<FString> MapList;
 	TArray<FString> CookDirectories = InCookDirectories;
-	// Add the default map section
-	GEditor->LoadMapListFromIni(TEXT("AlwaysCookMaps"), MapList);
-
-	for ( const auto &IniMapSection : IniMapSections )
+	
+	if ( IsRealtimeMode() == false )
 	{
-		GEditor->LoadMapListFromIni(*IniMapSection, MapList);
-	}
+		TArray<FString> MapList;
+		// Add the default map section
+		GEditor->LoadMapListFromIni(TEXT("AlwaysCookMaps"), MapList);
 
-	// Add any map sections specified on command line
-	/*GEditor->ParseMapSectionIni(*Params, MapList);
-	for (int32 MapIdx = 0; MapIdx < MapList.Num(); MapIdx++)
+		for ( const auto &IniMapSection : IniMapSections )
+		{
+			GEditor->LoadMapListFromIni(*IniMapSection, MapList);
+		}
+
+		for (int32 MapIdx = 0; MapIdx < MapList.Num(); MapIdx++)
+		{
+			AddFileToCook( FilesInPath, MapList[MapIdx]);
+		}
+	}
+	else
 	{
-		AddFileToCook( FilesInPath, MapList[MapIdx]);
-	}*/
+		// if we are cooking in the editor then don't include the always cook maps list because we just want to cook whats in memory
+		check( CookMaps.Num() > 0 );
+	}
+	
 
 
 	// Also append any cookdirs from the project ini files; these dirs are relative to the game content directory
