When calling FBlueprintEditorUtils::RemoveGraph, the graph was not being removed from LastEditedDocuments, which ran the risk of deleted Graphs being referenced and attempting to open them when opening up the Blueprint.

diff --git a/Engine/Source/Editor/UnrealEd/Private/Kismet2/BlueprintEditorUtils.cpp b/Engine/Source/Editor/UnrealEd/Private/Kismet2/BlueprintEditorUtils.cpp
--- a/Engine/Source/Editor/UnrealEd/Private/Kismet2/BlueprintEditorUtils.cpp
+++ b/Engine/Source/Editor/UnrealEd/Private/Kismet2/BlueprintEditorUtils.cpp
@@ -1945,7 +1945,17 @@ void FBlueprintEditorUtils::RemoveGraph(UBlueprint* Blueprint, class UEdGraph* G
 			Blueprint->DelegateSignatureGraphs.Remove( GraphToRemove );
 			Blueprint->FunctionGraphs.Remove( GraphToRemove );
 			Blueprint->UbergraphPages.Remove( GraphToRemove );
-			Blueprint->LastEditedDocuments.Remove( GraphToRemove );
+
+			// Can't just call Remove, the object is wrapped in a struct
+			for(int EditedDocIdx = 0; EditedDocIdx < Blueprint->LastEditedDocuments.Num(); ++EditedDocIdx)
+			{
+				if(Blueprint->LastEditedDocuments[EditedDocIdx].EditedObject == GraphToRemove)
+				{
+					Blueprint->LastEditedDocuments.RemoveAt(EditedDocIdx);
+					break;
+				}
+			}
+
 			if(Blueprint->MacroGraphs.Remove( GraphToRemove ) > 0 ) 
 			{
 				//removes all macro nodes using this macro graph
