Gathering text from assets will no longer access elements of a native fixed-sized FText array improperly.
Gathering text from assets will no longer compile blueprints while loading packages.

diff --git a/Engine/Source/Editor/UnrealEd/Private/Commandlets/GatherTextFromAssetsCommandlet.cpp b/Engine/Source/Editor/UnrealEd/Private/Commandlets/GatherTextFromAssetsCommandlet.cpp
--- a/Engine/Source/Editor/UnrealEd/Private/Commandlets/GatherTextFromAssetsCommandlet.cpp
+++ b/Engine/Source/Editor/UnrealEd/Private/Commandlets/GatherTextFromAssetsCommandlet.cpp
@@ -431,7 +431,8 @@ namespace
 					FText* Text;
 					if(Property->ArrayDim > 1)
 					{
-						Text = reinterpret_cast<FText*>(Property->ContainerPtrToValuePtr<void>(ValueAddress, i));
+						uint8* ArrayValueAddress = reinterpret_cast<uint8*>(ValueAddress);
+						Text = reinterpret_cast<FText*>(ArrayValueAddress + Property->ElementSize * i));
 					}
 					else
 					{
@@ -692,6 +693,8 @@ void UGatherTextFromAssetsCommandlet::ProcessTextProperty(UTextProperty* InTextP
 
 int32 UGatherTextFromAssetsCommandlet::Main(const FString& Params)
 {
+	TGuardValue<bool> DisableCompileOnLoad(GForceDisableBlueprintCompileOnLoad, true);
+
 	// Parse command line.
 	TArray<FString> Tokens;
 	TArray<FString> Switches;
