That last CL I submitted totally did not have an extra parenthesis.

diff --git a/Engine/Source/Editor/UnrealEd/Private/Commandlets/GatherTextFromAssetsCommandlet.cpp b/Engine/Source/Editor/UnrealEd/Private/Commandlets/GatherTextFromAssetsCommandlet.cpp
--- a/Engine/Source/Editor/UnrealEd/Private/Commandlets/GatherTextFromAssetsCommandlet.cpp
+++ b/Engine/Source/Editor/UnrealEd/Private/Commandlets/GatherTextFromAssetsCommandlet.cpp
@@ -432,7 +432,7 @@ namespace
 					if(Property->ArrayDim > 1)
 					{
 						uint8* ArrayValueAddress = reinterpret_cast<uint8*>(ValueAddress);
-						Text = reinterpret_cast<FText*>(ArrayValueAddress + Property->ElementSize * i));
+						Text = reinterpret_cast<FText*>(ArrayValueAddress + Property->ElementSize * i);
 					}
 					else
 					{
