Removed some cooking log spam.

diff --git a/Engine/Source/Editor/UnrealEd/Private/CookOnTheFlyServer.cpp b/Engine/Source/Editor/UnrealEd/Private/CookOnTheFlyServer.cpp
--- a/Engine/Source/Editor/UnrealEd/Private/CookOnTheFlyServer.cpp
+++ b/Engine/Source/Editor/UnrealEd/Private/CookOnTheFlyServer.cpp
@@ -971,10 +971,12 @@ uint32 UCookOnTheFlyServer::TickCookOnTheSide( const float TimeSlice, uint32 &Co
 				SCOPE_TIMER(LoadPackage);
 				Package = LoadPackage( NULL, *BuildFilename, LOAD_None );
 			}
+#if DEBUG_COOKONTHEFLY
 			else
 			{
 				UE_LOG(LogCook, Display, TEXT("Package already loaded %s avoiding reload"), *BuildFilename );
 			}
+#endif
 
 			if( Package == NULL )
 			{
@@ -1040,7 +1042,9 @@ uint32 UCookOnTheFlyServer::TickCookOnTheSide( const float TimeSlice, uint32 &Co
 					//  mark the original request as processed (if this isn't actually the file they were requesting then it will fail)
 					//	and then also save our new request as processed so we don't do it again
 					PackagesToSave.AddUnique( Package );
+#if DEBUG_COOKONTHEFLY
 					UE_LOG( LogCook, Display, TEXT("Request for %s received going to save %s"), *BuildFilename, *PackageFilename );
+#endif
 					// CookedPackages.Add( ToBuild );
 					
 					// ToBuild.SetFilename( PackageFilename );
@@ -1060,7 +1064,9 @@ uint32 UCookOnTheFlyServer::TickCookOnTheSide( const float TimeSlice, uint32 &Co
 			// this could also happen if the source file doesn't exist which is often as we request files with different extensions when we are searching for files
 			// just return that we processed the cook request
 			// the network file manager will then handle the missing file and search somewhere else
+#if DEBUG_COOKONTHEFLY
 			UE_LOG(LogCook, Display, TEXT("Not cooking package %s"), *ToBuild.GetFilename().ToString());
+#endif
 			CookedPackages.Add( FFilePlatformRequest( ToBuild.GetFilename(), TargetPlatformNames ) );
 			continue;
 		}
@@ -1104,7 +1110,7 @@ uint32 UCookOnTheFlyServer::TickCookOnTheSide( const float TimeSlice, uint32 &Co
 					Obj->BeginCacheForCookedPlatformData( TargetPlatform );
 					if ( Obj->IsCachedCookedPlatformDataLoaded(TargetPlatform) == false )
 					{
-#if DEBUG_COOKONTHEFLY || 1
+#if DEBUG_COOKONTHEFLY
 						UE_LOG(LogCook, Display, TEXT("Object %s isn't cached yet"), *Obj->GetFullName());
 #endif
 						bIsAllDataCached = false;
@@ -1112,7 +1118,7 @@ uint32 UCookOnTheFlyServer::TickCookOnTheSide( const float TimeSlice, uint32 &Co
 
 					if ( Timer.IsTimeUp() && IsRealtimeMode() )
 					{
-#if DEBUG_COOKONTHEFLY || 1
+#if DEBUG_COOKONTHEFLY
 						UE_LOG(LogCook, Display, TEXT("Object %s took too long to cache"), *Obj->GetFullName());
 #endif
 						bIsAllDataCached = false;
