Landscape exports textures in correct color space UE-6379
Added export of landscape metallic maps and option for baking foliage into texture

diff --git a/Engine/Source/Developer/MeshUtilities/Private/MeshUtilities.cpp b/Engine/Source/Developer/MeshUtilities/Private/MeshUtilities.cpp
--- a/Engine/Source/Developer/MeshUtilities/Private/MeshUtilities.cpp
+++ b/Engine/Source/Developer/MeshUtilities/Private/MeshUtilities.cpp
@@ -3604,7 +3604,20 @@ void FMeshUtilities::CreateProxyMesh(
 			UniqueMaterials.Last().DiffuseSize = FIntPoint(1024, 1024);
 			// FIXME: Landscape material exporter currently renders world space normal map, so it can't be merged with other meshes normal maps
 			UniqueMaterials.Last().NormalSize = FIntPoint::ZeroValue;
-			MaterialExportUtils::ExportMaterial(Landscape, UniqueMaterials.Last());
+
+			// Use only landscape primitives for texture flattening
+			TSet<FPrimitiveComponentId> PrimitivesToHide;
+			for (TObjectIterator<UPrimitiveComponent> It; It; ++It)
+			{
+				UPrimitiveComponent* PrimitiveComp = *It;
+				const bool bTargetPrim = PrimitiveComp->GetOuter() == Landscape;
+				if (!bTargetPrim && PrimitiveComp->IsRegistered() && PrimitiveComp->SceneProxy)
+				{
+					PrimitivesToHide.Add(PrimitiveComp->SceneProxy->GetPrimitiveComponentId());
+				}
+			}
+			
+			MaterialExportUtils::ExportMaterial(Landscape, PrimitivesToHide, UniqueMaterials.Last());
 			
 			//Store the bounds for each component
 			ProxyBounds+= Landscape->GetComponentsBoundingBox(true);
diff --git a/Engine/Source/Editor/UnrealEd/Private/EditorExporters.cpp b/Engine/Source/Editor/UnrealEd/Private/EditorExporters.cpp
--- a/Engine/Source/Editor/UnrealEd/Private/EditorExporters.cpp
+++ b/Engine/Source/Editor/UnrealEd/Private/EditorExporters.cpp
@@ -2698,14 +2698,16 @@ namespace MaterialExportUtils
 		const FMatrix& ViewMatrix, 
 		const FMatrix& ProjectionMatrix,  
 		const TSet<FPrimitiveComponentId>& HiddenPrimitives, 
-		FIntPoint TargetSize, 
+		FIntPoint TargetSize,
+		float TargetGamma,
 		TArray<FColor>& OutSamples)
 	{
 		UTextureRenderTarget2D* RenderTargetTexture = new UTextureRenderTarget2D(FObjectInitializer());
 		check(RenderTargetTexture);
 		RenderTargetTexture->AddToRoot();
-		RenderTargetTexture->ClearColor = FLinearColor(0.0f, 0.0f, 0.0f, 0.0f);
-		RenderTargetTexture->InitCustomFormat(TargetSize.X, TargetSize.Y, PF_B8G8R8A8, true);
+		RenderTargetTexture->ClearColor = FLinearColor::Transparent;
+		RenderTargetTexture->TargetGamma = TargetGamma;
+		RenderTargetTexture->InitCustomFormat(TargetSize.X, TargetSize.Y, PF_FloatRGBA, false);
 		FTextureRenderTargetResource* RenderTargetResource = RenderTargetTexture->GameThread_GetRenderTargetResource();
 
 		FSceneViewFamilyContext ViewFamily(
@@ -2713,17 +2715,10 @@ namespace MaterialExportUtils
 				.SetWorldTimes(FApp::GetCurrentTime() - GStartTime, FApp::GetDeltaTime(), FApp::GetCurrentTime() - GStartTime)
 			);
 
-		ViewFamily.EngineShowFlags.DisableAdvancedFeatures();
-		ViewFamily.EngineShowFlags.MotionBlur = 0;
-		ViewFamily.EngineShowFlags.Lighting = 0;
-		ViewFamily.EngineShowFlags.PostProcessing = 0;
-		ViewFamily.EngineShowFlags.LightFunctions = 0;
-		ViewFamily.EngineShowFlags.DynamicShadows = 0;
-		ViewFamily.EngineShowFlags.Atmosphere = 0;
 		// To enable visualization mode
-		ViewFamily.EngineShowFlags.PostProcessing = 1;
-		ViewFamily.EngineShowFlags.PostProcessMaterial = 1;
-		ViewFamily.EngineShowFlags.VisualizeBuffer = 1; 
+		ViewFamily.EngineShowFlags.SetPostProcessing(true);
+		ViewFamily.EngineShowFlags.SetVisualizeBuffer(true);
+		ViewFamily.EngineShowFlags.SetTonemapper(false);
 
 		FSceneViewInitOptions ViewInitOptions;
 		ViewInitOptions.SetViewRectangle(FIntRect(0, 0, TargetSize.X, TargetSize.Y));
@@ -2737,32 +2732,21 @@ namespace MaterialExportUtils
 		ViewFamily.Views.Add(NewView);
 					
 		FCanvas Canvas(RenderTargetResource, NULL, FApp::GetCurrentTime() - GStartTime, FApp::GetDeltaTime(), FApp::GetCurrentTime() - GStartTime, Scene->GetFeatureLevel());
-		Canvas.Clear(FLinearColor::Black);
+		Canvas.Clear(FLinearColor::Transparent);
 		GetRendererModule().BeginRenderingViewFamily(&Canvas, &ViewFamily);
 
-		ENQUEUE_UNIQUE_RENDER_COMMAND_ONEPARAMETER(
-			UpdateThumbnailRTCommand,
-			FTextureRenderTargetResource*, RenderTargetResource, RenderTargetResource,
-		{
-			// Copy (resolve) the rendered thumbnail from the render target to its texture
-			RHICmdList.CopyToResolveTarget(
-				RenderTargetResource->GetRenderTargetTexture(),		// Source texture
-				RenderTargetResource->TextureRHI,					// Dest texture
-				false,												// Do we need the source image content again?
-				FResolveParams() );									// Resolve parameters
-		});
-
 		// Copy the contents of the remote texture to system memory
-		// NOTE: OutRawImageData must be a preallocated buffer!
 		OutSamples.Init(TargetSize.X*TargetSize.Y);
-		RenderTargetResource->ReadPixelsPtr(OutSamples.GetData(), FReadSurfaceDataFlags(), FIntRect(0, 0, TargetSize.X, TargetSize.Y));
+		FReadSurfaceDataFlags ReadSurfaceDataFlags;
+		ReadSurfaceDataFlags.SetLinearToGamma(false);
+		RenderTargetResource->ReadPixelsPtr(OutSamples.GetData(), ReadSurfaceDataFlags, FIntRect(0, 0, TargetSize.X, TargetSize.Y));
 		FlushRenderingCommands();
 					
 		RenderTargetTexture->RemoveFromRoot();
 		RenderTargetTexture = nullptr;
 	}
-	
-	bool ExportMaterial(ALandscapeProxy* InLandscape, FFlattenMaterial& OutFlattenMaterial)
+		
+	bool ExportMaterial(ALandscapeProxy* InLandscape, const TSet<FPrimitiveComponentId>& HiddenPrimitives, FFlattenMaterial& OutFlattenMaterial)
 	{
 		check(InLandscape);
 
@@ -2772,21 +2756,6 @@ namespace MaterialExportUtils
 		FVector LandscapeCenter = InLandscape->GetTransform().TransformPosition(MidPoint);
 		FVector LandscapeExtent = FVector(LandscapeRect.Size(), 0.f)*InLandscape->GetActorScale()*0.5f; 
 
-		// We need to hide all primitives except target landscape
-		TSet<FPrimitiveComponentId> PrimitivesToHide;
-		for (TObjectIterator<UPrimitiveComponent> It; It; ++It)
-		{
-			UPrimitiveComponent* PrimitiveComp = *It;
-
-			const bool bTargetPrim = PrimitiveComp->GetOuter() == InLandscape && 
-										PrimitiveComp->IsA<ULandscapeComponent>();
-					
-			if (!bTargetPrim && PrimitiveComp->IsRegistered() && PrimitiveComp->SceneProxy)
-			{
-				PrimitivesToHide.Add(PrimitiveComp->SceneProxy->GetPrimitiveComponentId());
-			}
-		}
-
 		FMatrix ViewMatrix = FTranslationMatrix(-LandscapeCenter);
 		ViewMatrix*= FInverseRotationMatrix(InLandscape->GetActorRotation());
 		ViewMatrix*= FMatrix(FPlane(1,	0,	0,	0),
@@ -2802,14 +2771,15 @@ namespace MaterialExportUtils
 			ZOffset);
 
 		FSceneInterface* Scene = InLandscape->GetWorld()->Scene;
-		
+						
 		// Render diffuse texture using BufferVisualizationMode=BaseColor
 		if (OutFlattenMaterial.DiffuseSize.X > 0 && 
 			OutFlattenMaterial.DiffuseSize.Y > 0)
 		{
 			static const FName BaseColorName("BaseColor");
-			RenderSceneToTexture(Scene, BaseColorName, ViewMatrix, ProjectionMatrix, PrimitivesToHide, 
-				OutFlattenMaterial.DiffuseSize, OutFlattenMaterial.DiffuseSamples);
+			const float BaseColorGamma = 1.0f; // BaseColor material already do transformation to gamma space
+			RenderSceneToTexture(Scene, BaseColorName, ViewMatrix, ProjectionMatrix, HiddenPrimitives, 
+				OutFlattenMaterial.DiffuseSize, BaseColorGamma, OutFlattenMaterial.DiffuseSamples);
 		}
 
 		// Render normal map using BufferVisualizationMode=WorldNormal
@@ -2818,8 +2788,19 @@ namespace MaterialExportUtils
 			OutFlattenMaterial.NormalSize.Y > 0)
 		{
 			static const FName WorldNormalName("WorldNormal");
-			RenderSceneToTexture(Scene, WorldNormalName, ViewMatrix, ProjectionMatrix, PrimitivesToHide, 
-				OutFlattenMaterial.NormalSize, OutFlattenMaterial.NormalSamples);
+			const float NormalColorGamma = 1.0f; // Dump normal texture in linear space
+			RenderSceneToTexture(Scene, WorldNormalName, ViewMatrix, ProjectionMatrix, HiddenPrimitives, 
+				OutFlattenMaterial.NormalSize, NormalColorGamma, OutFlattenMaterial.NormalSamples);
+		}
+
+		// Render metallic map using BufferVisualizationMode=Metallic
+		if (OutFlattenMaterial.MetallicSize.X > 0 && 
+			OutFlattenMaterial.MetallicSize.Y > 0)
+		{
+			static const FName MetallicName("Metallic");
+			const float MetallicColorGamma = 1.0f; // Dump metallic texture in linear space
+			RenderSceneToTexture(Scene, MetallicName, ViewMatrix, ProjectionMatrix, HiddenPrimitives, 
+				OutFlattenMaterial.MetallicSize, MetallicColorGamma, OutFlattenMaterial.MetallicSamples);
 		}
 
 		// Render roughness map using BufferVisualizationMode=Roughness
@@ -2827,30 +2808,32 @@ namespace MaterialExportUtils
 			OutFlattenMaterial.RoughnessSize.Y > 0)
 		{
 			static const FName RoughnessName("Roughness");
-			RenderSceneToTexture(Scene, RoughnessName, ViewMatrix, ProjectionMatrix, PrimitivesToHide, 
-				OutFlattenMaterial.RoughnessSize, OutFlattenMaterial.RoughnessSamples);
+			const float RoughnessColorGamma = 2.2f; // Roughness material powers color by 2.2, transform it back to linear
+			RenderSceneToTexture(Scene, RoughnessName, ViewMatrix, ProjectionMatrix, HiddenPrimitives, 
+				OutFlattenMaterial.RoughnessSize, RoughnessColorGamma, OutFlattenMaterial.RoughnessSamples);
 		}
 
-		// Render roughness map using BufferVisualizationMode=Specular
+		// Render specular map using BufferVisualizationMode=Specular
 		if (OutFlattenMaterial.SpecularSize.X > 0 && 
 			OutFlattenMaterial.SpecularSize.Y > 0)
 		{
 			static const FName SpecularName("Specular");
-			RenderSceneToTexture(Scene, SpecularName, ViewMatrix, ProjectionMatrix, PrimitivesToHide, 
-				OutFlattenMaterial.SpecularSize, OutFlattenMaterial.SpecularSamples);
+			const float SpecularColorGamma = 1.0f; // Dump specular texture in linear space
+			RenderSceneToTexture(Scene, SpecularName, ViewMatrix, ProjectionMatrix, HiddenPrimitives, 
+				OutFlattenMaterial.SpecularSize, SpecularColorGamma, OutFlattenMaterial.SpecularSamples);
 		}
 				
 		OutFlattenMaterial.MaterialId = InLandscape->GetLandscapeGuid();
 		return true;
 	}
 
-	UTexture2D* CreateTexture(UPackage* Outer, const FString& AssetName, FIntPoint Size, const TArray<FColor>& Samples, TextureCompressionSettings CompressionSettings, TextureGroup LODGroup, EObjectFlags Flags)
+	UTexture2D* CreateTexture(UPackage* Outer, const FString& AssetName, FIntPoint Size, const TArray<FColor>& Samples, TextureCompressionSettings CompressionSettings, TextureGroup LODGroup, EObjectFlags Flags, bool bSRGB)
 	{
 		FCreateTexture2DParameters TexParams;
 		TexParams.bUseAlpha = false;
 		TexParams.CompressionSettings = CompressionSettings;
 		TexParams.bDeferCompression = true;
-		TexParams.bSRGB = false;
+		TexParams.bSRGB = bSRGB;
 
 		if (Outer == nullptr)
 		{
@@ -2886,85 +2869,123 @@ namespace MaterialExportUtils
 		Material->TwoSided = false;
 		Material->SetShadingModel(MSM_DefaultLit);
 
-		// Set BaseColor
+		int32 MaterialNodeY = -150;
+		int32 MaterialNodeStepY = 180;
+
+		// BaseColor
 		if (InFlattenMaterial.DiffuseSamples.Num() > 1)
 		{
 			const FString AssetName = TEXT("T_") + AssetBaseName + TEXT("_D");
-			UTexture2D* Texture = CreateTexture(Outer, AssetName, InFlattenMaterial.DiffuseSize, InFlattenMaterial.DiffuseSamples, TC_Default, TEXTUREGROUP_World, Flags);
+			const bool bSRGB = true;
+			UTexture2D* Texture = CreateTexture(Outer, AssetName, InFlattenMaterial.DiffuseSize, InFlattenMaterial.DiffuseSamples, TC_Default, TEXTUREGROUP_World, Flags, bSRGB);
 			
 			UMaterialExpressionTextureSample* BasecolorExpression = ConstructObject<UMaterialExpressionTextureSample>(UMaterialExpressionTextureSample::StaticClass(), Material);
 			BasecolorExpression->Texture = Texture;
-			BasecolorExpression->SamplerType = EMaterialSamplerType::SAMPLERTYPE_LinearColor;
+			BasecolorExpression->SamplerType = EMaterialSamplerType::SAMPLERTYPE_Color;
 			BasecolorExpression->MaterialExpressionEditorX = -400;
-			BasecolorExpression->MaterialExpressionEditorY = -150;
+			BasecolorExpression->MaterialExpressionEditorY = MaterialNodeY;
 			Material->Expressions.Add(BasecolorExpression);
 			Material->BaseColor.Expression = BasecolorExpression;
+
+			MaterialNodeY+= MaterialNodeStepY;
 		}
 
-		// Set Normal
-		if (InFlattenMaterial.NormalSamples.Num() > 1)
+		// Metallic
+		if (InFlattenMaterial.RoughnessSamples.Num() > 1)
 		{
-			const FString AssetName = TEXT("T_") + AssetBaseName + TEXT("_N");
-			UTexture2D* Texture = CreateTexture(Outer, AssetName, InFlattenMaterial.NormalSize, InFlattenMaterial.NormalSamples, TC_Normalmap, TEXTUREGROUP_WorldNormalMap, Flags);
+			const FString AssetName = TEXT("T_") + AssetBaseName + TEXT("_M");
+			const bool bSRGB = false;
+			UTexture2D* Texture = CreateTexture(Outer, AssetName, InFlattenMaterial.MetallicSize, InFlattenMaterial.MetallicSamples, TC_Grayscale, TEXTUREGROUP_World, Flags, bSRGB);
 			
-			UMaterialExpressionTextureSample* NormalExpression = ConstructObject<UMaterialExpressionTextureSample>(UMaterialExpressionTextureSample::StaticClass(), Material);
-			NormalExpression->Texture = Texture;
-			NormalExpression->SamplerType = EMaterialSamplerType::SAMPLERTYPE_Normal;
-			NormalExpression->MaterialExpressionEditorX = -400;
-			NormalExpression->MaterialExpressionEditorY = 0;
-			Material->Expressions.Add(NormalExpression);
-			Material->Normal.Expression = NormalExpression;
+			UMaterialExpressionTextureSample* MetallicExpression = ConstructObject<UMaterialExpressionTextureSample>(UMaterialExpressionTextureSample::StaticClass(), Material);
+			MetallicExpression->Texture = Texture;
+			MetallicExpression->SamplerType = EMaterialSamplerType::SAMPLERTYPE_LinearGrayscale;
+			MetallicExpression->MaterialExpressionEditorX = -400;
+			MetallicExpression->MaterialExpressionEditorY = MaterialNodeY;
+			Material->Expressions.Add(MetallicExpression);
+			Material->Metallic.Expression = MetallicExpression;
+
+			MaterialNodeY+= MaterialNodeStepY;
+		}
+		else
+		{
+			// Set Metallic to constant
+			float Metallic = InFlattenMaterial.MetallicSamples.Num() ? InFlattenMaterial.MetallicSamples[0].R : 0.0f;
+			UMaterialExpressionConstant* MetallicExpression = ConstructObject<UMaterialExpressionConstant>(UMaterialExpressionConstant::StaticClass(), Material);
+			MetallicExpression->R = Metallic;
+			MetallicExpression->MaterialExpressionEditorX = -400;
+			MetallicExpression->MaterialExpressionEditorY = MaterialNodeY;
+			Material->Expressions.Add(MetallicExpression);
+			Material->Metallic.Expression = MetallicExpression;
+
+			MaterialNodeY+= MaterialNodeStepY;
 		}
 
+		// Specular
 		if (InFlattenMaterial.SpecularSamples.Num() > 1)
 		{
 			const FString AssetName = TEXT("T_") + AssetBaseName + TEXT("_S");
-			UTexture2D* Texture = CreateTexture(Outer, AssetName, InFlattenMaterial.SpecularSize, InFlattenMaterial.SpecularSamples, TC_Default, TEXTUREGROUP_World, Flags);
+			const bool bSRGB = false;
+			UTexture2D* Texture = CreateTexture(Outer, AssetName, InFlattenMaterial.SpecularSize, InFlattenMaterial.SpecularSamples, TC_Default, TEXTUREGROUP_World, Flags, bSRGB);
 			
 			UMaterialExpressionTextureSample* SpecularExpression = ConstructObject<UMaterialExpressionTextureSample>(UMaterialExpressionTextureSample::StaticClass(), Material);
 			SpecularExpression->Texture = Texture;
 			SpecularExpression->SamplerType = EMaterialSamplerType::SAMPLERTYPE_LinearColor;
 			SpecularExpression->MaterialExpressionEditorX = -400;
-			SpecularExpression->MaterialExpressionEditorY = 300;
+			SpecularExpression->MaterialExpressionEditorY = MaterialNodeY;
 			Material->Expressions.Add(SpecularExpression);
 			Material->Specular.Expression = SpecularExpression;
+
+			MaterialNodeY+= MaterialNodeStepY;
 		}
-		
-		// Set Roughness
+
+		// Roughness
 		if (InFlattenMaterial.RoughnessSamples.Num() > 1)
 		{
 			const FString AssetName = TEXT("T_") + AssetBaseName + TEXT("_R");
-			UTexture2D* Texture = CreateTexture(Outer, AssetName, InFlattenMaterial.RoughnessSize, InFlattenMaterial.RoughnessSamples, TC_Default, TEXTUREGROUP_World, Flags);
+			const bool bSRGB = false;
+			UTexture2D* Texture = CreateTexture(Outer, AssetName, InFlattenMaterial.RoughnessSize, InFlattenMaterial.RoughnessSamples, TC_Grayscale, TEXTUREGROUP_World, Flags, bSRGB);
 			
 			UMaterialExpressionTextureSample* RoughnessExpression = ConstructObject<UMaterialExpressionTextureSample>(UMaterialExpressionTextureSample::StaticClass(), Material);
 			RoughnessExpression->Texture = Texture;
-			RoughnessExpression->SamplerType = EMaterialSamplerType::SAMPLERTYPE_LinearColor;
-			RoughnessExpression->MaterialExpressionEditorX = -250;
-			RoughnessExpression->MaterialExpressionEditorY = 150;
+			RoughnessExpression->SamplerType = EMaterialSamplerType::SAMPLERTYPE_LinearGrayscale;
+			RoughnessExpression->MaterialExpressionEditorX = -400;
+			RoughnessExpression->MaterialExpressionEditorY = MaterialNodeY;
 			Material->Expressions.Add(RoughnessExpression);
 			Material->Roughness.Expression = RoughnessExpression;
+
+			MaterialNodeY+= MaterialNodeStepY;
 		}
 		else
 		{
-			float Roughness = InFlattenMaterial.RoughnessSamples.Num() ? InFlattenMaterial.RoughnessSamples[0].R : 0.8f;
 			// Set Roughness to constant
+			float Roughness = InFlattenMaterial.RoughnessSamples.Num() ? InFlattenMaterial.RoughnessSamples[0].R : 0.8f;
 			UMaterialExpressionConstant* RoughnessExpression = ConstructObject<UMaterialExpressionConstant>(UMaterialExpressionConstant::StaticClass(), Material);
 			RoughnessExpression->R = Roughness;
-			RoughnessExpression->MaterialExpressionEditorX = -250;
-			RoughnessExpression->MaterialExpressionEditorY = 150;
+			RoughnessExpression->MaterialExpressionEditorX = -400;
+			RoughnessExpression->MaterialExpressionEditorY = MaterialNodeY;
 			Material->Expressions.Add(RoughnessExpression);
 			Material->Roughness.Expression = RoughnessExpression;
+
+			MaterialNodeY+= MaterialNodeStepY;
 		}
 
-		//Set Metallic as a default constant
+		// Normal
+		if (InFlattenMaterial.NormalSamples.Num() > 1)
 		{
-			UMaterialExpressionConstant* MetallicExpression = ConstructObject<UMaterialExpressionConstant>(UMaterialExpressionConstant::StaticClass(), Material);
-			MetallicExpression->R = 0.0f;
-			Material->Expressions.Add(MetallicExpression);
-			Material->Metallic.Expression = MetallicExpression;
+			const FString AssetName = TEXT("T_") + AssetBaseName + TEXT("_N");
+			const bool bSRGB = false;
+			UTexture2D* Texture = CreateTexture(Outer, AssetName, InFlattenMaterial.NormalSize, InFlattenMaterial.NormalSamples, TC_Normalmap, TEXTUREGROUP_WorldNormalMap, Flags, bSRGB);
+			
+			UMaterialExpressionTextureSample* NormalExpression = ConstructObject<UMaterialExpressionTextureSample>(UMaterialExpressionTextureSample::StaticClass(), Material);
+			NormalExpression->Texture = Texture;
+			NormalExpression->SamplerType = EMaterialSamplerType::SAMPLERTYPE_Normal;
+			NormalExpression->MaterialExpressionEditorX = -400;
+			NormalExpression->MaterialExpressionEditorY = MaterialNodeY;
+			Material->Expressions.Add(NormalExpression);
+			Material->Normal.Expression = NormalExpression;
 
-			MetallicExpression->MaterialExpressionEditorX = -250;
-			MetallicExpression->MaterialExpressionEditorY = 0;
+			MaterialNodeY+= MaterialNodeStepY;
 		}
 							
 		Material->PostEditChange();
diff --git a/Engine/Source/Editor/UnrealEd/Public/MaterialExportUtils.h b/Engine/Source/Editor/UnrealEd/Public/MaterialExportUtils.h
--- a/Engine/Source/Editor/UnrealEd/Public/MaterialExportUtils.h
+++ b/Engine/Source/Editor/UnrealEd/Public/MaterialExportUtils.h
@@ -11,6 +11,7 @@ namespace MaterialExportUtils
 		FFlattenMaterial()
 			: DiffuseSize(512, 512)
 			, NormalSize(512, 512)
+			, MetallicSize(0, 0)
 			, RoughnessSize(0, 0)
 			, SpecularSize(0, 0)
 		{}
@@ -19,11 +20,13 @@ namespace MaterialExportUtils
 		
 		FIntPoint		DiffuseSize;
 		FIntPoint		NormalSize;
+		FIntPoint		MetallicSize;	
 		FIntPoint		RoughnessSize;	
 		FIntPoint		SpecularSize;	
 			
 		TArray<FColor>	DiffuseSamples;
 		TArray<FColor>	NormalSamples;
+		TArray<FColor>	MetallicSamples;
 		TArray<FColor>	RoughnessSamples;
 		TArray<FColor>	SpecularSamples;
 	};
@@ -54,10 +57,11 @@ namespace MaterialExportUtils
 	 * Flattens specified landscape material
 	 *
 	 * @param InLandscape			Target landscape
+	 * @param HiddenPrimitives		Primitives to hide while rendering scene to texture
 	 * @param OutFlattenMaterial	Output flattened material
 	 * @return						Whether operation was successful
 	 */
-	UNREALED_API bool ExportMaterial(ALandscapeProxy* InLandscape, FFlattenMaterial& OutFlattenMaterial);
+	UNREALED_API bool ExportMaterial(ALandscapeProxy* InLandscape, const TSet<FPrimitiveComponentId>& HiddenPrimitives, FFlattenMaterial& OutFlattenMaterial);
 
 	/**
 	 * Creates UMaterial object from a flatten material
diff --git a/Engine/Source/Editor/WorldBrowser/Private/Tiles/WorldTileCollectionModel.cpp b/Engine/Source/Editor/WorldBrowser/Private/Tiles/WorldTileCollectionModel.cpp
--- a/Engine/Source/Editor/WorldBrowser/Private/Tiles/WorldTileCollectionModel.cpp
+++ b/Engine/Source/Editor/WorldBrowser/Private/Tiles/WorldTileCollectionModel.cpp
@@ -22,10 +22,10 @@
 #include "Engine/LevelStreaming.h"
 #include "GameFramework/WorldSettings.h"
 #include "Engine/Light.h"
+#include "Foliage/InstancedFoliageActor.h"
 
 #define LOCTEXT_NAMESPACE "WorldBrowser"
 
-
 static const FName HeightmapLayerName = FName("__Heightmap__");
 
 
@@ -1949,26 +1949,33 @@ bool FWorldTileCollectionModel::GenerateLODLevels(FLevelModelList InLevelList, i
 			{
 				VertexPos-= LandscapeWorldLocation;
 			}
-								
-			// This is texture resolution for a landscape, probably needs to be calculated using landscape size
-			LandscapeFlattenMaterial.DiffuseSize = FIntPoint(1024, 1024);
-			
-			if (SimplificationDetails.bGenerateLandscapeNormalMap)
-			{
-				LandscapeFlattenMaterial.NormalSize = FIntPoint(1024, 1024);
-			}
 
-			if (SimplificationDetails.bGenerateLandscapeRoughnessMap)
+			// Filter out primitives for landscape texture flattening
+			TSet<FPrimitiveComponentId> PrimitivesToHide;
+			for (TObjectIterator<UPrimitiveComponent> It; It; ++It)
 			{
-				LandscapeFlattenMaterial.RoughnessSize = FIntPoint(1024, 1024);
-			}
+				UPrimitiveComponent* PrimitiveComp = *It;
+				UObject* PrimitiveOuter = PrimitiveComp->GetOuter();
+				
+				const bool bTargetPrim = 
+					PrimitiveComp->GetOuter() == Landscape || 
+					(SimplificationDetails.bBakeFoliageToLandscape && PrimitiveOuter->IsA(AInstancedFoliageActor::StaticClass()));
 
-			if (SimplificationDetails.bGenerateLandscapeSpecularMap)
-			{
-				LandscapeFlattenMaterial.SpecularSize = FIntPoint(1024, 1024);
+				if (!bTargetPrim && PrimitiveComp->IsRegistered() && PrimitiveComp->SceneProxy)
+				{
+					PrimitivesToHide.Add(PrimitiveComp->SceneProxy->GetPrimitiveComponentId());
+				}
 			}
+								
+			// This is texture resolution for a landscape, probably needs to be calculated using landscape size
+			const FIntPoint LandscapeTextureSize(1024, 1024);
+			LandscapeFlattenMaterial.DiffuseSize	= LandscapeTextureSize;
+			LandscapeFlattenMaterial.NormalSize		= SimplificationDetails.bGenerateLandscapeNormalMap ? LandscapeTextureSize : FIntPoint::ZeroValue;
+			LandscapeFlattenMaterial.MetallicSize	= SimplificationDetails.bGenerateLandscapeMetallicMap ? LandscapeTextureSize : FIntPoint::ZeroValue;
+			LandscapeFlattenMaterial.RoughnessSize	= SimplificationDetails.bGenerateLandscapeRoughnessMap ? LandscapeTextureSize : FIntPoint::ZeroValue;
+			LandscapeFlattenMaterial.SpecularSize	= SimplificationDetails.bGenerateLandscapeSpecularMap ? LandscapeTextureSize : FIntPoint::ZeroValue;
 			
-			ExportMaterial(Landscape, LandscapeFlattenMaterial);
+			ExportMaterial(Landscape, PrimitivesToHide, LandscapeFlattenMaterial);
 		
 			FString LandscapeBaseAssetName = Landscape->GetName();
 			// Construct landscape material
diff --git a/Engine/Source/Runtime/Engine/Classes/Engine/Level.h b/Engine/Source/Runtime/Engine/Classes/Engine/Level.h
--- a/Engine/Source/Runtime/Engine/Classes/Engine/Level.h
+++ b/Engine/Source/Runtime/Engine/Classes/Engine/Level.h
@@ -246,6 +246,10 @@ struct ENGINE_API FLevelSimplificationDetails
 	UPROPERTY(Category=Landscape, EditAnywhere)
 	bool bGenerateLandscapeNormalMap;
 
+	// Whether to generate metallic map for landscape static mesh
+	UPROPERTY(Category=Landscape, EditAnywhere)
+	bool bGenerateLandscapeMetallicMap;
+
 	// Whether to generate roughness map for landscape static mesh
 	UPROPERTY(Category=Landscape, EditAnywhere)
 	bool bGenerateLandscapeRoughnessMap;
diff --git a/Engine/Source/Runtime/Engine/Private/Level.cpp b/Engine/Source/Runtime/Engine/Private/Level.cpp
--- a/Engine/Source/Runtime/Engine/Private/Level.cpp
+++ b/Engine/Source/Runtime/Engine/Private/Level.cpp
@@ -166,6 +166,7 @@ FLevelSimplificationDetails::FLevelSimplificationDetails()
  : DetailsPercentage(70.f)
  , LandscapeExportLOD(7)
  , bGenerateLandscapeNormalMap(true)
+ , bGenerateLandscapeMetallicMap(false)
  , bGenerateLandscapeRoughnessMap(false)
  , bGenerateLandscapeSpecularMap(false)
  , bBakeFoliageToLandscape(false)
@@ -178,6 +179,7 @@ bool FLevelSimplificationDetails::operator == (const FLevelSimplificationDetails
 		DetailsPercentage == Other.DetailsPercentage &&
 		LandscapeExportLOD == Other.LandscapeExportLOD &&
 		bGenerateLandscapeNormalMap == Other.bGenerateLandscapeNormalMap &&
+		bGenerateLandscapeMetallicMap == Other.bGenerateLandscapeMetallicMap &&
 		bGenerateLandscapeRoughnessMap == Other.bGenerateLandscapeRoughnessMap &&
 		bGenerateLandscapeSpecularMap == Other.bGenerateLandscapeSpecularMap &&
 		bBakeFoliageToLandscape == Other.bBakeFoliageToLandscape;
