Fixed: Editor crash while generating mesh proxy

diff --git a/Engine/Source/Editor/UnrealEd/Private/EditorExporters.cpp b/Engine/Source/Editor/UnrealEd/Private/EditorExporters.cpp
--- a/Engine/Source/Editor/UnrealEd/Private/EditorExporters.cpp
+++ b/Engine/Source/Editor/UnrealEd/Private/EditorExporters.cpp
@@ -2827,7 +2827,7 @@ namespace MaterialExportUtils
 		return true;
 	}
 
-	UTexture2D* CreateTexture(UPackage* Outer, const FString& AssetName, FIntPoint Size, const TArray<FColor>& Samples, TextureCompressionSettings CompressionSettings, TextureGroup LODGroup, EObjectFlags Flags, bool bSRGB)
+	UTexture2D* CreateTexture(UPackage* Outer, const FString& AssetLongName, FIntPoint Size, const TArray<FColor>& Samples, TextureCompressionSettings CompressionSettings, TextureGroup LODGroup, EObjectFlags Flags, bool bSRGB)
 	{
 		FCreateTexture2DParameters TexParams;
 		TexParams.bUseAlpha = false;
@@ -2837,12 +2837,12 @@ namespace MaterialExportUtils
 
 		if (Outer == nullptr)
 		{
-			Outer = CreatePackage(NULL, *AssetName);
+			Outer = CreatePackage(NULL, *AssetLongName);
 			Outer->FullyLoad();
 			Outer->Modify();
 		}
 
-		UTexture2D* Texture = FImageUtils::CreateTexture2D(Size.X, Size.Y, Samples, Outer, AssetName, Flags, TexParams);
+		UTexture2D* Texture = FImageUtils::CreateTexture2D(Size.X, Size.Y, Samples, Outer, FPackageName::GetShortName(AssetLongName), Flags, TexParams);
 		Texture->LODGroup = LODGroup;
 		Texture->PostEditChange();
 		
@@ -2876,8 +2876,9 @@ namespace MaterialExportUtils
 		if (InFlattenMaterial.DiffuseSamples.Num() > 1)
 		{
 			const FString AssetName = TEXT("T_") + AssetBaseName + TEXT("_D");
+			const FString AssetLongName = AssetBasePath + AssetName;
 			const bool bSRGB = true;
-			UTexture2D* Texture = CreateTexture(Outer, AssetName, InFlattenMaterial.DiffuseSize, InFlattenMaterial.DiffuseSamples, TC_Default, TEXTUREGROUP_World, Flags, bSRGB);
+			UTexture2D* Texture = CreateTexture(Outer, AssetLongName, InFlattenMaterial.DiffuseSize, InFlattenMaterial.DiffuseSamples, TC_Default, TEXTUREGROUP_World, Flags, bSRGB);
 			
 			UMaterialExpressionTextureSample* BasecolorExpression = ConstructObject<UMaterialExpressionTextureSample>(UMaterialExpressionTextureSample::StaticClass(), Material);
 			BasecolorExpression->Texture = Texture;
@@ -2894,8 +2895,9 @@ namespace MaterialExportUtils
 		if (InFlattenMaterial.RoughnessSamples.Num() > 1)
 		{
 			const FString AssetName = TEXT("T_") + AssetBaseName + TEXT("_M");
+			const FString AssetLongName = AssetBasePath + AssetName;
 			const bool bSRGB = false;
-			UTexture2D* Texture = CreateTexture(Outer, AssetName, InFlattenMaterial.MetallicSize, InFlattenMaterial.MetallicSamples, TC_Grayscale, TEXTUREGROUP_World, Flags, bSRGB);
+			UTexture2D* Texture = CreateTexture(Outer, AssetBasePath, InFlattenMaterial.MetallicSize, InFlattenMaterial.MetallicSamples, TC_Grayscale, TEXTUREGROUP_World, Flags, bSRGB);
 			
 			UMaterialExpressionTextureSample* MetallicExpression = ConstructObject<UMaterialExpressionTextureSample>(UMaterialExpressionTextureSample::StaticClass(), Material);
 			MetallicExpression->Texture = Texture;
@@ -2925,8 +2927,9 @@ namespace MaterialExportUtils
 		if (InFlattenMaterial.SpecularSamples.Num() > 1)
 		{
 			const FString AssetName = TEXT("T_") + AssetBaseName + TEXT("_S");
+			const FString AssetLongName = AssetBasePath + AssetName;
 			const bool bSRGB = false;
-			UTexture2D* Texture = CreateTexture(Outer, AssetName, InFlattenMaterial.SpecularSize, InFlattenMaterial.SpecularSamples, TC_Default, TEXTUREGROUP_World, Flags, bSRGB);
+			UTexture2D* Texture = CreateTexture(Outer, AssetLongName, InFlattenMaterial.SpecularSize, InFlattenMaterial.SpecularSamples, TC_Default, TEXTUREGROUP_World, Flags, bSRGB);
 			
 			UMaterialExpressionTextureSample* SpecularExpression = ConstructObject<UMaterialExpressionTextureSample>(UMaterialExpressionTextureSample::StaticClass(), Material);
 			SpecularExpression->Texture = Texture;
@@ -2943,8 +2946,9 @@ namespace MaterialExportUtils
 		if (InFlattenMaterial.RoughnessSamples.Num() > 1)
 		{
 			const FString AssetName = TEXT("T_") + AssetBaseName + TEXT("_R");
+			const FString AssetLongName = AssetBasePath + AssetName;
 			const bool bSRGB = false;
-			UTexture2D* Texture = CreateTexture(Outer, AssetName, InFlattenMaterial.RoughnessSize, InFlattenMaterial.RoughnessSamples, TC_Grayscale, TEXTUREGROUP_World, Flags, bSRGB);
+			UTexture2D* Texture = CreateTexture(Outer, AssetLongName, InFlattenMaterial.RoughnessSize, InFlattenMaterial.RoughnessSamples, TC_Grayscale, TEXTUREGROUP_World, Flags, bSRGB);
 			
 			UMaterialExpressionTextureSample* RoughnessExpression = ConstructObject<UMaterialExpressionTextureSample>(UMaterialExpressionTextureSample::StaticClass(), Material);
 			RoughnessExpression->Texture = Texture;
@@ -2974,8 +2978,9 @@ namespace MaterialExportUtils
 		if (InFlattenMaterial.NormalSamples.Num() > 1)
 		{
 			const FString AssetName = TEXT("T_") + AssetBaseName + TEXT("_N");
+			const FString AssetLongName = AssetBasePath + AssetName;
 			const bool bSRGB = false;
-			UTexture2D* Texture = CreateTexture(Outer, AssetName, InFlattenMaterial.NormalSize, InFlattenMaterial.NormalSamples, TC_Normalmap, TEXTUREGROUP_WorldNormalMap, Flags, bSRGB);
+			UTexture2D* Texture = CreateTexture(Outer, AssetLongName, InFlattenMaterial.NormalSize, InFlattenMaterial.NormalSamples, TC_Normalmap, TEXTUREGROUP_WorldNormalMap, Flags, bSRGB);
 			
 			UMaterialExpressionTextureSample* NormalExpression = ConstructObject<UMaterialExpressionTextureSample>(UMaterialExpressionTextureSample::StaticClass(), Material);
 			NormalExpression->Texture = Texture;
