Fixed: Diffuse texture color space in material flattening

diff --git a/Engine/Source/Editor/UnrealEd/Private/EditorExporters.cpp b/Engine/Source/Editor/UnrealEd/Private/EditorExporters.cpp
--- a/Engine/Source/Editor/UnrealEd/Private/EditorExporters.cpp
+++ b/Engine/Source/Editor/UnrealEd/Private/EditorExporters.cpp
@@ -2634,7 +2634,7 @@ namespace MaterialExportUtils
 			RenderTargetDiffuse->ClearColor = FLinearColor(0.0f, 0.0f, 0.0f, 0.0f);
 			RenderTargetDiffuse->InitCustomFormat(
 				OutFlattenMaterial.DiffuseSize.X, 
-				OutFlattenMaterial.DiffuseSize.Y, PF_B8G8R8A8, true);
+				OutFlattenMaterial.DiffuseSize.Y, PF_B8G8R8A8, false);
 			
 			//
 			OutFlattenMaterial.DiffuseSamples.Empty(OutFlattenMaterial.DiffuseSize.X * OutFlattenMaterial.DiffuseSize.Y);
