UE-6326 Re-parenting a child blueprint and compiling the parent will cause a crash

diff --git a/Engine/Source/Editor/UnrealEd/Private/Kismet2/KismetReinstanceUtilities.cpp b/Engine/Source/Editor/UnrealEd/Private/Kismet2/KismetReinstanceUtilities.cpp
--- a/Engine/Source/Editor/UnrealEd/Private/Kismet2/KismetReinstanceUtilities.cpp
+++ b/Engine/Source/Editor/UnrealEd/Private/Kismet2/KismetReinstanceUtilities.cpp
@@ -68,7 +68,7 @@ FBlueprintCompileReinstancer::FBlueprintCompileReinstancer(UClass* InClassToRein
 			TArray<UObject*> ObjectsToChange;
 			const bool bIncludeDerivedClasses = false;
 			GetObjectsOfClass(ClassToReinstance, ObjectsToChange, bIncludeDerivedClasses);
-			for ( auto ObjIt = ObjectsToChange.CreateConstIterator(); ObjIt; ++ObjIt )
+			for (auto ObjIt = ObjectsToChange.CreateConstIterator(); ObjIt; ++ObjIt)
 			{
 				(*ObjIt)->SetClass(DuplicatedClass);
 			}
@@ -81,7 +81,10 @@ FBlueprintCompileReinstancer::FBlueprintCompileReinstancer(UClass* InClassToRein
 				UBlueprint* ChildBP = Cast<UBlueprint>(ChildClass->ClassGeneratedBy);
 				if (ChildBP)
 				{
-					if (ChildBP->HasAnyFlags(RF_BeingRegenerated))
+					const bool bClassIsDirectlyGeneratedByTheBlueprint = (ChildBP->GeneratedClass == ChildClass)
+						|| (ChildBP->SkeletonGeneratedClass == ChildClass);
+
+					if (ChildBP->HasAnyFlags(RF_BeingRegenerated) || !bClassIsDirectlyGeneratedByTheBlueprint)
 					{
 						if (ChildClass->GetSuperClass() == ClassToReinstance)
 						{
