Fix for iterative cooking in the editor.

diff --git a/Engine/Source/Editor/UnrealEd/Classes/CookOnTheSide/CookOnTheFlyServer.h b/Engine/Source/Editor/UnrealEd/Classes/CookOnTheSide/CookOnTheFlyServer.h
--- a/Engine/Source/Editor/UnrealEd/Classes/CookOnTheSide/CookOnTheFlyServer.h
+++ b/Engine/Source/Editor/UnrealEd/Classes/CookOnTheSide/CookOnTheFlyServer.h
@@ -171,6 +171,13 @@ public:
 
 		FFilePlatformRequest() { }
 
+
+		FFilePlatformRequest( const FName& InFileName, const FName& InPlatformName ) : Filename( InFileName )
+		{
+			Platformnames.Add( InPlatformName );
+		}
+
+
 		FFilePlatformRequest( const FName& InFilename, const TArray<FName>& InPlatformname ) : Filename( InFilename )
 		{
 			Platformnames = InPlatformname;
@@ -616,7 +623,6 @@ private:
 	TAutoPtr<class FSandboxPlatformFile> SandboxFile;
 	bool bIsSavingPackage; // used to stop recursive mark package dirty functions
 
-
 	FThreadSafeQueue<struct FRecompileRequest*> RecompileRequests;
 	FFilenameQueue CookRequests; // list of requested files
 	FThreadSafeUnsolicitedPackagesList UnsolicitedCookedPackages;
@@ -636,17 +642,12 @@ private:
 	// used by GetCached * Filename functions
 	mutable TMap<FName, FCachedPackageFilename> PackageFilenameCache; // filename cache (only process the string operations once)
 
-
 	// declared mutable as it's used purely as a cache and don't want to have to declare all the functions as non const just because of this cache
 	// used by IniSettingsOutOfDate and GetCurrentIniStrings 
 	mutable TMap<FName, TArray<FString>> CachedIniVersionStringsMap;
 
 public:
 
-	// void WarmCookedPackages(const FString& AssetRegistryPath, const TArray<FName>& TargetPlatformNames);
-
-	
-
 	enum ECookOnTheSideResult
 	{
 		COSR_CookedMap			= 0x00000001,
@@ -875,6 +876,16 @@ private:
 	 */
 	bool ShouldCook(const FString& InFileName, const FName& InPlatformName);
 
+	/**
+	 * Initialize the sandbox 
+	 */
+	void InitializeSandbox();
+
+	/**
+	 * Clean up the sandbox
+	 */
+	void TermSandbox();
+
 	/**
 	 * GetDependencies
 	 * 
@@ -1023,7 +1034,10 @@ private:
 	FString GetOutputDirectoryOverride() const;
 
 	/** Cleans sandbox folders for all target platforms */
-	void CleanSandbox(const TArray<ITargetPlatform*>& Platforms);
+	void CleanSandbox( const bool bIterative );
+
+	/** Populate the cooked packages list from the on disk content using time stamps and dependencies to figure out if they are ok */
+	void PopulateCookedPackagesFromDisk( const TArray<ITargetPlatform*>& Platforms );
 
 	/** Generates asset registry */
 	void GenerateAssetRegistry(const TArray<ITargetPlatform*>& Platforms);
diff --git a/Engine/Source/Editor/UnrealEd/Private/CookOnTheFlyServer.cpp b/Engine/Source/Editor/UnrealEd/Private/CookOnTheFlyServer.cpp
--- a/Engine/Source/Editor/UnrealEd/Private/CookOnTheFlyServer.cpp
+++ b/Engine/Source/Editor/UnrealEd/Private/CookOnTheFlyServer.cpp
@@ -434,7 +434,7 @@ bool UCookOnTheFlyServer::StartNetworkFileServer( const bool BindAnyPort )
 	check( IsCookOnTheFlyMode() );
 	//GetDerivedDataCacheRef().WaitForQuiescence(false);
 
-	CreateSandboxFile();
+	InitializeSandbox();
 
 	ITargetPlatformManagerModule& TPM = GetTargetPlatformManagerRef();
 	const TArray<ITargetPlatform*>& Platforms = TPM.GetCookingTargetPlatforms();
@@ -480,9 +480,6 @@ bool UCookOnTheFlyServer::StartNetworkFileServer( const bool BindAnyPort )
 		NetworkFileServers.Add( HttpFileServer );
 	}
 
-	// clean sandbox when we start the network file server
-	CleanSandbox(Platforms);
-
 	// loop while waiting for requests
 	GIsRequestingExit = false;
 	return true;
@@ -1501,6 +1498,9 @@ bool UCookOnTheFlyServer::SaveCookedPackage( UPackage* Package, uint32 SaveFlags
 
 bool UCookOnTheFlyServer::ShouldCook(const FString& InFileName, const FName &InPlatformName)
 {
+	// this only does anything if we are using iterative cooking
+	return true;
+	/*
 	bool bDoCook = false;
 
 	FString PkgFile;
@@ -1553,7 +1553,7 @@ bool UCookOnTheFlyServer::ShouldCook(const FString& InFileName, const FName &InP
 		bDoCook |= bCookPackage;
 	}
 
-	return bDoCook;
+	return bDoCook;*/
 }
 
 bool UCookOnTheFlyServer::SaveCookedPackage( UPackage* Package, uint32 SaveFlags, bool& bOutWasUpToDate, TArray<FName> &TargetPlatformNames )
@@ -1767,13 +1767,11 @@ void UCookOnTheFlyServer::Initialize( ECookMode::Type DesiredCookMode, ECookInit
 		CookByTheBookOptions->bGenerateStreamingInstallManifests = IsCookFlagSet(ECookInitializationFlags::GenerateStreamingInstallManifest);
 	}
 
-
 	ITargetPlatformManagerModule& TPM = GetTargetPlatformManagerRef();
 	const TArray<ITargetPlatform*>& Platforms = TPM.GetCookingTargetPlatforms();
 
 	// always generate the asset registry before starting to cook, for either method
 	GenerateAssetRegistry(Platforms);
-	
 }
 
 uint32 UCookOnTheFlyServer::NumConnections() const
@@ -2016,8 +2014,52 @@ bool UCookOnTheFlyServer::SaveCurrentIniSettings( const ITargetPlatform* TargetP
 	return true;
 }
 
-void UCookOnTheFlyServer::CleanSandbox(const TArray<ITargetPlatform*>& Platforms)
+void UCookOnTheFlyServer::PopulateCookedPackagesFromDisk( const TArray<ITargetPlatform*>& Platforms )
 {
+	// check each package to see if it's out of date
+	FPackageDependencyInfoModule& PDInfoModule = FModuleManager::LoadModuleChecked<FPackageDependencyInfoModule>("PackageDependencyInfo");
+
+	// list of directories to skip
+	TArray<FString> DirectoriesToSkip;
+	TArray<FString> DirectoriesToNotRecurse;
+
+	// See what files are out of date in the sandbox folder
+	for (int32 Index = 0; Index < Platforms.Num(); Index++)
+	{
+		ITargetPlatform* Target = Platforms[Index];
+		FString SandboxDirectory = GetSandboxDirectory(Target->PlatformName());
+		FName PlatformFName( *Target->PlatformName() );
+		// use the timestamp grabbing visitor
+		IPlatformFile& PlatformFile = FPlatformFileManager::Get().GetPlatformFile();
+		FLocalTimestampDirectoryVisitor Visitor(PlatformFile, DirectoriesToSkip, DirectoriesToNotRecurse, false);
+
+		PlatformFile.IterateDirectory(*SandboxDirectory, Visitor);
+
+		for (TMap<FString, FDateTime>::TIterator TimestampIt(Visitor.FileTimes); TimestampIt; ++TimestampIt)
+		{
+			FString CookedFilename = TimestampIt.Key();
+			FDateTime CookedTimestamp = TimestampIt.Value();
+			FString StandardCookedFilename = CookedFilename.Replace(*SandboxDirectory, *(FPaths::GetRelativePathToRoot()));
+			FDateTime DependentTimestamp;
+
+			if (PDInfoModule.DeterminePackageDependentTimeStamp(*(FPaths::GetBaseFilename(StandardCookedFilename, false)), DependentTimestamp) == true)
+			{
+				double Diff = (CookedTimestamp - DependentTimestamp).GetTotalSeconds();
+
+				if (Diff >= 0.0)
+				{
+					CookedPackages.Add(FFilePlatformRequest( FName(*StandardCookedFilename), PlatformFName) );
+				}
+			}
+		}
+	}
+}
+
+void UCookOnTheFlyServer::CleanSandbox( const bool bIterative )
+{
+	ITargetPlatformManagerModule& TPM = GetTargetPlatformManagerRef();
+	const TArray<ITargetPlatform*>& Platforms = TPM.GetCookingTargetPlatforms();	
+
 	// before we can delete any cooked files we need to make sure that we have finished writing them
 	UPackage::WaitForAsyncFileWrites();
 
@@ -2034,8 +2076,7 @@ void UCookOnTheFlyServer::CleanSandbox(const TArray<ITargetPlatform*>& Platforms
 			CacheIniVersionStringsMap(Platform);
 		}
 
-
-		if (IsCookFlagSet(ECookInitializationFlags::Iterative) == false)
+		if (bIterative == false)
 		{
 			// for now we are going to wipe the cooked directory
 			for (int32 Index = 0; Index < Platforms.Num(); Index++)
@@ -2194,28 +2235,20 @@ void UCookOnTheFlyServer::CollectFilesToCook(TArray<FString>& FilesInPath, const
 {
 	TArray<FString> CookDirectories = InCookDirectories;
 	
-	if ( IsRealtimeMode() == false )
-	{
-		TArray<FString> MapList;
-		// Add the default map section
-		GEditor->LoadMapListFromIni(TEXT("AlwaysCookMaps"), MapList);
-
-		for ( const auto &IniMapSection : IniMapSections )
-		{
-			GEditor->LoadMapListFromIni(*IniMapSection, MapList);
-		}
+	
+	TArray<FString> MapList;
+	// Add the default map section
+	GEditor->LoadMapListFromIni(TEXT("AlwaysCookMaps"), MapList);
 
-		for (int32 MapIdx = 0; MapIdx < MapList.Num(); MapIdx++)
-		{
-			AddFileToCook( FilesInPath, MapList[MapIdx]);
-		}
+	for ( const auto &IniMapSection : IniMapSections )
+	{
+		GEditor->LoadMapListFromIni(*IniMapSection, MapList);
 	}
-	else
+
+	for (int32 MapIdx = 0; MapIdx < MapList.Num(); MapIdx++)
 	{
-		// if we are cooking in the editor then don't include the always cook maps list because we just want to cook whats in memory
-		check( CookMaps.Num() > 0 );
+		AddFileToCook( FilesInPath, MapList[MapIdx]);
 	}
-	
 
 
 	// Also append any cookdirs from the project ini files; these dirs are relative to the game content directory
@@ -2539,9 +2572,6 @@ void UCookOnTheFlyServer::CookByTheBookFinished()
 	const float TotalCookTime = (float)(FPlatformTime::Seconds() - CookByTheBookOptions->CookStartTime);
 	UE_LOG(LogCook, Display, TEXT("Cook by the book total time in tick %fs total time %f"), CookByTheBookOptions->CookTime, TotalCookTime);
 
-	check( SandboxFile );
-	SandboxFile = NULL;
-
 	CookByTheBookOptions->bRunning = false;
 }
 
@@ -2622,6 +2652,31 @@ void UCookOnTheFlyServer::CreateSandboxFile()
 	SandboxFile->Initialize(&FPlatformFileManager::Get().GetPlatformFile(), *FString::Printf(TEXT("-sandbox=\"%s\""), *OutputDirectory));
 }
 
+void UCookOnTheFlyServer::InitializeSandbox()
+{
+	if ( SandboxFile == NULL )
+	{
+		ITargetPlatformManagerModule& TPM = GetTargetPlatformManagerRef();
+		const TArray<ITargetPlatform*>& TargetPlatforms = TPM.GetCookingTargetPlatforms();
+
+		CreateSandboxFile();
+
+		if ( IsCookFlagSet(ECookInitializationFlags::Iterative) )
+		{
+			PopulateCookedPackagesFromDisk(TargetPlatforms);
+		}
+
+		CleanSandbox(IsCookFlagSet(ECookInitializationFlags::Iterative));
+	}
+}
+
+void UCookOnTheFlyServer::TermSandbox()
+{
+	ClearAllCookedData();
+	ClearPackageFilenameCache();
+	SandboxFile = NULL;
+}
+
 void UCookOnTheFlyServer::StartCookByTheBook( const FCookByTheBookStartupOptions& CookByTheBookStartupOptions )
 {
 
@@ -2637,66 +2692,21 @@ void UCookOnTheFlyServer::StartCookByTheBook( const FCookByTheBookStartupOptions
 	const FString& BasedOnReleaseVersion = CookByTheBookStartupOptions.BasedOnReleaseVersion;
 
 	check( IsInGameThread() );
-
-	// can transition from cook on the fly to cook by the book from the editor
 	check( IsCookByTheBookMode() );
 
-	/// reference ini parsing code
-	/*FString SectionStr;
-	if (FParse::Value(InCmdParams, TEXT("MAPINISECTION="), SectionStr))
-	{
-		if (SectionStr.Contains(TEXT("+")))
-		{
-			TArray<FString> Sections;
-			SectionStr.ParseIntoArray(&Sections,TEXT("+"),true);
-			for (int32 Index = 0; Index < Sections.Num(); Index++)
-			{
-				LoadMapListFromIni(Sections[Index], OutMapList);
-			}
-		}
-		else
-		{
-			LoadMapListFromIni(SectionStr, OutMapList);
-		}
-	}*/
-	
-
-	/*
-	// copied from original cook commandlet
-	// need to fill these in
-	bCookAll = Switches.Contains(TEXT("COOKALL"));   // Cook everything
-	bLeakTest = Switches.Contains(TEXT("LEAKTEST"));   // Test for UObject leaks
-
-	*/
-
-
-
 	if ( CookByTheBookOptions->DlcName != DLCName )
 	{
 		// we are going to change the state of dlc we need to clean out our package filename cache (the generated filename cache is dependent on this key)
-		ClearPackageFilenameCache();
-		
 		CookByTheBookOptions->DlcName = DLCName;
 		check( OutputDirectoryOverride.Len() == 0 );
-	}
 
-	CreateSandboxFile();
+		TermSandbox();
+	}
 
-	// if we are cooking dlc then we don't need to clean sandbox
-	CleanSandbox(TargetPlatforms);
-	// deprecated is handled in the clean sandbox function above
-	/*TArray<ITargetPlatform*> ModifiedTargetPlatforms;
-	IniSettingsOutOfDate(TargetPlatforms, ModifiedTargetPlatforms);
-	for ( const auto& ModifiedTargetPlatform : ModifiedTargetPlatforms )
-	{
-		ClearPlatformCookedData( FName( *ModifiedTargetPlatform->PlatformName() ) );
-	}*/
+	InitializeSandbox();
 
-	if ( IsRealtimeMode() )
-	{
-		// from this point on (after we call CleanSandbox) if we are in the editor we can act iterative
-		CookFlags |= ECookInitializationFlags::Iterative;
-	}
+	// iteratively clean any old files out of the sandbox (check if ini settings are out of date and clean out any stale files)
+	CleanSandbox(true);
 
 	CookByTheBookOptions->bRunning = true;
 	CookByTheBookOptions->bCancel = false;
