[UE-6487] When blueprints are created directly from assets, need to let the BP node database know the BP was modified, as part of the setup.

diff --git a/Engine/Source/Editor/UnrealEd/Private/Kismet2/Kismet2.cpp b/Engine/Source/Editor/UnrealEd/Private/Kismet2/Kismet2.cpp
--- a/Engine/Source/Editor/UnrealEd/Private/Kismet2/Kismet2.cpp
+++ b/Engine/Source/Editor/UnrealEd/Private/Kismet2/Kismet2.cpp
@@ -1359,8 +1359,9 @@ UBlueprint* FKismetEditorUtilities::CreateBlueprintUsingAsset(UObject* Asset, bo
 			// Add node to the SCS
 			NewBP->SimpleConstructionScript->AddNode(NewNode);
 
-			// Recompile skeleton because of the new component we added
-			FKismetEditorUtilities::GenerateBlueprintSkeleton(NewBP, true);
+			// Recompile skeleton because of the new component we added (and 
+			// broadcast the change to those that care, like the BP node database)
+			FBlueprintEditorUtils::MarkBlueprintAsStructurallyModified(NewBP);
 
 			// Open in BP editor if desired
 			if(bOpenInEditor)
