Fixed crash in Merge Actors tool UE-6642

diff --git a/Engine/Source/Editor/UnrealEd/Private/Dialogs/SMeshMergingDialog.cpp b/Engine/Source/Editor/UnrealEd/Private/Dialogs/SMeshMergingDialog.cpp
--- a/Engine/Source/Editor/UnrealEd/Private/Dialogs/SMeshMergingDialog.cpp
+++ b/Engine/Source/Editor/UnrealEd/Private/Dialogs/SMeshMergingDialog.cpp
@@ -37,13 +37,17 @@ void SMeshMergingDialog::Construct(const FArguments& InArgs)
 	FLevelEditorModule& LevelEditor = FModuleManager::GetModuleChecked<FLevelEditorModule>( "LevelEditor" );
 	LevelEditor.OnActorSelectionChanged().AddSP(this, &SMeshMergingDialog::OnActorSelectionChanged);
 	
-	// Setup available resolutions for an atlased lightmap
-	const auto& LightMapGroup = GSystemSettings.TextureLODSettings.GetTextureLODGroup(TEXTUREGROUP_Lightmap);
-	for (int32 Index = LightMapGroup.MinLODMipCount; Index <= LightMapGroup.MaxLODMipCount; Index++)
+
+	// Setup available resolutions for lightmap
+	const int32 MinTexResolution = 1 << FTextureLODSettings::FTextureLODGroup().MinLODMipCount;
+	const int32 MaxTexResolution = 1 << FTextureLODSettings::FTextureLODGroup().MaxLODMipCount;
+	for (int32 LightmapRes = MinTexResolution; LightmapRes <= MaxTexResolution; LightmapRes*=2)
 	{
-		LightMapResolutionOptions.Add(MakeShareable(new FString(FString::FormatAsNumber(1 << Index))));
+		LightMapResolutionOptions.Add(MakeShareable(new FString(FString::FormatAsNumber(LightmapRes))));
 	}
-	
+
+	MergingSettings.TargetLightMapResolution = FMath::Clamp(MergingSettings.TargetLightMapResolution, MinTexResolution, MaxTexResolution);
+		
 	// Setup available UV channels for an atlased lightmap
 	for (int32 Index = 0; Index < MAX_MESH_TEXTURE_COORDS; Index++)
 	{
