Fix crash when "Fracture Mesh" is pressed in the destructible mesh editor (UE-6607 and GitHub #685)

diff --git a/Engine/Source/Editor/UnrealEd/Private/Factories/ApexDestructibleAssetImport.cpp b/Engine/Source/Editor/UnrealEd/Private/Factories/ApexDestructibleAssetImport.cpp
--- a/Engine/Source/Editor/UnrealEd/Private/Factories/ApexDestructibleAssetImport.cpp
+++ b/Engine/Source/Editor/UnrealEd/Private/Factories/ApexDestructibleAssetImport.cpp
@@ -734,7 +734,7 @@ bool SetApexDestructibleAsset(UDestructibleMesh& DestructibleMesh, NxDestructibl
 	for(TObjectIterator<UDestructibleComponent> It; It; ++It)
 	{
 		UDestructibleComponent* DestructibleComponent = *It;
-		if(DestructibleComponent->SkeletalMesh == &DestructibleMesh)
+		if(DestructibleComponent->SkeletalMesh == &DestructibleMesh && DestructibleComponent->IsPhysicsStateCreated())
 		{
 			DestructibleComponent->DestroyPhysicsState();
 		}
