Fixed null pointer dereference in HandleUpdateLandscapeEditorData console command
#686

diff --git a/Engine/Source/Editor/UnrealEd/Private/UnrealEdSrv.cpp b/Engine/Source/Editor/UnrealEd/Private/UnrealEdSrv.cpp
--- a/Engine/Source/Editor/UnrealEd/Private/UnrealEdSrv.cpp
+++ b/Engine/Source/Editor/UnrealEd/Private/UnrealEdSrv.cpp
@@ -560,25 +560,28 @@ bool UUnrealEdEngine::HandleUpdateLandscapeEditorDataCommand( const TCHAR* Str,
 					{
 						// Fix level inconsistency for landscape component and collision component
 						ULandscapeHeightfieldCollisionComponent* Collision = Comp->CollisionComponent.Get();
-						if (Collision && Comp->GetLandscapeProxy()->GetLevel() != Collision->GetLandscapeProxy()->GetLevel())
+						if (Collision != nullptr)
 						{
-							ALandscapeProxy* FromProxy = Collision->GetLandscapeProxy();
-							ALandscapeProxy* DestProxy = Comp->GetLandscapeProxy();
-							// From MoveToLevelTool
-							FromProxy->CollisionComponents.Remove(Collision);
-							Collision->UnregisterComponent();
-							Collision->DetachFromParent(true);
-							Collision->Rename(NULL, DestProxy);
-							DestProxy->CollisionComponents.Add(Collision);
-							Collision->AttachTo( DestProxy->GetRootComponent(), NAME_None, EAttachLocation::KeepWorldPosition );
-							SelectProxies.Add(FromProxy);
-							SelectProxies.Add(DestProxy);
-						}
+							if (Comp->GetLandscapeProxy()->GetLevel() != Collision->GetLandscapeProxy()->GetLevel())
+							{
+								ALandscapeProxy* FromProxy = Collision->GetLandscapeProxy();
+								ALandscapeProxy* DestProxy = Comp->GetLandscapeProxy();
+								// From MoveToLevelTool
+								FromProxy->CollisionComponents.Remove(Collision);
+								Collision->UnregisterComponent();
+								Collision->DetachFromParent(true);
+								Collision->Rename(NULL, DestProxy);
+								DestProxy->CollisionComponents.Add(Collision);
+								Collision->AttachTo( DestProxy->GetRootComponent(), NAME_None, EAttachLocation::KeepWorldPosition );
+								SelectProxies.Add(FromProxy);
+								SelectProxies.Add(DestProxy);
+							}
 
-						// Fix Dominant Layer Data
-						if (bHasPhysicalMaterial && Collision->DominantLayerData.GetBulkDataSize() == 0)
-						{
-							Comp->UpdateCollisionLayerData();
+							// Fix Dominant Layer Data
+							if (bHasPhysicalMaterial && Collision->DominantLayerData.GetBulkDataSize() == 0)
+							{
+								Comp->UpdateCollisionLayerData();
+							}
 						}
 					}
 				}
