Got rid of deleting generated class after blueprints duplication. Not needed anymore.

diff --git a/Engine/Source/Editor/UnrealEd/Private/Kismet2/BlueprintEditorUtils.cpp b/Engine/Source/Editor/UnrealEd/Private/Kismet2/BlueprintEditorUtils.cpp
--- a/Engine/Source/Editor/UnrealEd/Private/Kismet2/BlueprintEditorUtils.cpp
+++ b/Engine/Source/Editor/UnrealEd/Private/Kismet2/BlueprintEditorUtils.cpp
@@ -1393,18 +1393,6 @@ void FBlueprintEditorUtils::PropagateParentBlueprintDefaults(UClass* ClassToProp
 	}
 }
 
-void DestroyGeneratedClass(TSubclassOf<UObject> &ClassToDestroy)
-{
-	if (ClassToDestroy != nullptr)
-	{
-		// We just need to rename the object, so the name won't collide.
-		ClassToDestroy->Rename(TEXT("None"), nullptr, REN_DontCreateRedirectors);
-
-		// The rest will do GC after we lose reference to it.
-		ClassToDestroy = nullptr;
-	}
-}
-
 void FBlueprintEditorUtils::PostDuplicateBlueprint(UBlueprint* Blueprint)
 {
 	// Only recompile after duplication if this isn't PIE
@@ -1428,11 +1416,6 @@ void FBlueprintEditorUtils::PostDuplicateBlueprint(UBlueprint* Blueprint)
 			TArray<UTimelineTemplate*> Timelines = Blueprint->Timelines;
 			Blueprint->Timelines.Empty();
 
-			// Delete the existing class references, the compile will create new ones
-			DestroyGeneratedClass(Blueprint->GeneratedClass);
-			DestroyGeneratedClass(Blueprint->SkeletonGeneratedClass);
-				
-
 			Blueprint->GeneratedClass = nullptr;
 			Blueprint->SkeletonGeneratedClass = nullptr;
 
