When '-immersive' is used, only start immersive if a project file was specified

diff --git a/Engine/Source/Editor/UnrealEd/Private/UnrealEd.cpp b/Engine/Source/Editor/UnrealEd/Private/UnrealEd.cpp
--- a/Engine/Source/Editor/UnrealEd/Private/UnrealEd.cpp
+++ b/Engine/Source/Editor/UnrealEd/Private/UnrealEd.cpp
@@ -97,7 +97,7 @@ int32 EditorInit( IEngineLoop& EngineLoop )
 	FPlatformSplash::Hide();
 
 	// Are we in immersive mode?
-	const bool bIsImmersive = FParse::Param( FCommandLine::Get(), TEXT( "immersive" ) );
+	const bool bIsImmersive = FPaths::IsProjectFilePathSet() && FParse::Param( FCommandLine::Get(), TEXT( "immersive" ) );
 
 	// Do final set up on the editor frame and show it
 	{
