#Anim: github: pull request on calculating animation length when importing

PR #670: Don't skip the last frame when sampling FBX animation due to anim length not perfectly being divisible by ResampleRate (Contributed by [REDACTED_PERSON])

https://github.com/EpicGames/UnrealEngine/pull/670

diff --git a/Engine/Source/Editor/UnrealEd/Private/SkeletalMeshEdit.cpp b/Engine/Source/Editor/UnrealEd/Private/SkeletalMeshEdit.cpp
--- a/Engine/Source/Editor/UnrealEd/Private/SkeletalMeshEdit.cpp
+++ b/Engine/Source/Editor/UnrealEd/Private/SkeletalMeshEdit.cpp
@@ -975,6 +975,7 @@ bool UnFbx::FFbxImporter::ImportCurveToAnimSequence(class UAnimSequence * Target
 
 bool UnFbx::FFbxImporter::ImportAnimation(USkeleton* Skeleton, UAnimSequence * DestSeq, const FString& FileName, TArray<FbxNode*>& SortedLinks, TArray<FbxNode*>& NodeArray, FbxAnimStack* CurAnimStack, const int32 ResampleRate, const FbxTimeSpan AnimTimeSpan)
 {
+	// @todo : the length might need to change w.r.t. sampling keys
 	FbxTime SequenceLength = AnimTimeSpan.GetDuration();
 	float PreviousSequenceLength = DestSeq->SequenceLength;
 
@@ -1110,6 +1111,8 @@ bool UnFbx::FFbxImporter::ImportAnimation(USkeleton* Skeleton, UAnimSequence * D
 	FMatrix AddedMatrix = Converter.ConvertMatrix(FbxAddedMatrix);
 
 	int32 TotalNumKeys = 0;
+	const int32 NumSamplingKeys = FMath::FloorToInt(AnimTimeSpan.GetDuration().GetSecondDouble() * ResampleRate);
+	const FbxTime TimeIncrement = AnimTimeSpan.GetDuration() / NumSamplingKeys;
 	for(int32 SourceTrackIdx = 0; SourceTrackIdx < FbxRawBoneNames.Num(); ++SourceTrackIdx)
 	{
 		int32 NumKeysForTrack = 0;
@@ -1132,8 +1135,7 @@ bool UnFbx::FFbxImporter::ImportAnimation(USkeleton* Skeleton, UAnimSequence * D
 			FbxNode* Link = SortedLinks[SourceTrackIdx];
 			FbxNode * LinkParent = Link->GetParent();
 			
-			int32 index = 0;
-			for (FbxTime CurTime = AnimTimeSpan.GetStart(); CurTime <= AnimTimeSpan.GetStop(); CurTime += FbxLongLong(46186158000) / ResampleRate, index++)
+			for(FbxTime CurTime = AnimTimeSpan.GetStart(); CurTime <= AnimTimeSpan.GetStop(); CurTime += TimeIncrement)
 			{
 				// save global trasnform
 				FbxAMatrix GlobalMatrix = Link->EvaluateGlobalTransform(CurTime);
