Allowed UFactory::StaticImportObject to call FactoryCreateNew even when it has a filename set

UDN - bCreateNew=true, and Drag & Drop files with a Factory

This can be useful if your factory needs knowledge of the filename, but doesn't actually do anything with the file data.

In order to try and preserve any existing behavior, the factories are sorted such that these calls to FactoryCreateNew will happen only if there are no factories present that would actually import the file data.

diff --git a/Engine/Source/Editor/UnrealEd/Private/Factories/Factory.cpp b/Engine/Source/Editor/UnrealEd/Private/Factories/Factory.cpp
--- a/Engine/Source/Editor/UnrealEd/Private/Factories/Factory.cpp
+++ b/Engine/Source/Editor/UnrealEd/Private/Factories/Factory.cpp
@@ -226,12 +226,28 @@ UObject* UFactory::StaticImportObject
 			}
 		}
 
-		struct FCompareUFactoryImportPriority
+		Factories.Sort([](const UFactory& A, const UFactory& B) -> bool
 		{
-			// Use > operator because we want higher priorities earlier in the list
-			FORCEINLINE bool operator()(const UFactory& A, const UFactory& B) const { return A.ImportPriority > B.ImportPriority; }
-		};
-		Factories.Sort( FCompareUFactoryImportPriority() );
+			// First sort so that higher priorities are earlier in the list
+			if( A.ImportPriority > B.ImportPriority )
+			{
+				return true;
+			}
+			else if( A.ImportPriority < B.ImportPriority )
+			{
+				return false;
+			}
+
+			// Then sort so that factories that only create new assets are tried after those that actually import the file data (when they have an equivalent priority)
+			const bool bFactoryAImportsFiles = !A.CanCreateNew();
+			const bool bFactoryBImportsFiles = !B.CanCreateNew();
+			if( bFactoryAImportsFiles && !bFactoryBImportsFiles )
+			{
+				return true;
+			}
+
+			return false;
+		});
 	}
 
 	bool bLoadedFile = false;
@@ -243,12 +259,9 @@ UObject* UFactory::StaticImportObject
 		UObject* Result = NULL;
 		if( Factory->CanCreateNew() )
 		{
-			if( FCString::Stricmp(Filename,TEXT(""))==0 )
-			{
-				UE_LOG(LogFactory, Log,  TEXT("FactoryCreateNew: %s with %s (%i %i %s)"), *Class->GetName(), *Factories[i]->GetClass()->GetName(), Factory->bCreateNew, Factory->bText, Filename );
-				Factory->ParseParms( Parms );
-				Result = Factory->FactoryCreateNew( Class, InOuter, Name, Flags, NULL, Warn );
-			}
+			UE_LOG(LogFactory, Log,  TEXT("FactoryCreateNew: %s with %s (%i %i %s)"), *Class->GetName(), *Factories[i]->GetClass()->GetName(), Factory->bCreateNew, Factory->bText, Filename );
+			Factory->ParseParms( Parms );
+			Result = Factory->FactoryCreateNew( Class, InOuter, Name, Flags, NULL, Warn );
 		}
 		else if( FCString::Stricmp(Filename,TEXT(""))!=0 )
 		{
