Refactored general implementation of processing properties while gathering text from assets to improve handling of native, fixed-size arrays.

diff --git a/Engine/Source/Editor/UnrealEd/Private/Commandlets/GatherTextFromAssetsCommandlet.cpp b/Engine/Source/Editor/UnrealEd/Private/Commandlets/GatherTextFromAssetsCommandlet.cpp
--- a/Engine/Source/Editor/UnrealEd/Private/Commandlets/GatherTextFromAssetsCommandlet.cpp
+++ b/Engine/Source/Editor/UnrealEd/Private/Commandlets/GatherTextFromAssetsCommandlet.cpp
@@ -374,81 +374,57 @@ namespace
 
 		void ProcessProperty(UProperty* const Property, void* const ValueAddress)
 		{
-			// Property is a struct property.
+			UTextProperty* const TextProperty = Cast<UTextProperty>(Property);
+			UArrayProperty* const ArrayProperty = Cast<UArrayProperty>(Property);
 			UStructProperty* const StructProperty = Cast<UStructProperty>(Property);
-			if(StructProperty)
+
+			// Handles both native, fixed-size arrays and plain old non-array properties.
+			for(int32 i = 0; i < Property->ArrayDim; ++i)
 			{
-				// Struct property may be a native array of fixed size.
-				for(int32 i = 0; i < Property->ArrayDim; ++i)
+				void* const ElementValueAddress = reinterpret_cast<uint8*>(ValueAddress) + Property->ElementSize * i;
+
+				// Property is a text property.
+				if (TextProperty)
 				{
-					uint8* StructureValueAddress = reinterpret_cast<uint8*>(ValueAddress);
-					if(Property->ArrayDim > 1)
+					// Text property may be a native array of fixed size.
+					for(int32 i = 0; i < Property->ArrayDim; ++i)
 					{
-						StructureValueAddress += StructProperty->Struct->GetStructureSize() * i;
-					}
+						FText* Text = reinterpret_cast<FText*>(ElementValueAddress);
 
-					// Iterate over all properties of the struct's class.
-					for (TFieldIterator<UProperty>PropIt(StructProperty->Struct, EFieldIteratorFlags::IncludeSuper, EFieldIteratorFlags::ExcludeDeprecated, EFieldIteratorFlags::IncludeInterfaces); PropIt; ++PropIt)
-					{
-						ProcessProperty( *PropIt, PropIt->ContainerPtrToValuePtr<void>(StructureValueAddress) );
+						if( FTextInspector::GetFlags(*Text) & ETextFlag::ConvertedProperty )
+						{
+							ObjectPackage->MarkPackageDirty();
+						}
+
+						bool Fixed = false;
+						Commandlet->ProcessTextProperty(TextProperty, Text, Object, /*OUT*/Fixed );
+						if( Fixed )
+						{
+							ObjectPackage->MarkPackageDirty();
+						}
 					}
 				}
-			}
-
-			// Property is an array property.
-			UArrayProperty* const ArrayProperty = Cast<UArrayProperty>(Property);
-			if(ArrayProperty)
-			{
-				for(int32 i = 0; i < Property->ArrayDim; ++i)
+				// Property is a DYNAMIC array property.
+				else if (ArrayProperty)
 				{
-					void* ArrayValueAddress;
-					if(Property->ArrayDim > 1)
-					{
-						ArrayValueAddress = Property->ContainerPtrToValuePtr<void>(ValueAddress, i);
-					}
-					else
-					{
-						ArrayValueAddress = ValueAddress;
-					}
-
-					// Iterate over all objects of the array.
-					FScriptArrayHelper ScriptArrayHelper(ArrayProperty, ArrayValueAddress);
-					const int32 ElementCount = ScriptArrayHelper.Num();
-					for(int32 j = 0; j < ElementCount; ++j)
+					for(int32 i = 0; i < Property->ArrayDim; ++i)
 					{
-						ProcessProperty( ArrayProperty->Inner, ScriptArrayHelper.GetRawPtr(j) );
+						// Iterate over all objects of the array.
+						FScriptArrayHelper ScriptArrayHelper(ArrayProperty, ElementValueAddress);
+						const int32 ElementCount = ScriptArrayHelper.Num();
+						for(int32 j = 0; j < ElementCount; ++j)
+						{
+							ProcessProperty( ArrayProperty->Inner, ScriptArrayHelper.GetRawPtr(j) );
+						}
 					}
 				}
-			}
-
-			// Property is a text property.
-			UTextProperty* const TextProperty = Cast<UTextProperty>(Property);
-			if(TextProperty)
-			{
-				// Text property may be a native array of fixed size.
-				for(int32 i = 0; i < Property->ArrayDim; ++i)
+				// Property is a struct property.
+				else if (StructProperty)
 				{
-					FText* Text;
-					if(Property->ArrayDim > 1)
-					{
-						uint8* ArrayValueAddress = reinterpret_cast<uint8*>(ValueAddress);
-						Text = reinterpret_cast<FText*>(ArrayValueAddress + Property->ElementSize * i);
-					}
-					else
-					{
-						Text = reinterpret_cast<FText*>(ValueAddress);
-					}
-
-					if( FTextInspector::GetFlags(*Text) & ETextFlag::ConvertedProperty )
-					{
-						ObjectPackage->MarkPackageDirty();
-					}
-
-					bool Fixed = false;
-					Commandlet->ProcessTextProperty(TextProperty, Text, Object, /*OUT*/Fixed );
-					if( Fixed )
+					// Iterate over all properties of the struct's class.
+					for (TFieldIterator<UProperty>PropIt(StructProperty->Struct, EFieldIteratorFlags::IncludeSuper, EFieldIteratorFlags::ExcludeDeprecated, EFieldIteratorFlags::IncludeInterfaces); PropIt; ++PropIt)
 					{
-						ObjectPackage->MarkPackageDirty();
+						ProcessProperty( *PropIt, PropIt->ContainerPtrToValuePtr<void>(ElementValueAddress) );
 					}
 				}
 			}
@@ -527,6 +503,10 @@ void UGatherTextFromAssetsCommandlet::ProcessPackages( const TArray< UPackage* >
 					UE_LOG(LogGatherTextFromAssetsCommandlet, Warning, TEXT("%s - Invalid generated class!"), *Blueprint->GetFullName());
 				}
 			}
+			else if ( Object->IsA( UUserDefinedStruct::StaticClass() ) )
+			{
+				int x = 0;
+			}
 			else if ( Object->IsA( UDataTable::StaticClass() ) )
 			{
 				UDataTable* DataTable = Cast<UDataTable>(Object);
