Removed debug conditional case in implementation of gathering text from assets.

diff --git a/Engine/Source/Editor/UnrealEd/Private/Commandlets/GatherTextFromAssetsCommandlet.cpp b/Engine/Source/Editor/UnrealEd/Private/Commandlets/GatherTextFromAssetsCommandlet.cpp
--- a/Engine/Source/Editor/UnrealEd/Private/Commandlets/GatherTextFromAssetsCommandlet.cpp
+++ b/Engine/Source/Editor/UnrealEd/Private/Commandlets/GatherTextFromAssetsCommandlet.cpp
@@ -503,10 +503,6 @@ void UGatherTextFromAssetsCommandlet::ProcessPackages( const TArray< UPackage* >
 					UE_LOG(LogGatherTextFromAssetsCommandlet, Warning, TEXT("%s - Invalid generated class!"), *Blueprint->GetFullName());
 				}
 			}
-			else if ( Object->IsA( UUserDefinedStruct::StaticClass() ) )
-			{
-				int x = 0;
-			}
 			else if ( Object->IsA( UDataTable::StaticClass() ) )
 			{
 				UDataTable* DataTable = Cast<UDataTable>(Object);
