Improved splash screen and initialization process.

diff --git a/Engine/Source/Developer/TargetPlatform/Private/TargetPlatformManagerModule.cpp b/Engine/Source/Developer/TargetPlatform/Private/TargetPlatformManagerModule.cpp
--- a/Engine/Source/Developer/TargetPlatform/Private/TargetPlatformManagerModule.cpp
+++ b/Engine/Source/Developer/TargetPlatform/Private/TargetPlatformManagerModule.cpp
@@ -569,8 +569,10 @@ protected:
 			UE_LOG(LogTargetPlatformManager, Error, TEXT("No target platforms found!"));
 		}
 
+		FScopedSlowTask SlowTask(Modules.Num());
 		for (int32 Index = 0; Index < Modules.Num(); Index++)
 		{
+			SlowTask.EnterProgressFrame(1);
 			ITargetPlatformModule* Module = FModuleManager::LoadModulePtr<ITargetPlatformModule>(Modules[Index]);
 			if (Module)
 			{
diff --git a/Engine/Source/Editor/UnrealEd/Private/Editor.cpp b/Engine/Source/Editor/UnrealEd/Private/Editor.cpp
--- a/Engine/Source/Editor/UnrealEd/Private/Editor.cpp
+++ b/Engine/Source/Editor/UnrealEd/Private/Editor.cpp
@@ -507,6 +507,8 @@ extern void StripUnusedPackagesFromList(TArray<FString>& PackageList, const FStr
 
 void UEditorEngine::Init(IEngineLoop* InEngineLoop)
 {
+	FScopedSlowTask SlowTask(100);
+
 	DECLARE_SCOPE_CYCLE_COUNTER(TEXT("Editor Engine Initialized"), STAT_EditorEngineStartup, STATGROUP_LoadTime);
 
 	check(!HasAnyFlags(RF_ClassDefaultObject));
@@ -522,6 +524,7 @@ void UEditorEngine::Init(IEngineLoop* InEngineLoop)
 	FLevelStreamingGCHelper::OnGCStreamedOutLevels.AddUObject(this, &UEditorEngine::OnGCStreamedOutLevels);
 	
 	// Init editor.
+	SlowTask.EnterProgressFrame(40);
 	GEditor = this;
 	InitEditor(InEngineLoop);
 
@@ -533,60 +536,71 @@ void UEditorEngine::Init(IEngineLoop* InEngineLoop)
 	// Load all of the runtime modules that the game needs.  The game is part of the editor, so we'll need these loaded.
 	UGameEngine::LoadRuntimeEngineStartupModules();
 
+	SlowTask.EnterProgressFrame(50);
 
 	// Load all editor modules here
 	{
-		FModuleManager::Get().LoadModule(TEXT("Documentation"));
-		FModuleManager::Get().LoadModule(TEXT("WorkspaceMenuStructure"));
-		FModuleManager::Get().LoadModule(TEXT("MainFrame"));
-		FModuleManager::Get().LoadModule(TEXT("GammaUI"));
-		FModuleManager::Get().LoadModule(TEXT("OutputLog"));
-		FModuleManager::Get().LoadModule(TEXT("SourceControl"));
-		FModuleManager::Get().LoadModule(TEXT("TextureCompressor"));
-		FModuleManager::Get().LoadModule(TEXT("MeshUtilities"));
-		FModuleManager::Get().LoadModule(TEXT("MovieSceneTools"));
-		FModuleManager::Get().LoadModule(TEXT("ModuleUI"));
-		FModuleManager::Get().LoadModule(TEXT("Toolbox"));
-		FModuleManager::Get().LoadModule(TEXT("ClassViewer"));
-		FModuleManager::Get().LoadModule(TEXT("ContentBrowser"));
-		FModuleManager::Get().LoadModule(TEXT("AssetTools"));
-		FModuleManager::Get().LoadModule(TEXT("GraphEditor"));
-		FModuleManager::Get().LoadModule(TEXT("KismetCompiler"));
-		FModuleManager::Get().LoadModule(TEXT("Kismet"));
-		FModuleManager::Get().LoadModule(TEXT("Persona"));
-		FModuleManager::Get().LoadModule(TEXT("LevelEditor"));
-		FModuleManager::Get().LoadModule(TEXT("MainFrame"));
-		FModuleManager::Get().LoadModule(TEXT("PropertyEditor"));
-		FModuleManager::Get().LoadModule(TEXT("EditorStyle"));
-		FModuleManager::Get().LoadModule(TEXT("PackagesDialog"));
-		FModuleManager::Get().LoadModule(TEXT("AssetRegistry"));
-		FModuleManager::Get().LoadModule(TEXT("DetailCustomizations"));
-		FModuleManager::Get().LoadModule(TEXT("ComponentVisualizers"));
-		FModuleManager::Get().LoadModule(TEXT("Layers"));
-		FModuleManager::Get().LoadModule(TEXT("AutomationWindow"));
-		FModuleManager::Get().LoadModule(TEXT("AutomationController"));
-		FModuleManager::Get().LoadModule(TEXT("DeviceManager"));
-		FModuleManager::Get().LoadModule(TEXT("ProfilerClient"));
-//		FModuleManager::Get().LoadModule(TEXT("Search"));
-		FModuleManager::Get().LoadModule(TEXT("SessionFrontend"));
-		FModuleManager::Get().LoadModule(TEXT("ProjectLauncher"));
-		FModuleManager::Get().LoadModule(TEXT("SettingsEditor"));
-		FModuleManager::Get().LoadModule(TEXT("EditorSettingsViewer"));
-		FModuleManager::Get().LoadModule(TEXT("ProjectSettingsViewer"));
-		FModuleManager::Get().LoadModule(TEXT("AndroidRuntimeSettings"));
-		FModuleManager::Get().LoadModule(TEXT("AndroidPlatformEditor"));
-		FModuleManager::Get().LoadModule(TEXT("IOSRuntimeSettings"));
-		FModuleManager::Get().LoadModule(TEXT("IOSPlatformEditor"));
-		FModuleManager::Get().LoadModule(TEXT("Blutility"));
-		FModuleManager::Get().LoadModule(TEXT("OnlineBlueprintSupport"));
-		FModuleManager::Get().LoadModule(TEXT("XmlParser"));
-		FModuleManager::Get().LoadModule(TEXT("UserFeedback"));
-		FModuleManager::Get().LoadModule(TEXT("GameplayTagsEditor"));
-		FModuleManager::Get().LoadModule(TEXT("UndoHistory"));
-		FModuleManager::Get().LoadModule(TEXT("DeviceProfileEditor"));
-		FModuleManager::Get().LoadModule(TEXT("SourceCodeAccess"));
-		FModuleManager::Get().LoadModule(TEXT("BehaviorTreeEditor"));
-		FModuleManager::Get().LoadModule(TEXT("HardwareTargeting"));
+		static const TCHAR* ModuleNames[] =
+		{
+			TEXT("Documentation"),
+			TEXT("WorkspaceMenuStructure"),
+			TEXT("MainFrame"),
+			TEXT("GammaUI"),
+			TEXT("OutputLog"),
+			TEXT("SourceControl"),
+			TEXT("TextureCompressor"),
+			TEXT("MeshUtilities"),
+			TEXT("MovieSceneTools"),
+			TEXT("ModuleUI"),
+			TEXT("Toolbox"),
+			TEXT("ClassViewer"),
+			TEXT("ContentBrowser"),
+			TEXT("AssetTools"),
+			TEXT("GraphEditor"),
+			TEXT("KismetCompiler"),
+			TEXT("Kismet"),
+			TEXT("Persona"),
+			TEXT("LevelEditor"),
+			TEXT("MainFrame"),
+			TEXT("PropertyEditor"),
+			TEXT("EditorStyle"),
+			TEXT("PackagesDialog"),
+			TEXT("AssetRegistry"),
+			TEXT("DetailCustomizations"),
+			TEXT("ComponentVisualizers"),
+			TEXT("Layers"),
+			TEXT("AutomationWindow"),
+			TEXT("AutomationController"),
+			TEXT("DeviceManager"),
+			TEXT("ProfilerClient"),
+//			TEXT("Search"),
+			TEXT("SessionFrontend"),
+			TEXT("ProjectLauncher"),
+			TEXT("SettingsEditor"),
+			TEXT("EditorSettingsViewer"),
+			TEXT("ProjectSettingsViewer"),
+			TEXT("AndroidRuntimeSettings"),
+			TEXT("AndroidPlatformEditor"),
+			TEXT("IOSRuntimeSettings"),
+			TEXT("IOSPlatformEditor"),
+			TEXT("Blutility"),
+			TEXT("OnlineBlueprintSupport"),
+			TEXT("XmlParser"),
+			TEXT("UserFeedback"),
+			TEXT("GameplayTagsEditor"),
+			TEXT("UndoHistory"),
+			TEXT("DeviceProfileEditor"),
+			TEXT("SourceCodeAccess"),
+			TEXT("BehaviorTreeEditor"),
+			TEXT("HardwareTargeting")
+		};
+
+		FScopedSlowTask SlowTask(ARRAY_COUNT(ModuleNames));
+		for (const TCHAR* ModuleName : ModuleNames)
+		{
+			SlowTask.EnterProgressFrame(1);
+			FModuleManager::Get().LoadModule(ModuleName);
+		}
 
 		if (!IsRunningCommandlet())
 		{
@@ -618,6 +632,8 @@ void UEditorEngine::Init(IEngineLoop* InEngineLoop)
 		FModuleManager::Get().LoadModule(TEXT("HotReload"));
 	}
 
+	SlowTask.EnterProgressFrame(10);
+
 	float BSPTexelScale = 100.0f;
 	if( GetDefault<ULevelEditorViewportSettings>()->bUsePowerOf2SnapSize )
 	{
diff --git a/Engine/Source/Editor/UnrealEd/Private/FeedbackContextEditor.cpp b/Engine/Source/Editor/UnrealEd/Private/FeedbackContextEditor.cpp
--- a/Engine/Source/Editor/UnrealEd/Private/FeedbackContextEditor.cpp
+++ b/Engine/Source/Editor/UnrealEd/Private/FeedbackContextEditor.cpp
@@ -475,16 +475,26 @@ void FFeedbackContextEditor::ProgressReported( const float TotalProgressInterp,
 				break;
 			}
 		}
+	}
 
-		if (!DisplayMessage.IsEmpty() && TotalProgressInterp > 0)
+	if (FPlatformSplash::IsShown())
+	{
+		if (!DisplayMessage.IsEmpty())
 		{
-            FFormatOrderedArguments Args;
-            Args.Add(DisplayMessage);
-            Args.Add(int(TotalProgressInterp * 100.f));
-
-            DisplayMessage = FText::Format(NSLOCTEXT("FeedbackContextEditor", "ProgressDisplayText", "{0} ({1}%)"), Args);
+			// Animate a dot to show progress
+			FFormatOrderedArguments Args;
+			Args.Add(DisplayMessage);
+			Args.Add(int(TotalProgressInterp * 100.f));
+
+			double CurrentTime = FPlatformTime::Seconds() / 5.0;
+			int32 Progress = static_cast<int32>((CurrentTime - FMath::FloorToDouble(CurrentTime)) * 5.0);
+			static const TCHAR Spaces[] = TEXT("     ");
+			FString Dots = FString(Progress, Spaces) + TEXT(".") + FString(5 - Progress, Spaces);
+			Args.Add(FText::FromString(Dots));
+			DisplayMessage = FText::Format(NSLOCTEXT("FeedbackContextEditor", "ProgressDisplayText", "{0}{2} ({1}%)"), Args);
 		}
-		FPlatformSplash::SetSplashText( SplashTextType::StartupProgress, *DisplayMessage.ToString() );
+
+		FPlatformSplash::SetSplashText(SplashTextType::StartupProgress, *DisplayMessage.ToString());
 	}
 }
 
diff --git a/Engine/Source/Runtime/Core/Private/Mac/MacPlatformSplash.cpp b/Engine/Source/Runtime/Core/Private/Mac/MacPlatformSplash.cpp
--- a/Engine/Source/Runtime/Core/Private/Mac/MacPlatformSplash.cpp
+++ b/Engine/Source/Runtime/Core/Private/Mac/MacPlatformSplash.cpp
@@ -351,6 +351,11 @@ void FMacPlatformSplash::Hide()
 	}
 }
 
+bool FMacPlatformSplash::IsShown()
+{
+	return (GSplashWindow != NULL);
+}
+
 void FMacPlatformSplash::SetSplashText(const SplashTextType::Type InType, const TCHAR* InText)
 {
 	// We only want to bother drawing startup progress in the editor, since this information is
diff --git a/Engine/Source/Runtime/Core/Private/Misc/AutomationTest.cpp b/Engine/Source/Runtime/Core/Private/Misc/AutomationTest.cpp
--- a/Engine/Source/Runtime/Core/Private/Misc/AutomationTest.cpp
+++ b/Engine/Source/Runtime/Core/Private/Misc/AutomationTest.cpp
@@ -133,8 +133,11 @@ bool FAutomationTestFramework::RunSmokeTests()
 			TMap<FString, FAutomationTestExecutionInfo> OutExecutionInfoMap;
 
 			// Run each valid test
+			FScopedSlowTask SlowTask(TestInfo.Num());
+
 			for ( int TestIndex = 0; TestIndex < TestInfo.Num(); ++TestIndex )
 			{
+				SlowTask.EnterProgressFrame(1);
 				if (TestInfo[TestIndex].GetTestType() == EAutomationTestType::ATT_SmokeTest )
 				{
 					FString TestCommand = TestInfo[TestIndex].GetTestName();
diff --git a/Engine/Source/Runtime/Core/Private/Windows/WindowsPlatformSplash.cpp b/Engine/Source/Runtime/Core/Private/Windows/WindowsPlatformSplash.cpp
--- a/Engine/Source/Runtime/Core/Private/Windows/WindowsPlatformSplash.cpp
+++ b/Engine/Source/Runtime/Core/Private/Windows/WindowsPlatformSplash.cpp
@@ -730,6 +730,10 @@ void FWindowsPlatformSplash::Hide()
 }
 
 
+bool FWindowsPlatformSplash::IsShown()
+{
+	return (GSplashScreenThread != nullptr);
+}
 
 /**
  * Sets the text displayed on the splash screen (for startup/loading progress)
@@ -746,11 +750,11 @@ void FWindowsPlatformSplash::SetSplashText( const SplashTextType::Type InType, c
 		// Only allow copyright text displayed while loading the game.  Editor displays all.
 		if( InType == SplashTextType::CopyrightInfo || GIsEditor )
 		{
+			// Take a critical section since the splash thread may already be repainting using this text
+			FScopeLock ScopeLock(&GSplashScreenSynchronizationObject);
+
 			bool bWasUpdated = false;
 			{
-				// Take a critical section since the splash thread may already be repainting using this text
-				FScopeLock ScopeLock( &GSplashScreenSynchronizationObject );
-
 				// Update splash text
 				if( FCString::Strcmp( InText, *GSplashScreenText[ InType ].ToString() ) != 0 )
 				{
diff --git a/Engine/Source/Runtime/Core/Public/GenericPlatform/GenericPlatformSplash.h b/Engine/Source/Runtime/Core/Public/GenericPlatform/GenericPlatformSplash.h
--- a/Engine/Source/Runtime/Core/Public/GenericPlatform/GenericPlatformSplash.h
+++ b/Engine/Source/Runtime/Core/Public/GenericPlatform/GenericPlatformSplash.h
@@ -62,4 +62,12 @@ struct CORE_API FGenericPlatformSplash
 
 	}
 
+	/**
+	 * Return whether the splash screen is being shown or not
+	 */
+	FORCEINLINE static bool IsShown()
+	{
+		return true;
+	}
+
 };
diff --git a/Engine/Source/Runtime/Core/Public/Mac/MacPlatformSplash.h b/Engine/Source/Runtime/Core/Public/Mac/MacPlatformSplash.h
--- a/Engine/Source/Runtime/Core/Public/Mac/MacPlatformSplash.h
+++ b/Engine/Source/Runtime/Core/Public/Mac/MacPlatformSplash.h
@@ -27,6 +27,11 @@ struct CORE_API FMacPlatformSplash : public FGenericPlatformSplash
 	 * @param	InText		Text to display
 	 */
 	static void SetSplashText( const SplashTextType::Type InType, const TCHAR* InText );
+
+	/**
+	 * Return whether the splash screen is being shown or not
+	 */
+	static bool IsShown();
 };
 
 typedef FMacPlatformSplash FPlatformSplash;
diff --git a/Engine/Source/Runtime/Core/Public/Windows/WindowsPlatformSplash.h b/Engine/Source/Runtime/Core/Public/Windows/WindowsPlatformSplash.h
--- a/Engine/Source/Runtime/Core/Public/Windows/WindowsPlatformSplash.h
+++ b/Engine/Source/Runtime/Core/Public/Windows/WindowsPlatformSplash.h
@@ -24,6 +24,11 @@ struct CORE_API FWindowsPlatformSplash
 	 * @param InText Text to display.
 	 */
 	static void SetSplashText( const SplashTextType::Type InType, const TCHAR* InText );
+
+	/**
+	 * Return whether the splash screen is being shown or not
+	 */
+	static bool IsShown();
 };
 
 
diff --git a/Engine/Source/Runtime/Engine/Private/GlobalShader.cpp b/Engine/Source/Runtime/Engine/Private/GlobalShader.cpp
--- a/Engine/Source/Runtime/Engine/Private/GlobalShader.cpp
+++ b/Engine/Source/Runtime/Engine/Private/GlobalShader.cpp
@@ -349,12 +349,15 @@ TShaderMap<FGlobalShaderType>* GetGlobalShaderMap(EShaderPlatform Platform, bool
 	// If the global shader map hasn't been created yet, create it.
 	if(!GGlobalShaderMap[Platform])
 	{
-		// verify that all shader source files are intact
-		VerifyShaderSourceFiles();
-
 		// GetGlobalShaderMap is called the first time during startup in the main thread.
 		check(IsInGameThread());
 
+		FScopedSlowTask SlowTask(70);
+
+		// verify that all shader source files are intact
+		SlowTask.EnterProgressFrame(20);
+		VerifyShaderSourceFiles();
+
 		GGlobalShaderMap[Platform] = new TShaderMap<FGlobalShaderType>();
 
 		bool bLoadedFromCacheFile = false;
@@ -363,6 +366,8 @@ TShaderMap<FGlobalShaderType>* GetGlobalShaderMap(EShaderPlatform Platform, bool
 		// This method is used exclusively with cooked content, since the DDC is not present
 		if (FPlatformProperties::RequiresCookedData())
 		{
+			SlowTask.EnterProgressFrame(50);
+
 			TArray<uint8> GlobalShaderData;
 			FString GlobalShaderCacheFilename = FPaths::GetRelativePathToRoot() / GetGlobalShaderCacheFilename(Platform);
 			FPaths::MakeStandardFilename(GlobalShaderCacheFilename);
@@ -396,9 +401,11 @@ TShaderMap<FGlobalShaderType>* GetGlobalShaderMap(EShaderPlatform Platform, bool
 			FGlobalShaderMapId ShaderMapId(Platform);
 
 			TArray<uint8> CachedData;
+			SlowTask.EnterProgressFrame(40);
 			const FString DataKey = GetGlobalShaderMapKeyString(ShaderMapId, Platform);
 
 			// Find the shader map in the derived data cache
+			SlowTask.EnterProgressFrame(10);
 			if (GetDerivedDataCacheRef().GetSynchronous(*DataKey, CachedData))
 			{
 				FMemoryReader Ar(CachedData, true);
diff --git a/Engine/Source/Runtime/Launch/Private/Launch.cpp b/Engine/Source/Runtime/Launch/Private/Launch.cpp
--- a/Engine/Source/Runtime/Launch/Private/Launch.cpp
+++ b/Engine/Source/Runtime/Launch/Private/Launch.cpp
@@ -115,15 +115,26 @@ int32 GuardedMain( const TCHAR* CmdLine )
 		return ErrorLevel;
 	}
 
-#if WITH_EDITOR
-	if( GIsEditor )
 	{
-		ErrorLevel = EditorInit( GEngineLoop );
-	}
-	else
+		FScopedSlowTask SlowTask(100, NSLOCTEXT("EngineInit", "EngineInit_Loading", "Loading..."));
+
+		// EnginePreInit leaves 20% unused in its slow task.
+		// Here we consume 80% immediately so that the percentage value on the splash screen doesn't change from one slow task to the next.
+		// (Note, we can't include the call to EnginePreInit in this ScopedSlowTask, because the engine isn't fully initialized at that point)
+		SlowTask.EnterProgressFrame(80);
+
+		SlowTask.EnterProgressFrame(20);
+
+#if WITH_EDITOR
+		if (GIsEditor)
+		{
+			ErrorLevel = EditorInit(GEngineLoop);
+		}
+		else
 #endif
-	{
-		ErrorLevel = EngineInit();
+		{
+			ErrorLevel = EngineInit();
+		}
 	}
 
 	UE_LOG(LogLoad, Log, TEXT("Full Startup: %.2f seconds (BP compile: %.2f seconds)"), FPlatformTime::Seconds() - GStartTime, GBlueprintCompileTime);
diff --git a/Engine/Source/Runtime/Launch/Private/LaunchEngineLoop.cpp b/Engine/Source/Runtime/Launch/Private/LaunchEngineLoop.cpp
--- a/Engine/Source/Runtime/Launch/Private/LaunchEngineLoop.cpp
+++ b/Engine/Source/Runtime/Launch/Private/LaunchEngineLoop.cpp
@@ -1287,6 +1287,8 @@ int32 FEngineLoop::PreInit( const TCHAR* CmdLine )
 		// Initialize shader types before loading any shaders
 		InitializeShaderTypes();
 
+		SlowTask.EnterProgressFrame(30);
+	
 		// Load the global shaders.
 		if (GetGlobalShaderMap(GMaxRHIFeatureLevel) == NULL && GIsRequestingExit)
 		{
@@ -1298,6 +1300,7 @@ int32 FEngineLoop::PreInit( const TCHAR* CmdLine )
 		// package names from ini files and register them with FPackageName system.
 		FPackageName::RegisterShortPackageNamesForUObjectModules();
 
+		SlowTask.EnterProgressFrame(5);
 
 		// Make sure all UObject classes are registered and default properties have been initialized
 		ProcessNewlyLoadedUObjects();
@@ -1309,7 +1312,7 @@ int32 FEngineLoop::PreInit( const TCHAR* CmdLine )
 		UMaterialInterface::AssertDefaultMaterialsPostLoaded();
 	}
 
-	SlowTask.EnterProgressFrame(10);
+	SlowTask.EnterProgressFrame(5);
 
 	// Tell the module manager is may now process newly-loaded UObjects when new C++ modules are loaded
 	FModuleManager::Get().StartProcessingNewlyLoadedObjects();
@@ -1377,7 +1380,7 @@ int32 FEngineLoop::PreInit( const TCHAR* CmdLine )
         // do any post appInit processing, before the render thread is started.
         FPlatformMisc::PlatformPostInit(true);
     }
-	SlowTask.EnterProgressFrame(10);
+	SlowTask.EnterProgressFrame(5);
 
 	if (GUseThreadedRendering)
 	{
@@ -1407,7 +1410,7 @@ int32 FEngineLoop::PreInit( const TCHAR* CmdLine )
 
 	SetIsServerForOnlineSubsystemsDelegate(FQueryIsRunningServer::CreateStatic(&IsServerDelegateForOSS));
 
-	SlowTask.EnterProgressFrame(50);
+	SlowTask.EnterProgressFrame(5);
 
 	if (!bHasEditorToken)
 	{
@@ -1702,6 +1705,7 @@ int32 FEngineLoop::PreInit( const TCHAR* CmdLine )
 	//run automation smoke tests now that everything is setup to run
 	FAutomationTestFramework::GetInstance().RunSmokeTests();
 
+	// Note we still have 20% remaining on the slow task: this will be used by the Editor/Engine initialization next
 	return 0;
 }
 
@@ -1751,19 +1755,24 @@ void FEngineLoop::LoadPreInitModules()
 
 bool FEngineLoop::LoadStartupCoreModules()
 {
+	FScopedSlowTask SlowTask(100);
+
 	DECLARE_SCOPE_CYCLE_COUNTER(TEXT("Loading Startup Modules"), STAT_StartupModules, STATGROUP_LoadTime);
 
 	bool bSuccess = true;
 
 	// Load all Runtime modules
+	SlowTask.EnterProgressFrame(10);
 	{
 		FModuleManager::Get().LoadModule(TEXT("Core"));
 		FModuleManager::Get().LoadModule(TEXT("Networking"));
 	}
 
+	SlowTask.EnterProgressFrame(10);
 	FPlatformMisc::LoadStartupModules();
 
 	// initialize messaging
+	SlowTask.EnterProgressFrame(10);
 	if (FPlatformProcess::SupportsMultithreading())
 	{
 		FModuleManager::LoadModuleChecked<IMessagingModule>("Messaging");
@@ -1777,11 +1786,13 @@ bool FEngineLoop::LoadStartupCoreModules()
 		EngineService = new FEngineService();
 	}
 
+	SlowTask.EnterProgressFrame(10);
 #if WITH_EDITOR
 		FModuleManager::LoadModuleChecked<IEditorStyleModule>("EditorStyle");
 #endif //WITH_EDITOR
 
 	// Load all Development modules
+	SlowTask.EnterProgressFrame(30);
 	if (!IsRunningDedicatedServer())
 	{
 		FModuleManager::Get().LoadModule("Slate");
@@ -1798,6 +1809,7 @@ bool FEngineLoop::LoadStartupCoreModules()
 		FModuleManager::Get().LoadModule("FunctionalTesting");
 #endif	//WITH_UNREAL_DEVELOPER_TOOLS
 
+	SlowTask.EnterProgressFrame(30);
 #if (WITH_EDITOR && !(UE_BUILD_SHIPPING || UE_BUILD_TEST))
 	// HACK: load BT editor as early as possible for statically initialized assets (non cooked BT assets needs it)
 	// cooking needs this module too
@@ -1839,18 +1851,23 @@ bool FEngineLoop::LoadStartupCoreModules()
 
 bool FEngineLoop::LoadStartupModules()
 {
+	FScopedSlowTask SlowTask(3);
+
+	SlowTask.EnterProgressFrame(1);
 	// Load any modules that want to be loaded before default modules are loaded up.
 	if (!IProjectManager::Get().LoadModulesForProject(ELoadingPhase::PreDefault) || !IPluginManager::Get().LoadModulesForEnabledPlugins(ELoadingPhase::PreDefault))
 	{
 		return false;
 	}
 
+	SlowTask.EnterProgressFrame(1);
 	// Load modules that are configured to load in the default phase
 	if (!IProjectManager::Get().LoadModulesForProject(ELoadingPhase::Default) || !IPluginManager::Get().LoadModulesForEnabledPlugins(ELoadingPhase::Default))
 	{
 		return false;
 	}
 
+	SlowTask.EnterProgressFrame(1);
 	// Load any modules that want to be loaded after default modules are loaded up.
 	if (!IProjectManager::Get().LoadModulesForProject(ELoadingPhase::PostDefault) || !IPluginManager::Get().LoadModulesForEnabledPlugins(ELoadingPhase::PostDefault))
 	{
diff --git a/Engine/Source/Runtime/Projects/Private/ModuleDescriptor.cpp b/Engine/Source/Runtime/Projects/Private/ModuleDescriptor.cpp
--- a/Engine/Source/Runtime/Projects/Private/ModuleDescriptor.cpp
+++ b/Engine/Source/Runtime/Projects/Private/ModuleDescriptor.cpp
@@ -331,8 +331,10 @@ bool FModuleDescriptor::IsLoadedInCurrentConfiguration() const
 
 void FModuleDescriptor::LoadModulesForPhase(ELoadingPhase::Type LoadingPhase, const TArray<FModuleDescriptor>& Modules, TMap<FName, EModuleLoadResult>& ModuleLoadErrors)
 {
+	FScopedSlowTask SlowTask(Modules.Num());
 	for(int Idx = 0; Idx < Modules.Num(); Idx++)
 	{
+		SlowTask.EnterProgressFrame(1);
 		const FModuleDescriptor& Descriptor = Modules[Idx];
 
 		// Don't need to do anything if this module is already loaded
diff --git a/Engine/Source/Runtime/Projects/Private/PluginManager.cpp b/Engine/Source/Runtime/Projects/Private/PluginManager.cpp
--- a/Engine/Source/Runtime/Projects/Private/PluginManager.cpp
+++ b/Engine/Source/Runtime/Projects/Private/PluginManager.cpp
@@ -409,9 +409,13 @@ bool FPluginManager::LoadModulesForEnabledPlugins( const ELoadingPhase::Type Loa
 		return false;
 	}
 
+	FScopedSlowTask SlowTask(AllPlugins.Num());
+
 	// Load plugins!
 	for( const TSharedRef< FPluginInstance > Plugin : AllPlugins )
 	{
+		SlowTask.EnterProgressFrame(1);
+
 		if ( Plugin->bEnabled )
 		{
 			TMap<FName, EModuleLoadResult> ModuleLoadFailures;
diff --git a/Engine/Source/Runtime/ShaderCore/Private/ShaderCore.cpp b/Engine/Source/Runtime/ShaderCore/Private/ShaderCore.cpp
--- a/Engine/Source/Runtime/ShaderCore/Private/ShaderCore.cpp
+++ b/Engine/Source/Runtime/ShaderCore/Private/ShaderCore.cpp
@@ -252,8 +252,10 @@ void VerifyShaderSourceFiles()
 		// get the list of shader files that can be used
 		TArray<FString> ShaderSourceFiles;
 		GetAllShaderSourceFiles(ShaderSourceFiles);
+		FScopedSlowTask SlowTask(ShaderSourceFiles.Num());
 		for( int32 ShaderFileIdx=0; ShaderFileIdx < ShaderSourceFiles.Num(); ShaderFileIdx++ )
 		{
+			SlowTask.EnterProgressFrame(1);
 			FString FileContents;
 			// load each shader source file. This will cache the shader source data after it has been verified
 			LoadShaderSourceFileChecked(*ShaderSourceFiles[ShaderFileIdx], FileContents);
@@ -509,8 +511,11 @@ void BuildShaderFileToUniformBufferMap(TMap<FString, TArray<const TCHAR*> >& Sha
 		TArray<FString> ShaderSourceFiles;
 		GetAllShaderSourceFiles(ShaderSourceFiles);
 
+		FScopedSlowTask SlowTask(ShaderSourceFiles.Num());
+
 		for (int32 FileIndex = 0; FileIndex < ShaderSourceFiles.Num(); FileIndex++)
 		{
+			SlowTask.EnterProgressFrame(1);
 			FString ShaderFileContents;
 
 			LoadShaderSourceFileChecked(*ShaderSourceFiles[FileIndex], ShaderFileContents);
