UE4 - Added back Period for Time material expressions on the CPU (only used if enabling checkbox)

diff --git a/Engine/Shaders/ShaderVersion.usf b/Engine/Shaders/ShaderVersion.usf
--- a/Engine/Shaders/ShaderVersion.usf
+++ b/Engine/Shaders/ShaderVersion.usf
@@ -2,4 +2,4 @@
 // Each time the console command is executed it generates a new GUID. As this file is included
 // in common.usf (which should be included in any shader) it allows to invalidate the shader DDC.
 // 
-// GUID = {CAC8D6F0-2E83-4C3D-A767-584D1E3DC80x}
+// GUID = {CAC8D6F0-2E83-4C3D-A767-584D1E3DC80y}
diff --git a/Engine/Source/Editor/UnrealEd/Private/Lightmass/LightmassRender.cpp b/Engine/Source/Editor/UnrealEd/Private/Lightmass/LightmassRender.cpp
--- a/Engine/Source/Editor/UnrealEd/Private/Lightmass/LightmassRender.cpp
+++ b/Engine/Source/Editor/UnrealEd/Private/Lightmass/LightmassRender.cpp
@@ -25,101 +25,101 @@ struct FLightmassMaterialCompiler : public FProxyMaterialCompiler
 	{}
 
 	// gets value stored by SetMaterialProperty()
-	virtual EShaderFrequency GetCurrentShaderFrequency() const
+	virtual EShaderFrequency GetCurrentShaderFrequency() const override
 	{
 		// not used by Lightmass
 		return SF_Pixel;
 	}
 
-	virtual int32 ParticleMacroUV()
+	virtual int32 ParticleMacroUV() override
 	{
 		return Compiler->ParticleMacroUV();
 	}
 
-	virtual int32 ParticleRelativeTime()
+	virtual int32 ParticleRelativeTime() override
 	{
 		return Compiler->Constant(0.0f);
 	}
 
-	virtual int32 ParticleMotionBlurFade()
+	virtual int32 ParticleMotionBlurFade() override
 	{
 		return Compiler->Constant(1.0f);
 	}
 
-	virtual int32 ParticleDirection()
+	virtual int32 ParticleDirection() override
 	{
 		return Compiler->Constant3(0.0f, 0.0f, 0.0f);
 	}
 
-	virtual int32 ParticleSpeed()
+	virtual int32 ParticleSpeed() override
 	{
 		return Compiler->Constant(0.0f);
 	}
 	
-	virtual int32 ParticleSize()
+	virtual int32 ParticleSize() override
 	{
 		return Compiler->Constant2(0.0f,0.0f);
 	}
 
-	virtual int32 WorldPosition(EWorldPositionIncludedOffsets WorldPositionIncludedOffsets)
+	virtual int32 WorldPosition(EWorldPositionIncludedOffsets WorldPositionIncludedOffsets) override
 	{
 		//UE_LOG(LogLightmassRender, Log, TEXT("Lightmass material compiler has encountered WorldPosition... Forcing constant (0.0f,0.0f,0.0f)."));
 		return Compiler->Constant3(0.0f,0.0f,0.0f);
 	}
 
-	virtual int32 ObjectWorldPosition() 
+	virtual int32 ObjectWorldPosition() override
 	{
 		//UE_LOG(LogLightmassRender, Log, TEXT("Lightmass material compiler has encountered ObjectWorldPosition... Forcing constant (0.0f,0.0f,0.0f)."));
 		return Compiler->Constant3(0.0f,0.0f,0.0f);
 	}
 
-	virtual int32 ObjectRadius() 
+	virtual int32 ObjectRadius() override
 	{
 		//UE_LOG(LogLightmassRender, Log, TEXT("Lightmass material compiler has encountered ObjectRadius... Forcing constant 500.0f."));
 		return Compiler->Constant(500);
 	}
 
-	virtual int32 ObjectBounds() 
+	virtual int32 ObjectBounds() override
 	{
 		//UE_LOG(LogLightmassRender, Log, TEXT("Lightmass material compiler has encountered ObjectBounds... Forcing constant (0,0,0)."));
 		return Compiler->Constant3(0,0,0);
 	}
 
-	virtual int32 DistanceCullFade() 
+	virtual int32 DistanceCullFade() override
 	{
 		return Compiler->Constant(1.0f);
 	}
 
-	virtual int32 ActorWorldPosition() 
+	virtual int32 ActorWorldPosition() override
 	{
 		return Compiler->Constant3(0.0f,0.0f,0.0f);
 	}
 
-	virtual int32 CameraWorldPosition()
+	virtual int32 CameraWorldPosition() override
 	{
 		//UE_LOG(LogLightmassRender, Log, TEXT("Lightmass material compiler has encountered CameraWorldPosition... Forcing constant (0.0f,0.0f,0.0f)."));
 		return Compiler->Constant3(0.0f,0.0f,0.0f);
 	}
 
-	virtual int32 CameraVector()
+	virtual int32 CameraVector() override
 	{
 		//UE_LOG(LogLightmassRender, Log, TEXT("Lightmass material compiler has encountered CameraVector... Forcing constant (0.0f,0.0f,1.0f)."));
 		return Compiler->Constant3(0.0f,0.0f,1.0f);
 	}
 
-	virtual int32 LightVector()
+	virtual int32 LightVector() override
 	{
 		//UE_LOG(LogLightmassRender, Log, TEXT("Lightmass material compiler has encountered LightVector... Forcing constant (1.0f,0.0f,0.0f)."));
 		return Compiler->Constant3(1.0f,0.0f,0.0f);
 	}
 
-	virtual int32 ReflectionVector()
+	virtual int32 ReflectionVector() override
 	{
 		//UE_LOG(LogLightmassRender, Log, TEXT("Lightmass material compiler has encountered ReflectionVector... Forcing constant (0.0f,0.0f,-1.0f)."));
 		return Compiler->Constant3(0.0f,0.0f,-1.0f);
 	}
 
-	virtual int32 ReflectionAboutCustomWorldNormal(int32 CustomWorldNormal, int32 bNormalizeCustomWorldNormal)
+	virtual int32 ReflectionAboutCustomWorldNormal(int32 CustomWorldNormal, int32 bNormalizeCustomWorldNormal) override
 	{
 		//UE_LOG(LogLightmassRender, Log, TEXT("Lightmass material compiler has encountered ReflectionAboutCustomNormalVector... Forcing constant (0.0f,0.0f,-1.0f)."));
 		return Compiler->Constant3(0.0f,0.0f,-1.0f);
@@ -128,7 +128,7 @@ struct FLightmassMaterialCompiler : public FProxyMaterialCompiler
 	/**
 	 *	Generate shader code for transforming a vector
 	 */
-	virtual int32 TransformVector(uint8 SourceCoordType,uint8 DestCoordType,int32 A)
+	virtual int32 TransformVector(uint8 SourceCoordType,uint8 DestCoordType,int32 A) override
 	{
 		//UE_LOG(LogLightmassRender, Log, TEXT("Lightmass material compiler has encountered TransformVector... Passing thru source vector untouched."));
 		return A;
@@ -140,33 +140,33 @@ struct FLightmassMaterialCompiler : public FProxyMaterialCompiler
 	 *	@param	CoordType - type of transform to apply. see EMaterialExpressionTransformPosition 
 	 *	@param	A - index for input vector parameter's code
 	 */
-	virtual int32 TransformPosition(uint8 SourceCoordType,uint8 DestCoordType,int32 A)
+	virtual int32 TransformPosition(uint8 SourceCoordType,uint8 DestCoordType,int32 A) override
 	{
 		//UE_LOG(LogLightmassRender, Log, TEXT("Lightmass material compiler has encountered TransformPosition... Passing thru source vector untouched."));
 		return A;
 	}
 
-	virtual int32 VertexColor()
+	virtual int32 VertexColor() override
 	{
 		//UE_LOG(LogLightmassRender, Log, TEXT("Lightmass material compiler has encountered VertexColor... Forcing constant (1.0f,1.0f,1.0f,1.0f)."));
 		return Compiler->Constant4(1.0f,1.0f,1.0f,1.0f);
 	}
 
-	virtual int32 RealTime()
+	virtual int32 RealTime(bool bPeriodic, float Period) override
 	{
 		//UE_LOG(LogLightmassRender, Log, TEXT("Lightmass material compiler has encountered RealTime... Forcing constant 0.0f."));
 		return Compiler->Constant(0.0f);
 	}
 
-	virtual int32 GameTime()
+	virtual int32 GameTime(bool bPeriodic, float Period) override
 	{
 		//UE_LOG(LogLightmassRender, Log, TEXT("Lightmass material compiler has encountered GameTime... Forcing constant 0.0f."));
 		return Compiler->Constant(0.0f);
 	}
 
-	virtual int32 LightmassReplace(int32 Realtime, int32 Lightmass) { return Lightmass; }
+	virtual int32 LightmassReplace(int32 Realtime, int32 Lightmass) override { return Lightmass; }
 
-	virtual int32 GIReplace(int32 Direct, int32 StaticIndirect, int32 DynamicIndirect) { return StaticIndirect; }
+	virtual int32 GIReplace(int32 Direct, int32 StaticIndirect, int32 DynamicIndirect) override { return StaticIndirect; }
 };
 
 /**
diff --git a/Engine/Source/Runtime/Engine/Classes/Materials/MaterialExpressionTime.h b/Engine/Source/Runtime/Engine/Classes/Materials/MaterialExpressionTime.h
--- a/Engine/Source/Runtime/Engine/Classes/Materials/MaterialExpressionTime.h
+++ b/Engine/Source/Runtime/Engine/Classes/Materials/MaterialExpressionTime.h
@@ -2,6 +2,7 @@
 
 
 #pragma once
+
 #include "Materials/MaterialExpression.h"
 #include "MaterialExpressionTime.generated.h"
 
@@ -14,6 +15,13 @@ class UMaterialExpressionTime : public UMaterialExpression
 	UPROPERTY(EditAnywhere, Category=MaterialExpressionTime)
 	uint32 bIgnorePause:1;
 
+	/** Enables or disables the Period value. */
+	UPROPERTY()
+	uint32 bOverride_Period:1;
+
+	/** Time will loop around once it gets to Period. */
+	UPROPERTY(EditAnywhere, Category=MaterialExpressionTime, meta = (editcondition = "bOverride_Period", ClampMin = "0.0", ToolTip = "Period at which to wrap around time"))
+	float Period;
 
 	// Begin UMaterialExpression Interface
 	virtual int32 Compile(class FMaterialCompiler* Compiler, int32 OutputIndex, int32 MultiplexIndex) override;
@@ -22,6 +30,3 @@ class UMaterialExpressionTime : public UMaterialExpression
 	// End UMaterialExpression Interface
 
 };
-
-
-
diff --git a/Engine/Source/Runtime/Engine/Private/Materials/HLSLMaterialTranslator.h b/Engine/Source/Runtime/Engine/Private/Materials/HLSLMaterialTranslator.h
--- a/Engine/Source/Runtime/Engine/Private/Materials/HLSLMaterialTranslator.h
+++ b/Engine/Source/Runtime/Engine/Private/Materials/HLSLMaterialTranslator.h
@@ -1511,14 +1511,44 @@ protected:
 		return AddUniformExpression(new FMaterialUniformExpressionConstant(FLinearColor(X,Y,Z,W),MCT_Float4),MCT_Float4,TEXT("MaterialFloat4(%0.8f,%0.8f,%0.8f,%0.8f)"),X,Y,Z,W);
 	}
 
-	virtual int32 GameTime() override
+	virtual int32 GameTime(bool bPeriodic, float Period) override
 	{
-		return AddInlinedCodeChunk(MCT_Float, TEXT("View.GameTime"));
+		if (!bPeriodic)
+		{
+			return AddInlinedCodeChunk(MCT_Float, TEXT("View.GameTime"));
+		}
+		else if (Period == 0.0f)
+		{
+			return Constant(0.0f);
+		}
+
+		return AddUniformExpression(
+			new FMaterialUniformExpressionFmod(
+				new FMaterialUniformExpressionTime(),
+				new FMaterialUniformExpressionConstant(FLinearColor(Period, Period, Period, Period), MCT_Float)
+				),
+			MCT_Float, TEXT("")
+			);
 	}
 
-	virtual int32 RealTime() override
+	virtual int32 RealTime(bool bPeriodic, float Period) override
 	{
-		return AddInlinedCodeChunk(MCT_Float, TEXT("View.RealTime"));
+		if (!bPeriodic)
+		{
+			return AddInlinedCodeChunk(MCT_Float, TEXT("View.RealTime"));
+		}
+		else if (Period == 0.0f)
+		{
+			return Constant(0.0f);
+		}
+
+		return AddUniformExpression(
+			new FMaterialUniformExpressionFmod(
+				new FMaterialUniformExpressionRealTime(),
+				new FMaterialUniformExpressionConstant(FLinearColor(Period, Period, Period, Period), MCT_Float)
+				),
+			MCT_Float, TEXT("")
+			);
 	}
 
 	virtual int32 PeriodicHint(int32 PeriodicCode) override
@@ -2861,16 +2891,37 @@ protected:
 			return Errorf(TEXT("Couldn't determine result type of component mask %u%u%u%u"),R,G,B,A);
 		};
 
-		return AddInlinedCodeChunk(
-			ResultType,
-			TEXT("%s.%s%s%s%s"),
-			*GetParameterCode(Vector),
+		FString MaskString = FString::Printf(TEXT("%s%s%s%s"),
 			R ? TEXT("r") : TEXT(""),
 			// If VectorType is set to MCT_Float which means it could be any of the float types, assume it is a float1
 			G ? (VectorType == MCT_Float ? TEXT("r") : TEXT("g")) : TEXT(""),
 			B ? (VectorType == MCT_Float ? TEXT("r") : TEXT("b")) : TEXT(""),
 			A ? (VectorType == MCT_Float ? TEXT("r") : TEXT("a")) : TEXT("")
 			);
+
+		auto* Expression = GetParameterUniformExpression(Vector);
+		if (Expression)
+		{
+			int8 Mask[4] = {-1, -1, -1, -1};
+			for (int32 Index = 0; Index < MaskString.Len(); ++Index)
+			{
+				Mask[Index] = SwizzleComponentToIndex(MaskString[Index]);
+			}
+			return AddUniformExpression(
+				new FMaterialUniformExpressionComponentSwizzle(Expression, Mask[0], Mask[1], Mask[2], Mask[3]),
+				ResultType,
+				TEXT("%s.%s"),
+				*GetParameterCode(Vector),
+				*MaskString
+				);
+		}
+
+		return AddInlinedCodeChunk(
+			ResultType,
+			TEXT("%s.%s"),
+			*GetParameterCode(Vector),
+			*MaskString
+			);
 	}
 
 	virtual int32 AppendVector(int32 A,int32 B) override
diff --git a/Engine/Source/Runtime/Engine/Private/Materials/MaterialExpressions.cpp b/Engine/Source/Runtime/Engine/Private/Materials/MaterialExpressions.cpp
--- a/Engine/Source/Runtime/Engine/Private/Materials/MaterialExpressions.cpp
+++ b/Engine/Source/Runtime/Engine/Private/Materials/MaterialExpressions.cpp
@@ -2716,16 +2716,32 @@ UMaterialExpressionTime::UMaterialExpressionTime(const FObjectInitializer& Objec
 
 	MenuCategories.Add(ConstructorStatics.NAME_Constants);
 	bShaderInputData = true;
+	Period = 0.0f;
+	bOverride_Period = false;
 }
 
 int32 UMaterialExpressionTime::Compile(class FMaterialCompiler* Compiler, int32 OutputIndex, int32 MultiplexIndex)
 {
-	return bIgnorePause ? Compiler->RealTime() : Compiler->GameTime();
+	return bIgnorePause ? Compiler->RealTime(bOverride_Period, Period) : Compiler->GameTime(bOverride_Period, Period);
 }
 
 void UMaterialExpressionTime::GetCaption(TArray<FString>& OutCaptions) const
 {
-	OutCaptions.Add(TEXT("Time"));
+	if (bOverride_Period)
+	{
+		if (Period == 0.0f)
+		{
+			OutCaptions.Add(TEXT("Time (Stopped)"));
+		}
+		else
+		{
+			OutCaptions.Add(FString::Printf(TEXT("Time (Period of %.2f)"), Period));
+		}
+	}
+	else
+	{
+		OutCaptions.Add(TEXT("Time"));
+	}
 }
 
 
@@ -2857,7 +2873,7 @@ UMaterialExpressionPanner::UMaterialExpressionPanner(const FObjectInitializer& O
 
 int32 UMaterialExpressionPanner::Compile(class FMaterialCompiler* Compiler, int32 OutputIndex, int32 MultiplexIndex)
 {
-	int32 TimeArg = Time.Expression ? Time.Compile(Compiler) : Compiler->GameTime();
+	int32 TimeArg = Time.Expression ? Time.Compile(Compiler) : Compiler->GameTime(false, 0.0f);
 	int32 Arg1;
 	int32 Arg2;
 	if (bFractionalPart)
@@ -2917,8 +2933,8 @@ UMaterialExpressionRotator::UMaterialExpressionRotator(const FObjectInitializer&
 
 int32 UMaterialExpressionRotator::Compile(class FMaterialCompiler* Compiler, int32 OutputIndex, int32 MultiplexIndex)
 {
-	int32	Cosine = Compiler->Cosine(Compiler->Mul(Time.Expression ? Time.Compile(Compiler) : Compiler->GameTime(),Compiler->Constant(Speed))),
-		Sine = Compiler->Sine(Compiler->Mul(Time.Expression ? Time.Compile(Compiler) : Compiler->GameTime(),Compiler->Constant(Speed))),
+	int32	Cosine = Compiler->Cosine(Compiler->Mul(Time.Expression ? Time.Compile(Compiler) : Compiler->GameTime(false, 0.0f),Compiler->Constant(Speed))),
+		Sine = Compiler->Sine(Compiler->Mul(Time.Expression ? Time.Compile(Compiler) : Compiler->GameTime(false, 0.0f),Compiler->Constant(Speed))),
 		RowX = Compiler->AppendVector(Cosine,Compiler->Mul(Compiler->Constant(-1.0f),Sine)),
 		RowY = Compiler->AppendVector(Sine,Cosine),
 		Origin = Compiler->Constant2(CenterX,CenterY),
diff --git a/Engine/Source/Runtime/Engine/Private/Materials/MaterialUniformExpressions.cpp b/Engine/Source/Runtime/Engine/Private/Materials/MaterialUniformExpressions.cpp
--- a/Engine/Source/Runtime/Engine/Private/Materials/MaterialUniformExpressions.cpp
+++ b/Engine/Source/Runtime/Engine/Private/Materials/MaterialUniformExpressions.cpp
@@ -490,6 +490,7 @@ IMPLEMENT_MATERIALUNIFORMEXPRESSION_TYPE(FMaterialUniformExpressionMin);
 IMPLEMENT_MATERIALUNIFORMEXPRESSION_TYPE(FMaterialUniformExpressionMax);
 IMPLEMENT_MATERIALUNIFORMEXPRESSION_TYPE(FMaterialUniformExpressionClamp);
 IMPLEMENT_MATERIALUNIFORMEXPRESSION_TYPE(FMaterialUniformExpressionSaturate);
+IMPLEMENT_MATERIALUNIFORMEXPRESSION_TYPE(FMaterialUniformExpressionComponentSwizzle);
 IMPLEMENT_MATERIALUNIFORMEXPRESSION_TYPE(FMaterialUniformExpressionFloor);
 IMPLEMENT_MATERIALUNIFORMEXPRESSION_TYPE(FMaterialUniformExpressionCeil);
 IMPLEMENT_MATERIALUNIFORMEXPRESSION_TYPE(FMaterialUniformExpressionFrac);
diff --git a/Engine/Source/Runtime/Engine/Private/Materials/MaterialUniformExpressions.h b/Engine/Source/Runtime/Engine/Private/Materials/MaterialUniformExpressions.h
--- a/Engine/Source/Runtime/Engine/Private/Materials/MaterialUniformExpressions.h
+++ b/Engine/Source/Runtime/Engine/Private/Materials/MaterialUniformExpressions.h
@@ -70,6 +70,7 @@ public:
 	{
 		return false;
 	}
+	virtual bool IsChangingPerFrame() const { return true; }
 	virtual bool IsIdentical(const FMaterialUniformExpression* OtherExpression) const
 	{
 		return GetType() == OtherExpression->GetType();
@@ -98,6 +99,7 @@ public:
 	{
 		return false;
 	}
+	virtual bool IsChangingPerFrame() const { return true; }
 	virtual bool IsIdentical(const FMaterialUniformExpression* OtherExpression) const
 	{
 		return GetType() == OtherExpression->GetType();
@@ -923,6 +925,109 @@ private:
 	TRefCountPtr<FMaterialUniformExpression> Input;
 };
 
+class FMaterialUniformExpressionComponentSwizzle : public FMaterialUniformExpression
+{
+	DECLARE_MATERIALUNIFORMEXPRESSION_TYPE(FMaterialUniformExpressionComponentSwizzle);
+public:
+
+	FMaterialUniformExpressionComponentSwizzle() {}
+	FMaterialUniformExpressionComponentSwizzle(FMaterialUniformExpression* InX, int8 InR, int8 InG, int8 InB, int8 InA) :
+		X(InX),
+		IndexR(InR),
+		IndexG(InG),
+		IndexB(InB),
+		IndexA(InA)
+	{
+		NumElements = 0;
+		if (InA >= 0)
+		{
+			check(InA <= 3);
+			++NumElements;
+			check(InB >= 0);
+		}
+
+		if (InB >= 0)
+		{
+			check(InB <= 3);
+			++NumElements;
+			check(InG >= 0);
+		}
+
+		if (InG >= 0)
+		{
+			check(InG <= 3);
+			++NumElements;
+		}
+
+		// At least one proper index
+		check(InR >= 0 && InR <= 3);
+		++NumElements;
+	}
+
+	// FMaterialUniformExpression interface.
+	virtual void Serialize(FArchive& Ar)
+	{
+		Ar << X;
+		Ar << IndexR;
+		Ar << IndexG;
+		Ar << IndexB;
+		Ar << IndexA;
+		Ar << NumElements;
+	}
+	virtual void GetNumberValue(const FMaterialRenderContext& Context, FLinearColor& OutValue) const
+	{
+		FLinearColor Temp = OutValue;
+		X->GetNumberValue(Context, Temp);
+		OutValue = Temp;
+		switch (NumElements)
+		{
+		case 4:
+			OutValue.A = Temp.Component(IndexA);
+			// Fallthrough...
+		case 3:
+			OutValue.B = Temp.Component(IndexB);
+			// Fallthrough...
+		case 2:
+			OutValue.G = Temp.Component(IndexG);
+			// Fallthrough...
+		case 1:
+			OutValue.R = Temp.Component(IndexR);
+			break;
+		default: UE_LOG(LogMaterial, Fatal, TEXT("Invalid number of swizzle elements: %d"), NumElements);
+		}
+	}
+	virtual bool IsConstant() const
+	{
+		return X->IsConstant();
+	}
+	virtual bool IsChangingPerFrame() const
+	{
+		return X->IsChangingPerFrame();
+	}
+	virtual bool IsIdentical(const FMaterialUniformExpression* OtherExpression) const
+	{
+		if (GetType() != OtherExpression->GetType())
+		{
+			return false;
+		}
+		auto* OtherSwizzle = (FMaterialUniformExpressionComponentSwizzle*)OtherExpression;
+		return X->IsIdentical(OtherSwizzle->X) &&
+			NumElements == OtherSwizzle->NumElements &&
+			IndexR == OtherSwizzle->IndexR &&
+			IndexG == OtherSwizzle->IndexG &&
+			IndexB == OtherSwizzle->IndexB &&
+			IndexA == OtherSwizzle->IndexA;
+	}
+
+private:
+	TRefCountPtr<FMaterialUniformExpression> X;
+	int8 IndexR;
+	int8 IndexG;
+	int8 IndexB;
+	int8 IndexA;
+	int8 NumElements;
+};
+
 /**
  */
 class FMaterialUniformExpressionFloor: public FMaterialUniformExpression
diff --git a/Engine/Source/Runtime/Engine/Public/MaterialCompiler.h b/Engine/Source/Runtime/Engine/Public/MaterialCompiler.h
--- a/Engine/Source/Runtime/Engine/Public/MaterialCompiler.h
+++ b/Engine/Source/Runtime/Engine/Public/MaterialCompiler.h
@@ -57,8 +57,8 @@ public:
 	virtual int32 Constant3(float X,float Y,float Z) = 0;
 	virtual int32 Constant4(float X,float Y,float Z,float W) = 0;
 
-	virtual int32 GameTime() = 0;
-	virtual int32 RealTime() = 0;
+	virtual int32 GameTime(bool bPeriodic, float Period) = 0;
+	virtual int32 RealTime(bool bPeriodic, float Period) = 0;
 	virtual int32 PeriodicHint(int32 PeriodicCode) { return PeriodicCode; }
 
 	virtual int32 Sine(int32 X) = 0;
@@ -216,8 +216,8 @@ public:
 	virtual int32 Constant3(float X,float Y,float Z) override { return Compiler->Constant3(X,Y,Z); }
 	virtual int32 Constant4(float X,float Y,float Z,float W) override { return Compiler->Constant4(X,Y,Z,W); }
 
-	virtual int32 GameTime() override { return Compiler->GameTime(); }
-	virtual int32 RealTime() override { return Compiler->RealTime(); }
+	virtual int32 GameTime(bool bPeriodic, float Period) override { return Compiler->GameTime(bPeriodic, Period); }
+	virtual int32 RealTime(bool bPeriodic, float Period) override { return Compiler->RealTime(bPeriodic, Period); }
 
 	virtual int32 PeriodicHint(int32 PeriodicCode) override { return Compiler->PeriodicHint(PeriodicCode); }
 
