UE-6390 Objects selected in the viewport, cannot be de-selected using the escape key

diff --git a/Engine/Source/Editor/UnrealEd/Private/EditorViewportClient.cpp b/Engine/Source/Editor/UnrealEd/Private/EditorViewportClient.cpp
--- a/Engine/Source/Editor/UnrealEd/Private/EditorViewportClient.cpp
+++ b/Engine/Source/Editor/UnrealEd/Private/EditorViewportClient.cpp
@@ -1873,7 +1873,7 @@ bool FEditorViewportClient::InputKey(FViewport* InViewport, int32 ControllerId,
 
 		bHandled = true;
 	}
-	else if (Key == EKeys::Escape && Event == IE_Pressed/* && TrackingTransaction.IsActive()*/) //@TODO: MODETOOLS: Now consuming Escape regardless of if we were previously tracking
+	else if (Key == EKeys::Escape && Event == IE_Pressed && bIsTracking)
 	{
 		// Pressing Escape cancels the current operation
 		AbortTracking();
