UE-6596 BLUEPRINTS: DestroyPersistentUberGraphFrame() causing some kind of mem corruption

diff --git a/Engine/Source/Editor/UnrealEd/Private/Kismet2/KismetReinstanceUtilities.cpp b/Engine/Source/Editor/UnrealEd/Private/Kismet2/KismetReinstanceUtilities.cpp
--- a/Engine/Source/Editor/UnrealEd/Private/Kismet2/KismetReinstanceUtilities.cpp
+++ b/Engine/Source/Editor/UnrealEd/Private/Kismet2/KismetReinstanceUtilities.cpp
@@ -90,6 +90,12 @@ FBlueprintCompileReinstancer::FBlueprintCompileReinstancer(UClass* InClassToRein
 						{
 							ReparentChild(ChildClass);
 						}
+
+						//TODO: some stronger condition would be nice
+						if (!bClassIsDirectlyGeneratedByTheBlueprint)
+						{
+							ObjectsThatShouldUseOldStuff.Add(ChildClass);
+						}
 					}
 					// If this is a direct child, change the parent and relink so the property chain is valid for reinstancing
 					else if( !ChildBP->HasAnyFlags(RF_NeedLoad) )
@@ -230,7 +236,7 @@ void FBlueprintCompileReinstancer::ReinstanceObjects(bool bAlwaysReinstance)
 
 		if (bShouldReinstance)
 		{
-			ReplaceInstancesOfClass(DuplicatedClass, ClassToReinstance, OriginalCDO);
+			ReplaceInstancesOfClass(DuplicatedClass, ClassToReinstance, OriginalCDO, &ObjectsThatShouldUseOldStuff);
 		}
 	
 		{ 
@@ -489,7 +495,7 @@ void FActorReplacementHelper::AttachChildActors(USceneComponent* RootComponent)
 	}
 }
 
-void FBlueprintCompileReinstancer::ReplaceInstancesOfClass(UClass* OldClass, UClass* NewClass, UObject*	OriginalCDO)
+void FBlueprintCompileReinstancer::ReplaceInstancesOfClass(UClass* OldClass, UClass* NewClass, UObject*	OriginalCDO, TSet<UObject*>* ObjectsThatShouldUseOldStuff)
 {
 	USelection* SelectedActors;
 	bool bSelectionChanged = false;
@@ -682,7 +688,11 @@ void FBlueprintCompileReinstancer::ReplaceInstancesOfClass(UClass* OldClass, UCl
 		TFindObjectReferencers<UObject> Referencers(SourceObjects, NULL, false);
 		for (TFindObjectReferencers<UObject>::TIterator It(Referencers); It; ++It)
 		{
-			Targets.Add(It.Value());
+			UObject* Referencer = It.Value();
+			if (!ObjectsThatShouldUseOldStuff || !ObjectsThatShouldUseOldStuff->Contains(Referencer))
+			{
+				Targets.Add(Referencer);
+			}
 		}
 	}
 
diff --git a/Engine/Source/Editor/UnrealEd/Public/Kismet2/KismetReinstanceUtilities.h b/Engine/Source/Editor/UnrealEd/Public/Kismet2/KismetReinstanceUtilities.h
--- a/Engine/Source/Editor/UnrealEd/Public/Kismet2/KismetReinstanceUtilities.h
+++ b/Engine/Source/Editor/UnrealEd/Public/Kismet2/KismetReinstanceUtilities.h
@@ -52,6 +52,9 @@ protected:
 
 	uint32 ClassToReinstanceDefaultValuesCRC;
 
+	/** Objects that should keep reference to old class */
+	TSet<UObject*> ObjectsThatShouldUseOldStuff;
+
 public:
 	virtual ~FBlueprintCompileReinstancer();
 
@@ -71,7 +74,7 @@ public:
 	void UpdateBytecodeReferences();
 
 	/** Worker function to replace all instances of OldClass with a new instance of NewClass */
-	static void ReplaceInstancesOfClass(UClass* OldClass, UClass* NewClass, UObject* OriginalCDO = NULL);
+	static void ReplaceInstancesOfClass(UClass* OldClass, UClass* NewClass, UObject* OriginalCDO = NULL, TSet<UObject*>* ObjectsThatShouldUseOldStuff = NULL);
 	
 	/** Verify that all instances of the duplicated class have been replaced and collected */
 	void VerifyReplacement();
