Nodes spawned for the EventGraph when creating an Animation Blueprint are now properly marked RF_Transactional.

diff --git a/Engine/Source/Editor/UnrealEd/Private/Kismet2/Kismet2.cpp b/Engine/Source/Editor/UnrealEd/Private/Kismet2/Kismet2.cpp
--- a/Engine/Source/Editor/UnrealEd/Private/Kismet2/Kismet2.cpp
+++ b/Engine/Source/Editor/UnrealEd/Private/Kismet2/Kismet2.cpp
@@ -468,6 +468,7 @@ UBlueprint* FKismetEditorUtilities::CreateBlueprint(UClass* ParentClass, UObject
 			EventNode->bOverrideFunction=true;
 			EventNode->CreateNewGuid();
 			EventNode->PostPlacedNewNode();
+			EventNode->SetFlags(RF_Transactional);
 			EventNode->AllocateDefaultPins();
 			UEdGraphSchema_K2::SetNodeMetaData(EventNode, FNodeMetadata::DefaultGraphNode);
 
@@ -479,6 +480,7 @@ UBlueprint* FKismetEditorUtilities::CreateBlueprint(UClass* ParentClass, UObject
 			GetOwnerNode->CreateNewGuid();
 			GetOwnerNode->PostPlacedNewNode();
 			GetOwnerNode->SetFromFunction(MakeNodeFunction);
+			GetOwnerNode->SetFlags(RF_Transactional);
 			GetOwnerNode->AllocateDefaultPins();
 			GetOwnerNode->NodePosX = EventNode->NodePosX;
 			GetOwnerNode->NodePosY = EventNode->NodePosY + EventNode->NodeHeight + 100;
