Editor: Expose the owning mode tools on FEdMode as GetModeManager() and pushed GetWorld() code into a virtual on the mode tools

diff --git a/Engine/Source/Editor/UnrealEd/Private/EdMode.cpp b/Engine/Source/Editor/UnrealEd/Private/EdMode.cpp
--- a/Engine/Source/Editor/UnrealEd/Private/EdMode.cpp
+++ b/Engine/Source/Editor/UnrealEd/Private/EdMode.cpp
@@ -931,6 +931,16 @@ bool FEdMode::UsesToolkits() const
 	return false;
 }
 
+UWorld* FEdMode::GetWorld() const
+{
+	return Owner->GetWorld();
+}
+
+class FEditorModeTools* FEdMode::GetModeManager() const
+{
+	return Owner;
+}
+
 bool FEdMode::StartTracking(FEditorViewportClient* InViewportClient, FViewport* InViewport)
 {
 	bool bResult = false;
diff --git a/Engine/Source/Editor/UnrealEd/Private/EditorModeManager.cpp b/Engine/Source/Editor/UnrealEd/Private/EditorModeManager.cpp
--- a/Engine/Source/Editor/UnrealEd/Private/EditorModeManager.cpp
+++ b/Engine/Source/Editor/UnrealEd/Private/EditorModeManager.cpp
@@ -121,6 +121,11 @@ class USelection* FEditorModeTools::GetSelectedObjects() const
 	return GEditor->GetSelectedObjects();
 }
 
+UWorld* FEditorModeTools::GetWorld() const
+{
+	return GEditor->GetEditorWorldContext().World();
+}
+
 void FEditorModeTools::OnEditorSelectionChanged(UObject* NewSelection)
 {
 	// If selecting an actor, move the pivot location.
diff --git a/Engine/Source/Editor/UnrealEd/Public/EdMode.h b/Engine/Source/Editor/UnrealEd/Public/EdMode.h
--- a/Engine/Source/Editor/UnrealEd/Public/EdMode.h
+++ b/Engine/Source/Editor/UnrealEd/Public/EdMode.h
@@ -214,10 +214,11 @@ public:
 	/** True if this mode uses a toolkit mode (eventually they all should) */
 	virtual bool UsesToolkits() const;
 
-	UWorld* GetWorld() const
-	{
-		return GEditor->GetEditorWorldContext().World();
-	}
+	/** Returns the world this toolkit is editing */
+	UWorld* GetWorld() const;
+
+	/** Returns the owning mode manager for this mode */
+	class FEditorModeTools* GetModeManager() const;
 
 	// Property Widgets
 
diff --git a/Engine/Source/Editor/UnrealEd/Public/EditorModeManager.h b/Engine/Source/Editor/UnrealEd/Public/EditorModeManager.h
--- a/Engine/Source/Editor/UnrealEd/Public/EditorModeManager.h
+++ b/Engine/Source/Editor/UnrealEd/Public/EditorModeManager.h
@@ -369,6 +369,12 @@ public:
 	 * @return the set of selected non-actor objects.
 	 */
 	virtual class USelection* GetSelectedObjects() const;
+
+	/**
+	 * Returns the world that is being edited by this mode manager
+	 */ 
+	virtual UWorld* GetWorld() const;
+
 protected:
 	/** 
 	 * Delegate handlers
