Fixed Matinee invert pan not working.  Decoupled this from FEditorModeTools as well

diff --git a/Engine/Source/Editor/Matinee/Private/Matinee.cpp b/Engine/Source/Editor/Matinee/Private/Matinee.cpp
--- a/Engine/Source/Editor/Matinee/Private/Matinee.cpp
+++ b/Engine/Source/Editor/Matinee/Private/Matinee.cpp
@@ -483,6 +483,9 @@ void FMatinee::InitMatinee(const EToolkitMode::Type Mode, const TSharedPtr< clas
 	// Setup "allow keyframe text selection" preference
 	GConfig->GetBool( TEXT("Matinee"), TEXT("AllowKeyframeTextSelection"), bAllowKeyframeTextSelection, GEditorUserSettingsIni );
 
+	bInvertPan = true;
+	GConfig->GetBool( TEXT("Matinee"), TEXT("InterpEdPanInvert"), bInvertPan, GEditorUserSettingsIni );
+
 	// Setup "lock camera pitch" preference
 	GetLockCameraPitchFromConfig();
 
diff --git a/Engine/Source/Editor/Matinee/Private/Matinee.h b/Engine/Source/Editor/Matinee/Private/Matinee.h
--- a/Engine/Source/Editor/Matinee/Private/Matinee.h
+++ b/Engine/Source/Editor/Matinee/Private/Matinee.h
@@ -1119,6 +1119,7 @@ public:
 	/** If true, the editor is modifying a CameraAnim, and functionality is tweaked appropriately */
 	bool	bEditingCameraAnim;
 
+	bool bInvertPan;
 	// Static list of all InterpTrack subclasses.
 	static TArray<UClass*>	InterpTrackClasses;
 	static bool			bInterpTrackClassesInitialized;
diff --git a/Engine/Source/Editor/Matinee/Private/MatineeMenus.cpp b/Engine/Source/Editor/Matinee/Private/MatineeMenus.cpp
--- a/Engine/Source/Editor/Matinee/Private/MatineeMenus.cpp
+++ b/Engine/Source/Editor/Matinee/Private/MatineeMenus.cpp
@@ -4225,16 +4225,14 @@ void FMatinee::OnMenuReduceKeys()
 /** Toggles interting of the panning the interp editor left and right */
 void FMatinee::OnToggleInvertPan()
 {
-	bool bInvertPan = GLevelEditorModeTools().GetInterpPanInvert();
-
 	bInvertPan = !bInvertPan;
 
-	GLevelEditorModeTools().SetInterpPanInvert(bInvertPan);
+	GConfig->SetBool(TEXT("Matinee"), TEXT("InterpEdPanInvert"), bInvertPan, GEditorUserSettingsIni);
 }
 
 bool FMatinee::IsInvertPanToggled()
 {
-	return GLevelEditorModeTools().GetInterpPanInvert();	
+	return bInvertPan;
 }
 
 /** Called when split translation and rotation is selected from a movement track context menu */
diff --git a/Engine/Source/Editor/Matinee/Private/MatineeViewportClient.cpp b/Engine/Source/Editor/Matinee/Private/MatineeViewportClient.cpp
--- a/Engine/Source/Editor/Matinee/Private/MatineeViewportClient.cpp
+++ b/Engine/Source/Editor/Matinee/Private/MatineeViewportClient.cpp
@@ -920,7 +920,7 @@ void FMatineeViewportClient::MouseMove(FViewport* Viewport, int32 X, int32 Y)
 		}
 		else if(bPanning)
 		{
-			const bool bInvertPanning = ModeTools->GetInterpPanInvert();
+			const bool bInvertPanning = InterpEd->IsInvertPanToggled();
 
 			float DeltaTime = (bInvertPanning ? -DeltaX : DeltaX) / InterpEd->PixelsPerSec;
 			InterpEd->ViewStartTime -= DeltaTime;
diff --git a/Engine/Source/Editor/UnrealEd/Private/EditorModeManager.cpp b/Engine/Source/Editor/UnrealEd/Private/EditorModeManager.cpp
--- a/Engine/Source/Editor/UnrealEd/Private/EditorModeManager.cpp
+++ b/Engine/Source/Editor/UnrealEd/Private/EditorModeManager.cpp
@@ -76,8 +76,7 @@ void FEditorModeTools::LoadConfig(void)
 		GEditorUserSettingsIni);
 	SetCoordSystem((ECoordSystem)Bogus);
 
-	
-	GConfig->GetBool( TEXT("UnEdViewport"), TEXT("InterpEdPanInvert"), bInterpPanInverted, GEditorUserSettingsIni );
+
 	LoadWidgetSettings();
 }
 
@@ -93,8 +92,6 @@ void FEditorModeTools::SaveConfig(void)
 	GConfig->SetInt(TEXT("FEditorModeTools"),TEXT("CoordSystem"),(int32)GetCoordSystem(bGetRawValue),
 		GEditorUserSettingsIni);
 
-	GConfig->SetBool( TEXT("UnEdViewport"), TEXT("InterpEdPanInvert"), bInterpPanInverted, GEditorUserSettingsIni );
-
 	SaveWidgetSettings();
 }
 
diff --git a/Engine/Source/Editor/UnrealEd/Public/EditorModeManager.h b/Engine/Source/Editor/UnrealEd/Public/EditorModeManager.h
--- a/Engine/Source/Editor/UnrealEd/Public/EditorModeManager.h
+++ b/Engine/Source/Editor/UnrealEd/Public/EditorModeManager.h
@@ -237,17 +237,6 @@ public:
 	 */
 	bool GetShowFriendlyVariableNames() const;
 
-	/**
-	 * Sets Interp Editor inversion of panning on drag
-	 */
-	void SetInterpPanInvert(bool bInInterpPanInverted) { bInterpPanInverted = bInInterpPanInverted; }
-
-	/**
-	 * Gets the current state of the interp editor panning
-	 * @ return - If true, the interp editor left-right panning is inverted
-	 */
-	bool GetInterpPanInvert(void) const { return bInterpPanInverted; }
-
 	/**
 	 * Sets a bookmark in the levelinfo file, allocating it if necessary.
 	 * 
@@ -403,9 +392,6 @@ protected:
 	/** If 1, draw the widget and let the user interact with it. */
 	bool bShowWidget;
 
-	/** If true, panning in the Interp Editor is inverted. If false, panning is normal*/
-	bool bInterpPanInverted;
-
 	/** if true, the viewports will hide all UI overlays */
 	bool bHideViewportUI;
 
