Fix missing forward declaration for ULevel.

diff --git a/Engine/Source/Editor/UnrealEd/Public/FileHelpers.h b/Engine/Source/Editor/UnrealEd/Public/FileHelpers.h
--- a/Engine/Source/Editor/UnrealEd/Public/FileHelpers.h
+++ b/Engine/Source/Editor/UnrealEd/Public/FileHelpers.h
@@ -8,6 +8,7 @@
 class FFileName;
 class FString;
 class FAssetData;
+class ULevel;
 
 enum EFileInteraction
 {
