Allow .obj files to be used when importing LODs.

GitHub pull request: https://github.com/EpicGames/UnrealEngine/pull/700

diff --git a/Engine/Source/Editor/UnrealEd/Private/FbxMeshUtils.cpp b/Engine/Source/Editor/UnrealEd/Private/FbxMeshUtils.cpp
--- a/Engine/Source/Editor/UnrealEd/Private/FbxMeshUtils.cpp
+++ b/Engine/Source/Editor/UnrealEd/Private/FbxMeshUtils.cpp
@@ -571,8 +571,12 @@ namespace FbxMeshUtils
 
 		FString ExtensionStr;
 
+		ExtensionStr += TEXT("All model files|*.fbx;*.obj|");
+
 		ExtensionStr += TEXT("FBX files|*.fbx|");
 
+		ExtensionStr += TEXT("Object files|*.obj|");
+
 		ExtensionStr += TEXT("All files|*.*");
 
 		// First, display the file open dialog for selecting the file.
