#ue4 - Convert all engine uses of GetComponents() to use TInlineComponentArray.

diff --git a/Engine/Plugins/Experimental/LiveEditor/Source/LiveEditor/Private/K2Node_LiveEditObject.cpp b/Engine/Plugins/Experimental/LiveEditor/Source/LiveEditor/Private/K2Node_LiveEditObject.cpp
--- a/Engine/Plugins/Experimental/LiveEditor/Source/LiveEditor/Private/K2Node_LiveEditObject.cpp
+++ b/Engine/Plugins/Experimental/LiveEditor/Source/LiveEditor/Private/K2Node_LiveEditObject.cpp
@@ -89,9 +89,9 @@ namespace UK2Node_LiveEditObjectStatics
 				//FString ComponentName = ComponentPropertyName.LeftChop(SplitIndex);
 				ComponentPropertyName = ComponentPropertyName.RightChop(SplitIndex+1);
 
-				TArray<UActorComponent*> ActorComponents;
+				TInlineComponentArray<UActorComponent*> ActorComponents;
 				AsActor->GetComponents(ActorComponents);
-				for ( TArray<UActorComponent*>::TIterator ComponentIt(ActorComponents); ComponentIt; ++ComponentIt )
+				for ( auto ComponentIt = ActorComponents.CreateIterator(); ComponentIt; ++ComponentIt )
 				{
 					UActorComponent *Component = *ComponentIt;
 					check( Component != NULL );
diff --git a/Engine/Plugins/Experimental/LiveEditor/Source/LiveEditor/Private/LiveEditorKismetLibrary.cpp b/Engine/Plugins/Experimental/LiveEditor/Source/LiveEditor/Private/LiveEditorKismetLibrary.cpp
--- a/Engine/Plugins/Experimental/LiveEditor/Source/LiveEditor/Private/LiveEditorKismetLibrary.cpp
+++ b/Engine/Plugins/Experimental/LiveEditor/Source/LiveEditor/Private/LiveEditorKismetLibrary.cpp
@@ -120,9 +120,9 @@ namespace nLiveEditorKismetLibrary
 					//FString ComponentName = ComponentPropertyName.LeftChop(SplitIndex);
 					ComponentPropertyName = ComponentPropertyName.RightChop(SplitIndex+1);
 
-					TArray<UActorComponent*> ActorComponents;
+					TInlineComponentArray<UActorComponent*> ActorComponents;
 					AsActor->GetComponents(ActorComponents);
-					for ( TArray<UActorComponent*>::TIterator ComponentIt(ActorComponents); ComponentIt && !Prop; ++ComponentIt )
+					for ( auto ComponentIt = ActorComponents.CreateIterator(); ComponentIt && !Prop; ++ComponentIt )
 					{
 						UActorComponent *Component = *ComponentIt;
 						check( Component != NULL );
diff --git a/Engine/Plugins/Experimental/LiveEditor/Source/LiveEditor/Private/SGraphPinLiveEditVar.cpp b/Engine/Plugins/Experimental/LiveEditor/Source/LiveEditor/Private/SGraphPinLiveEditVar.cpp
--- a/Engine/Plugins/Experimental/LiveEditor/Source/LiveEditor/Private/SGraphPinLiveEditVar.cpp
+++ b/Engine/Plugins/Experimental/LiveEditor/Source/LiveEditor/Private/SGraphPinLiveEditVar.cpp
@@ -71,9 +71,9 @@ void SGraphPinLiveEditVar::GenerateComboBoxIndexes( TArray< TSharedPtr<int32> >&
 	AActor *AsActor = Cast<AActor>(InClass->ClassDefaultObject);
 	if ( AsActor != NULL )
 	{
-		TArray<UActorComponent*> ActorComponents;
+		TInlineComponentArray<UActorComponent*> ActorComponents;
 		AsActor->GetComponents(ActorComponents);
-		for ( TArray<UActorComponent*>::TIterator ComponentIt(ActorComponents); ComponentIt; ++ComponentIt )
+		for ( auto ComponentIt = ActorComponents.CreateIterator(); ComponentIt; ++ComponentIt )
 		{
 			UActorComponent *Component = *ComponentIt;
 			check( Component != NULL );
diff --git a/Engine/Plugins/Runtime/LiveEditorListenServer/Source/LiveEditorListenServer/Private/LiveEditorListenServer.cpp b/Engine/Plugins/Runtime/LiveEditorListenServer/Source/LiveEditorListenServer/Private/LiveEditorListenServer.cpp
--- a/Engine/Plugins/Runtime/LiveEditorListenServer/Source/LiveEditorListenServer/Private/LiveEditorListenServer.cpp
+++ b/Engine/Plugins/Runtime/LiveEditorListenServer/Source/LiveEditorListenServer/Private/LiveEditorListenServer.cpp
@@ -76,9 +76,9 @@ namespace nLiveEditorListenServer
 					//FString ComponentName = ComponentPropertyName.LeftChop(SplitIndex);
 					ComponentPropertyName = ComponentPropertyName.RightChop(SplitIndex+1);
 
-					TArray<UActorComponent*> ActorComponents;
+					TInlineComponentArray<UActorComponent*> ActorComponents;
 					AsActor->GetComponents(ActorComponents);
-					for ( TArray<UActorComponent*>::TIterator ComponentIt(ActorComponents); ComponentIt && !Prop; ++ComponentIt )
+					for ( auto ComponentIt = ActorComponents.CreateIterator(); ComponentIt && !Prop; ++ComponentIt )
 					{
 						UActorComponent *Component = *ComponentIt;
 						check( Component != NULL );
diff --git a/Engine/Source/Developer/MeshUtilities/Private/MeshUtilities.cpp b/Engine/Source/Developer/MeshUtilities/Private/MeshUtilities.cpp
--- a/Engine/Source/Developer/MeshUtilities/Private/MeshUtilities.cpp
+++ b/Engine/Source/Developer/MeshUtilities/Private/MeshUtilities.cpp
@@ -3565,7 +3565,7 @@ void FMeshUtilities::CreateProxyMesh(
 		}
 		else
 		{
-			TArray<UStaticMeshComponent*> Components;
+			TInlineComponentArray<UStaticMeshComponent*> Components;
 			Actor->GetComponents<UStaticMeshComponent>(Components);
 			// TODO: support instanced static meshes
 			Components.RemoveAll([](UStaticMeshComponent* Val){ return Val->IsA(UInstancedStaticMeshComponent::StaticClass()); });
@@ -3909,7 +3909,7 @@ void FMeshUtilities::MergeActors(
 	// Collect static mesh components
 	for (AActor* Actor : SourceActors)
 	{
-		TArray<UStaticMeshComponent*> Components;
+		TInlineComponentArray<UStaticMeshComponent*> Components;
 		Actor->GetComponents<UStaticMeshComponent>(Components);
 		ComponentsToMerge.Append(Components);
 	}
diff --git a/Engine/Source/Developer/MovieSceneTools/Private/AnimationTrackEditor.cpp b/Engine/Source/Developer/MovieSceneTools/Private/AnimationTrackEditor.cpp
--- a/Engine/Source/Developer/MovieSceneTools/Private/AnimationTrackEditor.cpp
+++ b/Engine/Source/Developer/MovieSceneTools/Private/AnimationTrackEditor.cpp
@@ -258,7 +258,7 @@ USkeleton* FAnimationTrackEditor::AcquireSkeletonFromObjectGuid(const FGuid& Gui
 
 		if (Actor != NULL)
 		{
-			TArray<USkeletalMeshComponent*> SkeletalMeshComponents;
+			TInlineComponentArray<USkeletalMeshComponent*> SkeletalMeshComponents;
 			Actor->GetComponents(SkeletalMeshComponents);
 
 			for (int32 j = 0; j <SkeletalMeshComponents.Num(); ++j)
diff --git a/Engine/Source/Editor/DetailCustomizations/Private/ActorMaterialCategory.cpp b/Engine/Source/Editor/DetailCustomizations/Private/ActorMaterialCategory.cpp
--- a/Engine/Source/Editor/DetailCustomizations/Private/ActorMaterialCategory.cpp
+++ b/Engine/Source/Editor/DetailCustomizations/Private/ActorMaterialCategory.cpp
@@ -184,7 +184,7 @@ private:
 	/** Reference to the selected actors */
 	TArray< TWeakObjectPtr<AActor> >& SelectedActors;
 	/** Reference to the current actor's components */
-	TArray<USceneComponent*> CurActorComponents;
+	TInlineComponentArray<USceneComponent*> CurActorComponents;
 	/** The current material the iterator is stopped on */
 	UMaterialInterface* CurMaterial;
 	/** The current component using the current material */
diff --git a/Engine/Source/Editor/DetailCustomizations/Private/ComponentsTree.cpp b/Engine/Source/Editor/DetailCustomizations/Private/ComponentsTree.cpp
--- a/Engine/Source/Editor/DetailCustomizations/Private/ComponentsTree.cpp
+++ b/Engine/Source/Editor/DetailCustomizations/Private/ComponentsTree.cpp
@@ -211,7 +211,7 @@ void SComponentsTree::UpdateTree()
 		}
 
 		// Add non-scene components in Components array
-		TArray<UActorComponent*> Components;
+		TInlineComponentArray<UActorComponent*> Components;
 		Actor->GetComponents(Components);
 
 		for(int32 CompIdx=0; CompIdx<Components.Num(); CompIdx++)
diff --git a/Engine/Source/Editor/DetailCustomizations/Private/SkinnedMeshComponentDetails.cpp b/Engine/Source/Editor/DetailCustomizations/Private/SkinnedMeshComponentDetails.cpp
--- a/Engine/Source/Editor/DetailCustomizations/Private/SkinnedMeshComponentDetails.cpp
+++ b/Engine/Source/Editor/DetailCustomizations/Private/SkinnedMeshComponentDetails.cpp
@@ -67,7 +67,7 @@ bool FSkinnedMeshComponentDetails::FindUniqueUsedPhysicsAsset(IDetailLayoutBuild
 	{
 		if (AActor* Actor = Cast<AActor>(SelectionIt->Get()))
 		{
-			TArray<USkinnedMeshComponent*> SkinnedMeshComponents;
+			TInlineComponentArray<USkinnedMeshComponent*> SkinnedMeshComponents;
 			Actor->GetComponents(SkinnedMeshComponents);
 
 			for (int32 CompIdx = 0; CompIdx < SkinnedMeshComponents.Num(); CompIdx++)
diff --git a/Engine/Source/Editor/FoliageEdit/Private/FoliageEdMode.cpp b/Engine/Source/Editor/FoliageEdit/Private/FoliageEdMode.cpp
--- a/Engine/Source/Editor/FoliageEdit/Private/FoliageEdMode.cpp
+++ b/Engine/Source/Editor/FoliageEdit/Private/FoliageEdMode.cpp
@@ -1287,7 +1287,7 @@ void FEdModeFoliage::ApplyPaintBucket(AActor* Actor, bool bRemove)
 		// Remove all instances of the selected meshes
 		AInstancedFoliageActor* IFA = AInstancedFoliageActor::GetInstancedFoliageActorForCurrentLevel(GetWorld());
 
-		TArray<UActorComponent*> Components;
+		TInlineComponentArray<UActorComponent*> Components;
 		Actor->GetComponents(Components);
 
 		for (int32 ComponentIdx = 0; ComponentIdx < Components.Num(); ComponentIdx++)
@@ -1314,7 +1314,7 @@ void FEdModeFoliage::ApplyPaintBucket(AActor* Actor, bool bRemove)
 		TMap<UPrimitiveComponent*, TArray<FFoliagePaintBucketTriangle> > ComponentPotentialTriangles;
 
 		// Check all the components of the hit actor
-		TArray<UStaticMeshComponent*> StaticMeshComponents;
+		TInlineComponentArray<UStaticMeshComponent*> StaticMeshComponents;
 		Actor->GetComponents(StaticMeshComponents);
 
 		for (int32 ComponentIdx = 0; ComponentIdx < StaticMeshComponents.Num(); ComponentIdx++)
diff --git a/Engine/Source/Editor/Kismet/Private/SCSEditorViewportClient.cpp b/Engine/Source/Editor/Kismet/Private/SCSEditorViewportClient.cpp
--- a/Engine/Source/Editor/Kismet/Private/SCSEditorViewportClient.cpp
+++ b/Engine/Source/Editor/Kismet/Private/SCSEditorViewportClient.cpp
@@ -170,7 +170,7 @@ void FSCSEditorViewportClient::Tick(float DeltaSeconds)
 	AActor* PreviewActor = GetPreviewActor();
 	if (PreviewActor != NULL)
 	{
-		TArray<UPrimitiveComponent*> PrimitiveComponents;
+		TInlineComponentArray<UPrimitiveComponent*> PrimitiveComponents;
 		PreviewActor->GetComponents(PrimitiveComponents);
 
 		for (int32 CompIdx = 0; CompIdx < PrimitiveComponents.Num(); ++CompIdx)
@@ -330,7 +330,7 @@ void FSCSEditorViewportClient::ProcessClick(class FSceneView& View, class HHitPr
 			AActor* PreviewActor = GetPreviewActor();
 			if(ActorProxy && ActorProxy->Actor && ActorProxy->Actor == PreviewActor && ActorProxy->PrimComponent != NULL)
 			{
-				TArray<USceneComponent*> SceneComponents;
+				TInlineComponentArray<USceneComponent*> SceneComponents;
 				ActorProxy->Actor->GetComponents(SceneComponents);
 	
 				for(auto CompIt = SceneComponents.CreateConstIterator(); CompIt; ++CompIt)
@@ -987,7 +987,7 @@ void FSCSEditorViewportClient::RefreshPreviewBounds()
 	if(PreviewActor)
 	{
 		// Compute actor bounds as the sum of its visible parts
-		TArray<UPrimitiveComponent*> PrimitiveComponents;
+		TInlineComponentArray<UPrimitiveComponent*> PrimitiveComponents;
 		PreviewActor->GetComponents(PrimitiveComponents);
 
 		PreviewActorBounds = FBoxSphereBounds(ForceInitToZero);
diff --git a/Engine/Source/Editor/Kismet/Private/SSCSEditor.cpp b/Engine/Source/Editor/Kismet/Private/SSCSEditor.cpp
--- a/Engine/Source/Editor/Kismet/Private/SSCSEditor.cpp
+++ b/Engine/Source/Editor/Kismet/Private/SSCSEditor.cpp
@@ -328,7 +328,7 @@ UActorComponent* FSCSEditorTreeNode::FindComponentInstanceInActor(const AActor*
 		else if(ComponentTemplate != NULL)
 		{
 			// Look for a native component instance with a name that matches the template name
-			TArray<UActorComponent*> Components;
+			TInlineComponentArray<UActorComponent*> Components;
 			InActor->GetComponents(Components);
 
 			for(auto It = Components.CreateConstIterator(); It; ++It)
@@ -2429,7 +2429,7 @@ FSCSEditorTreeNodePtrType SSCSEditor::GetNodeFromActorComponent(const UActorComp
 				if(CDO)
 				{
 					// Iterate over the Components array and attempt to find a component with a matching name
-					TArray<UActorComponent*> Components;
+					TInlineComponentArray<UActorComponent*> Components;
 					CDO->GetComponents(Components);
 
 					for(auto It = Components.CreateConstIterator(); It; ++It)
@@ -2579,7 +2579,7 @@ void SSCSEditor::UpdateTree(bool bRegenerateTreeNodes)
 		if(CDO != NULL)
 		{
 			// Add native ActorComponent nodes to the root set first
-			TArray<UActorComponent*> Components;
+			TInlineComponentArray<UActorComponent*> Components;
 			CDO->GetComponents(Components);
 
 			for(auto CompIter = Components.CreateIterator(); CompIter; ++CompIter)
@@ -3115,7 +3115,7 @@ void SSCSEditor::OnTreeSelectionChanged(FSCSEditorTreeNodePtrType, ESelectInfo::
 	AActor* EditorActorInstance = SCS->GetComponentEditorActorInstance();
 	if (EditorActorInstance != NULL)
 	{
-		TArray<UPrimitiveComponent*> PrimitiveComponents;
+		TInlineComponentArray<UPrimitiveComponent*> PrimitiveComponents;
 		EditorActorInstance->GetComponents(PrimitiveComponents);
 
 		for (int32 Idx = 0; Idx < PrimitiveComponents.Num(); ++Idx)
diff --git a/Engine/Source/Editor/LevelEditor/Private/LevelEditorActions.cpp b/Engine/Source/Editor/LevelEditor/Private/LevelEditorActions.cpp
--- a/Engine/Source/Editor/LevelEditor/Private/LevelEditorActions.cpp
+++ b/Engine/Source/Editor/LevelEditor/Private/LevelEditorActions.cpp
@@ -1456,7 +1456,7 @@ void FLevelEditorActionCallbacks::OnSelectStationaryLightsExceedingOverlap()
 	{
 		AActor* Actor = *It;
 
-		TArray<ULightComponent*> Components;
+		TInlineComponentArray<ULightComponent*> Components;
 		Actor->GetComponents(Components);
 
 		for (int32 ComponentIndex = 0; ComponentIndex < Components.Num(); ComponentIndex++)
@@ -1586,7 +1586,7 @@ bool FLevelEditorActionCallbacks::SaveAnimationFromSkeletalMeshComponent(AActor
 	}
 
 	// find all skel components
-	TArray<class USkeletalMeshComponent *> SimSkelComponents;
+	TInlineComponentArray<class USkeletalMeshComponent *> SimSkelComponents;
 	SimActor->GetComponents<USkeletalMeshComponent>(SimSkelComponents);
 
 	if(SimSkelComponents.Num() > 0)
diff --git a/Engine/Source/Editor/Matinee/Private/MatineeTrackHelpers.cpp b/Engine/Source/Editor/Matinee/Private/MatineeTrackHelpers.cpp
--- a/Engine/Source/Editor/Matinee/Private/MatineeTrackHelpers.cpp
+++ b/Engine/Source/Editor/Matinee/Private/MatineeTrackHelpers.cpp
@@ -218,7 +218,7 @@ bool UMatineeTrackAnimControlHelper::PreCreateKeyframe( UInterpTrack *Track, flo
 
 	USkeletalMeshComponent * SkelMeshComp = NULL;
 
-	TArray<USkeletalMeshComponent*> SkeletalMeshComponents;
+	TInlineComponentArray<USkeletalMeshComponent*> SkeletalMeshComponents;
 	Actor->GetComponents(SkeletalMeshComponents);
 
 	for (int32 I=0; I<SkeletalMeshComponents.Num(); ++I)
diff --git a/Engine/Source/Editor/UnrealEd/Private/AssetSelection.cpp b/Engine/Source/Editor/UnrealEd/Private/AssetSelection.cpp
--- a/Engine/Source/Editor/UnrealEd/Private/AssetSelection.cpp
+++ b/Engine/Source/Editor/UnrealEd/Private/AssetSelection.cpp
@@ -198,7 +198,7 @@ namespace AssetSelectionUtils
 						ActorInfo.bHaveAttachedActor = true;
 					}
 
-					TArray<UActorComponent*> ActorComponents;
+					TInlineComponentArray<UActorComponent*> ActorComponents;
 					CurrentActor->GetComponents(ActorComponents);
 
 					for( UActorComponent* Component : ActorComponents )
@@ -668,7 +668,7 @@ bool FActorFactoryAssetProxy::IsActorValidForMaterialApplication( AActor* Target
 	// a material applied to it. Otherwise, it cannot.
 	if ( TargetActor )
 	{
-		TArray<UMeshComponent*> MeshComponents;
+		TInlineComponentArray<UMeshComponent*> MeshComponents;
 		TargetActor->GetComponents(MeshComponents);
 
 		bIsValid = (MeshComponents.Num() > 0);
@@ -711,7 +711,7 @@ bool FActorFactoryAssetProxy::ApplyMaterialToActor( AActor* TargetActor, UMateri
 			TArray<UActorComponent*> FoundMeshComponents;
 
 			// Find which mesh the user clicked on first.
-			TArray<USceneComponent*> SceneComponents;
+			TInlineComponentArray<USceneComponent*> SceneComponents;
 			TargetActor->GetComponents(SceneComponents);
 
 			for ( int32 ComponentIdx=0; ComponentIdx < SceneComponents.Num(); ComponentIdx++ )
diff --git a/Engine/Source/Editor/UnrealEd/Private/Dialogs/SMeshMergingDialog.cpp b/Engine/Source/Editor/UnrealEd/Private/Dialogs/SMeshMergingDialog.cpp
--- a/Engine/Source/Editor/UnrealEd/Private/Dialogs/SMeshMergingDialog.cpp
+++ b/Engine/Source/Editor/UnrealEd/Private/Dialogs/SMeshMergingDialog.cpp
@@ -346,7 +346,7 @@ void SMeshMergingDialog::GenerateNewPackageName()
 		AActor* Actor = Cast<AActor>(*Iter);
 		if (Actor)
 		{
-			TArray<UStaticMeshComponent*> SMComponets; 
+			TInlineComponentArray<UStaticMeshComponent*> SMComponets; 
 			Actor->GetComponents<UStaticMeshComponent>(SMComponets);
 			for (UStaticMeshComponent* Component : SMComponets)
 			{
diff --git a/Engine/Source/Editor/UnrealEd/Private/Dialogs/SMeshProxyDialog.cpp b/Engine/Source/Editor/UnrealEd/Private/Dialogs/SMeshProxyDialog.cpp
--- a/Engine/Source/Editor/UnrealEd/Private/Dialogs/SMeshProxyDialog.cpp
+++ b/Engine/Source/Editor/UnrealEd/Private/Dialogs/SMeshProxyDialog.cpp
@@ -660,7 +660,7 @@ void SMeshProxyDialog::GenerateNewProxyPackageName()
 		AActor* Actor = Cast<AActor>(*Iter);
 		if (Actor)
 		{
-			TArray<UStaticMeshComponent*> SMComponets; 
+			TInlineComponentArray<UStaticMeshComponent*> SMComponets; 
 			Actor->GetComponents<UStaticMeshComponent>(SMComponets);
 			for (UStaticMeshComponent* Component : SMComponets)
 			{
diff --git a/Engine/Source/Editor/UnrealEd/Private/DragTool_BoxSelect.cpp b/Engine/Source/Editor/UnrealEd/Private/DragTool_BoxSelect.cpp
--- a/Engine/Source/Editor/UnrealEd/Private/DragTool_BoxSelect.cpp
+++ b/Engine/Source/Editor/UnrealEd/Private/DragTool_BoxSelect.cpp
@@ -300,7 +300,7 @@ bool FDragTool_ActorBoxSelect::IntersectsBox( AActor& InActor, const FBox& InBox
 	if( !bActorIsHiddenByShowFlags && !InActor.IsHiddenEd() && !FActorEditorUtils::IsABuilderBrush(&InActor) && bActorRecentlyRendered )
 	{
 		// Iterate over all actor components, selecting out primitive components
-		TArray<UPrimitiveComponent*> PrimitiveComponents;
+		TInlineComponentArray<UPrimitiveComponent*> PrimitiveComponents;
 		InActor.GetComponents(PrimitiveComponents);
 
 		for (const UPrimitiveComponent* PrimitiveComponent : PrimitiveComponents)
diff --git a/Engine/Source/Editor/UnrealEd/Private/DragTool_FrustumSelect.cpp b/Engine/Source/Editor/UnrealEd/Private/DragTool_FrustumSelect.cpp
--- a/Engine/Source/Editor/UnrealEd/Private/DragTool_FrustumSelect.cpp
+++ b/Engine/Source/Editor/UnrealEd/Private/DragTool_FrustumSelect.cpp
@@ -255,7 +255,7 @@ bool FDragTool_ActorFrustumSelect::IntersectsFrustum( AActor& InActor, const FCo
 	if( !bActorIsHiddenByShowFlags && !InActor.IsHiddenEd() && !FActorEditorUtils::IsABuilderBrush(&InActor) )
 	{
 		// Iterate over all actor components, selecting out primitive components
-		TArray<UPrimitiveComponent*> Components;
+		TInlineComponentArray<UPrimitiveComponent*> Components;
 		InActor.GetComponents(Components);
 
 		for (const UPrimitiveComponent* PrimitiveComponent : Components)
diff --git a/Engine/Source/Editor/UnrealEd/Private/Editor.cpp b/Engine/Source/Editor/UnrealEd/Private/Editor.cpp
--- a/Engine/Source/Editor/UnrealEd/Private/Editor.cpp
+++ b/Engine/Source/Editor/UnrealEd/Private/Editor.cpp
@@ -5127,7 +5127,7 @@ static void CopyLightComponentProperties( const AActor& InOldActor, AActor& InNe
 	UActorComponent* LightComponentToCopy = NULL;
 
 	// Go through the old actor's components and look for a light component to copy.
-	TArray<UActorComponent*> OldActorComponents;
+	TInlineComponentArray<UActorComponent*> OldActorComponents;
 	InOldActor.GetComponents(OldActorComponents);
 
 	for( int32 CompToCopyIdx = 0; CompToCopyIdx < OldActorComponents.Num(); ++CompToCopyIdx )
@@ -5151,7 +5151,7 @@ static void CopyLightComponentProperties( const AActor& InOldActor, AActor& InNe
 	// Dont do anything if there is no valid light component to copy from
 	if( LightComponentToCopy )
 	{
-		TArray<UActorComponent*> NewActorComponents;
+		TInlineComponentArray<UActorComponent*> NewActorComponents;
 		InNewActor.GetComponents(NewActorComponents);
 
 		// Find a light component to overwrite in the new actor
@@ -5387,11 +5387,11 @@ void CopyActorComponentProperties( const AActor* SourceActor, AActor* DestActor,
 
 		// Construct a mapping from the default actor of its relevant component names to its actual components. Here relevant component
 		// names are those that match a name provided as a parameter.
-		TArray<UActorComponent*> CDOComponents;
+		TInlineComponentArray<UActorComponent*> CDOComponents;
 		SrcActorDefaultActor->GetComponents(CDOComponents);
 
 		TMap<FString, const UActorComponent*> NameToDefaultComponentMap; 
-		for ( TArray<UActorComponent*>::TConstIterator CompIter( CDOComponents ); CompIter; ++CompIter )
+		for ( TInlineComponentArray<UActorComponent*>::TConstIterator CompIter( CDOComponents ); CompIter; ++CompIter )
 		{
 			const UActorComponent* CurComp = *CompIter;
 			check( CurComp );
@@ -5405,11 +5405,11 @@ void CopyActorComponentProperties( const AActor* SourceActor, AActor* DestActor,
 
 		// Construct a mapping from the source actor of its relevant component names to its actual components. Here relevant component names
 		// are those that match a name provided as a parameter.
-		TArray<UActorComponent*> SourceComponents;
+		TInlineComponentArray<UActorComponent*> SourceComponents;
 		SourceActor->GetComponents(SourceComponents);
 
 		TMap<FString, const UActorComponent*> NameToSourceComponentMap;
-		for ( TArray<UActorComponent*>::TConstIterator CompIter( SourceComponents ); CompIter; ++CompIter )
+		for ( TInlineComponentArray<UActorComponent*>::TConstIterator CompIter( SourceComponents ); CompIter; ++CompIter )
 		{
 			const UActorComponent* CurComp = *CompIter;
 			check( CurComp );
@@ -5423,11 +5423,11 @@ void CopyActorComponentProperties( const AActor* SourceActor, AActor* DestActor,
 
 		bool bCopiedAnyProperty = false;
 
-		TArray<UActorComponent*> DestComponents;
+		TInlineComponentArray<UActorComponent*> DestComponents;
 		DestActor->GetComponents(DestComponents);
 
 		// Iterate through all of the destination actor's components to find the ones which should have properties copied into them.
-		for ( TArray<UActorComponent*>::TIterator DestCompIter( DestComponents ); DestCompIter; ++DestCompIter )
+		for ( TInlineComponentArray<UActorComponent*>::TIterator DestCompIter( DestComponents ); DestCompIter; ++DestCompIter )
 		{
 			UActorComponent* CurComp = *DestCompIter;
 			check( CurComp );
@@ -6936,7 +6936,8 @@ namespace EditorUtilities
 	// TargetComponents array is passed in populated to avoid repeated refetching and StartIndex 
 	// is updated as an optimization based on the assumption that the standard use case is iterating 
 	// over two component arrays that will be parallel in order
-	UActorComponent* FindMatchingComponentInstance( UActorComponent* SourceComponent, AActor* TargetActor, const TArray<UActorComponent*>& TargetComponents, int32& StartIndex )
+	template<class AllocatorType = FDefaultAllocator>
+	UActorComponent* FindMatchingComponentInstance( UActorComponent* SourceComponent, AActor* TargetActor, const TArray<UActorComponent*, AllocatorType>& TargetComponents, int32& StartIndex )
 	{
 		UActorComponent* TargetComponent = StartIndex < TargetComponents.Num() ? TargetComponents[ StartIndex ] : NULL;
 
@@ -7005,12 +7006,12 @@ namespace EditorUtilities
 
 	UActorComponent* FindMatchingComponentInstance( UActorComponent* SourceComponent, AActor* TargetActor )
 	{
-		TArray<UActorComponent*> TargetComponents;
 		UActorComponent* MatchingComponent = NULL;
 		int32 StartIndex = 0;
 
 		if (TargetActor)
 		{
+			TInlineComponentArray<UActorComponent*> TargetComponents;
 			TargetActor->GetComponents(TargetComponents);
 			MatchingComponent = FindMatchingComponentInstance( SourceComponent, TargetActor, TargetComponents, StartIndex );
 		}
@@ -7052,7 +7053,7 @@ namespace EditorUtilities
 					else if (SourceObjectPropertyValue->IsA(UActorComponent::StaticClass()) && InTargetObject->IsA(AActor::StaticClass()))
 					{
 						AActor* const TargetActor = Cast<AActor>(InTargetObject);
-						TArray<UActorComponent*> TargetComponents;
+						TInlineComponentArray<UActorComponent*> TargetComponents;
 						TargetActor->GetComponents(TargetComponents);
 
 						// We can try and fix-up an actor component reference from the PIE world to instead be the version from the persistent world
@@ -7221,8 +7222,8 @@ namespace EditorUtilities
 		}
 
 		// Copy component properties from source to target if they match. Note that the component lists may not be 1-1 due to context-specific components (e.g. editor-only sprites, etc.).
-		TArray<UActorComponent*> SourceComponents;
-		TArray<UActorComponent*> TargetComponents;
+		TInlineComponentArray<UActorComponent*> SourceComponents;
+		TInlineComponentArray<UActorComponent*> TargetComponents;
 
 		SourceActor->GetComponents(SourceComponents);
 		TargetActor->GetComponents(TargetComponents);
diff --git a/Engine/Source/Editor/UnrealEd/Private/EditorActor.cpp b/Engine/Source/Editor/UnrealEd/Private/EditorActor.cpp
--- a/Engine/Source/Editor/UnrealEd/Private/EditorActor.cpp
+++ b/Engine/Source/Editor/UnrealEd/Private/EditorActor.cpp
@@ -684,7 +684,7 @@ bool UUnrealEdEngine::edactDeleteSelected( UWorld* InWorld, bool bVerifyDeletion
 		AInstancedFoliageActor* IFA = AInstancedFoliageActor::GetInstancedFoliageActorForLevel(Level);
 		if( IFA )
 		{
-			TArray<UActorComponent*> Components;
+			TInlineComponentArray<UActorComponent*> Components;
 			Actor->GetComponents(Components);
 
 			for(int32 ComponentIndex = 0;ComponentIndex < Components.Num();ComponentIndex++)
@@ -1881,7 +1881,7 @@ void UUnrealEdEngine::edactSelectMatchingMaterial()
 		if( CurrentActor )
 		{
 			// Find the materials by iterating over every primitive component.
-			TArray<UPrimitiveComponent*> PrimitiveComponents;
+			TInlineComponentArray<UPrimitiveComponent*> PrimitiveComponents;
 			CurrentActor->GetComponents(PrimitiveComponents);
 
 			for (int32 ComponentIdx = 0; ComponentIdx < PrimitiveComponents.Num(); ComponentIdx++)
@@ -1916,7 +1916,7 @@ void UUnrealEdEngine::edactSelectMatchingMaterial()
 		// Do not bother checking hidden actors
 		if( !Actor->IsHiddenEd() )
 		{
-			TArray<UPrimitiveComponent*> PrimitiveComponents;
+			TInlineComponentArray<UPrimitiveComponent*> PrimitiveComponents;
 			Actor->GetComponents(PrimitiveComponents);
 
 			const int32 NumComponents = PrimitiveComponents.Num();
@@ -2010,7 +2010,7 @@ void UUnrealEdEngine::edactSelectRelevantLights( UWorld* InWorld )
 
 		if (Actor->GetLevel()->IsCurrentLevel() )
 		{
-			TArray<UPrimitiveComponent*> PrimitiveComponents;
+			TInlineComponentArray<UPrimitiveComponent*> PrimitiveComponents;
 			Actor->GetComponents(PrimitiveComponents);
 
 			// Gather static lighting info from each of the actor's components.
diff --git a/Engine/Source/Editor/UnrealEd/Private/EditorExporters.cpp b/Engine/Source/Editor/UnrealEd/Private/EditorExporters.cpp
--- a/Engine/Source/Editor/UnrealEd/Private/EditorExporters.cpp
+++ b/Engine/Source/Editor/UnrealEd/Private/EditorExporters.cpp
@@ -1504,7 +1504,7 @@ static void AddActorToOBJs(AActor* Actor, TArray<FOBJGeom*>& Objects, TSet<UMate
 	UStaticMeshComponent* StaticMeshComponent = NULL;
 	UStaticMesh* StaticMesh = NULL;
 
-	TArray<UStaticMeshComponent*> StaticMeshComponents;
+	TInlineComponentArray<UStaticMeshComponent*> StaticMeshComponents;
 	Actor->GetComponents(StaticMeshComponents);
 
 	for( int32 j=0; j<StaticMeshComponents.Num(); j++ )
diff --git a/Engine/Source/Editor/UnrealEd/Private/EditorSelectUtils.cpp b/Engine/Source/Editor/UnrealEd/Private/EditorSelectUtils.cpp
--- a/Engine/Source/Editor/UnrealEd/Private/EditorSelectUtils.cpp
+++ b/Engine/Source/Editor/UnrealEd/Private/EditorSelectUtils.cpp
@@ -236,7 +236,7 @@ void UUnrealEdEngine::ResetPivot()
 
 void UUnrealEdEngine::SetActorSelectionFlags (AActor* InActor)
 {
-	TArray<UActorComponent*> Components;
+	TInlineComponentArray<UActorComponent*> Components;
 	InActor->GetComponents(Components);
 
 	//for every component in the actor
diff --git a/Engine/Source/Editor/UnrealEd/Private/EditorServer.cpp b/Engine/Source/Editor/UnrealEd/Private/EditorServer.cpp
--- a/Engine/Source/Editor/UnrealEd/Private/EditorServer.cpp
+++ b/Engine/Source/Editor/UnrealEd/Private/EditorServer.cpp
@@ -4198,7 +4198,7 @@ void UEditorEngine::MoveViewportCamerasToActor(const TArray<AActor*> &Actors, bo
 			}
 			else
 			{
-				TArray<UPrimitiveComponent*> Components;
+				TInlineComponentArray<UPrimitiveComponent*> Components;
 				Actor->GetComponents(Components);
 
 				for(int32 ComponentIndex = 0; ComponentIndex < Components.Num(); ++ComponentIndex)
@@ -4739,7 +4739,7 @@ void UEditorEngine::AssignReplacementComponentsByActors(TArray<AActor*>& ActorsT
 		// if we are clearing the replacement, then we don't need to find a component
 		if (Replacement)
 		{
-			TArray<UPrimitiveComponent*> Components;
+			TInlineComponentArray<UPrimitiveComponent*> Components;
 			Replacement->GetComponents(Components);
 
 			for (int32 ComponentIndex = 0; ComponentIndex < Components.Num(); ComponentIndex++)
@@ -4759,7 +4759,7 @@ void UEditorEngine::AssignReplacementComponentsByActors(TArray<AActor*>& ActorsT
 	{
 		AActor* Actor = ActorsToReplace[ActorIndex];
 
-		TArray<UPrimitiveComponent*> Components;
+		TInlineComponentArray<UPrimitiveComponent*> Components;
 		Actor->GetComponents(Components);
 
 		for (int32 ComponentIndex = 0; ComponentIndex < Components.Num(); ComponentIndex++)
@@ -5667,7 +5667,7 @@ bool UEditorEngine::HandleSetReplacementCommand( const TCHAR* Str, FOutputDevice
 	// attempt to set replacement component for all selected actors
 	for( FSelectedActorIterator It(InWorld); It; ++It )
 	{
-		TArray<UPrimitiveComponent*> Components;
+		TInlineComponentArray<UPrimitiveComponent*> Components;
 		It->GetComponents(Components);
 
 		for (int32 ComponentIndex = 0; ComponentIndex < Components.Num(); ComponentIndex++)
@@ -5907,7 +5907,7 @@ bool UEditorEngine::HandleSetDetailModeCommand( const TCHAR* Str, FOutputDevice&
 			AActor* Actor = static_cast<AActor*>( *It );
 			checkSlow( Actor->IsA(AActor::StaticClass()) );
 
-			TArray<UPrimitiveComponent*> Components;
+			TInlineComponentArray<UPrimitiveComponent*> Components;
 			Actor->GetComponents(Components);
 
 			for(int32 ComponentIndex = 0;ComponentIndex < Components.Num();ComponentIndex++)
@@ -5962,7 +5962,7 @@ bool UEditorEngine::HandleSetDetailModeViewCommand( const TCHAR* Str, FOutputDev
 			AActor* Actor = static_cast<AActor*>( *It );
 			checkSlow( Actor->IsA(AActor::StaticClass()) );
 
-			TArray<UPrimitiveComponent*> Components;
+			TInlineComponentArray<UPrimitiveComponent*> Components;
 			Actor->GetComponents(Components);
 
 			for(int32 ComponentIndex = 0;ComponentIndex < Components.Num();ComponentIndex++)
diff --git a/Engine/Source/Editor/UnrealEd/Private/Factories/ActorFactory.cpp b/Engine/Source/Editor/UnrealEd/Private/Factories/ActorFactory.cpp
--- a/Engine/Source/Editor/UnrealEd/Private/Factories/ActorFactory.cpp
+++ b/Engine/Source/Editor/UnrealEd/Private/Factories/ActorFactory.cpp
@@ -405,7 +405,7 @@ void UActorFactoryDeferredDecal::PostSpawnActor(UObject* Asset, AActor* NewActor
 		GEditor->SetActorLabelUnique(NewActor, Material->GetName());
 
 		// Change properties
-		TArray<UDecalComponent*> DecalComponents;
+		TInlineComponentArray<UDecalComponent*> DecalComponents;
 		NewActor->GetComponents(DecalComponents);
 
 		UDecalComponent* DecalComponent = NULL;
@@ -437,7 +437,7 @@ void UActorFactoryDeferredDecal::PostCreateBlueprint( UObject* Asset, AActor* CD
 
 		if (Material != NULL)
 		{
-			TArray<UDecalComponent*> DecalComponents;
+			TInlineComponentArray<UDecalComponent*> DecalComponents;
 			CDO->GetComponents(DecalComponents);
 
 			UDecalComponent* DecalComponent = NULL;
diff --git a/Engine/Source/Editor/UnrealEd/Private/Fbx/FbxMainExport.cpp b/Engine/Source/Editor/UnrealEd/Private/Fbx/FbxMainExport.cpp
--- a/Engine/Source/Editor/UnrealEd/Private/Fbx/FbxMainExport.cpp
+++ b/Engine/Source/Editor/UnrealEd/Private/Fbx/FbxMainExport.cpp
@@ -1093,10 +1093,10 @@ FbxNode* FFbxExporter::ExportActor(AActor* Actor, AMatineeActor* InMatineeActor,
 
 		if( bExportComponents )
 		{
-			TArray<UMeshComponent*> MeshComponents;
+			TInlineComponentArray<UMeshComponent*> MeshComponents;
 			Actor->GetComponents(MeshComponents);
 
-			TArray<UActorComponent*> ComponentsToExport;
+			TInlineComponentArray<UActorComponent*> ComponentsToExport;
 			for( int32 ComponentIndex = 0; ComponentIndex < MeshComponents.Num(); ++ComponentIndex )
 			{
 				UMeshComponent* Component = MeshComponents[ComponentIndex];
diff --git a/Engine/Source/Editor/UnrealEd/Private/Kismet2/BlueprintEditorUtils.cpp b/Engine/Source/Editor/UnrealEd/Private/Kismet2/BlueprintEditorUtils.cpp
--- a/Engine/Source/Editor/UnrealEd/Private/Kismet2/BlueprintEditorUtils.cpp
+++ b/Engine/Source/Editor/UnrealEd/Private/Kismet2/BlueprintEditorUtils.cpp
@@ -5255,7 +5255,7 @@ void FBlueprintEditorUtils::UpdateRootComponentReference(UBlueprint* Blueprint)
 				if(ParentSceneRootComponent != NULL)
 				{
 					// Search for a scene component with the same name in the Blueprint CDO's Components list
-					TArray<USceneComponent*> SceneComponents;
+					TInlineComponentArray<USceneComponent*> SceneComponents;
 					BlueprintActorCDO->GetComponents(SceneComponents);
 					for(int i = 0; i < SceneComponents.Num(); ++i)
 					{
diff --git a/Engine/Source/Editor/UnrealEd/Private/Kismet2/Kismet2.cpp b/Engine/Source/Editor/UnrealEd/Private/Kismet2/Kismet2.cpp
--- a/Engine/Source/Editor/UnrealEd/Private/Kismet2/Kismet2.cpp
+++ b/Engine/Source/Editor/UnrealEd/Private/Kismet2/Kismet2.cpp
@@ -976,7 +976,7 @@ public:
 			USceneComponent* SingleSceneComp = NULL;
 			if(SelectedActors.Num() == 1)
 			{
-				TArray<UActorComponent*> Components;
+				TInlineComponentArray<UActorComponent*> Components;
 				SelectedActors[0]->GetComponents(Components);
 				for(UActorComponent* Component : Components)
 				{
diff --git a/Engine/Source/Editor/UnrealEd/Private/Layers/Layers.cpp b/Engine/Source/Editor/UnrealEd/Private/Layers/Layers.cpp
--- a/Engine/Source/Editor/UnrealEd/Private/Layers/Layers.cpp
+++ b/Engine/Source/Editor/UnrealEd/Private/Layers/Layers.cpp
@@ -692,7 +692,7 @@ void FLayers::RemoveViewFromActorViewVisibility( FLevelEditorViewportClient* Vie
 		}
 
 		// Find all registered primitive components and update the scene proxy with the actors updated visibility map
-		TArray<UPrimitiveComponent*> Components;
+		TInlineComponentArray<UPrimitiveComponent*> Components;
 		Actor->GetComponents(Components);
 
 		for( int32 ComponentIdx = 0; ComponentIdx < Components.Num(); ++ComponentIdx )
diff --git a/Engine/Source/Editor/UnrealEd/Private/LevelEditorViewport.cpp b/Engine/Source/Editor/UnrealEd/Private/LevelEditorViewport.cpp
--- a/Engine/Source/Editor/UnrealEd/Private/LevelEditorViewport.cpp
+++ b/Engine/Source/Editor/UnrealEd/Private/LevelEditorViewport.cpp
@@ -1284,7 +1284,7 @@ bool FLevelEditorViewportClient::DropObjectsAtCoordinates(int32 MouseX, int32 Mo
 					}
 
 					// Prevent future selection. This also prevents the hit proxy from interfering with placement logic.
-					TArray<UPrimitiveComponent*> PrimitiveComponents;
+					TInlineComponentArray<UPrimitiveComponent*> PrimitiveComponents;
 					NewActor->GetComponents(PrimitiveComponents);
 
 					for ( auto CompIt = PrimitiveComponents.CreateConstIterator(); CompIt; ++CompIt )
@@ -3485,7 +3485,7 @@ void FLevelEditorViewportClient::Draw(const FSceneView* View,FPrimitiveDrawInter
 				continue;
 			}
 
-			TArray<USceneComponent*> Components;
+			TInlineComponentArray<USceneComponent*> Components;
 			Actor->GetComponents(Components);
 
 			for (int32 ComponentIndex = 0 ; ComponentIndex < Components.Num(); ++ComponentIndex)
@@ -3832,7 +3832,7 @@ void FLevelEditorViewportClient::AddHoverEffect( FViewportHoverTarget& InHoverTa
 
 	if( ActorUnderCursor != NULL )
 	{
-		TArray<UPrimitiveComponent*> Components;
+		TInlineComponentArray<UPrimitiveComponent*> Components;
 		ActorUnderCursor->GetComponents(Components);
 
 		for(int32 ComponentIndex = 0;ComponentIndex < Components.Num();ComponentIndex++)
@@ -3864,7 +3864,7 @@ void FLevelEditorViewportClient::RemoveHoverEffect( FViewportHoverTarget& InHove
 	AActor* CurHoveredActor = InHoverTarget.HoveredActor;
 	if( CurHoveredActor != NULL )
 	{
-		TArray<UPrimitiveComponent*> Components;
+		TInlineComponentArray<UPrimitiveComponent*> Components;
 		CurHoveredActor->GetComponents(Components);
 
 		for(int32 ComponentIndex = 0;ComponentIndex < Components.Num();ComponentIndex++)
diff --git a/Engine/Source/Editor/UnrealEd/Private/Lightmass/Lightmass.cpp b/Engine/Source/Editor/UnrealEd/Private/Lightmass/Lightmass.cpp
--- a/Engine/Source/Editor/UnrealEd/Private/Lightmass/Lightmass.cpp
+++ b/Engine/Source/Editor/UnrealEd/Private/Lightmass/Lightmass.cpp
@@ -756,7 +756,7 @@ void FLightmassExporter::WriteVisibilityData( int32 Channel )
 				AActor* CurrentActor = Volume->OverrideVisibleActors[ActorIndex];
 				if (CurrentActor)
 				{
-					TArray<UPrimitiveComponent*> Components;
+					TInlineComponentArray<UPrimitiveComponent*> Components;
 					CurrentActor->GetComponents(Components);
 
 					for (int32 ComponentIndex = 0; ComponentIndex < Components.Num(); ComponentIndex++)
@@ -775,7 +775,7 @@ void FLightmassExporter::WriteVisibilityData( int32 Channel )
 				AActor* RemoveActor = Volume->OverrideInvisibleActors[RemoveActorIndex];
 				if (RemoveActor)
 				{
-					TArray<UPrimitiveComponent*> Components;
+					TInlineComponentArray<UPrimitiveComponent*> Components;
 					RemoveActor->GetComponents(Components);
 
 					for (int32 ComponentIndex = 0; ComponentIndex < Components.Num(); ComponentIndex++)
@@ -799,7 +799,7 @@ void FLightmassExporter::WriteVisibilityData( int32 Channel )
 						AActor* RemoveActor = Level->Actors[RemoveActorIndex];
 						if (RemoveActor)
 						{
-							TArray<UPrimitiveComponent*> PrimitiveComponents;
+							TInlineComponentArray<UPrimitiveComponent*> PrimitiveComponents;
 							RemoveActor->GetComponents(PrimitiveComponents);
 
 							for (int32 ComponentIndex = 0; ComponentIndex < PrimitiveComponents.Num(); ComponentIndex++)
@@ -2131,7 +2131,7 @@ void FLightmassExporter::WriteDebugInput( Lightmass::FDebugLightingInputData& In
 	{
 		for (FSelectedActorIterator It(World); It; ++It)
 		{
-			TArray<UPrimitiveComponent*> Components;
+			TInlineComponentArray<UPrimitiveComponent*> Components;
 			It->GetComponents(Components);
 
 			for (int32 ComponentIndex = 0; ComponentIndex < Components.Num(); ComponentIndex++)
diff --git a/Engine/Source/Editor/UnrealEd/Private/StaticLightingSystem/StaticLightingSystem.cpp b/Engine/Source/Editor/UnrealEd/Private/StaticLightingSystem/StaticLightingSystem.cpp
--- a/Engine/Source/Editor/UnrealEd/Private/StaticLightingSystem/StaticLightingSystem.cpp
+++ b/Engine/Source/Editor/UnrealEd/Private/StaticLightingSystem/StaticLightingSystem.cpp
@@ -668,7 +668,7 @@ void FStaticLightingSystem::InvalidateStaticLighting()
 				{
 					if (!Options.bOnlyBuildVisibility)
 					{
-						TArray<UActorComponent*> Components;
+						TInlineComponentArray<UActorComponent*> Components;
 						Actor->GetComponents(Components);
 
 						for (int32 ComponentIndex = 0; ComponentIndex < Components.Num(); ComponentIndex++)
@@ -730,7 +730,7 @@ void FStaticLightingSystem::PostInvalidateStaticLighting()
 				{
 					if (!Options.bOnlyBuildVisibility)
 					{
-						TArray<UPrimitiveComponent*> Components;
+						TInlineComponentArray<UPrimitiveComponent*> Components;
 						Actor->GetComponents(Components);
 
 						for (int32 ComponentIndex = 0; ComponentIndex < Components.Num(); ComponentIndex++)
@@ -936,7 +936,7 @@ void FStaticLightingSystem::GatherStaticLightingInfo(bool bRebuildDirtyGeometryF
 					bBuildLightingForLevel &&
 					(!Options.bOnlyBuildSelected || Actor->IsSelected());
 
-				TArray<UPrimitiveComponent*> Components;
+				TInlineComponentArray<UPrimitiveComponent*> Components;
 				Actor->GetComponents(Components);
 
 				if (bBuildActorLighting)
@@ -1062,7 +1062,7 @@ void FStaticLightingSystem::ApplyNewLightingData(bool bLightingSuccessful)
 
 				if (Actor && bLightingSuccessful && !Options.bOnlyBuildSelected)
 				{
-					TArray<ULightComponentBase*> Components;
+					TInlineComponentArray<ULightComponentBase*> Components;
 					Actor->GetComponents(Components);
 
 					for (int32 ComponentIndex = 0; ComponentIndex < Components.Num(); ComponentIndex++)
diff --git a/Engine/Source/Editor/UnrealEd/Private/ThumbnailHelpers.cpp b/Engine/Source/Editor/UnrealEd/Private/ThumbnailHelpers.cpp
--- a/Engine/Source/Editor/UnrealEd/Private/ThumbnailHelpers.cpp
+++ b/Engine/Source/Editor/UnrealEd/Private/ThumbnailHelpers.cpp
@@ -1106,7 +1106,7 @@ TArray<UPrimitiveComponent*> FBlueprintThumbnailScene::GetPooledVisualizableComp
 			// Instance native components based on the CDO template
 			AActor* CDO = Blueprint->GeneratedClass->GetDefaultObject<AActor>();
 
-			TArray<UActorComponent*> Components;
+			TInlineComponentArray<UActorComponent*> Components;
 			CDO->GetComponents(Components);
 
 			for ( auto CompIt = Components.CreateConstIterator(); CompIt; ++CompIt )
diff --git a/Engine/Source/Editor/UnrealEd/Private/ThumbnailRendering/BlueprintThumbnailRenderer.cpp b/Engine/Source/Editor/UnrealEd/Private/ThumbnailRendering/BlueprintThumbnailRenderer.cpp
--- a/Engine/Source/Editor/UnrealEd/Private/ThumbnailRendering/BlueprintThumbnailRenderer.cpp
+++ b/Engine/Source/Editor/UnrealEd/Private/ThumbnailRendering/BlueprintThumbnailRenderer.cpp
@@ -30,7 +30,7 @@ bool UBlueprintThumbnailRenderer::CanVisualizeAsset(UObject* Object)
 		// Try to find any visible primitive components in the native class' CDO
 		AActor* CDO = Blueprint->GeneratedClass->GetDefaultObject<AActor>();
 
-		TArray<UActorComponent*> Components;
+		TInlineComponentArray<UActorComponent*> Components;
 		CDO->GetComponents(Components);
 
 		for (auto CompIt = Components.CreateConstIterator(); CompIt; ++CompIt)
diff --git a/Engine/Source/Editor/UnrealEd/Private/UnrealEdEngine.cpp b/Engine/Source/Editor/UnrealEd/Private/UnrealEdEngine.cpp
--- a/Engine/Source/Editor/UnrealEd/Private/UnrealEdEngine.cpp
+++ b/Engine/Source/Editor/UnrealEd/Private/UnrealEdEngine.cpp
@@ -208,7 +208,7 @@ void UUnrealEdEngine::MakeSortedSpriteInfo(TArray<FSpriteCategoryInfo>& OutSorte
 			const AActor* CurDefaultClassActor = Class->GetDefaultObject<AActor>();
 			if ( CurDefaultClassActor )
 			{
-				TArray<UActorComponent*> Components;
+				TInlineComponentArray<UActorComponent*> Components;
 				CurDefaultClassActor->GetComponents(Components);
 
 				for ( auto* Comp : Components )
@@ -1058,7 +1058,7 @@ void UUnrealEdEngine::UpdateVolumeActorVisibility( UClass* InVolumeActorClass, F
 			AActor* ActorToUpdate = ActorsThatChanged[ ActorIdx ];
 
 			// Find all registered primitive components and update the scene proxy with the actors updated visibility map
-			TArray<UPrimitiveComponent*> PrimitiveComponents;
+			TInlineComponentArray<UPrimitiveComponent*> PrimitiveComponents;
 			ActorToUpdate->GetComponents(PrimitiveComponents);
 
 			for( int32 ComponentIdx = 0; ComponentIdx < PrimitiveComponents.Num(); ++ComponentIdx )
@@ -1126,7 +1126,7 @@ void UUnrealEdEngine::DrawComponentVisualizers(const FSceneView* View, FPrimitiv
 		if(Actor != NULL)
 		{
 			// Then iterate over components of that actor
-			TArray<UActorComponent*> Components;
+			TInlineComponentArray<UActorComponent*> Components;
 			Actor->GetComponents(Components);
 
 			for(int32 CompIdx=0; CompIdx<Components.Num(); CompIdx++)
@@ -1157,7 +1157,7 @@ void UUnrealEdEngine::DrawComponentVisualizersHUD(const FViewport* Viewport, con
 		if (Actor != NULL)
 		{
 			// Then iterate over components of that actor
-			TArray<UActorComponent*> Components;
+			TInlineComponentArray<UActorComponent*> Components;
 			Actor->GetComponents(Components);
 
 			for (int32 CompIdx = 0; CompIdx<Components.Num(); CompIdx++)
diff --git a/Engine/Source/Editor/UnrealEd/Private/UnrealEdSrv.cpp b/Engine/Source/Editor/UnrealEd/Private/UnrealEdSrv.cpp
--- a/Engine/Source/Editor/UnrealEd/Private/UnrealEdSrv.cpp
+++ b/Engine/Source/Editor/UnrealEd/Private/UnrealEdSrv.cpp
@@ -1791,7 +1791,7 @@ TArray<FPoly*> GetSelectedPolygons()
 		checkSlow( Actor->IsA(AActor::StaticClass()) );
 		FTransform ActorToWorld = Actor->ActorToWorld();
 		
-		TArray<UStaticMeshComponent*> StaticMeshComponents;
+		TInlineComponentArray<UStaticMeshComponent*> StaticMeshComponents;
 		Actor->GetComponents(StaticMeshComponents);
 
 		for(int32 j=0; j<StaticMeshComponents.Num(); j++)
diff --git a/Engine/Source/Editor/UnrealEd/Private/VertexSnapping.cpp b/Engine/Source/Editor/UnrealEd/Private/VertexSnapping.cpp
--- a/Engine/Source/Editor/UnrealEd/Private/VertexSnapping.cpp
+++ b/Engine/Source/Editor/UnrealEd/Private/VertexSnapping.cpp
@@ -308,7 +308,7 @@ void FVertexSnappingImpl::ClearSnappingHelpers( bool bClearImmediately )
 
 static void DrawSnapVertices( AActor* Actor, float PointSize, FPrimitiveDrawInterface* PDI )
 {
-	TArray<UActorComponent*> Components;
+	TInlineComponentArray<UActorComponent*> Components;
 	Actor->GetComponents(Components);
 
 	// Get the closest vertex on each component
@@ -508,7 +508,7 @@ FSnappingVertex FVertexSnappingImpl::GetClosestVertex( const TArray<FSnapActor>&
 		AActor* Actor = SnapActor.Actor;
 
 		// Get the closest vertex on each component
-		TArray<UPrimitiveComponent*> PrimitiveComponents;
+		TInlineComponentArray<UPrimitiveComponent*> PrimitiveComponents;
 		Actor->GetComponents(PrimitiveComponents);
 
 		for( int32 ComponentIndex  = 0; ComponentIndex < PrimitiveComponents.Num(); ++ComponentIndex )
diff --git a/Engine/Source/Runtime/AIModule/Private/BehaviorTree/BlueprintNodeHelpers.cpp b/Engine/Source/Runtime/AIModule/Private/BehaviorTree/BlueprintNodeHelpers.cpp
--- a/Engine/Source/Runtime/AIModule/Private/BehaviorTree/BlueprintNodeHelpers.cpp
+++ b/Engine/Source/Runtime/AIModule/Private/BehaviorTree/BlueprintNodeHelpers.cpp
@@ -219,7 +219,7 @@ namespace BlueprintNodeHelpers
 
 		if (OwningActor && !bFound)
 		{
-			TArray<UBehaviorTreeComponent*> Components;
+			TInlineComponentArray<UBehaviorTreeComponent*> Components;
 			OwningActor->GetComponents(Components);
 			for (int32 ComponentIndex = 0; ComponentIndex < Components.Num(); ComponentIndex++)
 			{
diff --git a/Engine/Source/Runtime/Engine/Private/AI/Navigation/NavigationSystem.cpp b/Engine/Source/Runtime/Engine/Private/AI/Navigation/NavigationSystem.cpp
--- a/Engine/Source/Runtime/Engine/Private/AI/Navigation/NavigationSystem.cpp
+++ b/Engine/Source/Runtime/Engine/Private/AI/Navigation/NavigationSystem.cpp
@@ -1331,18 +1331,16 @@ bool UNavigationSystem::IsNavigationRelevant(const AActor* TestActor) const
 		return true;
 	}
 
-	TArray<UActorComponent*> Components;
 	if (TestActor)
 	{
-		TestActor->GetComponents(Components);
-	}
-
-	for (int32 Idx = 0; Idx < Components.Num(); Idx++)
-	{
-		NavInterface = Cast<const INavRelevantInterface>(Components[Idx]);
-		if (NavInterface && NavInterface->IsNavigationRelevant())
+		TInlineComponentArray<UActorComponent*> Components;
+		for (int32 Idx = 0; Idx < Components.Num(); Idx++)
 		{
-			return true;
+			NavInterface = Cast<const INavRelevantInterface>(Components[Idx]);
+			if (NavInterface && NavInterface->IsNavigationRelevant())
+			{
+				return true;
+			}
 		}
 	}
 
@@ -2243,7 +2241,7 @@ void UNavigationSystem::UpdateNavOctreeAll(AActor* Actor)
 	{
 		UpdateNavOctree(Actor);
 
-		TArray<UActorComponent*> Components;
+		TInlineComponentArray<UActorComponent*> Components;
 		Actor->GetComponents(Components);
 
 		for (int32 Idx = 0; Idx < Components.Num(); Idx++)
@@ -2255,7 +2253,7 @@ void UNavigationSystem::UpdateNavOctreeAll(AActor* Actor)
 
 void UNavigationSystem::UpdateNavOctreeBounds(AActor* Actor)
 {
-	TArray<UActorComponent*> Components;
+	TInlineComponentArray<UActorComponent*> Components;
 	Actor->GetComponents(Components);
 
 	for (int32 Idx = 0; Idx < Components.Num(); Idx++)
@@ -2274,7 +2272,7 @@ void UNavigationSystem::ClearNavOctreeAll(AActor* Actor)
 	{
 		OnActorUnregistered(Actor);
 
-		TArray<UActorComponent*> Components;
+		TInlineComponentArray<UActorComponent*> Components;
 		Actor->GetComponents(Components);
 
 		for (int32 Idx = 0; Idx < Components.Num(); Idx++)
diff --git a/Engine/Source/Runtime/Engine/Private/ActorConstruction.cpp b/Engine/Source/Runtime/Engine/Private/ActorConstruction.cpp
--- a/Engine/Source/Runtime/Engine/Private/ActorConstruction.cpp
+++ b/Engine/Source/Runtime/Engine/Private/ActorConstruction.cpp
@@ -50,7 +50,7 @@ namespace
 		/** Called before UCS execution has started for the given Actor */
 		void PreProcessComponents(const AActor* InActor)
 		{
-			TArray<UActorComponent*> ActorComponents;
+			TInlineComponentArray<UActorComponent*> ActorComponents;
 			InActor->GetComponents(ActorComponents);
 			for (auto CompIt = ActorComponents.CreateConstIterator(); CompIt; ++CompIt)
 			{
@@ -186,7 +186,7 @@ void AActor::ResetPropertiesForConstruction()
 void AActor::DestroyConstructedComponents()
 {
 	// Remove all existing components
-	TArray<UActorComponent*> PreviouslyAttachedComponents;
+	TInlineComponentArray<UActorComponent*> PreviouslyAttachedComponents;
 	GetComponents(PreviouslyAttachedComponents);
 	for (int32 i = 0; i < PreviouslyAttachedComponents.Num(); i++)
 	{
diff --git a/Engine/Source/Runtime/Engine/Private/ActorEditor.cpp b/Engine/Source/Runtime/Engine/Private/ActorEditor.cpp
--- a/Engine/Source/Runtime/Engine/Private/ActorEditor.cpp
+++ b/Engine/Source/Runtime/Engine/Private/ActorEditor.cpp
@@ -45,7 +45,7 @@ void AActor::PostEditChangeProperty(FPropertyChangedEvent& PropertyChangedEvent)
 			AInstancedFoliageActor* IFA = AInstancedFoliageActor::GetInstancedFoliageActorForLevel(GetLevel());
 			if (IFA)
 			{
-				TArray<UActorComponent*> Components;
+				TInlineComponentArray<UActorComponent*> Components;
 				GetComponents(Components);
 
 				for ( int32 Idx = 0 ; Idx < Components.Num() ; ++Idx )
@@ -97,7 +97,7 @@ void AActor::PostEditMove(bool bFinished)
 			AInstancedFoliageActor* IFA = AInstancedFoliageActor::GetInstancedFoliageActorForLevel(GetLevel());
 			if (IFA)
 			{
-				TArray<UActorComponent*> Components;
+				TInlineComponentArray<UActorComponent*> Components;
 				GetComponents(Components);
 
 				for ( int32 Idx = 0 ; Idx < Components.Num() ; ++Idx )
@@ -294,7 +294,7 @@ TSharedPtr<ITransactionObjectAnnotation> AActor::GetTransactionAnnotation() cons
 void AActor::PreEditUndo()
 {
 	// Since child actor components will rebuild themselves get rid of the Actor before we make changes
-	TArray<UChildActorComponent*> ChildActorComponents;
+	TInlineComponentArray<UChildActorComponent*> ChildActorComponents;
 	GetComponents(ChildActorComponents);
 
 	for (UChildActorComponent* ChildActorComponent : ChildActorComponents)
@@ -711,7 +711,7 @@ void AActor::CheckForErrors()
 	}
 
 	// Route error checking to components.
-	TArray<UActorComponent*> Components;
+	TInlineComponentArray<UActorComponent*> Components;
 	GetComponents(Components);
 
 	for ( int32 ComponentIndex = 0 ; ComponentIndex < Components.Num() ; ++ComponentIndex )
diff --git a/Engine/Source/Runtime/Engine/Private/Blueprint.cpp b/Engine/Source/Runtime/Engine/Private/Blueprint.cpp
--- a/Engine/Source/Runtime/Engine/Private/Blueprint.cpp
+++ b/Engine/Source/Runtime/Engine/Private/Blueprint.cpp
@@ -41,7 +41,7 @@ static void ConformNativeComponents(UBlueprint* Blueprint)
 	{
 		if (AActor* BlueprintCDO = Cast<AActor>(BlueprintClass->ClassDefaultObject))
 		{
-			TArray<UActorComponent*> OldNativeComponents;
+			TInlineComponentArray<UActorComponent*> OldNativeComponents;
 			// collect the native components that this blueprint was serialized out 
 			// with (the native components it had last time it was saved)
 			BlueprintCDO->GetComponents(OldNativeComponents);
@@ -50,11 +50,11 @@ static void ConformNativeComponents(UBlueprint* Blueprint)
 			AActor* NativeCDO = CastChecked<AActor>(NativeSuperClass->ClassDefaultObject);
 			// collect the more up to date native components (directly from the 
 			// native super-class)
-			TArray<UActorComponent*> NewNativeComponents;
+			TInlineComponentArray<UActorComponent*> NewNativeComponents;
 			NativeCDO->GetComponents(NewNativeComponents);
 
 			// utility lambda for finding named components in a supplied list
-			auto FindNamedComponentLambda = [](FName const ComponentName, TArray<UActorComponent*> const& ComponentList)->UActorComponent*
+			auto FindNamedComponentLambda = [](FName const ComponentName, TInlineComponentArray<UActorComponent*> const& ComponentList)->UActorComponent*
 			{
 				UActorComponent* FoundComponent = nullptr;
 				for (UActorComponent* Component : ComponentList)
diff --git a/Engine/Source/Runtime/Engine/Private/Character.cpp b/Engine/Source/Runtime/Engine/Private/Character.cpp
--- a/Engine/Source/Runtime/Engine/Private/Character.cpp
+++ b/Engine/Source/Runtime/Engine/Private/Character.cpp
@@ -424,7 +424,7 @@ namespace MovementBaseUtility
 				}
 
 				// @TODO: We need to find a more efficient way of finding all ticking components in an actor.
-				TArray<UActorComponent*> Components;
+				TInlineComponentArray<UActorComponent*> Components;
 				NewBaseOwner->GetComponents(Components);
 				for (auto& Component : Components)
 				{
@@ -454,7 +454,7 @@ namespace MovementBaseUtility
 				BasedObjectTick.RemovePrerequisite(OldBaseOwner, OldBaseOwner->PrimaryActorTick);
 
 				// @TODO: We need to find a more efficient way of finding all ticking components in an actor.
-				TArray<UActorComponent*> Components;
+				TInlineComponentArray<UActorComponent*> Components;
 				OldBaseOwner->GetComponents(Components);
 				for (auto& Component : Components)
 				{
diff --git a/Engine/Source/Runtime/Engine/Private/ComponentInstanceDataCache.cpp b/Engine/Source/Runtime/Engine/Private/ComponentInstanceDataCache.cpp
--- a/Engine/Source/Runtime/Engine/Private/ComponentInstanceDataCache.cpp
+++ b/Engine/Source/Runtime/Engine/Private/ComponentInstanceDataCache.cpp
@@ -71,7 +71,7 @@ FComponentInstanceDataCache::FComponentInstanceDataCache(const AActor* Actor)
 {
 	if(Actor != NULL)
 	{
-		TArray<UActorComponent*> Components;
+		TInlineComponentArray<UActorComponent*> Components;
 		Actor->GetComponents(Components);
 
 		// Grab per-instance data we want to persist
@@ -102,7 +102,7 @@ void FComponentInstanceDataCache::ApplyToActor(AActor* Actor) const
 {
 	if(Actor != NULL)
 	{
-		TArray<UActorComponent*> Components;
+		TInlineComponentArray<UActorComponent*> Components;
 		Actor->GetComponents(Components);
 
 		// Apply per-instance data.
diff --git a/Engine/Source/Runtime/Engine/Private/ContentStreaming.cpp b/Engine/Source/Runtime/Engine/Private/ContentStreaming.cpp
--- a/Engine/Source/Runtime/Engine/Private/ContentStreaming.cpp
+++ b/Engine/Source/Runtime/Engine/Private/ContentStreaming.cpp
@@ -2596,7 +2596,7 @@ void FStreamingManagerTexture::BoostTextures( AActor* Actor, float BoostFactor )
 		TArray<UTexture*> Textures;
 		Textures.Empty( 32 );
 
-		TArray<UPrimitiveComponent*> Components;
+		TInlineComponentArray<UPrimitiveComponent*> Components;
 		Actor->GetComponents(Components);
 
 		for(int32 ComponentIndex = 0;ComponentIndex < Components.Num();ComponentIndex++)
@@ -2859,7 +2859,7 @@ void FStreamingManagerTexture::NotifyActorSpawned( AActor* Actor )
 {
 	if ( bUseDynamicStreaming )
 	{
-		TArray<UPrimitiveComponent*> Components;
+		TInlineComponentArray<UPrimitiveComponent*> Components;
 		Actor->GetComponents(Components);
 
 		for(int32 ComponentIndex = 0;ComponentIndex < Components.Num();ComponentIndex++)
@@ -2876,7 +2876,7 @@ void FStreamingManagerTexture::NotifyActorSpawned( AActor* Actor )
 /** Called when a spawned primitive is deleted. */
 void FStreamingManagerTexture::NotifyActorDestroyed( AActor* Actor )
 {
-	TArray<UPrimitiveComponent*> Components;
+	TInlineComponentArray<UPrimitiveComponent*> Components;
 	Actor->GetComponents(Components);
 
 	for(int32 ComponentIndex = 0;ComponentIndex < Components.Num();ComponentIndex++)
diff --git a/Engine/Source/Runtime/Engine/Private/DebugCameraHUD.cpp b/Engine/Source/Runtime/Engine/Private/DebugCameraHUD.cpp
--- a/Engine/Source/Runtime/Engine/Private/DebugCameraHUD.cpp
+++ b/Engine/Source/Runtime/Engine/Private/DebugCameraHUD.cpp
@@ -135,7 +135,7 @@ void ADebugCameraHUD::PostRender()
 				}
 				else
 				{
-					TArray<UMeshComponent*> Components;
+					TInlineComponentArray<UMeshComponent*> Components;
 					GetComponents(Components);
 
 					for ( int32 i=0; i<Components.Num(); i++ )
diff --git a/Engine/Source/Runtime/Engine/Private/GameViewportClient.cpp b/Engine/Source/Runtime/Engine/Private/GameViewportClient.cpp
--- a/Engine/Source/Runtime/Engine/Private/GameViewportClient.cpp
+++ b/Engine/Source/Runtime/Engine/Private/GameViewportClient.cpp
@@ -2272,7 +2272,7 @@ void UGameViewportClient::ShowCollisionOnSpawnedActors(AActor* Actor)
 {
 	CollisionComponentVisibilityMap& Mapping = GetCollisionComponentVisibilityMap();
 
-	TArray<UPrimitiveComponent*> Components;
+	TInlineComponentArray<UPrimitiveComponent*> Components;
 	check(Actor != nullptr);
 	Actor->GetComponents(Components);
 
diff --git a/Engine/Source/Runtime/Engine/Private/InstancedFoliage.cpp b/Engine/Source/Runtime/Engine/Private/InstancedFoliage.cpp
--- a/Engine/Source/Runtime/Engine/Private/InstancedFoliage.cpp
+++ b/Engine/Source/Runtime/Engine/Private/InstancedFoliage.cpp
@@ -1415,7 +1415,7 @@ void AInstancedFoliageActor::Serialize(FArchive& Ar)
 	// Clean up any old cluster components and convert to hierarchical instanced foliage.
 	if (Ar.CustomVer(FFoliageCustomVersion::GUID) < FFoliageCustomVersion::FoliageUsingHierarchicalISMC)
 	{
-		TArray<UInstancedStaticMeshComponent*> ClusterComponents;
+		TInlineComponentArray<UInstancedStaticMeshComponent*> ClusterComponents;
 		GetComponents(ClusterComponents);
 		for (UInstancedStaticMeshComponent* Component : ClusterComponents)
 		{
diff --git a/Engine/Source/Runtime/Engine/Private/Interpolation.cpp b/Engine/Source/Runtime/Engine/Private/Interpolation.cpp
--- a/Engine/Source/Runtime/Engine/Private/Interpolation.cpp
+++ b/Engine/Source/Runtime/Engine/Private/Interpolation.cpp
@@ -8146,7 +8146,7 @@ static void AddMaterialRefsForActor(
 {
 	if (Actor && !Actor->IsRootComponentStatic())
 	{
-		TArray<USceneComponent*> Components;
+		TInlineComponentArray<USceneComponent*> Components;
 		Actor->GetComponents(Components);
 
 		for (int32 MaterialIndex = 0; MaterialIndex < Materials.Num(); MaterialIndex++)
diff --git a/Engine/Source/Runtime/Engine/Private/LevelActor.cpp b/Engine/Source/Runtime/Engine/Private/LevelActor.cpp
--- a/Engine/Source/Runtime/Engine/Private/LevelActor.cpp
+++ b/Engine/Source/Runtime/Engine/Private/LevelActor.cpp
@@ -452,7 +452,7 @@ bool UWorld::DestroyActor( AActor* ThisActor, bool bNetForce, bool bShouldModify
 	TArray<AActor*> AttachedActors;
 	ThisActor->GetAttachedActors(AttachedActors);
 
-	TArray<USceneComponent*> SceneComponents;
+	TInlineComponentArray<USceneComponent*> SceneComponents;
 	ThisActor->GetComponents(SceneComponents);
 
 	if (AttachedActors.Num() > 0)
diff --git a/Engine/Source/Runtime/Engine/Private/LocalPlayer.cpp b/Engine/Source/Runtime/Engine/Private/LocalPlayer.cpp
--- a/Engine/Source/Runtime/Engine/Private/LocalPlayer.cpp
+++ b/Engine/Source/Runtime/Engine/Private/LocalPlayer.cpp
@@ -1223,7 +1223,7 @@ bool ULocalPlayer::HandleListPawnComponentsCommand( const TCHAR* Cmd, FOutputDev
 		APawn *Pawn = *It;	
 		UE_LOG(LogPlayerManagement, Log, TEXT("Components for pawn: %s (collision component: %s)"),*Pawn->GetName(),*Pawn->GetRootComponent()->GetName());
 
-		TArray<UActorComponent*> Components;
+		TInlineComponentArray<UActorComponent*> Components;
 		Pawn->GetComponents(Components);
 
 		for (int32 CompIdx = 0; CompIdx < Components.Num(); CompIdx++)
diff --git a/Engine/Source/Runtime/Engine/Private/Particles/ParticleModules_Location.cpp b/Engine/Source/Runtime/Engine/Private/Particles/ParticleModules_Location.cpp
--- a/Engine/Source/Runtime/Engine/Private/Particles/ParticleModules_Location.cpp
+++ b/Engine/Source/Runtime/Engine/Private/Particles/ParticleModules_Location.cpp
@@ -2313,7 +2313,7 @@ void UParticleModuleLocationSkelVertSurface::UpdateBoneIndicesList(FParticleEmit
 		{
 			if(ActorInst != NULL)
 			{
-				TArray<USkeletalMeshComponent*> Components;
+				TInlineComponentArray<USkeletalMeshComponent*> Components;
 				ActorInst->GetComponents(Components);
 
 				int32 InsertionIndex = 0;
diff --git a/Engine/Source/Runtime/Engine/Private/Particles/ParticleSystemRender.cpp b/Engine/Source/Runtime/Engine/Private/Particles/ParticleSystemRender.cpp
--- a/Engine/Source/Runtime/Engine/Private/Particles/ParticleSystemRender.cpp
+++ b/Engine/Source/Runtime/Engine/Private/Particles/ParticleSystemRender.cpp
@@ -6656,11 +6656,11 @@ void DrawParticleSystemHelpers(UParticleSystemComponent* InPSysComp, const FScen
 
 ENGINE_API void DrawParticleSystemHelpers(const FSceneView* View,FPrimitiveDrawInterface* PDI)
 {
-	TArray<UParticleSystemComponent*> PSCArray;
+	TInlineComponentArray<UParticleSystemComponent*> PSCArray;
 
 	for (TObjectIterator<AActor> It; It; ++It)
 	{
-		PSCArray.Empty();
+		PSCArray.Reset();
 		(*It)->GetComponents(PSCArray);
 
 		for (int PSCIndex = 0; PSCIndex < PSCArray.Num(); ++PSCIndex)
diff --git a/Engine/Source/Runtime/Engine/Private/PhysicsEngine/RadialForceComponent.cpp b/Engine/Source/Runtime/Engine/Private/PhysicsEngine/RadialForceComponent.cpp
--- a/Engine/Source/Runtime/Engine/Private/PhysicsEngine/RadialForceComponent.cpp
+++ b/Engine/Source/Runtime/Engine/Private/PhysicsEngine/RadialForceComponent.cpp
@@ -57,7 +57,7 @@ void URadialForceComponent::TickComponent(float DeltaTime, enum ELevelTick TickT
 				AActor* PokeOwner = PokeComp->GetOwner();
 				if(PokeOwner)
 				{
-					TArray<UMovementComponent*> MovementComponents;
+					TInlineComponentArray<UMovementComponent*> MovementComponents;
 					PokeOwner->GetComponents<UMovementComponent>(MovementComponents);
 					for(const auto& MovementComponent : MovementComponents)
 					{
@@ -119,7 +119,7 @@ void URadialForceComponent::FireImpulse()
 			PokeComp->AddRadialImpulse( Origin, Radius, ImpulseStrength, Falloff, bImpulseVelChange );
 
 			// see if this is a target for a movement component
-			TArray<UMovementComponent*> MovementComponents;
+			TInlineComponentArray<UMovementComponent*> MovementComponents;
 			PokeComp->GetOwner()->GetComponents<UMovementComponent>(MovementComponents);
 			for(const auto& MovementComponent : MovementComponents)
 			{
diff --git a/Engine/Source/Runtime/Engine/Private/PlayerController.cpp b/Engine/Source/Runtime/Engine/Private/PlayerController.cpp
--- a/Engine/Source/Runtime/Engine/Private/PlayerController.cpp
+++ b/Engine/Source/Runtime/Engine/Private/PlayerController.cpp
@@ -372,7 +372,7 @@ FString APlayerController::ConsoleCommand(const FString& Cmd,bool bWriteToLog)
 
 void APlayerController::CleanUpAudioComponents()
 {
-	TArray<UAudioComponent*> Components;
+	TInlineComponentArray<UAudioComponent*> Components;
 	GetComponents(Components);
 
 	for(int32 CompIndex = 0; CompIndex < Components.Num(); CompIndex++)
@@ -2290,7 +2290,7 @@ void APlayerController::BuildInputStack(TArray<UInputComponent*>& InputStack)
 			}
 
 			// See if there is another InputComponent that was added to the Pawn's components array (possibly by script).
-			TArray<UInputComponent*> Components;
+			TInlineComponentArray<UInputComponent*> Components;
 			ControlledPawn->GetComponents(Components);
 
 			for (int32 i=0; i < Components.Num(); i++)
@@ -4401,7 +4401,7 @@ void APlayerController::BuildHiddenComponentList(const FVector& ViewLocation, TS
 		AActor* HiddenActor = HiddenActors[ActorIndex];
 		if (HiddenActor != NULL)
 		{
-			TArray<UPrimitiveComponent*> Components;
+			TInlineComponentArray<UPrimitiveComponent*> Components;
 			HiddenActor->GetComponents(Components);
 
 			for (int32 ComponentIndex = 0; ComponentIndex < Components.Num(); ComponentIndex++)
diff --git a/Engine/Source/Runtime/Engine/Private/SCS_Node.cpp b/Engine/Source/Runtime/Engine/Private/SCS_Node.cpp
--- a/Engine/Source/Runtime/Engine/Private/SCS_Node.cpp
+++ b/Engine/Source/Runtime/Engine/Private/SCS_Node.cpp
@@ -317,7 +317,7 @@ USceneComponent* USCS_Node::GetParentComponentTemplate(UBlueprint* InBlueprint)
 			if(CDO != NULL)
 			{
 				// Find the component template in the CDO that matches the specified name
-				TArray<USceneComponent*> Components;
+				TInlineComponentArray<USceneComponent*> Components;
 				CDO->GetComponents(Components);
 
 				for(auto CompIt = Components.CreateIterator(); CompIt; ++CompIt)
diff --git a/Engine/Source/Runtime/Engine/Private/SceneCaptureComponent.cpp b/Engine/Source/Runtime/Engine/Private/SceneCaptureComponent.cpp
--- a/Engine/Source/Runtime/Engine/Private/SceneCaptureComponent.cpp
+++ b/Engine/Source/Runtime/Engine/Private/SceneCaptureComponent.cpp
@@ -194,7 +194,7 @@ void USceneCaptureComponent::HideActorComponents(AActor* InActor)
 {
 	if (InActor)
 	{
-		TArray<UPrimitiveComponent*> PrimitiveComponents;
+		TInlineComponentArray<UPrimitiveComponent*> PrimitiveComponents;
 		InActor->GetComponents(PrimitiveComponents);
 		for (int32 ComponentIndex = 0, NumComponents = PrimitiveComponents.Num(); ComponentIndex < NumComponents; ++ComponentIndex)
 		{
diff --git a/Engine/Source/Runtime/Engine/Private/SimpleConstructionScript.cpp b/Engine/Source/Runtime/Engine/Private/SimpleConstructionScript.cpp
--- a/Engine/Source/Runtime/Engine/Private/SimpleConstructionScript.cpp
+++ b/Engine/Source/Runtime/Engine/Private/SimpleConstructionScript.cpp
@@ -251,7 +251,7 @@ void USimpleConstructionScript::FixupRootNodeParentReferences()
 				if(CDO != NULL)
 				{
 					// Look for the parent component in the CDO's components array
-					TArray<UActorComponent*> Components;
+					TInlineComponentArray<UActorComponent*> Components;
 					CDO->GetComponents(Components);
 
 					for (auto CompIter = Components.CreateConstIterator(); CompIter && !bWasFound; ++CompIter)
@@ -313,14 +313,14 @@ void USimpleConstructionScript::ExecuteScriptOnActor(AActor* Actor, const FTrans
 {
 	if(RootNodes.Num() > 0)
 	{
-		TArray<UActorComponent*> InstancedComponents;
+		TInlineComponentArray<UActorComponent*> InstancedComponents;
 		for(auto NodeIt = RootNodes.CreateIterator(); NodeIt; ++NodeIt)
 		{
 			USCS_Node* RootNode = *NodeIt;
 			if(RootNode != nullptr)
 			{
 				// Get all native scene components
-				TArray<USceneComponent*> Components;
+				TInlineComponentArray<USceneComponent*> Components;
 				Actor->GetComponents(Components);
 
 				// Get the native root component; if it's not set, the first native scene component will be used as root. This matches what's done in the SCS editor.
@@ -639,7 +639,7 @@ void USimpleConstructionScript::ValidateSceneRootNodes()
 			bHasSceneComponentRootNodes = CDO->GetRootComponent() != nullptr;
 			if(!bHasSceneComponentRootNodes)
 			{
-				TArray<USceneComponent*> SceneComponents;
+				TInlineComponentArray<USceneComponent*> SceneComponents;
 				CDO->GetComponents(SceneComponents);
 				bHasSceneComponentRootNodes = SceneComponents.Num() > 0;
 			}
diff --git a/Engine/Source/Runtime/Engine/Private/UnrealEngine.cpp b/Engine/Source/Runtime/Engine/Private/UnrealEngine.cpp
--- a/Engine/Source/Runtime/Engine/Private/UnrealEngine.cpp
+++ b/Engine/Source/Runtime/Engine/Private/UnrealEngine.cpp
@@ -10225,7 +10225,7 @@ void UEngine::CopyPropertiesForUnrelatedObjects(UObject* OldObject, UObject* New
 	AActor* NewActor = Cast<AActor>(NewObject);
 	if(NewActor != NULL)
 	{
-		TArray<UActorComponent*> Components;
+		TInlineComponentArray<UActorComponent*> Components;
 		NewActor->GetComponents(Components);
 
 		for(int32 i=0; i<Components.Num(); i++)
diff --git a/Engine/Source/Runtime/Engine/Private/UnrealExporter.cpp b/Engine/Source/Runtime/Engine/Private/UnrealExporter.cpp
--- a/Engine/Source/Runtime/Engine/Private/UnrealExporter.cpp
+++ b/Engine/Source/Runtime/Engine/Private/UnrealExporter.cpp
@@ -643,7 +643,7 @@ void UExporter::ExportObjectInner(const FExportObjectInnerContext* Context, UObj
 			// Export anything extra for the components. Used for instanced foliage.
 			// This is done after the actor properties so these are set when regenerating the extra data objects.
 			TArray<UActorComponent*> Components;
-			Actor->GetComponents(Components);
+			Actor->GetComponents<UActorComponent, FDefaultAllocator>(Components);
 			ExportComponentExtra(Context, Components, Ar, PortFlags);
 		}
 	}
diff --git a/Engine/Source/Runtime/Engine/Private/World.cpp b/Engine/Source/Runtime/Engine/Private/World.cpp
--- a/Engine/Source/Runtime/Engine/Private/World.cpp
+++ b/Engine/Source/Runtime/Engine/Private/World.cpp
@@ -1238,7 +1238,7 @@ void UWorld::UpdateCullDistanceVolumes()
 		// Establish base line of LD specified cull distances.
 		for( FActorIterator It(this); It; ++It )
 		{
-			TArray<UPrimitiveComponent*> PrimitiveComponents;
+			TInlineComponentArray<UPrimitiveComponent*> PrimitiveComponents;
 			It->GetComponents(PrimitiveComponents);
 			for (UPrimitiveComponent* PrimitiveComponent : PrimitiveComponents)
 			{
diff --git a/Engine/Source/Runtime/GameplayAbilities/Private/Abilities/GameplayAbilityWorldReticle_ActorVisualization.cpp b/Engine/Source/Runtime/GameplayAbilities/Private/Abilities/GameplayAbilityWorldReticle_ActorVisualization.cpp
--- a/Engine/Source/Runtime/GameplayAbilities/Private/Abilities/GameplayAbilityWorldReticle_ActorVisualization.cpp
+++ b/Engine/Source/Runtime/GameplayAbilities/Private/Abilities/GameplayAbilityWorldReticle_ActorVisualization.cpp
@@ -34,7 +34,7 @@ void AGameplayAbilityWorldReticle_ActorVisualization::InitializeReticleVisualiza
 	if (VisualizationActor)
 	{
 		//Get components
-		TArray<UMeshComponent*> MeshComps;
+		TInlineComponentArray<UMeshComponent*> MeshComps;
 		USceneComponent* MyRoot = GetRootComponent();
 		VisualizationActor->GetComponents(MeshComps);
 		check(MyRoot);
@@ -67,10 +67,6 @@ void AGameplayAbilityWorldReticle_ActorVisualization::InitializeReticleVisualiza
 
 void AGameplayAbilityWorldReticle_ActorVisualization::EndPlay(const EEndPlayReason::Type EndPlayReason)
 {
-	TArray<UActorComponent*> CompArray;
-	GetComponents(CompArray);
-	CompArray.Num();
-
 	Super::EndPlay(EndPlayReason);
 }
 
diff --git a/Engine/Source/Runtime/Landscape/Private/Landscape.cpp b/Engine/Source/Runtime/Landscape/Private/Landscape.cpp
--- a/Engine/Source/Runtime/Landscape/Private/Landscape.cpp
+++ b/Engine/Source/Runtime/Landscape/Private/Landscape.cpp
@@ -2076,7 +2076,7 @@ void ALandscapeProxy::FlushFoliageComponents()
 	FoliageCache.ClearCache();
 
 	// Destroy any owned foliage components
-	TArray<UHierarchicalInstancedStaticMeshComponent*> FoliageComps;
+	TInlineComponentArray<UHierarchicalInstancedStaticMeshComponent*> FoliageComps;
 	GetComponents(FoliageComps);
 	for (UHierarchicalInstancedStaticMeshComponent* Component : FoliageComps)
 	{
@@ -2692,7 +2692,7 @@ void ALandscapeProxy::Tick(float DeltaSeconds)
 
 	TSet<UHierarchicalInstancedStaticMeshComponent *> ToDestroy;
 
-	TArray<UHierarchicalInstancedStaticMeshComponent*> FoliageComps;
+	TInlineComponentArray<UHierarchicalInstancedStaticMeshComponent*> FoliageComps;
 	GetComponents(FoliageComps);
 	for (UHierarchicalInstancedStaticMeshComponent* Component : FoliageComps)
 	{
diff --git a/Engine/Source/Runtime/Landscape/Private/LandscapeEdit.cpp b/Engine/Source/Runtime/Landscape/Private/LandscapeEdit.cpp
--- a/Engine/Source/Runtime/Landscape/Private/LandscapeEdit.cpp
+++ b/Engine/Source/Runtime/Landscape/Private/LandscapeEdit.cpp
@@ -2261,7 +2261,7 @@ void ALandscapeProxy::Import(FGuid Guid, int32 VertsX, int32 VertsY,
 
 bool ALandscapeProxy::ExportToRawMesh(int32 InExportLOD, FRawMesh& OutRawMesh) const
 {
-	TArray<ULandscapeComponent*> RegisteredLandscapeComponents;
+	TInlineComponentArray<ULandscapeComponent*> RegisteredLandscapeComponents;
 	GetComponents<ULandscapeComponent>(RegisteredLandscapeComponents);
 
 	const FIntRect LandscapeSectionRect = GetBoundingRect();
@@ -3538,7 +3538,7 @@ void ALandscapeProxy::RecreateCollisionComponents()
 	CollisionComponents.Empty();
 
 	// Destroy any owned collision components
-	TArray<ULandscapeHeightfieldCollisionComponent*> CollisionComps;
+	TInlineComponentArray<ULandscapeHeightfieldCollisionComponent*> CollisionComps;
 	GetComponents(CollisionComps);
 	for (ULandscapeHeightfieldCollisionComponent* Component : CollisionComps)
 	{
diff --git a/Engine/Source/Runtime/Renderer/Private/RendererScene.cpp b/Engine/Source/Runtime/Renderer/Private/RendererScene.cpp
--- a/Engine/Source/Runtime/Renderer/Private/RendererScene.cpp
+++ b/Engine/Source/Runtime/Renderer/Private/RendererScene.cpp
@@ -530,7 +530,7 @@ void FScene::UpdatePrimitiveTransform(UPrimitiveComponent* Primitive)
 	// If the root component of an actor is being moved, update all the actor position of the other components sharing that actor
 	if (Owner && Owner->GetRootComponent() == Primitive)
 	{
-		TArray<UPrimitiveComponent*> Components;
+		TInlineComponentArray<UPrimitiveComponent*> Components;
 		Owner->GetComponents(Components);
 		for (int32 ComponentIndex = 0; ComponentIndex < Components.Num(); ComponentIndex++)
 		{
