Revised FixLevelScriptActorBindings to generate a warning if a target delegate isnt found rather than assert.

diff --git a/Engine/Source/Editor/UnrealEd/Private/Kismet2/BlueprintEditorUtils.cpp b/Engine/Source/Editor/UnrealEd/Private/Kismet2/BlueprintEditorUtils.cpp
--- a/Engine/Source/Editor/UnrealEd/Private/Kismet2/BlueprintEditorUtils.cpp
+++ b/Engine/Source/Editor/UnrealEd/Private/Kismet2/BlueprintEditorUtils.cpp
@@ -5646,12 +5646,19 @@ bool FBlueprintEditorUtils::FixLevelScriptActorBindings(ALevelScriptActor* Level
 				{
 					// Grab the MC delegate we need to add to
 					FMulticastScriptDelegate* TargetDelegate = EventNode->GetTargetDelegate();
-					check(TargetDelegate);
+					if( TargetDelegate == nullptr )
+					{
 
-					// Create the delegate, and add it if it doesn't already exist
-					FScriptDelegate Delegate;
-					Delegate.BindUFunction(LevelScriptActor, TargetFunction);
-					TargetDelegate->AddUnique(Delegate);
+						// Create the delegate, and add it if it doesn't already exist
+						FScriptDelegate Delegate;
+						Delegate.BindUFunction(LevelScriptActor, TargetFunction);
+						TargetDelegate->AddUnique(Delegate);
+					}
+					else
+					{
+						UE_LOG(LogBlueprint, Error, TEXT("Unable to bind event for node: %s (Owner: %s)! Can't get FMulticastScriptDelegate Target Delegate"), *GetNameSafe(EventNode), *GetNameSafe(EventNode->EventOwner));
+						bWasSuccessful = false;
+					}
 				}
 				else
 				{
