#vr Added editor option to support HMD control in editor viewports

diff --git a/Engine/Plugins/Runtime/OculusRift/Source/OculusRift/Private/OculusRiftHMD.cpp b/Engine/Plugins/Runtime/OculusRift/Source/OculusRift/Private/OculusRiftHMD.cpp
--- a/Engine/Plugins/Runtime/OculusRift/Source/OculusRift/Private/OculusRiftHMD.cpp
+++ b/Engine/Plugins/Runtime/OculusRift/Source/OculusRift/Private/OculusRiftHMD.cpp
@@ -1707,7 +1707,7 @@ bool FOculusRiftHMD::IsHeadTrackingAllowed() const
 #endif//WITH_EDITOR
 	return Hmd && 
 #if WITH_EDITOR
-		(!EdEngine || EdEngine->bUseVRPreviewForPlayWorld) &&
+		(!EdEngine || EdEngine->bUseVRPreviewForPlayWorld || GetDefault<ULevelEditorPlaySettings>()->ViewportGetsHMDControl) &&
 #endif//WITH_EDITOR
 		(Flags.bHeadTrackingEnforced || GEngine->IsStereoscopic3D());
 }
diff --git a/Engine/Source/Editor/UnrealEd/Classes/Settings/LevelEditorPlaySettings.h b/Engine/Source/Editor/UnrealEd/Classes/Settings/LevelEditorPlaySettings.h
--- a/Engine/Source/Editor/UnrealEd/Classes/Settings/LevelEditorPlaySettings.h
+++ b/Engine/Source/Editor/UnrealEd/Classes/Settings/LevelEditorPlaySettings.h
@@ -135,6 +135,10 @@ public:
 	UPROPERTY(config, EditAnywhere, Category=PlayInEditor)
 	TEnumAsByte<ELabelAnchorMode> MouseControlLabelPosition;
 
+	/** Should Play-in-Viewport respect HMD orientations (default = false) */
+	UPROPERTY(config, EditAnywhere, Category = PlayInEditor, meta = (ToolTip = "Whether or not HMD orientation should be used when playing in viewport"))
+	bool ViewportGetsHMDControl;
+
 	/** Whether to automatically recompile blueprints on PIE */
 	UPROPERTY(config, EditAnywhere, Category=PlayInEditor, meta=(ToolTip="Automatically recompile blueprints used by the current level when initiating a Play In Editor session"))
 	bool AutoRecompileBlueprints;
