Refactored gather text commandlet base class to handle relative paths taken from configuration files automatically. Cleaner, less duplication.

diff --git a/Engine/Source/Editor/UnrealEd/Classes/Commandlets/GatherTextCommandletBase.h b/Engine/Source/Editor/UnrealEd/Classes/Commandlets/GatherTextCommandletBase.h
--- a/Engine/Source/Editor/UnrealEd/Classes/Commandlets/GatherTextCommandletBase.h
+++ b/Engine/Source/Editor/UnrealEd/Classes/Commandlets/GatherTextCommandletBase.h
@@ -261,9 +261,11 @@ public:
 	static FString MungeLogOutput( const FString& InString );
 
 	// Wrappers for extracting config values
-	bool GetConfigBool( const TCHAR* Section, const TCHAR* Key, bool& OutValue, const FString& Filename );
-	bool GetConfigString( const TCHAR* Section, const TCHAR* Key, FString& OutValue, const FString& Filename );
-	int32 GetConfigArray( const TCHAR* Section, const TCHAR* Key, TArray<FString>& OutArr, const FString& Filename );
+	bool GetBoolFromConfig( const TCHAR* Section, const TCHAR* Key, bool& OutValue, const FString& Filename );
+	bool GetStringFromConfig( const TCHAR* Section, const TCHAR* Key, FString& OutValue, const FString& Filename );
+	bool GetPathFromConfig( const TCHAR* Section, const TCHAR* Key, FString& OutValue, const FString& Filename );
+	int32 GetStringArrayFromConfig( const TCHAR* Section, const TCHAR* Key, TArray<FString>& OutArr, const FString& Filename );
+	int32 GetPathArrayFromConfig( const TCHAR* Section, const TCHAR* Key, TArray<FString>& OutArr, const FString& Filename );
 
 protected:
 	TSharedPtr< FManifestInfo > ManifestInfo;
diff --git a/Engine/Source/Editor/UnrealEd/Private/Commandlets/GatherTextCommandletBase.cpp b/Engine/Source/Editor/UnrealEd/Private/Commandlets/GatherTextCommandletBase.cpp
--- a/Engine/Source/Editor/UnrealEd/Private/Commandlets/GatherTextCommandletBase.cpp
+++ b/Engine/Source/Editor/UnrealEd/Private/Commandlets/GatherTextCommandletBase.cpp
@@ -156,7 +156,7 @@ FString UGatherTextCommandletBase::MungeLogOutput( const FString& InString )
 }
 
 
-bool UGatherTextCommandletBase::GetConfigBool( const TCHAR* Section, const TCHAR* Key, bool& OutValue, const FString& Filename )
+bool UGatherTextCommandletBase::GetBoolFromConfig( const TCHAR* Section, const TCHAR* Key, bool& OutValue, const FString& Filename )
 {
 	bool bSuccess = GConfig->GetBool( Section, Key, OutValue, Filename );
 	
@@ -167,7 +167,7 @@ bool UGatherTextCommandletBase::GetConfigBool( const TCHAR* Section, const TCHAR
 	return bSuccess;
 }
 
-bool UGatherTextCommandletBase::GetConfigString( const TCHAR* Section, const TCHAR* Key, FString& OutValue, const FString& Filename )
+bool UGatherTextCommandletBase::GetStringFromConfig( const TCHAR* Section, const TCHAR* Key, FString& OutValue, const FString& Filename )
 {
 	bool bSuccess = GConfig->GetString( Section, Key, OutValue, Filename );
 
@@ -178,7 +178,28 @@ bool UGatherTextCommandletBase::GetConfigString( const TCHAR* Section, const TCH
 	return bSuccess;
 }
 
-int32 UGatherTextCommandletBase::GetConfigArray( const TCHAR* Section, const TCHAR* Key, TArray<FString>& OutArr, const FString& Filename )
+bool UGatherTextCommandletBase::GetPathFromConfig( const TCHAR* Section, const TCHAR* Key, FString& OutValue, const FString& Filename )
+{
+	bool bSuccess = GetStringFromConfig( Section, Key, OutValue, Filename );
+
+	if( bSuccess )
+	{
+		if (FPaths::IsRelative(OutValue))
+		{
+			if (!FPaths::GameDir().IsEmpty())
+			{
+				OutValue = FPaths::Combine( *( FPaths::GameDir() ), *OutValue );
+			}
+			else
+			{
+				OutValue = FPaths::Combine( *( FPaths::EngineDir() ), *OutValue );
+			}
+		}
+	}
+	return bSuccess;
+}
+
+int32 UGatherTextCommandletBase::GetStringArrayFromConfig( const TCHAR* Section, const TCHAR* Key, TArray<FString>& OutArr, const FString& Filename )
 {
 	int32 count = GConfig->GetArray( Section, Key, OutArr, Filename );
 
@@ -189,6 +210,26 @@ int32 UGatherTextCommandletBase::GetConfigArray( const TCHAR* Section, const TCH
 	return count;
 }
 
+int32 UGatherTextCommandletBase::GetPathArrayFromConfig( const TCHAR* Section, const TCHAR* Key, TArray<FString>& OutArr, const FString& Filename )
+{
+	int32 count = GetStringArrayFromConfig( Section, Key, OutArr, Filename );
+
+	for (int32 i = 0; i < count; ++i)
+	{
+		if (FPaths::IsRelative(OutArr[i]))
+		{
+			if (!FPaths::GameDir().IsEmpty())
+			{
+				OutArr[i] = FPaths::Combine( *( FPaths::GameDir() ), *OutArr[i] );
+			}
+			else
+			{
+				OutArr[i] = FPaths::Combine( *( FPaths::EngineDir() ), *OutArr[i] );
+			}
+		}
+	}
+	return count;
+}
 
 bool FManifestInfo::AddManifestDependency( const FString& InManifestFilePath )
 {
diff --git a/Engine/Source/Editor/UnrealEd/Private/Commandlets/GatherTextFromAssetsCommandlet.cpp b/Engine/Source/Editor/UnrealEd/Private/Commandlets/GatherTextFromAssetsCommandlet.cpp
--- a/Engine/Source/Editor/UnrealEd/Private/Commandlets/GatherTextFromAssetsCommandlet.cpp
+++ b/Engine/Source/Editor/UnrealEd/Private/Commandlets/GatherTextFromAssetsCommandlet.cpp
@@ -707,7 +707,7 @@ int32 UGatherTextFromAssetsCommandlet::Main(const FString& Params)
 
 	//Modules to Preload
 	TArray<FString> ModulesToPreload;
-	GetConfigArray(*SectionName, TEXT("ModulesToPreload"), ModulesToPreload, GatherTextConfigPath);
+	GetStringArrayFromConfig(*SectionName, TEXT("ModulesToPreload"), ModulesToPreload, GatherTextConfigPath);
 
 	for (const FString& ModuleName : ModulesToPreload)
 	{
@@ -716,7 +716,7 @@ int32 UGatherTextFromAssetsCommandlet::Main(const FString& Params)
 
 	//Include paths
 	TArray<FString> IncludePaths;
-	GetConfigArray(*SectionName, TEXT("IncludePaths"), IncludePaths, GatherTextConfigPath);
+	GetPathArrayFromConfig(*SectionName, TEXT("IncludePaths"), IncludePaths, GatherTextConfigPath);
 
 	if (IncludePaths.Num() == 0)
 	{
@@ -727,17 +727,6 @@ int32 UGatherTextFromAssetsCommandlet::Main(const FString& Params)
 	for (FString& IncludePath : IncludePaths)
 	{
 		FPaths::NormalizeDirectoryName(IncludePath);
-		if (FPaths::IsRelative(IncludePath))
-		{
-			if (!FPaths::GameDir().IsEmpty())
-			{
-				IncludePath = FPaths::Combine( *( FPaths::GameDir() ), *IncludePath );
-			}
-			else
-			{
-				IncludePath = FPaths::Combine( *( FPaths::EngineDir() ), *IncludePath );
-			}
-		}
 
 		IncludePath = FPaths::ConvertRelativePathToFull(IncludePath);
 
@@ -759,11 +748,11 @@ int32 UGatherTextFromAssetsCommandlet::Main(const FString& Params)
 
 	//Exclude paths
 	TArray<FString> ExcludePaths;
-	GetConfigArray(*SectionName, TEXT("ExcludePaths"), ExcludePaths, GatherTextConfigPath);
+	GetStringArrayFromConfig(*SectionName, TEXT("ExcludePaths"), ExcludePaths, GatherTextConfigPath);
 
 	//package extensions
 	TArray<FString> PackageExts;
-	GetConfigArray(*SectionName, TEXT("PackageExtensions"), PackageExts, GatherTextConfigPath);
+	GetStringArrayFromConfig(*SectionName, TEXT("PackageExtensions"), PackageExts, GatherTextConfigPath);
 
 	if (PackageExts.Num() == 0)
 	{
@@ -775,7 +764,7 @@ int32 UGatherTextFromAssetsCommandlet::Main(const FString& Params)
 
 	//asset class exclude
 	TArray<FString> ExcludeClasses;
-	GetConfigArray(*SectionName, TEXT("ExcludeClasses"), ExcludeClasses, GatherTextConfigPath);
+	GetStringArrayFromConfig(*SectionName, TEXT("ExcludeClasses"), ExcludeClasses, GatherTextConfigPath);
 
 	FAssetRegistryModule& AssetRegistryModule = FModuleManager::LoadModuleChecked<FAssetRegistryModule>(TEXT("AssetRegistry"));
 	AssetRegistryModule.Get().SearchAllAssets( true );
@@ -805,26 +794,11 @@ int32 UGatherTextFromAssetsCommandlet::Main(const FString& Params)
 	}
 
 	//Get whether we should fix broken properties that we find.
-	GetConfigBool(*SectionName, TEXT("bFixBroken"), bFixBroken, GatherTextConfigPath);
+	GetBoolFromConfig(*SectionName, TEXT("bFixBroken"), bFixBroken, GatherTextConfigPath);
 
 	// Add any manifest dependencies if they were provided
 	TArray<FString> ManifestDependenciesList;
-	GetConfigArray(*SectionName, TEXT("ManifestDependencies"), ManifestDependenciesList, GatherTextConfigPath);
-	
-	for (FString& ManifestDependency : ManifestDependenciesList)
-	{
-		if (FPaths::IsRelative(ManifestDependency))
-		{
-			if (!FPaths::GameDir().IsEmpty())
-			{
-				ManifestDependency = FPaths::Combine( *( FPaths::GameDir() ), *ManifestDependency );
-			}
-			else
-			{
-				ManifestDependency = FPaths::Combine( *( FPaths::EngineDir() ), *ManifestDependency );
-			}
-		}
-	}
+	GetPathArrayFromConfig(*SectionName, TEXT("ManifestDependencies"), ManifestDependenciesList, GatherTextConfigPath);
 
 	if( !ManifestInfo->AddManifestDependencies( ManifestDependenciesList ) )
 	{
diff --git a/Engine/Source/Editor/UnrealEd/Private/Commandlets/GatherTextFromMetadataCommandlet.cpp b/Engine/Source/Editor/UnrealEd/Private/Commandlets/GatherTextFromMetadataCommandlet.cpp
--- a/Engine/Source/Editor/UnrealEd/Private/Commandlets/GatherTextFromMetadataCommandlet.cpp
+++ b/Engine/Source/Editor/UnrealEd/Private/Commandlets/GatherTextFromMetadataCommandlet.cpp
@@ -51,7 +51,7 @@ int32 UGatherTextFromMetaDataCommandlet::Main( const FString& Params )
 
 	//Modules to Preload
 	TArray<FString> ModulesToPreload;
-	GetConfigArray(*SectionName, TEXT("ModulesToPreload"), ModulesToPreload, GatherTextConfigPath);
+	GetStringArrayFromConfig(*SectionName, TEXT("ModulesToPreload"), ModulesToPreload, GatherTextConfigPath);
 
 	for (const FString& ModuleName : ModulesToPreload)
 	{
@@ -60,7 +60,7 @@ int32 UGatherTextFromMetaDataCommandlet::Main( const FString& Params )
 
 	//Include paths
 	TArray<FString> IncludePaths;
-	GetConfigArray(*SectionName, TEXT("IncludePaths"), IncludePaths, GatherTextConfigPath);
+	GetPathArrayFromConfig(*SectionName, TEXT("IncludePaths"), IncludePaths, GatherTextConfigPath);
 
 	if (IncludePaths.Num() == 0)
 	{
@@ -71,17 +71,6 @@ int32 UGatherTextFromMetaDataCommandlet::Main( const FString& Params )
 	for (FString& IncludePath : IncludePaths)
 	{
 		FPaths::NormalizeDirectoryName(IncludePath);
-		if (FPaths::IsRelative(IncludePath))
-		{
-			if (!FPaths::GameDir().IsEmpty())
-			{
-				IncludePath = FPaths::Combine( *( FPaths::GameDir() ), *IncludePath );
-			}
-			else
-			{
-				IncludePath = FPaths::Combine( *( FPaths::EngineDir() ), *IncludePath );
-			}
-		}
 
 		IncludePath = FPaths::ConvertRelativePathToFull(IncludePath);
 
@@ -103,11 +92,11 @@ int32 UGatherTextFromMetaDataCommandlet::Main( const FString& Params )
 
 	//Exclude paths
 	TArray<FString> ExcludePaths;
-	GetConfigArray(*SectionName, TEXT("ExcludePaths"), ExcludePaths, GatherTextConfigPath);
+	GetStringArrayFromConfig(*SectionName, TEXT("ExcludePaths"), ExcludePaths, GatherTextConfigPath);
 
 	//Required module names
 	TArray<FString> RequiredModuleNames;
-	GetConfigArray(*SectionName, TEXT("RequiredModuleNames"), RequiredModuleNames, GatherTextConfigPath);
+	GetStringArrayFromConfig(*SectionName, TEXT("RequiredModuleNames"), RequiredModuleNames, GatherTextConfigPath);
 
 	// Pre-load all required modules so that UFields from those modules can be guaranteed a chance to have their metadata gathered.
 	for(const FString& RequiredModuleName : RequiredModuleNames)
@@ -116,10 +105,10 @@ int32 UGatherTextFromMetaDataCommandlet::Main( const FString& Params )
 	}
 
 	FGatherParameters Arguments;
-	GetConfigArray(*SectionName, TEXT("InputKeys"), Arguments.InputKeys, GatherTextConfigPath);
-	GetConfigArray(*SectionName, TEXT("OutputNamespaces"), Arguments.OutputNamespaces, GatherTextConfigPath);
+	GetStringArrayFromConfig(*SectionName, TEXT("InputKeys"), Arguments.InputKeys, GatherTextConfigPath);
+	GetStringArrayFromConfig(*SectionName, TEXT("OutputNamespaces"), Arguments.OutputNamespaces, GatherTextConfigPath);
 	TArray<FString> OutputKeys;
-	GetConfigArray(*SectionName, TEXT("OutputKeys"), OutputKeys, GatherTextConfigPath);
+	GetStringArrayFromConfig(*SectionName, TEXT("OutputKeys"), OutputKeys, GatherTextConfigPath);
 	for(const auto& OutputKey : OutputKeys)
 	{
 		Arguments.OutputKeys.Add(FText::FromString(OutputKey));
@@ -130,22 +119,8 @@ int32 UGatherTextFromMetaDataCommandlet::Main( const FString& Params )
 
 	// Add any manifest dependencies if they were provided
 	TArray<FString> ManifestDependenciesList;
-	GetConfigArray(*SectionName, TEXT("ManifestDependencies"), ManifestDependenciesList, GatherTextConfigPath);
-	
-	for (FString& ManifestDependency : ManifestDependenciesList)
-	{
-		if (FPaths::IsRelative(ManifestDependency))
-		{
-			if (!FPaths::GameDir().IsEmpty())
-			{
-				ManifestDependency = FPaths::Combine( *( FPaths::GameDir() ), *ManifestDependency );
-			}
-			else
-			{
-				ManifestDependency = FPaths::Combine( *( FPaths::EngineDir() ), *ManifestDependency );
-			}
-		}
-	}
+	GetPathArrayFromConfig(*SectionName, TEXT("ManifestDependencies"), ManifestDependenciesList, GatherTextConfigPath);
+
 
 	if( !ManifestInfo->AddManifestDependencies( ManifestDependenciesList ) )
 	{
diff --git a/Engine/Source/Editor/UnrealEd/Private/Commandlets/GatherTextFromSourceCommandlet.cpp b/Engine/Source/Editor/UnrealEd/Private/Commandlets/GatherTextFromSourceCommandlet.cpp
--- a/Engine/Source/Editor/UnrealEd/Private/Commandlets/GatherTextFromSourceCommandlet.cpp
+++ b/Engine/Source/Editor/UnrealEd/Private/Commandlets/GatherTextFromSourceCommandlet.cpp
@@ -61,7 +61,7 @@ int32 UGatherTextFromSourceCommandlet::Main( const FString& Params )
 
 	//Include paths
 	TArray<FString> IncludePaths;
-	GetConfigArray(*SectionName, TEXT("IncludePaths"), IncludePaths, GatherTextConfigPath);
+	GetStringArrayFromConfig(*SectionName, TEXT("IncludePaths"), IncludePaths, GatherTextConfigPath);
 
 	if (IncludePaths.Num() == 0)
 	{
@@ -71,11 +71,11 @@ int32 UGatherTextFromSourceCommandlet::Main( const FString& Params )
 
 	//Exclude paths
 	TArray<FString> ExcludePaths;
-	GetConfigArray(*SectionName, TEXT("ExcludePaths"), ExcludePaths, GatherTextConfigPath);
+	GetStringArrayFromConfig(*SectionName, TEXT("ExcludePaths"), ExcludePaths, GatherTextConfigPath);
 
 	// Check the config section for source filters. e.g. *.cpp
 	TArray<FString> SourceFileSearchFilters;
-	GetConfigArray(*SectionName, TEXT("SourceFileSearchFilters"), SourceFileSearchFilters, GatherTextConfigPath);
+	GetStringArrayFromConfig(*SectionName, TEXT("SourceFileSearchFilters"), SourceFileSearchFilters, GatherTextConfigPath);
 
 	if (SourceFileSearchFilters.Num() == 0)
 	{
@@ -155,22 +155,8 @@ int32 UGatherTextFromSourceCommandlet::Main( const FString& Params )
 
 	// Add any manifest dependencies if they were provided
 	TArray<FString> ManifestDependenciesList;
-	GetConfigArray(*SectionName, TEXT("ManifestDependencies"), ManifestDependenciesList, GatherTextConfigPath);
+	GetPathArrayFromConfig(*SectionName, TEXT("ManifestDependencies"), ManifestDependenciesList, GatherTextConfigPath);
 	
-	for (FString& ManifestDependency : ManifestDependenciesList)
-	{
-		if (FPaths::IsRelative(ManifestDependency))
-		{
-			if (!FPaths::GameDir().IsEmpty())
-			{
-				ManifestDependency = FPaths::Combine( *( FPaths::GameDir() ), *ManifestDependency );
-			}
-			else
-			{
-				ManifestDependency = FPaths::Combine( *( FPaths::EngineDir() ), *ManifestDependency );
-			}
-		}
-	}
 
 	if( !ManifestInfo->AddManifestDependencies( ManifestDependenciesList ) )
 	{
diff --git a/Engine/Source/Editor/UnrealEd/Private/Commandlets/GenerateGatherArchiveCommandlet.cpp b/Engine/Source/Editor/UnrealEd/Private/Commandlets/GenerateGatherArchiveCommandlet.cpp
--- a/Engine/Source/Editor/UnrealEd/Private/Commandlets/GenerateGatherArchiveCommandlet.cpp
+++ b/Engine/Source/Editor/UnrealEd/Private/Commandlets/GenerateGatherArchiveCommandlet.cpp
@@ -60,7 +60,7 @@ int32 UGenerateGatherArchiveCommandlet::Main( const FString& Params )
 
 	// Get manifest name.
 	FString ManifestName;
-	if( !GetConfigString( *SectionName, TEXT("ManifestName"), ManifestName, GatherTextConfigPath ) )
+	if( !GetStringFromConfig( *SectionName, TEXT("ManifestName"), ManifestName, GatherTextConfigPath ) )
 	{
 		UE_LOG( LogGenerateArchiveCommandlet, Error, TEXT("No manifest name specified.") );
 		return -1;
@@ -68,7 +68,7 @@ int32 UGenerateGatherArchiveCommandlet::Main( const FString& Params )
 
 	// Get source culture.
 	FString SourceCulture;
-	if( GetConfigString( *SectionName, TEXT("SourceCulture"), SourceCulture, GatherTextConfigPath ) )
+	if( GetStringFromConfig( *SectionName, TEXT("SourceCulture"), SourceCulture, GatherTextConfigPath ) )
 	{
 		if( I18N.GetCulture( SourceCulture ).IsValid() )
 		{
@@ -78,7 +78,7 @@ int32 UGenerateGatherArchiveCommandlet::Main( const FString& Params )
 
 	// Get cultures to generate.
 	TArray<FString> CulturesToGenerate;
-	GetConfigArray(*SectionName, TEXT("CulturesToGenerate"), CulturesToGenerate, GatherTextConfigPath);
+	GetStringArrayFromConfig(*SectionName, TEXT("CulturesToGenerate"), CulturesToGenerate, GatherTextConfigPath);
 	
 	if( CulturesToGenerate.Num() == 0 )
 	{
@@ -96,27 +96,15 @@ int32 UGenerateGatherArchiveCommandlet::Main( const FString& Params )
 
 	// Get destination path.
 	FString DestinationPath;
-	if( !GetConfigString( *SectionName, TEXT("DestinationPath"), DestinationPath, GatherTextConfigPath ) )
+	if( !GetPathFromConfig( *SectionName, TEXT("DestinationPath"), DestinationPath, GatherTextConfigPath ) )
 	{
 		UE_LOG( LogGenerateArchiveCommandlet, Error, TEXT("No destination path specified.") );
 		return -1;
 	}
 
-	if (FPaths::IsRelative(DestinationPath))
-	{
-		if (!FPaths::GameDir().IsEmpty())
-		{
-			DestinationPath = FPaths::Combine( *( FPaths::GameDir() ), *DestinationPath );
-		}
-		else
-		{
-			DestinationPath = FPaths::Combine( *( FPaths::EngineDir() ), *DestinationPath );
-		}
-	}
-
 	// Get archive name.
 	FString ArchiveName;
-	if( !( GetConfigString(* SectionName, TEXT("ArchiveName"), ArchiveName, GatherTextConfigPath ) ) )
+	if( !( GetStringFromConfig(* SectionName, TEXT("ArchiveName"), ArchiveName, GatherTextConfigPath ) ) )
 	{
 		UE_LOG(LogGenerateArchiveCommandlet, Error, TEXT("No archive name specified."));
 		return -1;
@@ -124,7 +112,7 @@ int32 UGenerateGatherArchiveCommandlet::Main( const FString& Params )
 
 	// Get bPurgeOldEmptyEntries option.
 	bool ShouldPurgeOldEmptyEntries;
-	if ( !GetConfigBool( *SectionName, TEXT("bPurgeOldEmptyEntries"), ShouldPurgeOldEmptyEntries, GatherTextConfigPath) )
+	if ( !GetBoolFromConfig( *SectionName, TEXT("bPurgeOldEmptyEntries"), ShouldPurgeOldEmptyEntries, GatherTextConfigPath) )
 	{
 		ShouldPurgeOldEmptyEntries = false;
 	}
diff --git a/Engine/Source/Editor/UnrealEd/Private/Commandlets/GenerateGatherManifestCommandlet.cpp b/Engine/Source/Editor/UnrealEd/Private/Commandlets/GenerateGatherManifestCommandlet.cpp
--- a/Engine/Source/Editor/UnrealEd/Private/Commandlets/GenerateGatherManifestCommandlet.cpp
+++ b/Engine/Source/Editor/UnrealEd/Private/Commandlets/GenerateGatherManifestCommandlet.cpp
@@ -54,27 +54,15 @@ int32 UGenerateGatherManifestCommandlet::Main( const FString& Params )
 
 	// Get destination path.
 	FString DestinationPath;
-	if( !GetConfigString( *SectionName, TEXT("DestinationPath"), DestinationPath, GatherTextConfigPath ) )
+	if( !GetPathFromConfig( *SectionName, TEXT("DestinationPath"), DestinationPath, GatherTextConfigPath ) )
 	{
 		UE_LOG( LogGenerateManifestCommandlet, Error, TEXT("No destination path specified.") );
 		return -1;
 	}
 
-	if (FPaths::IsRelative(DestinationPath))
-	{
-		if (!FPaths::GameDir().IsEmpty())
-		{
-			DestinationPath = FPaths::Combine( *( FPaths::GameDir() ), *DestinationPath );
-		}
-		else
-		{
-			DestinationPath = FPaths::Combine( *( FPaths::EngineDir() ), *DestinationPath );
-		}
-	}
-
 	// Get manifest name.
 	FString ManifestName;
-	if( !GetConfigString( *SectionName, TEXT("ManifestName"), ManifestName, GatherTextConfigPath ) )
+	if( !GetStringFromConfig( *SectionName, TEXT("ManifestName"), ManifestName, GatherTextConfigPath ) )
 	{
 		UE_LOG( LogGenerateManifestCommandlet, Error, TEXT("No manifest name specified.") );
 		return -1;
@@ -82,22 +70,7 @@ int32 UGenerateGatherManifestCommandlet::Main( const FString& Params )
 
 	//Grab any manifest dependencies
 	TArray<FString> ManifestDependenciesList;
-	GetConfigArray(*SectionName, TEXT("ManifestDependencies"), ManifestDependenciesList, GatherTextConfigPath);
-
-	for (FString& ManifestDependency : ManifestDependenciesList)
-	{
-		if (FPaths::IsRelative(ManifestDependency))
-		{
-			if (!FPaths::GameDir().IsEmpty())
-			{
-				ManifestDependency = FPaths::Combine( *( FPaths::GameDir() ), *ManifestDependency );
-			}
-			else
-			{
-				ManifestDependency = FPaths::Combine( *( FPaths::EngineDir() ), *ManifestDependency );
-			}
-		}
-	}
+	GetPathArrayFromConfig(*SectionName, TEXT("ManifestDependencies"), ManifestDependenciesList, GatherTextConfigPath);
 
 	if( ManifestDependenciesList.Num() > 0 )
 	{
diff --git a/Engine/Source/Editor/UnrealEd/Private/Commandlets/GenerateTextLocalizationReportCommandlet.cpp b/Engine/Source/Editor/UnrealEd/Private/Commandlets/GenerateTextLocalizationReportCommandlet.cpp
--- a/Engine/Source/Editor/UnrealEd/Private/Commandlets/GenerateTextLocalizationReportCommandlet.cpp
+++ b/Engine/Source/Editor/UnrealEd/Private/Commandlets/GenerateTextLocalizationReportCommandlet.cpp
@@ -83,43 +83,19 @@ int32 UGenerateTextLocalizationReportCommandlet::Main(const FString& Params)
 	bool bConflictReport = false;
 	
 	// Get source path.
-	if( !( GetConfigString( *SectionName, TEXT("SourcePath"), SourcePath, GatherTextConfigPath ) ) )
+	if( !( GetPathFromConfig( *SectionName, TEXT("SourcePath"), SourcePath, GatherTextConfigPath ) ) )
 	{
 		UE_LOG(LogGenerateTextLocalizationReportCommandlet, Error, TEXT("No source path specified."));
 		return -1;
 	}
 
-	if (FPaths::IsRelative(SourcePath))
-	{
-		if (!FPaths::GameDir().IsEmpty())
-		{
-			SourcePath = FPaths::Combine( *( FPaths::GameDir() ), *SourcePath );
-		}
-		else
-		{
-			SourcePath = FPaths::Combine( *( FPaths::EngineDir() ), *SourcePath );
-		}
-	}
-
 	// Get destination path.
-	if( !( GetConfigString( *SectionName, TEXT("DestinationPath"), DestinationPath, GatherTextConfigPath ) ) )
+	if( !( GetPathFromConfig( *SectionName, TEXT("DestinationPath"), DestinationPath, GatherTextConfigPath ) ) )
 	{
 		UE_LOG(LogGenerateTextLocalizationReportCommandlet, Error, TEXT("No destination path specified."));
 		return -1;
 	}
 
-	if (FPaths::IsRelative(DestinationPath))
-	{
-		if (!FPaths::GameDir().IsEmpty())
-		{
-			DestinationPath = FPaths::Combine( *( FPaths::GameDir() ), *DestinationPath );
-		}
-		else
-		{
-			DestinationPath = FPaths::Combine( *( FPaths::EngineDir() ), *DestinationPath );
-		}
-	}
-
 	// Get the timestamp from the commandline, if not provided we will use the current time.
 	const FString* TimeStampParamVal = ParamVals.Find(FString(TEXT("TimeStamp")));
 	if ( TimeStampParamVal && !TimeStampParamVal->IsEmpty() )
@@ -127,8 +103,8 @@ int32 UGenerateTextLocalizationReportCommandlet::Main(const FString& Params)
 		CmdlineTimeStamp = *TimeStampParamVal;
 	}
 
-	GetConfigBool( *SectionName, TEXT("bWordCountReport"), bWordCountReport, GatherTextConfigPath );
-	GetConfigBool( *SectionName, TEXT("bConflictReport"), bConflictReport, GatherTextConfigPath );
+	GetBoolFromConfig( *SectionName, TEXT("bWordCountReport"), bWordCountReport, GatherTextConfigPath );
+	GetBoolFromConfig( *SectionName, TEXT("bConflictReport"), bConflictReport, GatherTextConfigPath );
 
 
 	if( bWordCountReport )
@@ -162,21 +138,21 @@ bool UGenerateTextLocalizationReportCommandlet::ProcessWordCountReport(const FSt
 	TimeStamp = (CmdlineTimeStamp.IsEmpty()) ? *FDateTime::Now().ToString() : CmdlineTimeStamp;
 
 	// Get manifest name.
-	if( !( GetConfigString( *SectionName, TEXT("ManifestName"), ManifestName, GatherTextConfigPath ) ) )
+	if( !( GetStringFromConfig( *SectionName, TEXT("ManifestName"), ManifestName, GatherTextConfigPath ) ) )
 	{
 		UE_LOG(LogGenerateTextLocalizationReportCommandlet, Error, TEXT("No manifest name specified."));
 		return false;
 	}
 
 	// Get report name.
-	if( !( GetConfigString( *SectionName, TEXT("WordCountReportName"), WordCountReportName, GatherTextConfigPath ) ) )
+	if( !( GetStringFromConfig( *SectionName, TEXT("WordCountReportName"), WordCountReportName, GatherTextConfigPath ) ) )
 	{
 		UE_LOG(LogGenerateTextLocalizationReportCommandlet, Error, TEXT("No word count report name specified."));
 		return false;
 	}
 
 	// Get cultures to generate.
-	GetConfigArray( *SectionName, TEXT("CulturesToGenerate"), CulturesToGenerate, GatherTextConfigPath );
+	GetStringArrayFromConfig( *SectionName, TEXT("CulturesToGenerate"), CulturesToGenerate, GatherTextConfigPath );
 
 	if( CulturesToGenerate.Num() == 0 )
 	{
@@ -381,7 +357,7 @@ bool UGenerateTextLocalizationReportCommandlet::ProcessConflictReport(const FStr
 	TimeStamp = (CmdlineTimeStamp.IsEmpty()) ? *FDateTime::Now().ToString() : CmdlineTimeStamp;
 
 	// Get report name.
-	if( !( GetConfigString( *SectionName, TEXT("ConflictReportName"), ConflictReportName, GatherTextConfigPath ) ) )
+	if( !( GetStringFromConfig( *SectionName, TEXT("ConflictReportName"), ConflictReportName, GatherTextConfigPath ) ) )
 	{
 		UE_LOG(LogGenerateTextLocalizationReportCommandlet, Error, TEXT("No conflict report name specified."));
 		return false;
diff --git a/Engine/Source/Editor/UnrealEd/Private/Commandlets/GenerateTextLocalizationResourceCommandlet.cpp b/Engine/Source/Editor/UnrealEd/Private/Commandlets/GenerateTextLocalizationResourceCommandlet.cpp
--- a/Engine/Source/Editor/UnrealEd/Private/Commandlets/GenerateTextLocalizationResourceCommandlet.cpp
+++ b/Engine/Source/Editor/UnrealEd/Private/Commandlets/GenerateTextLocalizationResourceCommandlet.cpp
@@ -53,26 +53,15 @@ int32 UGenerateTextLocalizationResourceCommandlet::Main(const FString& Params)
 
 	// Get source path.
 	FString SourcePath;
-	if( !( GetConfigString( *SectionName, TEXT("SourcePath"), SourcePath, GatherTextConfigPath ) ) )
+	if( !( GetPathFromConfig( *SectionName, TEXT("SourcePath"), SourcePath, GatherTextConfigPath ) ) )
 	{
 		UE_LOG(LogGenerateTextLocalizationResourceCommandlet, Error, TEXT("No source path specified."));
 		return -1;
 	}
 
-	if (FPaths::IsRelative(SourcePath))
-	{
-		if (!FPaths::GameDir().IsEmpty())
-		{
-			SourcePath = FPaths::Combine( *( FPaths::GameDir() ), *SourcePath );
-		}
-		else
-		{
-			SourcePath = FPaths::Combine( *( FPaths::EngineDir() ), *SourcePath );
-		}
-	}
 	// Get manifest name.
 	FString ManifestName;
-	if( !( GetConfigString( *SectionName, TEXT("ManifestName"), ManifestName, GatherTextConfigPath ) ) )
+	if( !( GetStringFromConfig( *SectionName, TEXT("ManifestName"), ManifestName, GatherTextConfigPath ) ) )
 	{
 		UE_LOG(LogGenerateTextLocalizationResourceCommandlet, Error, TEXT("No manifest name specified."));
 		return -1;
@@ -80,7 +69,7 @@ int32 UGenerateTextLocalizationResourceCommandlet::Main(const FString& Params)
 
 	// Get cultures to generate.
 	TArray<FString> CulturesToGenerate;
-	GetConfigArray( *SectionName, TEXT("CulturesToGenerate"), CulturesToGenerate, GatherTextConfigPath );
+	GetStringArrayFromConfig( *SectionName, TEXT("CulturesToGenerate"), CulturesToGenerate, GatherTextConfigPath );
 
 	if( CulturesToGenerate.Num() == 0 )
 	{
@@ -98,27 +87,15 @@ int32 UGenerateTextLocalizationResourceCommandlet::Main(const FString& Params)
 
 	// Get destination path.
 	FString DestinationPath;
-	if( !( GetConfigString( *SectionName, TEXT("DestinationPath"), DestinationPath, GatherTextConfigPath ) ) )
+	if( !( GetPathFromConfig( *SectionName, TEXT("DestinationPath"), DestinationPath, GatherTextConfigPath ) ) )
 	{
 		UE_LOG(LogGenerateTextLocalizationResourceCommandlet, Error, TEXT("No destination path specified."));
 		return -1;
 	}
 
-	if (FPaths::IsRelative(DestinationPath))
-	{
-		if (!FPaths::GameDir().IsEmpty())
-		{
-			DestinationPath = FPaths::Combine( *( FPaths::GameDir() ), *DestinationPath );
-		}
-		else
-		{
-			DestinationPath = FPaths::Combine( *( FPaths::EngineDir() ), *DestinationPath );
-		}
-	}
-
 	// Get resource name.
 	FString ResourceName;
-	if( !( GetConfigString( *SectionName, TEXT("ResourceName"), ResourceName, GatherTextConfigPath ) ) )
+	if( !( GetStringFromConfig( *SectionName, TEXT("ResourceName"), ResourceName, GatherTextConfigPath ) ) )
 	{
 		UE_LOG(LogGenerateTextLocalizationResourceCommandlet, Error, TEXT("No resource name specified."));
 		return -1;
diff --git a/Engine/Source/Editor/UnrealEd/Private/Commandlets/InternationalizationConditioningCommandlet.cpp b/Engine/Source/Editor/UnrealEd/Private/Commandlets/InternationalizationConditioningCommandlet.cpp
--- a/Engine/Source/Editor/UnrealEd/Private/Commandlets/InternationalizationConditioningCommandlet.cpp
+++ b/Engine/Source/Editor/UnrealEd/Private/Commandlets/InternationalizationConditioningCommandlet.cpp
@@ -1,4 +1,4 @@
-// Copyright 1998-2015 Epic Games, Inc. All Rights Reserved.
+﻿// Copyright 1998-2015 Epic Games, Inc. All Rights Reserved.
 
 #include "UnrealEd.h"
 #include "Internationalization/InternationalizationArchive.h"
@@ -298,7 +298,7 @@ bool UInternationalizationConditioningCommandlet::ProcessManifest( const FString
 {
 	FString ManifestName = TEXT("Manifest.txt");
 
-	GetConfigString( *SectionName, TEXT("ManifestName"), ManifestName, GatherTextConfigPath );
+	GetStringFromConfig( *SectionName, TEXT("ManifestName"), ManifestName, GatherTextConfigPath );
 
 	// Build info about the primary language
 	TArray<FString> PrimaryFilenames;
@@ -386,10 +386,10 @@ bool UInternationalizationConditioningCommandlet::ProcessArchive( const FString&
 	TArray<FString> TargetCultures;
 	bool bAppendToExistingArchive = true;
 
-	GetConfigString( *SectionName, TEXT("ArchiveName"), ArchiveName, GatherTextConfigPath );
-	GetConfigArray( *SectionName, TEXT("ProcessLanguage"), LanguagesToProcess, GatherTextConfigPath );
-	GetConfigArray( *SectionName, TEXT("TargetCulture"), TargetCultures, GatherTextConfigPath );
-	GetConfigBool( *SectionName, TEXT("bAppendToExistingArchive"), bAppendToExistingArchive, GatherTextConfigPath );
+	GetStringFromConfig( *SectionName, TEXT("ArchiveName"), ArchiveName, GatherTextConfigPath );
+	GetStringArrayFromConfig( *SectionName, TEXT("ProcessLanguage"), LanguagesToProcess, GatherTextConfigPath );
+	GetStringArrayFromConfig( *SectionName, TEXT("TargetCulture"), TargetCultures, GatherTextConfigPath );
+	GetBoolFromConfig( *SectionName, TEXT("bAppendToExistingArchive"), bAppendToExistingArchive, GatherTextConfigPath );
 
 	// Build info about the primary language
 	TArray<FString> PrimaryFilenames;
@@ -584,13 +584,13 @@ int32 UInternationalizationConditioningCommandlet::Main( const FString& Params )
 	bool bGenerateManifestFromLocIni = false;
 	
 	// Get the common settings from config
-	GetConfigString( *SectionName, TEXT("SourcePath"), SourcePath, GatherTextConfigPath );
-	GetConfigString( *SectionName, TEXT("DestinationPath"), DestinationPath, GatherTextConfigPath );
-	GetConfigString( *SectionName, TEXT("PrimaryLanguage"), PrimaryLangExt, GatherTextConfigPath );
-	GetConfigArray( *SectionName, TEXT("ProcessLanguage"), LanguagesToProcess, GatherTextConfigPath );
+	GetStringFromConfig( *SectionName, TEXT("SourcePath"), SourcePath, GatherTextConfigPath );
+	GetStringFromConfig( *SectionName, TEXT("DestinationPath"), DestinationPath, GatherTextConfigPath );
+	GetStringFromConfig( *SectionName, TEXT("PrimaryLanguage"), PrimaryLangExt, GatherTextConfigPath );
+	GetStringArrayFromConfig( *SectionName, TEXT("ProcessLanguage"), LanguagesToProcess, GatherTextConfigPath );
 
-	GetConfigBool( *SectionName, TEXT("bGenerateManifestFromLocIni"), bGenerateManifestFromLocIni, GatherTextConfigPath );
-	GetConfigBool( *SectionName, TEXT("bGenerateArchiveFromLocIni"), bGenerateArchiveFromLocIni, GatherTextConfigPath );
+	GetBoolFromConfig( *SectionName, TEXT("bGenerateManifestFromLocIni"), bGenerateManifestFromLocIni, GatherTextConfigPath );
+	GetBoolFromConfig( *SectionName, TEXT("bGenerateArchiveFromLocIni"), bGenerateArchiveFromLocIni, GatherTextConfigPath );
 
 	// Load legacy localization files.
 	LoadLegacyLocalizationFiles(SourcePath, PrimaryLangExt, LanguagesToProcess);
diff --git a/Engine/Source/Editor/UnrealEd/Private/Commandlets/InternationalizationExportCommandlet.cpp b/Engine/Source/Editor/UnrealEd/Private/Commandlets/InternationalizationExportCommandlet.cpp
--- a/Engine/Source/Editor/UnrealEd/Private/Commandlets/InternationalizationExportCommandlet.cpp
+++ b/Engine/Source/Editor/UnrealEd/Private/Commandlets/InternationalizationExportCommandlet.cpp
@@ -1039,7 +1039,7 @@ bool UInternationalizationExportCommandlet::DoExport( const FString& SourcePath,
 {
 	// Get manifest name.
 	FString ManifestName;
-	if( !GetConfigString( *SectionName, TEXT("ManifestName"), ManifestName, ConfigPath ) )
+	if( !GetStringFromConfig( *SectionName, TEXT("ManifestName"), ManifestName, ConfigPath ) )
 	{
 		UE_LOG( LogInternationalizationExportCommandlet, Error, TEXT("No manifest name specified.") );
 		return false;
@@ -1047,7 +1047,7 @@ bool UInternationalizationExportCommandlet::DoExport( const FString& SourcePath,
 
 	// Get archive name.
 	FString ArchiveName;
-	if( !( GetConfigString(* SectionName, TEXT("ArchiveName"), ArchiveName, ConfigPath ) ) )
+	if( !( GetStringFromConfig(* SectionName, TEXT("ArchiveName"), ArchiveName, ConfigPath ) ) )
 	{
 		UE_LOG(LogInternationalizationExportCommandlet, Error, TEXT("No archive name specified."));
 		return false;
@@ -1055,7 +1055,7 @@ bool UInternationalizationExportCommandlet::DoExport( const FString& SourcePath,
 
 	// Get culture directory setting, default to true if not specified (used to allow picking of export directory with windows file dialog from Translation Editor)
 	bool bUseCultureDirectory = true;
-	if (!(GetConfigBool(*SectionName, TEXT("bUseCultureDirectory"), bUseCultureDirectory, ConfigPath)))
+	if (!(GetBoolFromConfig(*SectionName, TEXT("bUseCultureDirectory"), bUseCultureDirectory, ConfigPath)))
 	{
 		bUseCultureDirectory = true;
 	}
@@ -1187,7 +1187,7 @@ bool UInternationalizationExportCommandlet::DoImport(const FString& SourcePath,
 {
 	// Get manifest name.
 	FString ManifestName;
-	if( !GetConfigString( *SectionName, TEXT("ManifestName"), ManifestName, ConfigPath ) )
+	if( !GetStringFromConfig( *SectionName, TEXT("ManifestName"), ManifestName, ConfigPath ) )
 	{
 		UE_LOG( LogInternationalizationExportCommandlet, Error, TEXT("No manifest name specified.") );
 		return false;
@@ -1195,7 +1195,7 @@ bool UInternationalizationExportCommandlet::DoImport(const FString& SourcePath,
 
 	// Get archive name.
 	FString ArchiveName;
-	if( !( GetConfigString(* SectionName, TEXT("ArchiveName"), ArchiveName, ConfigPath ) ) )
+	if( !( GetStringFromConfig(* SectionName, TEXT("ArchiveName"), ArchiveName, ConfigPath ) ) )
 	{
 		UE_LOG(LogInternationalizationExportCommandlet, Error, TEXT("No archive name specified."));
 		return false;
@@ -1203,7 +1203,7 @@ bool UInternationalizationExportCommandlet::DoImport(const FString& SourcePath,
 
 	// Get culture directory setting, default to true if not specified (used to allow picking of import directory with file open dialog from Translation Editor)
 	bool bUseCultureDirectory = true;
-	if (!(GetConfigBool(*SectionName, TEXT("bUseCultureDirectory"), bUseCultureDirectory, ConfigPath)))
+	if (!(GetBoolFromConfig(*SectionName, TEXT("bUseCultureDirectory"), bUseCultureDirectory, ConfigPath)))
 	{
 		bUseCultureDirectory = true;
 	}
@@ -1355,52 +1355,28 @@ int32 UInternationalizationExportCommandlet::Main( const FString& Params )
 
 
 	FString SourcePath; // Source path to the root folder that manifest/archive files live in.
-	if( !GetConfigString( *SectionName, TEXT("SourcePath"), SourcePath, ConfigPath ) )
+	if( !GetPathFromConfig( *SectionName, TEXT("SourcePath"), SourcePath, ConfigPath ) )
 	{
 		UE_LOG( LogInternationalizationExportCommandlet, Error, TEXT("No source path specified.") );
 		return -1;
 	}
 
-	if (FPaths::IsRelative(SourcePath))
-	{
-		if (!FPaths::GameDir().IsEmpty())
-		{
-			SourcePath = FPaths::Combine( *( FPaths::GameDir() ), *SourcePath );
-		}
-		else
-		{
-			SourcePath = FPaths::Combine( *( FPaths::EngineDir() ), *SourcePath );
-		}
-	}
-
 	FString DestinationPath; // Destination path that we will write files to.
-	if( !GetConfigString( *SectionName, TEXT("DestinationPath"), DestinationPath, ConfigPath ) )
+	if( !GetPathFromConfig( *SectionName, TEXT("DestinationPath"), DestinationPath, ConfigPath ) )
 	{
 		UE_LOG( LogInternationalizationExportCommandlet, Error, TEXT("No destination path specified.") );
 		return -1;
 	}
 
-	if (FPaths::IsRelative(DestinationPath))
-	{
-		if (!FPaths::GameDir().IsEmpty())
-		{
-			DestinationPath = FPaths::Combine( *( FPaths::GameDir() ), *DestinationPath );
-		}
-		else
-		{
-			DestinationPath = FPaths::Combine( *( FPaths::EngineDir() ), *DestinationPath );
-		}
-	}
-
 	FString Filename; // Name of the file to read or write from
-	if (!GetConfigString(*SectionName, TEXT("PortableObjectName"), Filename, ConfigPath))
+	if (!GetStringFromConfig(*SectionName, TEXT("PortableObjectName"), Filename, ConfigPath))
 	{
 		UE_LOG(LogInternationalizationExportCommandlet, Error, TEXT("No portable object name specified."));
 		return -1;
 	}
 
 	// Get cultures to generate.
-	if( GetConfigArray(*SectionName, TEXT("CulturesToGenerate"), CulturesToGenerate, ConfigPath) == 0 )
+	if( GetStringArrayFromConfig(*SectionName, TEXT("CulturesToGenerate"), CulturesToGenerate, ConfigPath) == 0 )
 	{
 		UE_LOG(LogInternationalizationExportCommandlet, Error, TEXT("No cultures specified for generation."));
 		return -1;
@@ -1410,8 +1386,8 @@ int32 UInternationalizationExportCommandlet::Main( const FString& Params )
 	bool bDoExport = false;
 	bool bDoImport = false;
 
-	GetConfigBool( *SectionName, TEXT("bExportLoc"), bDoExport, ConfigPath );
-	GetConfigBool( *SectionName, TEXT("bImportLoc"), bDoImport, ConfigPath );
+	GetBoolFromConfig( *SectionName, TEXT("bExportLoc"), bDoExport, ConfigPath );
+	GetBoolFromConfig( *SectionName, TEXT("bImportLoc"), bDoImport, ConfigPath );
 	
 	if( !bDoImport && !bDoExport )
 	{
diff --git a/Engine/Source/Editor/UnrealEd/Private/Commandlets/RepairLocalizationDataCommandlet.cpp b/Engine/Source/Editor/UnrealEd/Private/Commandlets/RepairLocalizationDataCommandlet.cpp
--- a/Engine/Source/Editor/UnrealEd/Private/Commandlets/RepairLocalizationDataCommandlet.cpp
+++ b/Engine/Source/Editor/UnrealEd/Private/Commandlets/RepairLocalizationDataCommandlet.cpp
@@ -177,27 +177,15 @@ int32 URepairLocalizationDataCommandlet::Main(const FString& Params)
 
 	// Get destination path.
 	FString DestinationPath;
-	if( !GetConfigString( *SectionName, TEXT("DestinationPath"), DestinationPath, GatherTextConfigPath ) )
+	if( !GetPathFromConfig( *SectionName, TEXT("DestinationPath"), DestinationPath, GatherTextConfigPath ) )
 	{
 		UE_LOG( LogRepairLocalizationDataCommandlet, Error, TEXT("No destination path specified.") );
 		return -1;
 	}
 
-	if (FPaths::IsRelative(DestinationPath))
-	{
-		if (!FPaths::GameDir().IsEmpty())
-		{
-			DestinationPath = FPaths::Combine( *( FPaths::GameDir() ), *DestinationPath );
-		}
-		else
-		{
-			DestinationPath = FPaths::Combine( *( FPaths::EngineDir() ), *DestinationPath );
-		}
-	}
-
 	// Get manifest name.
 	FString ManifestName;
-	if( !GetConfigString( *SectionName, TEXT("ManifestName"), ManifestName, GatherTextConfigPath ) )
+	if( !GetStringFromConfig( *SectionName, TEXT("ManifestName"), ManifestName, GatherTextConfigPath ) )
 	{
 		UE_LOG( LogRepairLocalizationDataCommandlet, Error, TEXT("No manifest name specified.") );
 		return -1;
@@ -205,7 +193,7 @@ int32 URepairLocalizationDataCommandlet::Main(const FString& Params)
 
 	// Get archive name.
 	FString ArchiveName;
-	if( !( GetConfigString(* SectionName, TEXT("ArchiveName"), ArchiveName, GatherTextConfigPath ) ) )
+	if( !( GetStringFromConfig(* SectionName, TEXT("ArchiveName"), ArchiveName, GatherTextConfigPath ) ) )
 	{
 		UE_LOG(LogRepairLocalizationDataCommandlet, Error, TEXT("No archive name specified."));
 		return -1;
@@ -213,7 +201,7 @@ int32 URepairLocalizationDataCommandlet::Main(const FString& Params)
 
 	// Get cultures to generate.
 	TArray<FString> CulturesToGenerate;
-	GetConfigArray(*SectionName, TEXT("CulturesToGenerate"), CulturesToGenerate, GatherTextConfigPath);
+	GetStringArrayFromConfig(*SectionName, TEXT("CulturesToGenerate"), CulturesToGenerate, GatherTextConfigPath);
 
 	if( CulturesToGenerate.Num() == 0 )
 	{
