Corrected order of import and export operations in internationalization import/export commandlet so that the existing PO file is imported before it is replaced by the export operation, when both operations are requested simultaneously.

diff --git a/Engine/Source/Editor/UnrealEd/Private/Commandlets/InternationalizationExportCommandlet.cpp b/Engine/Source/Editor/UnrealEd/Private/Commandlets/InternationalizationExportCommandlet.cpp
--- a/Engine/Source/Editor/UnrealEd/Private/Commandlets/InternationalizationExportCommandlet.cpp
+++ b/Engine/Source/Editor/UnrealEd/Private/Commandlets/InternationalizationExportCommandlet.cpp
@@ -1386,8 +1386,8 @@ int32 UInternationalizationExportCommandlet::Main( const FString& Params )
 	bool bDoExport = false;
 	bool bDoImport = false;
 
-	GetBoolFromConfig( *SectionName, TEXT("bExportLoc"), bDoExport, ConfigPath );
 	GetBoolFromConfig( *SectionName, TEXT("bImportLoc"), bDoImport, ConfigPath );
+	GetBoolFromConfig( *SectionName, TEXT("bExportLoc"), bDoExport, ConfigPath );
 	
 	if( !bDoImport && !bDoExport )
 	{
@@ -1395,20 +1395,20 @@ int32 UInternationalizationExportCommandlet::Main( const FString& Params )
 		return -1;
 	}
 
-	if( bDoExport )
+	if( bDoImport )
 	{
-		if (!DoExport(SourcePath, DestinationPath, Filename))
+		if (!DoImport(SourcePath, DestinationPath, Filename))
 		{
-			UE_LOG(LogInternationalizationExportCommandlet, Error, TEXT("Failed to export localization files."));
+			UE_LOG(LogInternationalizationExportCommandlet, Error, TEXT("Failed to import localization files."));
 			return -1;
 		}
 	}
 
-	if( bDoImport )
+	if( bDoExport )
 	{
-		if (!DoImport(SourcePath, DestinationPath, Filename))
+		if (!DoExport(SourcePath, DestinationPath, Filename))
 		{
-			UE_LOG(LogInternationalizationExportCommandlet, Error, TEXT("Failed to import localization files."));
+			UE_LOG(LogInternationalizationExportCommandlet, Error, TEXT("Failed to export localization files."));
 			return -1;
 		}
 	}
