Added source control integration to the localization report generating commandlet.

diff --git a/Engine/Source/Editor/UnrealEd/Private/Commandlets/GenerateTextLocalizationReportCommandlet.cpp b/Engine/Source/Editor/UnrealEd/Private/Commandlets/GenerateTextLocalizationReportCommandlet.cpp
--- a/Engine/Source/Editor/UnrealEd/Private/Commandlets/GenerateTextLocalizationReportCommandlet.cpp
+++ b/Engine/Source/Editor/UnrealEd/Private/Commandlets/GenerateTextLocalizationReportCommandlet.cpp
@@ -340,6 +340,16 @@ bool UGenerateTextLocalizationReportCommandlet::ProcessWordCountReport(const FSt
 		ReportData.SetEntry(NewRowIndex, CultureStr, FString::FromInt( TranslatedWords ) );
 	}
 
+	if( SourceControlInfo.IsValid() )
+	{
+		FText SCCErrorText;
+		if (!SourceControlInfo->CheckOutFile(ReportFilePath, SCCErrorText))
+		{
+			UE_LOG(LogGenerateTextLocalizationReportCommandlet, Error, TEXT("Check out of file %s failed: %s"), *ReportFilePath, *SCCErrorText.ToString());
+			return false;
+		}
+	}
+
 	if ( !FFileHelper::SaveStringToFile( ReportData.ToCSV(), *ReportFilePath ) )
 	{
 		UE_LOG(LogGenerateTextLocalizationReportCommandlet, Error, TEXT("Unable to save report at %s."), *ReportFilePath);
@@ -368,6 +378,16 @@ bool UGenerateTextLocalizationReportCommandlet::ProcessConflictReport(const FStr
 
 	ReportStr = FConflictReportInfo::GetInstance().ToString();
 
+	if( SourceControlInfo.IsValid() )
+	{
+		FText SCCErrorText;
+		if (!SourceControlInfo->CheckOutFile(ReportFilePath, SCCErrorText))
+		{
+			UE_LOG(LogGenerateTextLocalizationReportCommandlet, Error, TEXT("Check out of file %s failed: %s"), *ReportFilePath, *SCCErrorText.ToString());
+			return false;
+		}
+	}
+
 	if ( !FFileHelper::SaveStringToFile( ReportStr, *ReportFilePath ) )
 	{
 		UE_LOG(LogGenerateTextLocalizationReportCommandlet, Error, TEXT("Unable to save report at %s."), *ReportFilePath);
