Disabled the automatic revert of added/checked out files when exiting localization commandlets. It made the combination of "-EnableSCC" and "-DisableSCCSubmit" useless, as the changes would be reverted upon completion. The changes need to not be reverted in order for the localization dashboard to leverage the commandlets.

diff --git a/Engine/Source/Editor/UnrealEd/Private/Commandlets/GatherTextCommandletBase.cpp b/Engine/Source/Editor/UnrealEd/Private/Commandlets/GatherTextCommandletBase.cpp
--- a/Engine/Source/Editor/UnrealEd/Private/Commandlets/GatherTextCommandletBase.cpp
+++ b/Engine/Source/Editor/UnrealEd/Private/Commandlets/GatherTextCommandletBase.cpp
@@ -723,18 +723,12 @@ FGatherTextSCC::~FGatherTextSCC()
 {
 	if( CheckedOutFiles.Num() > 0 )
 	{
-		UE_LOG(LogGatherTextCommandletBase, Log, TEXT("Source Control wrapper shutting down with checked out files, cleaning up."));
-		FText SCCErrorStr;
-		if (!CleanUp(SCCErrorStr))
-		{
-			UE_LOG(LogGatherTextCommandletBase, Error, TEXT("%s"), *SCCErrorStr.ToString());
-		}
+		UE_LOG(LogGatherTextCommandletBase, Log, TEXT("Source Control wrapper shutting down with checked out files."));
 	}
 
 	ISourceControlModule::Get().GetProvider().Close();
 }
 
-
 bool FGatherTextSCC::CheckOutFile(const FString& InFile, FText& OutError)
 {
 	if ( InFile.IsEmpty() )
