Fixed existing duplicate matinee groups not being overridden when importing new tracks over the top

diff --git a/Engine/Source/Editor/UnrealEd/Private/Fbx/FbxMatineeImport.cpp b/Engine/Source/Editor/UnrealEd/Private/Fbx/FbxMatineeImport.cpp
--- a/Engine/Source/Editor/UnrealEd/Private/Fbx/FbxMatineeImport.cpp
+++ b/Engine/Source/Editor/UnrealEd/Private/Fbx/FbxMatineeImport.cpp
@@ -225,7 +225,7 @@ bool FFbxImporter::ImportMatineeSequence(AMatineeActor* InMatineeActor)
 		else
 		{
 			// Check to see if the scene node name matches a Matinee group name
-			InMatineeActor->FindFirstGroupInstByName( FString( Node->GetName() ) );
+			FoundGroupInst = InMatineeActor->FindFirstGroupInstByName( FString( Node->GetName() ) );
 		}
 
 		if( FoundGroupInst != NULL )
