Behavior Tree editor no longer checks out blackboards that are not dirty

Note: There is the possiblity that this might be best applied to all asset editors, so I have left this as a flag that can be enabled if we want to in the future.

UE-2992 - EDITOR: BEHAVIOR TREES: Blackboard is saved even though it is not dirty

diff --git a/Engine/Source/Editor/BehaviorTreeEditor/Private/BehaviorTreeEditor.cpp b/Engine/Source/Editor/BehaviorTreeEditor/Private/BehaviorTreeEditor.cpp
--- a/Engine/Source/Editor/BehaviorTreeEditor/Private/BehaviorTreeEditor.cpp
+++ b/Engine/Source/Editor/BehaviorTreeEditor/Private/BehaviorTreeEditor.cpp
@@ -62,6 +62,8 @@ FBehaviorTreeEditor::FBehaviorTreeEditor()
 
 	BehaviorTree = nullptr;
 	BlackboardData = nullptr;
+
+	bCheckDirtyOnAssetSave = true;
 }
 
 FBehaviorTreeEditor::~FBehaviorTreeEditor()
diff --git a/Engine/Source/Editor/UnrealEd/Private/Toolkits/AssetEditorToolkit.cpp b/Engine/Source/Editor/UnrealEd/Private/Toolkits/AssetEditorToolkit.cpp
--- a/Engine/Source/Editor/UnrealEd/Private/Toolkits/AssetEditorToolkit.cpp
+++ b/Engine/Source/Editor/UnrealEd/Private/Toolkits/AssetEditorToolkit.cpp
@@ -25,6 +25,7 @@ const FName FAssetEditorToolkit::ToolbarTabId( TEXT( "AssetEditorToolkit_Toolbar
 
 FAssetEditorToolkit::FAssetEditorToolkit()
 	: bIsToolbarFocusable(false)
+	, bCheckDirtyOnAssetSave(false)
 {
 	WorkspaceMenuCategory = FWorkspaceItem::NewGroup(LOCTEXT("WorkspaceMenu_BaseAssetEditor", "Asset Editor"));
 }
@@ -375,7 +376,7 @@ void FAssetEditorToolkit::SaveAsset_Execute()
 			}
 		}
 
-		FEditorFileUtils::PromptForCheckoutAndSave( PackagesToSave, /*bCheckDirty=*/ false, /*bPromptToSave=*/ false );
+		FEditorFileUtils::PromptForCheckoutAndSave( PackagesToSave, bCheckDirtyOnAssetSave, /*bPromptToSave=*/ false );
 	}
 }
 
diff --git a/Engine/Source/Editor/UnrealEd/Public/Toolkits/AssetEditorToolkit.h b/Engine/Source/Editor/UnrealEd/Public/Toolkits/AssetEditorToolkit.h
--- a/Engine/Source/Editor/UnrealEd/Public/Toolkits/AssetEditorToolkit.h
+++ b/Engine/Source/Editor/UnrealEd/Public/Toolkits/AssetEditorToolkit.h
@@ -239,6 +239,9 @@ protected:
 	/** Controls our internal layout */
 	TSharedPtr<FTabManager> TabManager;
 
+	/** Whether only dirty assets should be prompted about on save - otherwise all edited assets will be prompted to the user for save/check-out */
+	bool bCheckDirtyOnAssetSave;
+
 private:
 	/** The toolkit standalone host; may be NULL */
 	TWeakPtr< class SStandaloneAssetEditorToolkitHost > StandaloneHost;
