Fixed Clang compile errors

diff --git a/Engine/Source/Editor/UnrealEd/Private/Toolkits/AssetEditorToolkit.cpp b/Engine/Source/Editor/UnrealEd/Private/Toolkits/AssetEditorToolkit.cpp
--- a/Engine/Source/Editor/UnrealEd/Private/Toolkits/AssetEditorToolkit.cpp
+++ b/Engine/Source/Editor/UnrealEd/Private/Toolkits/AssetEditorToolkit.cpp
@@ -24,8 +24,8 @@ TWeakPtr< IToolkitHost > FAssetEditorToolkit::PreviousWorldCentricToolkitHostFor
 const FName FAssetEditorToolkit::ToolbarTabId( TEXT( "AssetEditorToolkit_Toolbar" ) );
 
 FAssetEditorToolkit::FAssetEditorToolkit()
-	: bIsToolbarFocusable(false)
-	, bCheckDirtyOnAssetSave(false)
+	: bCheckDirtyOnAssetSave(false)
+	, bIsToolbarFocusable(false)
 {
 	WorkspaceMenuCategory = FWorkspaceItem::NewGroup(LOCTEXT("WorkspaceMenu_BaseAssetEditor", "Asset Editor"));
 }
diff --git a/Engine/Source/Editor/UnrealEd/Public/Kismet2/BlueprintEditorUtils.h b/Engine/Source/Editor/UnrealEd/Public/Kismet2/BlueprintEditorUtils.h
--- a/Engine/Source/Editor/UnrealEd/Public/Kismet2/BlueprintEditorUtils.h
+++ b/Engine/Source/Editor/UnrealEd/Public/Kismet2/BlueprintEditorUtils.h
@@ -5,9 +5,9 @@
 #include "Editor/ClassViewer/Public/ClassViewerModule.h"
 #include "EdGraph/EdGraph.h"
 #include "EdGraphSchema_K2.h"
+#include "K2Node_EditablePinBase.h"
 
 class USCS_Node;
-struct FKismetUserDeclaredFunctionMetadata;
 
 /** 
   * Flags describing how to handle graph removal
