BP, Reinstancer: CDO is always replaced after recompilation. Even when bReinstanceOnlyWhenNecessary is used and regular objects aren't reinstanced.

diff --git a/Engine/Source/Editor/UnrealEd/Private/Kismet2/KismetReinstanceUtilities.cpp b/Engine/Source/Editor/UnrealEd/Private/Kismet2/KismetReinstanceUtilities.cpp
--- a/Engine/Source/Editor/UnrealEd/Private/Kismet2/KismetReinstanceUtilities.cpp
+++ b/Engine/Source/Editor/UnrealEd/Private/Kismet2/KismetReinstanceUtilities.cpp
@@ -17,6 +17,96 @@ DEFINE_STAT(EKismetReinstancerStats_UpdateBytecodeReferences);
 DEFINE_STAT(EKismetReinstancerStats_RecompileChildClasses);
 DEFINE_STAT(EKismetReinstancerStats_ReplaceClassNoReinsancing);
 
+struct FReplaceReferenceHelper
+{
+	static void IncludeCDO(UClass* OldClass, UClass* NewClass, TMap<UObject*, UObject*> &OldToNewInstanceMap, TArray<UObject*> &SourceObjects, UObject* OriginalCDO)
+	{
+		UObject* OldCDO = OldClass->GetDefaultObject();
+		UObject* NewCDO = NewClass->GetDefaultObject();
+
+		// Add the old->new CDO mapping into the fixup map
+		OldToNewInstanceMap.Add(OldCDO, NewCDO);
+		// Add in the old CDO to this pass, so CDO references are fixed up
+		SourceObjects.Add(OldCDO);
+
+		if (OriginalCDO)
+		{
+			OldToNewInstanceMap.Add(OriginalCDO, NewCDO);
+			SourceObjects.Add(OriginalCDO);
+		}
+	}
+
+	static void FindAndReplaceReferences(TArray<UObject*>& SourceObjects, TSet<UObject*>* ObjectsThatShouldUseOldStuff, TArray<UObject*> &ObjectsToReplace, TMap<UObject*, UObject*>& OldToNewInstanceMap, TMap<FStringAssetReference, UObject*>& ReinstancedObjectsWeakReferenceMap)
+	{
+		// Find everything that references these objects
+		TSet<UObject *> Targets;
+		{
+			BP_SCOPED_COMPILER_EVENT_STAT(EKismetReinstancerStats_FindReferencers);
+
+			TFindObjectReferencers<UObject> Referencers(SourceObjects, NULL, false);
+			for (TFindObjectReferencers<UObject>::TIterator It(Referencers); It; ++It)
+			{
+				UObject* Referencer = It.Value();
+				if (!ObjectsThatShouldUseOldStuff || !ObjectsThatShouldUseOldStuff->Contains(Referencer))
+				{
+					Targets.Add(Referencer);
+				}
+			}
+		}
+
+		{
+			BP_SCOPED_COMPILER_EVENT_STAT(EKismetReinstancerStats_ReplaceReferences);
+
+			for (TSet<UObject *>::TIterator It(Targets); It; ++It)
+			{
+				UObject* Obj = *It;
+				if (!ObjectsToReplace.Contains(Obj)) // Don't bother trying to fix old objects, this would break them
+				{
+					// The class for finding and replacing weak references.
+					// We can't relay on "standard" weak references replacement as
+					// it depends on FStringAssetReference::ResolveObject, which
+					// tries to find the object with the stored path. It is
+					// impossible, cause above we deleted old actors (after
+					// spawning new ones), so during objects traverse we have to
+					// find FStringAssetReferences with the raw given path taken
+					// before deletion of old actors and fix them.
+					class ReferenceReplace : public FArchiveReplaceObjectRef < UObject >
+					{
+					public:
+						ReferenceReplace(UObject* InSearchObject, const TMap<UObject*, UObject*>& InReplacementMap, TMap<FStringAssetReference, UObject*> WeakReferencesMap)
+							: FArchiveReplaceObjectRef<UObject>(InSearchObject, InReplacementMap, false, false, false, true), WeakReferencesMap(WeakReferencesMap)
+						{
+							SerializeSearchObject();
+						}
+
+						FArchive& operator<<(FStringAssetReference& Ref) override
+						{
+							const UObject*const* PtrToObjPtr = WeakReferencesMap.Find(Ref);
+
+							if (PtrToObjPtr != nullptr)
+							{
+								Ref = *PtrToObjPtr;
+							}
+
+							return *this;
+						}
+
+						FArchive& operator<<(FAssetPtr& Ref) override
+						{
+							return operator<<(Ref.GetUniqueID());
+						}
+
+					private:
+						const TMap<FStringAssetReference, UObject*>& WeakReferencesMap;
+					};
+
+					ReferenceReplace ReplaceAr(Obj, OldToNewInstanceMap, ReinstancedObjectsWeakReferenceMap);
+				}
+			}
+		}
+	}
+};
+
 /////////////////////////////////////////////////////////////////////////////////
 // FBlueprintCompileReinstancer
 
@@ -230,6 +320,12 @@ void FBlueprintCompileReinstancer::ReinstanceObjects(bool bAlwaysReinstance)
 					}
 				}
 
+				TArray<UObject*> SourceObjects;
+				TMap<UObject*, UObject*> OldToNewInstanceMap;
+				TMap<FStringAssetReference, UObject*> ReinstancedObjectsWeakReferenceMap;
+				FReplaceReferenceHelper::IncludeCDO(DuplicatedClass, ClassToReinstance, OldToNewInstanceMap, SourceObjects, OriginalCDO);
+				FReplaceReferenceHelper::FindAndReplaceReferences(SourceObjects, &ObjectsThatShouldUseOldStuff, ObjectsToReplace, OldToNewInstanceMap, ReinstancedObjectsWeakReferenceMap);
+
 				bShouldReinstance = false;
 			}
 		}
@@ -661,9 +757,6 @@ void FBlueprintCompileReinstancer::ReplaceInstancesOfClass(UClass* OldClass, UCl
 		}
 	}
 
-	UObject* OldCDO = OldClass->GetDefaultObject();
-	UObject* NewCDO = NewClass->GetDefaultObject();
-
 	// Now replace any pointers to the old archetypes/instances with pointers to the new one
 	TArray<UObject*> SourceObjects;
 	TArray<UObject*> DstObjects;
@@ -673,86 +766,14 @@ void FBlueprintCompileReinstancer::ReplaceInstancesOfClass(UClass* OldClass, UCl
 
 	SourceObjects.Append(DstObjects);
 
-	// Add in the old class/CDO to this pass, so class/CDO references are fixed up
-	SourceObjects.Add(OldClass);
-	SourceObjects.Add(OldCDO);
-	if (OriginalCDO)
-	{
-		SourceObjects.Add(OriginalCDO);
-	}
-
-	// Find everything that references these objects
-	TSet<UObject *> Targets;
-	{ BP_SCOPED_COMPILER_EVENT_STAT(EKismetReinstancerStats_FindReferencers);
-
-		TFindObjectReferencers<UObject> Referencers(SourceObjects, NULL, false);
-		for (TFindObjectReferencers<UObject>::TIterator It(Referencers); It; ++It)
-		{
-			UObject* Referencer = It.Value();
-			if (!ObjectsThatShouldUseOldStuff || !ObjectsThatShouldUseOldStuff->Contains(Referencer))
-			{
-				Targets.Add(Referencer);
-			}
-		}
-	}
-
-	// Add the old->new class/CDO mapping into the fixup map
+	// Add the old->new class mapping into the fixup map
 	OldToNewInstanceMap.Add(OldClass, NewClass);
-	OldToNewInstanceMap.Add(OldCDO, NewCDO);
-	if (OriginalCDO)
-	{
-		OldToNewInstanceMap.Add(OriginalCDO, NewCDO);
-	}
-
-	{ BP_SCOPED_COMPILER_EVENT_STAT(EKismetReinstancerStats_ReplaceReferences);
 
-		for (TSet<UObject *>::TIterator It(Targets); It; ++It)
-		{
-			UObject* Obj = *It;
-			if (!ObjectsToReplace.Contains(Obj)) // Don't bother trying to fix old objects, this would break them
-			{
-				// The class for finding and replacing weak references.
-				// We can't relay on "standard" weak references replacement as
-				// it depends on FStringAssetReference::ResolveObject, which
-				// tries to find the object with the stored path. It is
-				// impossible, cause above we deleted old actors (after
-				// spawning new ones), so during objects traverse we have to
-				// find FStringAssetReferences with the raw given path taken
-				// before deletion of old actors and fix them.
-				class ReferenceReplace : public FArchiveReplaceObjectRef<UObject>
-				{
-				public:
-					ReferenceReplace(UObject* InSearchObject, const TMap<UObject*, UObject*>& InReplacementMap, TMap<FStringAssetReference, UObject*> WeakReferencesMap)
-						: FArchiveReplaceObjectRef<UObject>(InSearchObject, InReplacementMap, false, false, false, true), WeakReferencesMap(WeakReferencesMap)
-					{
-						SerializeSearchObject();
-					}
-
-					FArchive& operator<<(FStringAssetReference& Ref) override
-					{
-						const UObject*const* PtrToObjPtr = WeakReferencesMap.Find(Ref);
-
-						if (PtrToObjPtr != nullptr)
-						{
-							Ref = *PtrToObjPtr;
-						}
-
-						return *this;
-					}
-
-					FArchive& operator<<(FAssetPtr& Ref) override
-					{
-						return operator<<(Ref.GetUniqueID());
-					}
-
-				private:
-					const TMap<FStringAssetReference, UObject*>& WeakReferencesMap;
-				};
+	// Add in the old class to this pass, so class references are fixed up
+	SourceObjects.Add(OldClass);
 
-				ReferenceReplace ReplaceAr(Obj, OldToNewInstanceMap, ReinstancedObjectsWeakReferenceMap);
-			}
-		}
-	}
+	FReplaceReferenceHelper::IncludeCDO(OldClass, NewClass, OldToNewInstanceMap, SourceObjects, OriginalCDO);
+	FReplaceReferenceHelper::FindAndReplaceReferences(SourceObjects, ObjectsThatShouldUseOldStuff, ObjectsToReplace, OldToNewInstanceMap, ReinstancedObjectsWeakReferenceMap);
 
 	{ BP_SCOPED_COMPILER_EVENT_STAT(EKismetReinstancerStats_ReplacementConstruction);
 
