Added BP Single pane editor and BP In World editing experimental flag.

diff --git a/Engine/Source/Editor/UnrealEd/Classes/Settings/EditorExperimentalSettings.h b/Engine/Source/Editor/UnrealEd/Classes/Settings/EditorExperimentalSettings.h
--- a/Engine/Source/Editor/UnrealEd/Classes/Settings/EditorExperimentalSettings.h
+++ b/Engine/Source/Editor/UnrealEd/Classes/Settings/EditorExperimentalSettings.h
@@ -29,6 +29,10 @@ public:
 	UPROPERTY(EditAnywhere, config, Category=Tools, meta=(DisplayName="Editor Utility Blueprints (Blutility)"))
 	bool bEnableEditorUtilityBlueprints;
 
+	/** Enable Single Layout BP Editor and In World BP Editing (WIP). */
+	UPROPERTY(EditAnywhere, config, Category = Tools, meta = (DisplayName = "Single Layout BP Editor and In World BP Editing"))
+	bool bInWorldBPEditing;
+
 	/** The Messaging Debugger provides a visual utility for debugging the messaging system. */
 	UPROPERTY(EditAnywhere, config, Category=Tools, meta=(DisplayName="Messaging Debugger"))
 	bool bMessagingDebugger;
