Tagged some less common editor viewport settings as advanced

diff --git a/Engine/Source/Editor/UnrealEd/Classes/Settings/LevelEditorViewportSettings.h b/Engine/Source/Editor/UnrealEd/Classes/Settings/LevelEditorViewportSettings.h
--- a/Engine/Source/Editor/UnrealEd/Classes/Settings/LevelEditorViewportSettings.h
+++ b/Engine/Source/Editor/UnrealEd/Classes/Settings/LevelEditorViewportSettings.h
@@ -305,11 +305,11 @@ private:
 public:
 
 	/** If true actor snap will be enabled in the editor **/
-	UPROPERTY(config, Category=GridSnapping, VisibleDefaultsOnly)
+	UPROPERTY(config, Category=GridSnapping, VisibleDefaultsOnly,AdvancedDisplay)
 	uint32 bEnableActorSnap:1;
 
 	/** Global actor snap scale for the editor */
-	UPROPERTY(config, Category=GridSnapping, VisibleDefaultsOnly)
+	UPROPERTY(config, Category=GridSnapping, VisibleDefaultsOnly,AdvancedDisplay)
 	float ActorSnapScale;
 
 	/** Global actor snap distance setting for the editor */
@@ -349,7 +349,7 @@ public:
 	uint32 bEnableViewportHoverFeedback:1;
 
 	/** If enabled, selected objects will be highlighted with brackets in all modes rather than a special highlight color. */
-	UPROPERTY(EditAnywhere, config, Category=LookAndFeel, meta=(DisplayName = "Highlight Selected Objects with Brackets"))
+	UPROPERTY(EditAnywhere, config, Category=LookAndFeel, AdvancedDisplay, meta=(DisplayName = "Highlight Selected Objects with Brackets"))
 	uint32 bHighlightWithBrackets:1;
 
 	/** If checked all orthographic view ports are linked to the same position and move together. */
@@ -373,11 +373,11 @@ public:
 	float SelectionHighlightIntensity;
 
 	/** Sets the intensity of the overlay displayed when an object is selected */
-	UPROPERTY(EditAnywhere, config, Category=LookAndFeel, meta=(DisplayName = "BSP Surface Highlight Intensity" ,ClampMin = "0", UIMin = "0", UIMax = "1"))
+	UPROPERTY(EditAnywhere, config, Category=LookAndFeel, AdvancedDisplay, meta=(DisplayName = "BSP Surface Highlight Intensity" ,ClampMin = "0", UIMin = "0", UIMax = "1"))
 	float BSPSelectionHighlightIntensity;
 
 	/** Sets the intensity of the overlay displayed when an object is hovered */
-	UPROPERTY(EditAnywhere, config, Category=LookAndFeel, meta=(DisplayName = "Hover Highlight Intensity" ,ClampMin = "0", UIMin = "0", UIMax = "20"))
+	UPROPERTY(EditAnywhere, config, Category=LookAndFeel, AdvancedDisplay, meta=(DisplayName = "Hover Highlight Intensity" ,ClampMin = "0", UIMin = "0", UIMax = "20"))
 	float HoverHighlightIntensity;
 
 	/** Enables the editor perspective camera to be dropped at the last PlayInViewport cam position */
@@ -393,22 +393,22 @@ public:
 	float CameraPreviewSize;
 
 	/** Distance from the camera to place actors which are dropped on nothing in the view port. */
-	UPROPERTY(EditAnywhere, config, Category=LookAndFeel, meta=(DisplayName = "Background Drop Distance"))
+	UPROPERTY(EditAnywhere, config, Category=LookAndFeel, AdvancedDisplay, meta=(DisplayName = "Background Drop Distance"))
 	float BackgroundDropDistance;
 
 	/** A list of meshes that can be used as preview mesh in the editor view port by holding down the backslash key */
 	UPROPERTY(EditAnywhere, config, Category=Preview, meta=(AllowedClasses = "StaticMesh"))
 	TArray<FStringAssetReference> PreviewMeshes;
 
-	UPROPERTY(EditAnywhere, config, Category=LookAndFeel, meta=(ClampMin = "0.01", UIMin = "0.01", UIMax = "5"))
+	UPROPERTY(EditAnywhere, config, AdvancedDisplay, Category=LookAndFeel, meta=(ClampMin = "0.01", UIMin = "0.01", UIMax = "5"))
 	float BillboardScale;
 
 	/** The size adjustment to apply to the translate/rotate/scale widgets (in Unreal units). */
-	UPROPERTY(EditAnywhere, config, Category=LookAndFeel, meta=(ClampMin="-10",ClampMax="150") )
+	UPROPERTY(EditAnywhere, config, Category=LookAndFeel, AdvancedDisplay, meta=(ClampMin="-10",ClampMax="150") )
 	int32 TransformWidgetSizeAdjustment;
 
 	/** When enabled, engine stats that are enabled in level viewports are preserved between editor sessions */
-	UPROPERTY(EditAnywhere, config, Category = LookAndFeel)
+	UPROPERTY(EditAnywhere, config, AdvancedDisplay, Category = LookAndFeel)
 	uint32 bSaveEngineStats : 1;
 
 	/** Specify the units used by the measuring tool */
