#ue4 UE-3952:  Fix for not being able to base blueprints on blueprints based on object.  BPs now always considered extendable

diff --git a/Engine/Source/Editor/UnrealEd/Private/Kismet2/Kismet2.cpp b/Engine/Source/Editor/UnrealEd/Private/Kismet2/Kismet2.cpp
--- a/Engine/Source/Editor/UnrealEd/Private/Kismet2/Kismet2.cpp
+++ b/Engine/Source/Editor/UnrealEd/Private/Kismet2/Kismet2.cpp
@@ -860,9 +860,12 @@ bool FKismetEditorUtilities::CanCreateBlueprintOfClass(const UClass* Class)
 		&& !Class->HasAnyClassFlags(CLASS_NewerVersionExists)
 		&& (!Class->ClassGeneratedBy || (bAllowDerivedBlueprints && !IsClassABlueprintSkeleton(Class)));
 
+	const bool bIsBPGC = (Cast<UBlueprintGeneratedClass>(Class) != nullptr);
+
 	const UEdGraphSchema_K2* K2Schema = GetDefault<UEdGraphSchema_K2>();
 	const bool bIsValidClass = Class->GetBoolMetaDataHierarchical(FBlueprintMetadata::MD_IsBlueprintBase)
-		|| (Class == UObject::StaticClass());
+		|| (Class == UObject::StaticClass())
+		|| bIsBPGC;  // BPs are always considered inheritable
 
 	return bCanCreateBlueprint && bIsValidClass;
 }
