Re-enabled preserve local transform in skeletal mesh import when animations are enabled.

diff --git a/Engine/Source/Editor/UnrealEd/Classes/Factories/FbxImportUI.h b/Engine/Source/Editor/UnrealEd/Classes/Factories/FbxImportUI.h
--- a/Engine/Source/Editor/UnrealEd/Classes/Factories/FbxImportUI.h
+++ b/Engine/Source/Editor/UnrealEd/Classes/Factories/FbxImportUI.h
@@ -108,8 +108,8 @@ class UFbxImportUI : public UObject
 	class UFbxAnimSequenceImportData* AnimSequenceImportData;
 
 	/** Type of asset to import from the FBX file */
-	UPROPERTY(EditAnywhere, AdvancedDisplay, config, Category=Animation, meta=(editcondition="bImportAnimations", ImportType = "Animation"))
-	bool	bPreserveLocalTransform;
+	UPROPERTY(EditAnywhere, AdvancedDisplay, config, Category=Animation, meta=(editcondition="bImportAnimations", ImportType = "SkeletalMesh|Animation"))
+	bool bPreserveLocalTransform;
 
 	/** Import data used when importing textures */
 	UPROPERTY(EditAnywhere, Instanced, Category=Material)
