Fixed issue where brush vertices were being rendered on all viewports, not just the level viewports.

diff --git a/Engine/Source/Editor/UnrealEd/Private/EdMode.cpp b/Engine/Source/Editor/UnrealEd/Private/EdMode.cpp
--- a/Engine/Source/Editor/UnrealEd/Private/EdMode.cpp
+++ b/Engine/Source/Editor/UnrealEd/Private/EdMode.cpp
@@ -11,7 +11,6 @@
 #include "LevelEditor.h"
 #include "Toolkits/ToolkitManager.h"
 #include "EditorLevelUtils.h"
-#include "DynamicMeshBuilder.h"
 
 #include "ActorEditorUtils.h"
 #include "EditorStyle.h"
@@ -584,96 +583,6 @@ FModeTool* FEdMode::FindTool( EModeTools InID )
 
 void FEdMode::Render(const FSceneView* View,FViewport* Viewport,FPrimitiveDrawInterface* PDI)
 {
-	if( GEditor->bShowBrushMarkerPolys )
-	{
-		// Draw translucent polygons on brushes and volumes
-
-		for( TActorIterator<ABrush> It(GetWorld()); It; ++ It )
-		{
-			ABrush* Brush = *It;
-
-			// Brush->Brush is checked to safe from brushes that were created without having their brush members attached.
-			if( Brush->Brush && (FActorEditorUtils::IsABuilderBrush(Brush) || Brush->IsVolumeBrush()) && Owner->GetSelectedActors()->IsSelected(Brush) )
-			{
-				// Build a mesh by basically drawing the triangles of each 
-				FDynamicMeshBuilder MeshBuilder;
-				int32 VertexOffset = 0;
-
-				for( int32 PolyIdx = 0 ; PolyIdx < Brush->Brush->Polys->Element.Num() ; ++PolyIdx )
-				{
-					const FPoly* Poly = &Brush->Brush->Polys->Element[PolyIdx];
-
-					if( Poly->Vertices.Num() > 2 )
-					{
-						const FVector Vertex0 = Poly->Vertices[0];
-						FVector Vertex1 = Poly->Vertices[1];
-
-						MeshBuilder.AddVertex(Vertex0, FVector2D::ZeroVector, FVector(1,0,0), FVector(0,1,0), FVector(0,0,1), FColor::White);
-						MeshBuilder.AddVertex(Vertex1, FVector2D::ZeroVector, FVector(1,0,0), FVector(0,1,0), FVector(0,0,1), FColor::White);
-
-						for( int32 VertexIdx = 2 ; VertexIdx < Poly->Vertices.Num() ; ++VertexIdx )
-						{
-							const FVector Vertex2 = Poly->Vertices[VertexIdx];
-							MeshBuilder.AddVertex(Vertex2, FVector2D::ZeroVector, FVector(1,0,0), FVector(0,1,0), FVector(0,0,1), FColor::White);
-							MeshBuilder.AddTriangle(VertexOffset,VertexOffset + VertexIdx,VertexOffset+VertexIdx-1);
-							Vertex1 = Vertex2;
-						}
-
-						// Increment the vertex offset so the next polygon uses the correct vertex indices.
-						VertexOffset += Poly->Vertices.Num();
-					}
-				}
-
-				// Allocate the material proxy and register it so it can be deleted properly once the rendering is done with it.
-				FDynamicColoredMaterialRenderProxy* MaterialProxy = new FDynamicColoredMaterialRenderProxy(GEngine->EditorBrushMaterial->GetRenderProxy(false),Brush->GetWireColor());
-				PDI->RegisterDynamicResource( MaterialProxy );
-
-				// Flush the mesh triangles.
-				MeshBuilder.Draw(PDI, Brush->ActorToWorld().ToMatrixWithScale(), MaterialProxy, SDPG_World, 0.f);
-			}
-		}
-	}
-
-	const bool bIsInGameView = !Viewport->GetClient() || Viewport->GetClient()->IsInGameView();
-	if (Owner->ShouldDrawBrushVertices() && !bIsInGameView)
-	{
-		UTexture2D* VertexTexture = GetVertexTexture();
-		const float TextureSizeX = VertexTexture->GetSizeX() * 0.170f;
-		const float TextureSizeY = VertexTexture->GetSizeY() * 0.170f;
-
-		for (FSelectionIterator It(*Owner->GetSelectedActors()); It; ++It)
-		{
-			AActor* SelectedActor = static_cast<AActor*>(*It);
-			checkSlow(SelectedActor->IsA(AActor::StaticClass()));
-
-			ABrush* Brush = Cast< ABrush >(SelectedActor);
-			if (Brush && Brush->Brush && !FActorEditorUtils::IsABuilderBrush(Brush))
-			{
-				for (int32 p = 0; p < Brush->Brush->Polys->Element.Num(); ++p)
-				{
-					FTransform BrushTransform = Brush->ActorToWorld();
-
-					FPoly* poly = &Brush->Brush->Polys->Element[p];
-					for (int32 VertexIndex = 0; VertexIndex < poly->Vertices.Num(); ++VertexIndex)
-					{
-						const FVector& PolyVertex = poly->Vertices[VertexIndex];
-						const FVector WorldLocation = BrushTransform.TransformPosition(PolyVertex);
-						
-						const float Scale = View->WorldToScreen( WorldLocation ).W * ( 4.0f / View->ViewRect.Width() / View->ViewMatrices.ProjMatrix.M[0][0] );
-
-						const FColor Color(Brush->GetWireColor());
-						PDI->SetHitProxy(new HBSPBrushVert(Brush, &poly->Vertices[VertexIndex]));
-
-						PDI->DrawSprite(WorldLocation, TextureSizeX * Scale, TextureSizeY * Scale, VertexTexture->Resource, Color, SDPG_World, 0.0f, 0.0f, 0.0f, 0.0f, SE_BLEND_Masked );
-
-						PDI->SetHitProxy(NULL);
-			
-					}
-				}
-			}
-		}
-	}
-
 	// Let the current mode tool render if it wants to
 	FModeTool* tool = GetCurrentTool();
 	if( tool )
@@ -681,8 +590,6 @@ void FEdMode::Render(const FSceneView* View,FViewport* Viewport,FPrimitiveDrawIn
 		tool->Render( View, Viewport, PDI );
 	}
 
-	AGroupActor::DrawBracketsForGroups(PDI, Viewport);
-
 	if(UsesPropertyWidgets())
 	{
 		bool bHitTesting = PDI->IsHitTesting();
diff --git a/Engine/Source/Editor/UnrealEd/Private/LevelEditorViewport.cpp b/Engine/Source/Editor/UnrealEd/Private/LevelEditorViewport.cpp
--- a/Engine/Source/Editor/UnrealEd/Private/LevelEditorViewport.cpp
+++ b/Engine/Source/Editor/UnrealEd/Private/LevelEditorViewport.cpp
@@ -43,7 +43,7 @@
 #include "AssetRegistryModule.h"
 #include "Animation/VertexAnim/VertexAnimation.h"
 #include "InstancedFoliage.h"
-
+#include "DynamicMeshBuilder.h"
 #include "Editor/ActorPositioning.h"
 #include "NotificationManager.h"
 #include "SNotificationList.h"
@@ -3430,6 +3430,9 @@ void FLevelEditorViewportClient::Draw(const FSceneView* View,FPrimitiveDrawInter
 
 	FEditorViewportClient::Draw(View,PDI);
 
+	DrawBrushDetails(View, PDI);
+	AGroupActor::DrawBracketsForGroups(PDI, Viewport);
+
 	if (EngineShowFlags.StreamingBounds)
 	{
 		DrawTextureStreamingBounds(View, PDI);
@@ -3540,6 +3543,97 @@ void FLevelEditorViewportClient::Draw(const FSceneView* View,FPrimitiveDrawInter
 	Mark.Pop();
 }
 
+void FLevelEditorViewportClient::DrawBrushDetails(const FSceneView* View, FPrimitiveDrawInterface* PDI)
+{
+	if (GEditor->bShowBrushMarkerPolys)
+	{
+		// Draw translucent polygons on brushes and volumes
+
+		for (TActorIterator<ABrush> It(GetWorld()); It; ++It)
+		{
+			ABrush* Brush = *It;
+
+			// Brush->Brush is checked to safe from brushes that were created without having their brush members attached.
+			if (Brush->Brush && (FActorEditorUtils::IsABuilderBrush(Brush) || Brush->IsVolumeBrush()) && ModeTools->GetSelectedActors()->IsSelected(Brush))
+			{
+				// Build a mesh by basically drawing the triangles of each 
+				FDynamicMeshBuilder MeshBuilder;
+				int32 VertexOffset = 0;
+
+				for (int32 PolyIdx = 0; PolyIdx < Brush->Brush->Polys->Element.Num(); ++PolyIdx)
+				{
+					const FPoly* Poly = &Brush->Brush->Polys->Element[PolyIdx];
+
+					if (Poly->Vertices.Num() > 2)
+					{
+						const FVector Vertex0 = Poly->Vertices[0];
+						FVector Vertex1 = Poly->Vertices[1];
+
+						MeshBuilder.AddVertex(Vertex0, FVector2D::ZeroVector, FVector(1, 0, 0), FVector(0, 1, 0), FVector(0, 0, 1), FColor::White);
+						MeshBuilder.AddVertex(Vertex1, FVector2D::ZeroVector, FVector(1, 0, 0), FVector(0, 1, 0), FVector(0, 0, 1), FColor::White);
+
+						for (int32 VertexIdx = 2; VertexIdx < Poly->Vertices.Num(); ++VertexIdx)
+						{
+							const FVector Vertex2 = Poly->Vertices[VertexIdx];
+							MeshBuilder.AddVertex(Vertex2, FVector2D::ZeroVector, FVector(1, 0, 0), FVector(0, 1, 0), FVector(0, 0, 1), FColor::White);
+							MeshBuilder.AddTriangle(VertexOffset, VertexOffset + VertexIdx, VertexOffset + VertexIdx - 1);
+							Vertex1 = Vertex2;
+						}
+
+						// Increment the vertex offset so the next polygon uses the correct vertex indices.
+						VertexOffset += Poly->Vertices.Num();
+					}
+				}
+
+				// Allocate the material proxy and register it so it can be deleted properly once the rendering is done with it.
+				FDynamicColoredMaterialRenderProxy* MaterialProxy = new FDynamicColoredMaterialRenderProxy(GEngine->EditorBrushMaterial->GetRenderProxy(false), Brush->GetWireColor());
+				PDI->RegisterDynamicResource(MaterialProxy);
+
+				// Flush the mesh triangles.
+				MeshBuilder.Draw(PDI, Brush->ActorToWorld().ToMatrixWithScale(), MaterialProxy, SDPG_World, 0.f);
+			}
+		}
+	}
+	
+	if (ModeTools->ShouldDrawBrushVertices() && !IsInGameView())
+	{
+		UTexture2D* VertexTexture = GEngine->DefaultBSPVertexTexture;
+		const float TextureSizeX = VertexTexture->GetSizeX() * 0.170f;
+		const float TextureSizeY = VertexTexture->GetSizeY() * 0.170f;
+
+		for (FSelectionIterator It(*ModeTools->GetSelectedActors()); It; ++It)
+		{
+			AActor* SelectedActor = static_cast<AActor*>(*It);
+			checkSlow(SelectedActor->IsA(AActor::StaticClass()));
+
+			ABrush* Brush = Cast< ABrush >(SelectedActor);
+			if (Brush && Brush->Brush && !FActorEditorUtils::IsABuilderBrush(Brush))
+			{
+				for (int32 p = 0; p < Brush->Brush->Polys->Element.Num(); ++p)
+				{
+					FTransform BrushTransform = Brush->ActorToWorld();
+
+					FPoly* poly = &Brush->Brush->Polys->Element[p];
+					for (int32 VertexIndex = 0; VertexIndex < poly->Vertices.Num(); ++VertexIndex)
+					{
+						const FVector& PolyVertex = poly->Vertices[VertexIndex];
+						const FVector WorldLocation = BrushTransform.TransformPosition(PolyVertex);
+
+						const float Scale = View->WorldToScreen(WorldLocation).W * (4.0f / View->ViewRect.Width() / View->ViewMatrices.ProjMatrix.M[0][0]);
+
+						const FColor Color(Brush->GetWireColor());
+						PDI->SetHitProxy(new HBSPBrushVert(Brush, &poly->Vertices[VertexIndex]));
+
+						PDI->DrawSprite(WorldLocation, TextureSizeX * Scale, TextureSizeY * Scale, VertexTexture->Resource, Color, SDPG_World, 0.0f, 0.0f, 0.0f, 0.0f, SE_BLEND_Masked);
+
+						PDI->SetHitProxy(NULL);
+					}
+				}
+			}
+		}
+	}
+}
+
 /**
  * Updates the audio listener for this viewport 
  *
diff --git a/Engine/Source/Editor/UnrealEd/Public/LevelEditorViewport.h b/Engine/Source/Editor/UnrealEd/Public/LevelEditorViewport.h
--- a/Engine/Source/Editor/UnrealEd/Public/LevelEditorViewport.h
+++ b/Engine/Source/Editor/UnrealEd/Public/LevelEditorViewport.h
@@ -143,12 +143,12 @@ public:
 	virtual void Tick(float DeltaSeconds) override;
 	virtual bool InputWidgetDelta( FViewport* Viewport, EAxisList::Type CurrentAxis, FVector& Drag, FRotator& Rot, FVector& Scale ) override;
 	virtual TSharedPtr<FDragTool> MakeDragTool( EDragTool::Type DragToolType ) override;
-	virtual bool IsLevelEditorClient() const { return ParentLevelEditor.IsValid(); }
+	virtual bool IsLevelEditorClient() const override { return ParentLevelEditor.IsValid(); }
 	virtual void TrackingStarted( const struct FInputEventState& InInputState, bool bIsDraggingWidget, bool bNudge ) override;
 	virtual void TrackingStopped() override;
 	virtual void AbortTracking() override;
 	virtual FVector GetWidgetLocation() const override;
-	virtual FMatrix GetWidgetCoordSystem() const;
+	virtual FMatrix GetWidgetCoordSystem() const override;
 	virtual void SetupViewForRendering( FSceneViewFamily& ViewFamily, FSceneView& View ) override;
 	virtual FLinearColor GetBackgroundColor() const override;
 	virtual int32 GetCameraSpeedSetting() const override;
@@ -610,6 +610,9 @@ private:
 	/** Project the specified actors into the world according to the current drag parameters */
 	void ProjectActorsIntoWorld(const TArray<AActor*>& Actors, FViewport* Viewport, const FVector& Drag, const FRotator& Rot);
 
+	/** Draw additional details for brushes in the world */
+	void DrawBrushDetails(const FSceneView* View, FPrimitiveDrawInterface* PDI);
+
 public:
 	/** Static: List of objects we're hovering over */
 	static TSet< FViewportHoverTarget > HoveredObjects;
