fixed cvar error

diff --git a/Engine/Source/Editor/LevelEditor/Private/LevelEditorActions.cpp b/Engine/Source/Editor/LevelEditor/Private/LevelEditorActions.cpp
--- a/Engine/Source/Editor/LevelEditor/Private/LevelEditorActions.cpp
+++ b/Engine/Source/Editor/LevelEditor/Private/LevelEditorActions.cpp
@@ -457,7 +457,7 @@ void FLevelEditorActionCallbacks::SetMaterialQualityLevel( EMaterialQualityLevel
 
 	//Ensure the material quality cvar is also set.
 	static IConsoleVariable* MaterialQualityLevelVar = IConsoleManager::Get().FindConsoleVariable(TEXT("r.MaterialQualityLevel"));
-	MaterialQualityLevelVar->Set(NewQualityLevel, ECVF_SetByConsole);
+	MaterialQualityLevelVar->Set(NewQualityLevel, ECVF_SetByScalability);
 
 	GUnrealEd->RedrawAllViewports();
 }
diff --git a/Engine/Source/Editor/UnrealEd/Private/Editor/EditorUserSettings.cpp b/Engine/Source/Editor/UnrealEd/Private/Editor/EditorUserSettings.cpp
--- a/Engine/Source/Editor/UnrealEd/Private/Editor/EditorUserSettings.cpp
+++ b/Engine/Source/Editor/UnrealEd/Private/Editor/EditorUserSettings.cpp
@@ -21,7 +21,7 @@ void UEditorUserSettings::PostInitProperties()
 
 	//Ensure the material quality cvar is set to the settings loaded.
 	static IConsoleVariable* MaterialQualityLevelVar = IConsoleManager::Get().FindConsoleVariable(TEXT("r.MaterialQualityLevel"));
-	MaterialQualityLevelVar->Set(MaterialQualityLevel, ECVF_SetByConsole);
+	MaterialQualityLevelVar->Set(MaterialQualityLevel, ECVF_SetByScalability);
 }
 
 void UEditorUserSettings::PostEditChangeProperty( FPropertyChangedEvent& PropertyChangedEvent )
