Updated editor viewports & editor viewport clients to allow for inertial movement and animated/lerped transitions when Slate is asleep
- SEditorViewportClient now has an additional optional parameter in its constructor for the SEditorViewport that the client is attached to
- SBasePoseViewport, SDestructibleMeshEditorViewport, and SPhATPreviewViewport were all updated to inherit from SEditorViewport (previously just inherited from SCompoundWidget)

diff --git a/Engine/Plugins/2D/Paper2D/Source/Paper2DEditor/Private/PaperEditorViewportClient.cpp b/Engine/Plugins/2D/Paper2D/Source/Paper2DEditor/Private/PaperEditorViewportClient.cpp
--- a/Engine/Plugins/2D/Paper2D/Source/Paper2DEditor/Private/PaperEditorViewportClient.cpp
+++ b/Engine/Plugins/2D/Paper2D/Source/Paper2DEditor/Private/PaperEditorViewportClient.cpp
@@ -4,6 +4,7 @@
 #include "TileSetEditor.h"
 #include "PaperEditorViewportClient.h"
 #include "CanvasTypes.h"
+#include "SEditorViewport.h"
 
 //////////////////////////////////////////////////////////////////////////
 // FAssetEditorModeTools
@@ -40,8 +41,8 @@ USelection* FAssetEditorModeTools::GetSelectedObjects() const
 //////////////////////////////////////////////////////////////////////////
 // FPaperEditorViewportClient
 
-FPaperEditorViewportClient::FPaperEditorViewportClient()
-	: FEditorViewportClient(new FAssetEditorModeTools())
+FPaperEditorViewportClient::FPaperEditorViewportClient(const TWeakPtr<SEditorViewport>& InEditorViewportWidget)
+	: FEditorViewportClient(new FAssetEditorModeTools(), nullptr, InEditorViewportWidget)
 	, CheckerboardTexture(nullptr)
 {
 	bOwnsModeTools = true;
diff --git a/Engine/Plugins/2D/Paper2D/Source/Paper2DEditor/Private/PaperEditorViewportClient.h b/Engine/Plugins/2D/Paper2D/Source/Paper2DEditor/Private/PaperEditorViewportClient.h
--- a/Engine/Plugins/2D/Paper2D/Source/Paper2DEditor/Private/PaperEditorViewportClient.h
+++ b/Engine/Plugins/2D/Paper2D/Source/Paper2DEditor/Private/PaperEditorViewportClient.h
@@ -37,7 +37,7 @@ class FPaperEditorViewportClient : public FEditorViewportClient
 {
 public:
 	/** Constructor */
-	FPaperEditorViewportClient();
+	explicit FPaperEditorViewportClient(const TWeakPtr<class SEditorViewport>& InEditorViewportWidget = nullptr);
 	~FPaperEditorViewportClient();
 
 	// FViewportClient interface
diff --git a/Engine/Source/Developer/LogVisualizer/Private/STimelineBar.cpp b/Engine/Source/Developer/LogVisualizer/Private/STimelineBar.cpp
--- a/Engine/Source/Developer/LogVisualizer/Private/STimelineBar.cpp
+++ b/Engine/Source/Developer/LogVisualizer/Private/STimelineBar.cpp
@@ -88,7 +88,7 @@ FReply STimelineBar::OnMouseButtonDoubleClick(const FGeometry& MyGeometry, const
 				float Radius = BoundingBox.GetExtent().Size();
 
 				FViewportCameraTransform ViewTransform;
-				ViewTransform.TransitionToLocation(Position, true);
+				ViewTransform.TransitionToLocation(Position, SharedThis(this), true);
 
 				float NewOrthoZoom;
 				const float AspectRatio = 1.777777f;
diff --git a/Engine/Source/Developer/MovieSceneTools/Private/DirectorTrackEditor.cpp b/Engine/Source/Developer/MovieSceneTools/Private/DirectorTrackEditor.cpp
--- a/Engine/Source/Developer/MovieSceneTools/Private/DirectorTrackEditor.cpp
+++ b/Engine/Source/Developer/MovieSceneTools/Private/DirectorTrackEditor.cpp
@@ -169,14 +169,14 @@ FShotSection::FShotSection( TSharedPtr<ISequencer> InSequencer, TSharedPtr<FShot
 	, StoredSize(ForceInit)
 	, StoredStartTime(0.f)
 	, VisibleTimeRange(TRange<float>::Empty())
-	, InternalViewportScene(NULL)
-	, InternalViewportClient(NULL)
+	, InternalViewportScene(nullptr)
+	, InternalViewportClient(nullptr)
 {
 	if (Camera.IsValid())
 	{
 		// @todo Sequencer optimize This code shouldn't even be called if this is offscreen
 		// the following code could be generated on demand when the section is scrolled onscreen
-		InternalViewportClient = MakeShareable(new FLevelEditorViewportClient());
+		InternalViewportClient = MakeShareable(new FLevelEditorViewportClient(nullptr));
 		InternalViewportClient->ViewportType = LVT_Perspective;
 		InternalViewportClient->bDisableInput = true;
 		InternalViewportClient->bDrawAxes = false;
@@ -184,7 +184,7 @@ FShotSection::FShotSection( TSharedPtr<ISequencer> InSequencer, TSharedPtr<FShot
 		InternalViewportClient->EngineShowFlags.DisableAdvancedFeatures();
 		InternalViewportClient->SetActorLock(Camera.Get());
 
-		InternalViewportScene = MakeShareable(new FSceneViewport(InternalViewportClient.Get(), NULL));
+		InternalViewportScene = MakeShareable(new FSceneViewport(InternalViewportClient.Get(), nullptr));
 		InternalViewportClient->Viewport = InternalViewportScene.Get();
 
 		CalculateThumbnailWidthAndResize();
@@ -197,7 +197,7 @@ FShotSection::~FShotSection()
 
 	if (InternalViewportClient.IsValid())
 	{
-		InternalViewportClient->Viewport = NULL;
+		InternalViewportClient->Viewport = nullptr;
 	}
 }
 
diff --git a/Engine/Source/Editor/Cascade/Private/CascadePreviewViewportClient.cpp b/Engine/Source/Editor/Cascade/Private/CascadePreviewViewportClient.cpp
--- a/Engine/Source/Editor/Cascade/Private/CascadePreviewViewportClient.cpp
+++ b/Engine/Source/Editor/Cascade/Private/CascadePreviewViewportClient.cpp
@@ -21,10 +21,9 @@
 
 #define LOCTEXT_NAMESPACE "CascadeViewportClient"
 
-FCascadeEdPreviewViewportClient::FCascadeEdPreviewViewportClient(TWeakPtr<FCascade> InCascade, TWeakPtr<SCascadePreviewViewport> InCascadeViewport)
-	: FEditorViewportClient(nullptr)
+FCascadeEdPreviewViewportClient::FCascadeEdPreviewViewportClient(TWeakPtr<FCascade> InCascade, const TSharedRef<SCascadePreviewViewport>& InCascadeViewport)
+	: FEditorViewportClient(nullptr, nullptr, StaticCastSharedRef<SEditorViewport>(InCascadeViewport))
 	, CascadePtr(InCascade)
-	, CascadeViewportPtr(InCascadeViewport)
 	, CascadePreviewScene(FPreviewScene::ConstructionValues()
 						.SetLightRotation(FRotator(-45.f, 180.f, 0.f))
 						.SetSkyBrightness(0.25f)
@@ -33,7 +32,7 @@ FCascadeEdPreviewViewportClient::FCascadeEdPreviewViewportClient(TWeakPtr<FCasca
 	, LightRotSpeed(0.22f)
 {
 
-	check(CascadePtr.IsValid() && CascadeViewportPtr.IsValid());
+	check(CascadePtr.IsValid() && EditorViewportWidget.IsValid());
 
 	UParticleSystem* ParticleSystem = CascadePtr.Pin()->GetParticleSystem();
 	UCascadeParticleSystemComponent* ParticleSystemComponent = CascadePtr.Pin()->GetParticleSystemComponent();
diff --git a/Engine/Source/Editor/Cascade/Private/CascadePreviewViewportClient.h b/Engine/Source/Editor/Cascade/Private/CascadePreviewViewportClient.h
--- a/Engine/Source/Editor/Cascade/Private/CascadePreviewViewportClient.h
+++ b/Engine/Source/Editor/Cascade/Private/CascadePreviewViewportClient.h
@@ -15,7 +15,7 @@ class FCascadeEdPreviewViewportClient : public FEditorViewportClient
 {
 public:
 	/** Constructor */
-	FCascadeEdPreviewViewportClient(TWeakPtr<FCascade> InCascade, TWeakPtr<SCascadePreviewViewport> InCascadeViewport);
+	FCascadeEdPreviewViewportClient(TWeakPtr<FCascade> InCascade, const TSharedRef<SCascadePreviewViewport>& InCascadeViewport);
 	~FCascadeEdPreviewViewportClient();
 
 	/** FEditorViewportClient interface */
@@ -64,9 +64,6 @@ public:
 private:
 	/** Pointer back to the ParticleSystem editor tool that owns us */
 	TWeakPtr<FCascade> CascadePtr;
-
-	/** Pointer back to the ParticleSystem viewport control that owns us */
-	TWeakPtr<SCascadePreviewViewport> CascadeViewportPtr;
 	
 	/** Preview mesh */
 	UStaticMeshComponent* FloorComponent;
diff --git a/Engine/Source/Editor/DestructibleMeshEditor/Private/SDestructibleMeshEditorViewport.cpp b/Engine/Source/Editor/DestructibleMeshEditor/Private/SDestructibleMeshEditorViewport.cpp
--- a/Engine/Source/Editor/DestructibleMeshEditor/Private/SDestructibleMeshEditorViewport.cpp
+++ b/Engine/Source/Editor/DestructibleMeshEditor/Private/SDestructibleMeshEditorViewport.cpp
@@ -29,7 +29,7 @@ protected:
 	TWeakObjectPtr<UDestructibleComponent> PreviewDestructibleComp;
 
 public:
-	FDestructibleMeshEditorViewportClient(TWeakPtr<IDestructibleMeshEditor> InDestructibleMeshEditor, FPreviewScene& InPreviewScene);
+	FDestructibleMeshEditorViewportClient(TWeakPtr<IDestructibleMeshEditor> InDestructibleMeshEditor, FPreviewScene& InPreviewScene, const TSharedRef<SDestructibleMeshEditorViewport>& InDestructibleMeshEditorViewport);
 
 	// FGCObject interface
 	virtual void AddReferencedObjects( FReferenceCollector& Collector ) override;
@@ -81,8 +81,8 @@ private:
 	TArray<class UDestructibleChunkParamsProxy*> UnusedProxies;
 };
 
-FDestructibleMeshEditorViewportClient::FDestructibleMeshEditorViewportClient(TWeakPtr<IDestructibleMeshEditor> InDestructibleMeshEditor, FPreviewScene& InPreviewScene)
-	: FEditorViewportClient(nullptr, &InPreviewScene)
+FDestructibleMeshEditorViewportClient::FDestructibleMeshEditorViewportClient(TWeakPtr<IDestructibleMeshEditor> InDestructibleMeshEditor, FPreviewScene& InPreviewScene, const TSharedRef<SDestructibleMeshEditorViewport>& InDestructibleMeshEditorViewport)
+	: FEditorViewportClient(nullptr, &InPreviewScene, StaticCastSharedRef<SEditorViewport>(InDestructibleMeshEditorViewport))
 	, DestructibleMeshEditorPtr(InDestructibleMeshEditor)
 {
 	SetViewMode(VMI_Lit);
@@ -479,41 +479,15 @@ void FDestructibleMeshEditorViewportClient::ImportFBXChunks()
 void SDestructibleMeshEditorViewport::Construct(const FArguments& InArgs)
 {
 	DestructibleMeshEditorPtr = InArgs._DestructibleMeshEditor;
-
 	CurrentViewMode = VMI_Lit;
 
-	this->ChildSlot
-		[
-			SNew(SVerticalBox)
-			+SVerticalBox::Slot()
-			.FillHeight(1)
-			[
-				SAssignNew( ViewportWidget, SViewport )
-				.EnableGammaCorrection(false)
-				.IsEnabled( FSlateApplication::Get().GetNormalExecutionAttribute() )
-				.ShowEffectWhenDisabled( false )
-			]
-		];
-
-	EditorViewportClient = MakeShareable( new FDestructibleMeshEditorViewportClient(DestructibleMeshEditorPtr, PreviewScene) );
-
-	EditorViewportClient->bSetListenerPosition = false;
-
-	EditorViewportClient->SetRealtime( false );
-	EditorViewportClient->VisibilityDelegate.BindSP( this, &SDestructibleMeshEditorViewport::IsVisible );
-
-	Viewport = MakeShareable( new FSceneViewport( EditorViewportClient.Get(), ViewportWidget ) );
-	EditorViewportClient->Viewport = Viewport.Get();
-
-	// The viewport widget needs an interface so it knows what should render
-	ViewportWidget->SetViewportInterface( Viewport.ToSharedRef() );
+	SEditorViewport::Construct(SEditorViewport::FArguments());
 
 	PreviewComponent = ConstructObject<UDestructibleComponent>(
 		UDestructibleComponent::StaticClass(), GetTransientPackage(), NAME_None, RF_Transient );
 
 	SetPreviewMesh(InArgs._ObjectToEdit);
 
-	BindCommands();
 	EditorViewportClient->BindCommands();
 
 	PreviewDepth = 0;
@@ -605,7 +579,7 @@ void SDestructibleMeshEditorViewport::RefreshViewport()
 #endif // WITH_APEX
 
 	// Invalidate the viewport's display.
-	Viewport->InvalidateDisplay();
+	SceneViewport->InvalidateDisplay();
 }
 
 void SDestructibleMeshEditorViewport::SetPreviewMesh(UDestructibleMesh* InDestructibleMesh)
@@ -690,10 +664,6 @@ UDestructibleComponent* SDestructibleMeshEditorViewport::GetDestructibleComponen
 	return PreviewComponent;
 }
 
-void SDestructibleMeshEditorViewport::BindCommands()
-{
-}
-
 void SDestructibleMeshEditorViewport::SetViewModeWireframe()
 {
 	if(CurrentViewMode != VMI_Wireframe)
@@ -706,7 +676,7 @@ void SDestructibleMeshEditorViewport::SetViewModeWireframe()
 	}
 
 	EditorViewportClient->SetViewMode(CurrentViewMode);
-	Viewport->Invalidate();
+	SceneViewport->Invalidate();
 
 }
 
@@ -714,3 +684,25 @@ bool SDestructibleMeshEditorViewport::IsInViewModeWireframeChecked() const
 {
 	return CurrentViewMode == VMI_Wireframe;
 }
+
+TSharedRef<FEditorViewportClient> SDestructibleMeshEditorViewport::MakeEditorViewportClient()
+{
+	EditorViewportClient = MakeShareable( new FDestructibleMeshEditorViewportClient(DestructibleMeshEditorPtr, PreviewScene, SharedThis(this)) );
+
+	EditorViewportClient->bSetListenerPosition = false;
+
+	EditorViewportClient->SetRealtime( false );
+	EditorViewportClient->VisibilityDelegate.BindSP( this, &SDestructibleMeshEditorViewport::IsVisible );
+
+	return EditorViewportClient.ToSharedRef();
+}
+
+TSharedPtr<SWidget> SDestructibleMeshEditorViewport::MakeViewportToolbar()
+{
+	return nullptr;
+}
+
+void SDestructibleMeshEditorViewport::BindCommands()
+{
+	// No commands. Overridden to prevent the base SEditorViewport commands from being bound.
+}
\ No newline at end of file
diff --git a/Engine/Source/Editor/DestructibleMeshEditor/Private/SDestructibleMeshEditorViewport.h b/Engine/Source/Editor/DestructibleMeshEditor/Private/SDestructibleMeshEditorViewport.h
--- a/Engine/Source/Editor/DestructibleMeshEditor/Private/SDestructibleMeshEditorViewport.h
+++ b/Engine/Source/Editor/DestructibleMeshEditor/Private/SDestructibleMeshEditorViewport.h
@@ -3,11 +3,12 @@
 #pragma once
 
 #include "PreviewScene.h"
+#include "SEditorViewport.h"
 
 /**
  * DestructibleMesh Editor Preview viewport widget
  */
-class SDestructibleMeshEditorViewport : public SCompoundWidget, public FGCObject, public FNotifyHook
+class SDestructibleMeshEditorViewport : public SEditorViewport, public FGCObject, public FNotifyHook
 {
 public:
 	SLATE_BEGIN_ARGS( SDestructibleMeshEditorViewport ){}
@@ -34,9 +35,6 @@ public:
 	/** The parent tab where this viewport resides */
 	TWeakPtr<SDockableTab> ParentTab;
 
-	/** Binds commands associated with the viewport client. */
-	void BindCommands();
-
 	/** 
 	 *	Causes chunks at the given PreviewDepth to be displayed in the viewport.  Clamped to the range [0, depth count), where (depth count) = the number of chunk hierarchy depths in the destructible.
 	 *
@@ -62,12 +60,18 @@ public:
 	void SetPreviewMesh(UDestructibleMesh* InDestructibleMesh);
 
 	/**
-	 *	Updates the preview mesh and other viewport specfic settings that go with it.
+	 *	Updates the preview mesh and other viewport specific settings that go with it.
 	 *
 	 *	@param	InDestructibleMesh		The Destructible mesh being viewed in the editor.
 	 */
 	void UpdatePreviewMesh(UDestructibleMesh* InDestructibleMesh);
 
+protected:
+	/** SEditorViewport interface */
+	virtual TSharedRef<FEditorViewportClient> MakeEditorViewportClient() override;
+	virtual TSharedPtr<SWidget> MakeViewportToolbar() override;
+	virtual void BindCommands() override;
+
 private:
 
 	/** Determines the visibility of the viewport. */
@@ -89,12 +93,6 @@ private:
 	/** Level viewport client */
 	TSharedPtr<class FDestructibleMeshEditorViewportClient> EditorViewportClient;
 
-	/** Slate viewport for rendering and I/O */
-	TSharedPtr<class FSceneViewport> Viewport;
-
-	/** Viewport widget */
-	TSharedPtr<class SViewport> ViewportWidget;
-
 	/** The currently selected view mode. */
 	EViewModeIndex CurrentViewMode;
 
diff --git a/Engine/Source/Editor/DetailCustomizations/Private/AnimMontageSegmentDetails.cpp b/Engine/Source/Editor/DetailCustomizations/Private/AnimMontageSegmentDetails.cpp
--- a/Engine/Source/Editor/DetailCustomizations/Private/AnimMontageSegmentDetails.cpp
+++ b/Engine/Source/Editor/DetailCustomizations/Private/AnimMontageSegmentDetails.cpp
@@ -8,8 +8,8 @@
 #define LOCTEXT_NAMESPACE "AnimMontageSegmentDetails"
 
 /////////////////////////////////////////////////////////////////////////
-FAnimationSegmentViewportClient::FAnimationSegmentViewportClient(FPreviewScene& InPreviewScene)
-	: FEditorViewportClient(nullptr, &InPreviewScene)
+FAnimationSegmentViewportClient::FAnimationSegmentViewportClient(FPreviewScene& InPreviewScene, const TWeakPtr<SEditorViewport>& InEditorViewportWidget)
+	: FEditorViewportClient(nullptr, &InPreviewScene, InEditorViewportWidget)
 {
 	SetViewMode(VMI_Lit);
 
diff --git a/Engine/Source/Editor/DetailCustomizations/Private/AnimMontageSegmentDetails.h b/Engine/Source/Editor/DetailCustomizations/Private/AnimMontageSegmentDetails.h
--- a/Engine/Source/Editor/DetailCustomizations/Private/AnimMontageSegmentDetails.h
+++ b/Engine/Source/Editor/DetailCustomizations/Private/AnimMontageSegmentDetails.h
@@ -154,7 +154,7 @@ protected:
 class FAnimationSegmentViewportClient : public FEditorViewportClient
 {
 public:
-	FAnimationSegmentViewportClient(FPreviewScene& InPreviewScene);
+	FAnimationSegmentViewportClient(FPreviewScene& InPreviewScene, const TWeakPtr<SEditorViewport>& InEditorViewportWidget = nullptr);
 
 	// FlEditorViewportClient interface
 	virtual FSceneInterface* GetScene() const override;
diff --git a/Engine/Source/Editor/Kismet/Private/SCSEditorViewportClient.cpp b/Engine/Source/Editor/Kismet/Private/SCSEditorViewportClient.cpp
--- a/Engine/Source/Editor/Kismet/Private/SCSEditorViewportClient.cpp
+++ b/Engine/Source/Editor/Kismet/Private/SCSEditorViewportClient.cpp
@@ -24,6 +24,7 @@
 #include "Engine/TextureCube.h"
 #include "Components/InstancedStaticMeshComponent.h"
 #include "Engine/StaticMesh.h"
+#include "SSCSEditorViewport.h"
 
 DEFINE_LOG_CATEGORY_STATIC(LogSCSEditorViewport, Log, All);
 
@@ -105,8 +106,8 @@ namespace
 /////////////////////////////////////////////////////////////////////////
 // FSCSEditorViewportClient
 
-FSCSEditorViewportClient::FSCSEditorViewportClient(TWeakPtr<FBlueprintEditor>& InBlueprintEditorPtr, FPreviewScene& InPreviewScene)
-	: FEditorViewportClient(nullptr, &InPreviewScene)
+FSCSEditorViewportClient::FSCSEditorViewportClient(TWeakPtr<FBlueprintEditor>& InBlueprintEditorPtr, FPreviewScene& InPreviewScene, const TSharedRef<SSCSEditorViewport>& InSCSEditorViewport)
+	: FEditorViewportClient(nullptr, &InPreviewScene, StaticCastSharedRef<SEditorViewport>(InSCSEditorViewport))
 	,BlueprintEditorPtr(InBlueprintEditorPtr)
 	,PreviewBlueprint(NULL)
 	,PreviewActorBounds(ForceInitToZero)
diff --git a/Engine/Source/Editor/Kismet/Private/SCSEditorViewportClient.h b/Engine/Source/Editor/Kismet/Private/SCSEditorViewportClient.h
--- a/Engine/Source/Editor/Kismet/Private/SCSEditorViewportClient.h
+++ b/Engine/Source/Editor/Kismet/Private/SCSEditorViewportClient.h
@@ -16,7 +16,7 @@ public:
 	 * @param InBlueprintEditorPtr A weak reference to the Blueprint Editor context.
 	 * @param InPreviewScene The preview scene to use.
 	 */
-	FSCSEditorViewportClient(TWeakPtr<class FBlueprintEditor>& InBlueprintEditorPtr, FPreviewScene& InPreviewScene);
+	FSCSEditorViewportClient(TWeakPtr<class FBlueprintEditor>& InBlueprintEditorPtr, FPreviewScene& InPreviewScene, const TSharedRef<SSCSEditorViewport>& InSCSEditorViewport);
 
 	/**
 	 * Destructor.
diff --git a/Engine/Source/Editor/Kismet/Private/SSCSEditorViewport.cpp b/Engine/Source/Editor/Kismet/Private/SSCSEditorViewport.cpp
--- a/Engine/Source/Editor/Kismet/Private/SSCSEditorViewport.cpp
+++ b/Engine/Source/Editor/Kismet/Private/SSCSEditorViewport.cpp
@@ -223,7 +223,7 @@ EVisibility SSCSEditorViewport::GetWidgetVisibility() const
 TSharedRef<FEditorViewportClient> SSCSEditorViewport::MakeEditorViewportClient()
 {
 	// Construct a new viewport client instance.
-	ViewportClient = MakeShareable(new FSCSEditorViewportClient(BlueprintEditorPtr, PreviewScene));
+	ViewportClient = MakeShareable(new FSCSEditorViewportClient(BlueprintEditorPtr, PreviewScene, SharedThis(this)));
 	ViewportClient->SetRealtime(true);
 	ViewportClient->bSetListenerPosition = false;
 	ViewportClient->VisibilityDelegate.BindSP(this, &SSCSEditorViewport::IsVisible);
diff --git a/Engine/Source/Editor/LevelEditor/Private/SLevelViewport.cpp b/Engine/Source/Editor/LevelEditor/Private/SLevelViewport.cpp
--- a/Engine/Source/Editor/LevelEditor/Private/SLevelViewport.cpp
+++ b/Engine/Source/Editor/LevelEditor/Private/SLevelViewport.cpp
@@ -286,7 +286,7 @@ void SLevelViewport::ConstructViewportOverlayContent()
 
 void SLevelViewport::ConstructLevelEditorViewportClient( const FArguments& InArgs )
 {
-	LevelViewportClient = MakeShareable( new FLevelEditorViewportClient() );
+	LevelViewportClient = MakeShareable( new FLevelEditorViewportClient(SharedThis(this)) );
 
 	// Default level viewport client values for settings that could appear in layout config ini
 	FLevelEditorViewportInstanceSettings ViewportInstanceSettings;
@@ -2953,7 +2953,7 @@ void SLevelViewport::PreviewActors( const TArray< AActor* >& ActorsToPreview )
 		{
 			auto CurActor = *ActorIt;
 
-			TSharedPtr< FLevelEditorViewportClient > ActorPreviewLevelViewportClient = MakeShareable( new FLevelEditorViewportClient() );
+			TSharedPtr< FLevelEditorViewportClient > ActorPreviewLevelViewportClient = MakeShareable( new FLevelEditorViewportClient(SharedThis(this)) );
 			{
 				// NOTE: We don't bother setting ViewLocation, ViewRotation, etc, here.  This is because we'll call
 				//       PushControllingActorDataToViewportClient() below which will do this!
diff --git a/Engine/Source/Editor/MaterialEditor/Private/SMaterialEditorViewport.cpp b/Engine/Source/Editor/MaterialEditor/Private/SMaterialEditorViewport.cpp
--- a/Engine/Source/Editor/MaterialEditor/Private/SMaterialEditorViewport.cpp
+++ b/Engine/Source/Editor/MaterialEditor/Private/SMaterialEditorViewport.cpp
@@ -19,7 +19,7 @@
 class FMaterialEditorViewportClient : public FEditorViewportClient
 {
 public:
-	FMaterialEditorViewportClient(TWeakPtr<IMaterialEditor> InMaterialEditor, FPreviewScene& InPreviewScene);
+	FMaterialEditorViewportClient(TWeakPtr<IMaterialEditor> InMaterialEditor, FPreviewScene& InPreviewScene, const TSharedRef<SMaterialEditorViewport>& InMaterialEditorViewport);
 
 	// FEditorViewportClient interface
 	virtual FLinearColor GetBackgroundColor() const override;
@@ -44,8 +44,8 @@ private:
 	TWeakPtr<IMaterialEditor> MaterialEditorPtr;
 };
 
-FMaterialEditorViewportClient::FMaterialEditorViewportClient(TWeakPtr<IMaterialEditor> InMaterialEditor, FPreviewScene& InPreviewScene)
-	: FEditorViewportClient( nullptr, &InPreviewScene )
+FMaterialEditorViewportClient::FMaterialEditorViewportClient(TWeakPtr<IMaterialEditor> InMaterialEditor, FPreviewScene& InPreviewScene, const TSharedRef<SMaterialEditorViewport>& InMaterialEditorViewport)
+	: FEditorViewportClient(nullptr, &InPreviewScene, StaticCastSharedRef<SEditorViewport>(InMaterialEditorViewport))
 	, MaterialEditorPtr(InMaterialEditor)
 {
 	// Setup defaults for the common draw helper.
@@ -171,7 +171,7 @@ void FMaterialEditorViewportClient::FocusViewportOnBounds(const FBoxSphereBounds
 	FVector CameraOffsetVector = ViewTransform.GetRotation().Vector() * -DistanceFromSphere;
 
 	ViewTransform.SetLookAt(Position);
-	ViewTransform.TransitionToLocation(Position + CameraOffsetVector, bInstant);
+	ViewTransform.TransitionToLocation(Position + CameraOffsetVector, EditorViewportWidget, bInstant);
 
 	// Tell the viewport to redraw itself.
 	Invalidate();
@@ -504,7 +504,7 @@ bool SMaterialEditorViewport::IsTogglePreviewBackgroundChecked() const
 
 TSharedRef<FEditorViewportClient> SMaterialEditorViewport::MakeEditorViewportClient() 
 {
-	EditorViewportClient = MakeShareable( new FMaterialEditorViewportClient(MaterialEditorPtr, PreviewScene) );
+	EditorViewportClient = MakeShareable( new FMaterialEditorViewportClient(MaterialEditorPtr, PreviewScene, SharedThis(this)) );
 	
 	EditorViewportClient->SetViewLocation( FVector::ZeroVector );
 	EditorViewportClient->SetViewRotation( FRotator::ZeroRotator );
diff --git a/Engine/Source/Editor/Matinee/Private/SMatineeRecorder.cpp b/Engine/Source/Editor/Matinee/Private/SMatineeRecorder.cpp
--- a/Engine/Source/Editor/Matinee/Private/SMatineeRecorder.cpp
+++ b/Engine/Source/Editor/Matinee/Private/SMatineeRecorder.cpp
@@ -112,7 +112,7 @@ void SMatineeRecorder::Construct(const FArguments& InArgs)
 	CameraModeComboBox->ComputeDesiredSize();
 
 	// Create an animation viewport client
-	LevelViewportClient = MakeShareable( new FLevelEditorViewportClient );
+	LevelViewportClient = MakeShareable( new FLevelEditorViewportClient(nullptr) );
 
 	Viewport = MakeShareable( new FSceneViewport( LevelViewportClient.Get(), ViewportWidget ) );
 	LevelViewportClient->Viewport = Viewport.Get();
diff --git a/Engine/Source/Editor/NiagaraEditor/Private/SNiagaraEffectEditorViewport.cpp b/Engine/Source/Editor/NiagaraEditor/Private/SNiagaraEffectEditorViewport.cpp
--- a/Engine/Source/Editor/NiagaraEditor/Private/SNiagaraEffectEditorViewport.cpp
+++ b/Engine/Source/Editor/NiagaraEditor/Private/SNiagaraEffectEditorViewport.cpp
@@ -17,7 +17,7 @@
 class FNiagaraEffectEditorViewportClient : public FEditorViewportClient
 {
 public:
-	FNiagaraEffectEditorViewportClient(TWeakPtr<INiagaraEffectEditor> InEffectEditor, FPreviewScene& InPreviewScene);
+	FNiagaraEffectEditorViewportClient(TWeakPtr<INiagaraEffectEditor> InEffectEditor, FPreviewScene& InPreviewScene, const TSharedRef<SNiagaraEffectEditorViewport>& InNiagaraEditorViewport);
 	
 	// FEditorViewportClient interface
 	virtual FLinearColor GetBackgroundColor() const override;
@@ -34,9 +34,9 @@ private:
 	TWeakPtr<INiagaraEffectEditor> EffectEditorPtr;
 };
 
-FNiagaraEffectEditorViewportClient::FNiagaraEffectEditorViewportClient(TWeakPtr<INiagaraEffectEditor> InEffectEditor, FPreviewScene& InPreviewScene)
-: FEditorViewportClient( nullptr, &InPreviewScene )
-, EffectEditorPtr(InEffectEditor)
+FNiagaraEffectEditorViewportClient::FNiagaraEffectEditorViewportClient(TWeakPtr<INiagaraEffectEditor> InEffectEditor, FPreviewScene& InPreviewScene, const TSharedRef<SNiagaraEffectEditorViewport>& InNiagaraEditorViewport)
+	: FEditorViewportClient(nullptr, &InPreviewScene, StaticCastSharedRef<SEditorViewport>(InNiagaraEditorViewport))
+	, EffectEditorPtr(InEffectEditor)
 {
 	// Setup defaults for the common draw helper.
 	DrawHelper.bDrawPivot = false;
@@ -239,7 +239,7 @@ bool SNiagaraEffectEditorViewport::IsTogglePreviewBackgroundChecked() const
 
 TSharedRef<FEditorViewportClient> SNiagaraEffectEditorViewport::MakeEditorViewportClient() 
 {
-	EditorViewportClient = MakeShareable( new FNiagaraEffectEditorViewportClient(nullptr, PreviewScene) );
+	EditorViewportClient = MakeShareable( new FNiagaraEffectEditorViewportClient(nullptr, PreviewScene, SharedThis(this)) );
 	
 	EditorViewportClient->SetViewLocation( FVector::ZeroVector );
 	EditorViewportClient->SetViewRotation( FRotator::ZeroRotator );
diff --git a/Engine/Source/Editor/Persona/Private/AnimationEditorViewportClient.cpp b/Engine/Source/Editor/Persona/Private/AnimationEditorViewportClient.cpp
--- a/Engine/Source/Editor/Persona/Private/AnimationEditorViewportClient.cpp
+++ b/Engine/Source/Editor/Persona/Private/AnimationEditorViewportClient.cpp
@@ -28,6 +28,7 @@
 #include "PhysicsEngine/PhysicsSettings.h"
 #include "Components/WindDirectionalSourceComponent.h"
 #include "Engine/StaticMesh.h"
+#include "SAnimationEditorViewport.h"
 
 namespace {
 	// Value from UE3
@@ -78,8 +79,8 @@ IMPLEMENT_HIT_PROXY( HPersonaBoneProxy, HHitProxy );
 /////////////////////////////////////////////////////////////////////////
 // FAnimationViewportClient
 
-FAnimationViewportClient::FAnimationViewportClient( FPreviewScene& InPreviewScene, TWeakPtr<FPersona> InPersonaPtr )
-	: FEditorViewportClient(nullptr, &InPreviewScene)
+FAnimationViewportClient::FAnimationViewportClient(FPreviewScene& InPreviewScene, TWeakPtr<FPersona> InPersonaPtr, const TSharedRef<SAnimationEditorViewport>& InAnimationEditorViewport)
+	: FEditorViewportClient(nullptr, &InPreviewScene, StaticCastSharedRef<SEditorViewport>(InAnimationEditorViewport))
 	, PersonaPtr( InPersonaPtr )
 	, bManipulating(false)
 	, bInTransaction(false)
diff --git a/Engine/Source/Editor/Persona/Private/AnimationEditorViewportClient.h b/Engine/Source/Editor/Persona/Private/AnimationEditorViewportClient.h
--- a/Engine/Source/Editor/Persona/Private/AnimationEditorViewportClient.h
+++ b/Engine/Source/Editor/Persona/Private/AnimationEditorViewportClient.h
@@ -50,7 +50,7 @@ protected:
 	void DisplayInfo( FCanvas* Canvas, FSceneView* View, bool bDisplayAllInfo );
 
 public:
-	FAnimationViewportClient( FPreviewScene& InPreviewScene, TWeakPtr<FPersona> InPersonaPtr );
+	FAnimationViewportClient(FPreviewScene& InPreviewScene, TWeakPtr<FPersona> InPersonaPtr, const TSharedRef<SAnimationEditorViewport>& InAnimationEditorViewport);
 	virtual ~FAnimationViewportClient();
 
 	// FEditorViewportClient interface
diff --git a/Engine/Source/Editor/Persona/Private/SAnimationEditorViewport.cpp b/Engine/Source/Editor/Persona/Private/SAnimationEditorViewport.cpp
--- a/Engine/Source/Editor/Persona/Private/SAnimationEditorViewport.cpp
+++ b/Engine/Source/Editor/Persona/Private/SAnimationEditorViewport.cpp
@@ -56,7 +56,7 @@ void SAnimationEditorViewport::Construct(const FArguments& InArgs, TSharedPtr<cl
 TSharedRef<FEditorViewportClient> SAnimationEditorViewport::MakeEditorViewportClient()
 {
 	// Create an animation viewport client
-	LevelViewportClient = MakeShareable(new FAnimationViewportClient(PersonaPtr.Pin()->GetPreviewScene(), PersonaPtr));
+	LevelViewportClient = MakeShareable(new FAnimationViewportClient(PersonaPtr.Pin()->GetPreviewScene(), PersonaPtr, SharedThis(this)));
 
 	LevelViewportClient->ViewportType = LVT_Perspective;
 	LevelViewportClient->bSetListenerPosition = false;
diff --git a/Engine/Source/Editor/PhAT/Private/PhATPreviewViewportClient.cpp b/Engine/Source/Editor/PhAT/Private/PhATPreviewViewportClient.cpp
--- a/Engine/Source/Editor/PhAT/Private/PhATPreviewViewportClient.cpp
+++ b/Engine/Source/Editor/PhAT/Private/PhATPreviewViewportClient.cpp
@@ -19,8 +19,8 @@
 #include "PhysicsEngine/PhysicsHandleComponent.h"
 #include "DrawDebugHelpers.h"
 
-FPhATEdPreviewViewportClient::FPhATEdPreviewViewportClient(TWeakPtr<FPhAT> InPhAT, TSharedPtr<FPhATSharedData> Data)
-	: FEditorViewportClient(nullptr, &Data->PreviewScene)
+FPhATEdPreviewViewportClient::FPhATEdPreviewViewportClient(TWeakPtr<FPhAT> InPhAT, TSharedPtr<FPhATSharedData> Data, const TSharedRef<SPhATPreviewViewport>& InPhATPreviewViewport)
+	: FEditorViewportClient(nullptr, &Data->PreviewScene, StaticCastSharedRef<SEditorViewport>(InPhATPreviewViewport))
 	, PhATPtr(InPhAT)
 	, SharedData(Data)
 	, MinPrimSize(0.5f)
diff --git a/Engine/Source/Editor/PhAT/Private/PhATPreviewViewportClient.h b/Engine/Source/Editor/PhAT/Private/PhATPreviewViewportClient.h
--- a/Engine/Source/Editor/PhAT/Private/PhATPreviewViewportClient.h
+++ b/Engine/Source/Editor/PhAT/Private/PhATPreviewViewportClient.h
@@ -17,7 +17,7 @@ class FPhATEdPreviewViewportClient : public FEditorViewportClient
 {
 public:
 	/** Constructor */
-	FPhATEdPreviewViewportClient(TWeakPtr<FPhAT> InPhAT, TSharedPtr<FPhATSharedData> Data);
+	FPhATEdPreviewViewportClient(TWeakPtr<FPhAT> InPhAT, TSharedPtr<FPhATSharedData> Data, const TSharedRef<SPhATPreviewViewport>& InPhATPreviewViewport);
 	~FPhATEdPreviewViewportClient();
 
 	/** FEditorViewportClient interface */
diff --git a/Engine/Source/Editor/PhAT/Private/SPhATPreviewViewport.cpp b/Engine/Source/Editor/PhAT/Private/SPhATPreviewViewport.cpp
--- a/Engine/Source/Editor/PhAT/Private/SPhATPreviewViewport.cpp
+++ b/Engine/Source/Editor/PhAT/Private/SPhATPreviewViewport.cpp
@@ -13,63 +13,30 @@
 
 SPhATPreviewViewport::~SPhATPreviewViewport()
 {
-	if (ViewportClient.IsValid())
+	if (EditorViewportClient.IsValid())
 	{
-		ViewportClient->Viewport->Destroy();
-		ViewportClient->Viewport = NULL;
+		EditorViewportClient->Viewport->Destroy();
+		EditorViewportClient->Viewport = NULL;
 	}
 }
 
 void SPhATPreviewViewport::SetViewportType(ELevelViewportType ViewType)
 {
-	ViewportClient->SetViewportType(ViewType);
+	EditorViewportClient->SetViewportType(ViewType);
 }
 
 void SPhATPreviewViewport::Construct(const FArguments& InArgs)
 {
 	PhATPtr = InArgs._PhAT;
 
-	this->ChildSlot
-	[
-		SAssignNew(ViewportWidget, SViewport)
-		.EnableGammaCorrection(false)
-		.IsEnabled(FSlateApplication::Get().GetNormalExecutionAttribute())
-		.ShowEffectWhenDisabled(false)
-		[
-			SNew(SOverlay)
-			+SOverlay::Slot()
-			.VAlign(VAlign_Top)
-			[
-				SNew(SPhATPreviewViewportToolBar)
-				.PhATPtr(PhATPtr)
-		 		.Visibility(this, &SPhATPreviewViewport::GetWidgetVisibility)
-				.IsEnabled(FSlateApplication::Get().GetNormalExecutionAttribute())
-			]
-		]
-	];
-	
-	ViewportClient = MakeShareable(new FPhATEdPreviewViewportClient(PhATPtr, PhATPtr.Pin()->GetSharedData()));
-
-	ViewportClient->bSetListenerPosition = false;
-
-	if (!FPhAT::IsPIERunning())
-	{
-		ViewportClient->SetRealtime(true);
-	}
-
-	ViewportClient->VisibilityDelegate.BindSP(this, &SPhATPreviewViewport::IsVisible);
-
-	Viewport = MakeShareable(new FSceneViewport(ViewportClient.Get(), ViewportWidget));
-	ViewportClient->Viewport = Viewport.Get();
-
-	// The viewport widget needs an interface so it knows what should render
-	ViewportWidget->SetViewportInterface(Viewport.ToSharedRef());
+	SEditorViewport::Construct(SEditorViewport::FArguments());
+	ViewportWidget->SetEnabled(FSlateApplication::Get().GetNormalExecutionAttribute());
 }
 
 void SPhATPreviewViewport::RefreshViewport()
 {
-	Viewport->Invalidate();
-	Viewport->InvalidateDisplay();
+	SceneViewport->Invalidate();
+	SceneViewport->InvalidateDisplay();
 }
 
 bool SPhATPreviewViewport::IsVisible() const
@@ -84,15 +51,35 @@ EVisibility SPhATPreviewViewport::GetWidgetVisibility() const
 
 TSharedPtr<FSceneViewport> SPhATPreviewViewport::GetViewport() const
 {
-	return Viewport;
+	return SceneViewport;
 }
 
 TSharedPtr<FPhATEdPreviewViewportClient> SPhATPreviewViewport::GetViewportClient() const
 {
-	return ViewportClient;
+	return EditorViewportClient;
 }
 
 TSharedPtr<SViewport> SPhATPreviewViewport::GetViewportWidget() const
 {
 	return ViewportWidget;
 }
+
+TSharedRef<FEditorViewportClient> SPhATPreviewViewport::MakeEditorViewportClient()
+{
+	EditorViewportClient = MakeShareable(new FPhATEdPreviewViewportClient(PhATPtr, PhATPtr.Pin()->GetSharedData(), SharedThis(this)));
+
+	EditorViewportClient->bSetListenerPosition = false;
+
+	EditorViewportClient->SetRealtime(!FPhAT::IsPIERunning());
+	EditorViewportClient->VisibilityDelegate.BindSP(this, &SPhATPreviewViewport::IsVisible);
+
+	return EditorViewportClient.ToSharedRef();
+}
+
+TSharedPtr<SWidget> SPhATPreviewViewport::MakeViewportToolbar()
+{
+	return SNew(SPhATPreviewViewportToolBar)
+			.PhATPtr(PhATPtr)
+			.Visibility(this, &SPhATPreviewViewport::GetWidgetVisibility)
+			.IsEnabled(FSlateApplication::Get().GetNormalExecutionAttribute());
+}
diff --git a/Engine/Source/Editor/PhAT/Private/SPhATPreviewViewport.h b/Engine/Source/Editor/PhAT/Private/SPhATPreviewViewport.h
--- a/Engine/Source/Editor/PhAT/Private/SPhATPreviewViewport.h
+++ b/Engine/Source/Editor/PhAT/Private/SPhATPreviewViewport.h
@@ -2,14 +2,15 @@
 
 #pragma once
 
-class FPhAT;
+#include "SEditorViewport.h"
 
+class FPhAT;
 
 /*-----------------------------------------------------------------------------
    SPhATViewport
 -----------------------------------------------------------------------------*/
 
-class SPhATPreviewViewport : public SCompoundWidget
+class SPhATPreviewViewport : public SEditorViewport
 {
 public:
 	SLATE_BEGIN_ARGS(SPhATPreviewViewport)
@@ -42,16 +43,15 @@ public:
 
 	void SetViewportType(ELevelViewportType ViewType);
 
+protected:
+	/** SEditorViewport interface */
+	virtual TSharedRef<FEditorViewportClient> MakeEditorViewportClient() override;
+	virtual TSharedPtr<SWidget> MakeViewportToolbar() override;
+
 private:
 	/** Pointer back to the PhysicsAsset editor tool that owns us */
 	TWeakPtr<FPhAT> PhATPtr;
 	
 	/** Level viewport client */
-	TSharedPtr<FPhATEdPreviewViewportClient> ViewportClient;
-
-	/** Slate viewport for rendering and I/O */
-	TSharedPtr<FSceneViewport> Viewport;
-
-	/** Viewport widget*/
-	TSharedPtr<SViewport> ViewportWidget;
+	TSharedPtr<FPhATEdPreviewViewportClient> EditorViewportClient;
 };
diff --git a/Engine/Source/Editor/StaticMeshEditor/Private/StaticMeshEditorViewportClient.cpp b/Engine/Source/Editor/StaticMeshEditor/Private/StaticMeshEditorViewportClient.cpp
--- a/Engine/Source/Editor/StaticMeshEditor/Private/StaticMeshEditorViewportClient.cpp
+++ b/Engine/Source/Editor/StaticMeshEditor/Private/StaticMeshEditorViewportClient.cpp
@@ -45,8 +45,8 @@ namespace {
 	static float AmbientCubemapIntensity = 0.4f;
 }
 
-FStaticMeshEditorViewportClient::FStaticMeshEditorViewportClient(TWeakPtr<IStaticMeshEditor> InStaticMeshEditor, TWeakPtr<SStaticMeshEditorViewport> InStaticMeshEditorViewport, FPreviewScene& InPreviewScene, UStaticMesh* InPreviewStaticMesh, UStaticMeshComponent* InPreviewStaticMeshComponent)
-	: FEditorViewportClient(nullptr, &InPreviewScene)
+FStaticMeshEditorViewportClient::FStaticMeshEditorViewportClient(TWeakPtr<IStaticMeshEditor> InStaticMeshEditor, const TSharedRef<SStaticMeshEditorViewport>& InStaticMeshEditorViewport, FPreviewScene& InPreviewScene, UStaticMesh* InPreviewStaticMesh, UStaticMeshComponent* InPreviewStaticMeshComponent)
+	: FEditorViewportClient(nullptr, &InPreviewScene, StaticCastSharedRef<SEditorViewport>(InStaticMeshEditorViewport))
 	, StaticMeshEditorPtr(InStaticMeshEditor)
 	, StaticMeshEditorViewportPtr(InStaticMeshEditorViewport)
 {
diff --git a/Engine/Source/Editor/StaticMeshEditor/Private/StaticMeshEditorViewportClient.h b/Engine/Source/Editor/StaticMeshEditor/Private/StaticMeshEditorViewportClient.h
--- a/Engine/Source/Editor/StaticMeshEditor/Private/StaticMeshEditorViewportClient.h
+++ b/Engine/Source/Editor/StaticMeshEditor/Private/StaticMeshEditorViewportClient.h
@@ -10,7 +10,7 @@ class SStaticMeshEditorViewport;
 class FStaticMeshEditorViewportClient : public FEditorViewportClient, public TSharedFromThis<FStaticMeshEditorViewportClient>
 {
 public:
-	FStaticMeshEditorViewportClient(TWeakPtr<IStaticMeshEditor> InStaticMeshEditor, TWeakPtr<SStaticMeshEditorViewport> InStaticMeshEditorViewport, FPreviewScene& InPreviewScene, UStaticMesh* InPreviewStaticMesh, UStaticMeshComponent* InPreviewStaticMeshComponent);
+	FStaticMeshEditorViewportClient(TWeakPtr<IStaticMeshEditor> InStaticMeshEditor, const TSharedRef<SStaticMeshEditorViewport>& InStaticMeshEditorViewport, FPreviewScene& InPreviewScene, UStaticMesh* InPreviewStaticMesh, UStaticMeshComponent* InPreviewStaticMeshComponent);
 
 	// FEditorViewportClient interface
 	virtual void MouseMove(FViewport* Viewport,int32 x, int32 y) override;
diff --git a/Engine/Source/Editor/UnrealEd/Private/EditorViewportClient.cpp b/Engine/Source/Editor/UnrealEd/Private/EditorViewportClient.cpp
--- a/Engine/Source/Editor/UnrealEd/Private/EditorViewportClient.cpp
+++ b/Engine/Source/Editor/UnrealEd/Private/EditorViewportClient.cpp
@@ -24,6 +24,7 @@
 #include "Debug/DebugDrawService.h"
 #include "Components/BillboardComponent.h"
 #include "EngineUtils.h"
+#include "SEditorViewport.h"
 
 #define LOCTEXT_NAMESPACE "EditorViewportClient"
 
@@ -90,7 +91,7 @@ public:
 
 
 FViewportCameraTransform::FViewportCameraTransform()
-	: TransitionStartTime(0)
+	: TransitionCurve( new FCurveSequence( 0.0f, FocusConstants::TransitionTime, ECurveEaseFunction::CubicOut ) )
 	, ViewLocation( FVector::ZeroVector )
 	, ViewRotation( FRotator::ZeroRotator )
 	, DesiredLocation( FVector::ZeroVector )
@@ -105,19 +106,19 @@ void FViewportCameraTransform::SetLocation( const FVector& Position )
 	DesiredLocation = ViewLocation;
 }
 
-void FViewportCameraTransform::TransitionToLocation( const FVector& InDesiredLocation, bool bInstant )
+void FViewportCameraTransform::TransitionToLocation(const FVector& InDesiredLocation, TWeakPtr<SWidget> EditorViewportWidget, bool bInstant)
 {
-	if( bInstant )
+	if( bInstant || !EditorViewportWidget.IsValid() )
 	{
 		SetLocation( InDesiredLocation );
-		TransitionStartTime = FSlateApplication::Get().GetCurrentTime() - FocusConstants::TransitionTime;
+		TransitionCurve->JumpToEnd();
 	}
 	else
 	{
 		DesiredLocation = InDesiredLocation;
 		StartLocation = ViewLocation;
 
-		TransitionStartTime = FSlateApplication::Get().GetCurrentTime();
+		TransitionCurve->Play(EditorViewportWidget.Pin().ToSharedRef());
 	}
 }
 
@@ -125,11 +126,9 @@ void FViewportCameraTransform::TransitionToLocation( const FVector& InDesiredLoc
 bool FViewportCameraTransform::UpdateTransition()
 {
 	bool bIsAnimating = false;
-	double TransitionProgress = FMath::Clamp( (FSlateApplication::Get().GetCurrentTime() - TransitionStartTime) / FocusConstants::TransitionTime, 0.0, 1.0);
-	if( TransitionProgress < 1.0 || ViewLocation != DesiredLocation )
+	if (TransitionCurve->IsPlaying() || ViewLocation != DesiredLocation)
 	{
-		const float Offset = (float)TransitionProgress - 1.0f;
-		float LerpWeight = Offset * Offset * Offset + 1.0f;
+		float LerpWeight = TransitionCurve->GetLerp();
 
 		if( LerpWeight == 1.0f )
 		{
@@ -194,7 +193,7 @@ int32 FEditorViewportClient::GetCameraSpeedSetting() const
 
 float const FEditorViewportClient::SafePadding = 0.075f;
 
-FEditorViewportClient::FEditorViewportClient(FEditorModeTools* InModeTools, FPreviewScene* InPreviewScene)
+FEditorViewportClient::FEditorViewportClient(FEditorModeTools* InModeTools, FPreviewScene* InPreviewScene, const TWeakPtr<SEditorViewport>& InEditorViewportWidget)
 	: bAllowMatineePreview(false)
 	, CameraSpeedSetting(4)
 	, ImmersiveDelegate()
@@ -217,7 +216,7 @@ FEditorViewportClient::FEditorViewportClient(FEditorModeTools* InModeTools, FPre
 	, bNeedsLinkedRedraw(false)
 	, bNeedsInvalidateHitProxy(false)
 	, bUsingOrbitCamera(false)
-	, bDisableInput( false )
+	, bDisableInput(false)
 	, bDrawAxes(true)
 	, bSetListenerPosition(false)
 	, LandscapeLODOverride(-1)
@@ -228,10 +227,10 @@ FEditorViewportClient::FEditorViewportClient(FEditorModeTools* InModeTools, FPre
 	, MouseDeltaTracker(new FMouseDeltaTracker)
 	, RecordingInterpEd(NULL)
 	, bHasMouseMovedSinceClick(false)
-	, CameraController( new FEditorCameraController() )
-	, CameraUserImpulseData( new FCameraControllerUserImpulseData() )
-	, TimeForForceRedraw( 0.0 )
-	, FlightCameraSpeedScale( 1.0f )
+	, CameraController(new FEditorCameraController())
+	, CameraUserImpulseData(new FCameraControllerUserImpulseData())
+	, TimeForForceRedraw(0.0)
+	, FlightCameraSpeedScale(1.0f)
 	, bUseControllingActorViewInfo(false)
 	, LastMouseX(0)
 	, LastMouseY(0)
@@ -253,6 +252,7 @@ FEditorViewportClient::FEditorViewportClient(FEditorModeTools* InModeTools, FPre
 	, bUsesDrawHelper(true)
 	, bIsSimulateInEditorViewport(false)
 	, bCameraLock(false)
+	, EditorViewportWidget(InEditorViewportWidget)
 	, PreviewScene(InPreviewScene)
 	, PerspViewModeIndex(DefaultPerspectiveViewMode)
 	, OrthoViewModeIndex(DefaultOrthoViewMode)
@@ -286,11 +286,11 @@ FEditorViewportClient::FEditorViewportClient(FEditorModeTools* InModeTools, FPre
 	DrawHelper.bDrawWorldBox = false;
 	DrawHelper.bDrawKillZ = false;
 	DrawHelper.bDrawGrid = true;
-	DrawHelper.GridColorAxis = FColor(160,160,160);
-	DrawHelper.GridColorMajor = FColor(144,144,144);
-	DrawHelper.GridColorMinor = FColor(128,128,128);
+	DrawHelper.GridColorAxis = FColor(160, 160, 160);
+	DrawHelper.GridColorMajor = FColor(144, 144, 144);
+	DrawHelper.GridColorMinor = FColor(128, 128, 128);
 	DrawHelper.PerspectiveGridSize = GridSize;
-	DrawHelper.NumCells = DrawHelper.PerspectiveGridSize / (CellSize * 2);
+	DrawHelper.NumCells = DrawHelper.PerspectiveGridSize / ( CellSize * 2 );
 
 	// Most editor viewports do not want motion blur.
 	EngineShowFlags.MotionBlur = 0;
@@ -489,14 +489,14 @@ void FEditorViewportClient::FocusViewportOnBox( const FBox& BoundingBox, bool bI
 			FVector CameraOffsetVector = ViewTransform.GetRotation().Vector() * -DistanceFromSphere;
 
 			ViewTransform.SetLookAt(Position);
-			ViewTransform.TransitionToLocation(Position + CameraOffsetVector, bInstant);
+			ViewTransform.TransitionToLocation(Position + CameraOffsetVector, EditorViewportWidget, bInstant);
 
 		}
 		else
 		{
 			// For ortho viewports just set the camera position to the center of the bounding volume.
 			//SetViewLocation( Position );
-			ViewTransform.TransitionToLocation(Position, bInstant);
+			ViewTransform.TransitionToLocation(Position, EditorViewportWidget, bInstant);
 
 			if( !(Viewport->KeyState(EKeys::LeftControl) || Viewport->KeyState(EKeys::RightControl)) )
 			{
@@ -1320,6 +1320,12 @@ void FEditorViewportClient::UpdateCameraMovement( float DeltaTime )
 			MoveViewportPerspectiveCamera(
 				NewViewLocation - GetViewLocation(),
 				NewViewRotation - GetViewRotation() );
+
+			// Invalidate the viewport widget
+			if (EditorViewportWidget.IsValid())
+			{
+				EditorViewportWidget.Pin()->Invalidate();
+			}
 		}
 	}
 }
diff --git a/Engine/Source/Editor/UnrealEd/Private/LevelEditorViewport.cpp b/Engine/Source/Editor/UnrealEd/Private/LevelEditorViewport.cpp
--- a/Engine/Source/Editor/UnrealEd/Private/LevelEditorViewport.cpp
+++ b/Engine/Source/Editor/UnrealEd/Private/LevelEditorViewport.cpp
@@ -47,6 +47,7 @@
 #include "Editor/ActorPositioning.h"
 #include "NotificationManager.h"
 #include "SNotificationList.h"
+#include "SLevelViewport.h"
 
 DEFINE_LOG_CATEGORY(LogEditorViewport);
 
@@ -1422,8 +1423,8 @@ void FTrackingTransaction::PromotePendingToActive()
 	}
 }
 
-FLevelEditorViewportClient::FLevelEditorViewportClient()
-	: FEditorViewportClient(&GLevelEditorModeTools(), NULL)
+FLevelEditorViewportClient::FLevelEditorViewportClient(const TSharedPtr<SLevelViewport>& InLevelViewport)
+	: FEditorViewportClient(&GLevelEditorModeTools(), nullptr, StaticCastSharedPtr<SEditorViewport>(InLevelViewport))
 	, ViewHiddenLayers()
 	, VolumeActorVisibility()
 	, LastEditorViewLocation( FVector::ZeroVector )
@@ -1440,13 +1441,13 @@ FLevelEditorViewportClient::FLevelEditorViewportClient()
 	, bIsTrackingBrushModification( false )
 	, bLockedCameraView(true)
 	, SpriteCategoryVisibility()
-	, World(NULL)
+	, World(nullptr)
 	, TrackingTransaction()
 	, DropPreviewMouseX(0)
 	, DropPreviewMouseY(0)
 	, bWasControlledByOtherViewport(false)
-	, ActorLockedByMatinee(NULL)
-	, ActorLockedToCamera(NULL)
+	, ActorLockedByMatinee(nullptr)
+	, ActorLockedToCamera(nullptr)
 {
 	// By default a level editor viewport is pointed to the editor world
 	SetReferenceToWorldContext(GEditor->GetEditorWorldContext());
diff --git a/Engine/Source/Editor/UnrealEd/Private/SEditorViewport.cpp b/Engine/Source/Editor/UnrealEd/Private/SEditorViewport.cpp
--- a/Engine/Source/Editor/UnrealEd/Private/SEditorViewport.cpp
+++ b/Engine/Source/Editor/UnrealEd/Private/SEditorViewport.cpp
@@ -14,6 +14,7 @@
 
 SEditorViewport::SEditorViewport()
 	: LastTickTime(0)
+	, bInvalidated(false)
 {
 }
 
@@ -63,7 +64,7 @@ void SEditorViewport::Construct( const FArguments& InArgs )
 
 	if ( Client->IsRealtime() )
 	{
-		ActiveTimerHandle = RegisterActiveTimer( 0.f, FWidgetActiveTimerDelegate::CreateSP( this, &SEditorViewport::EnsureTickWhileRealtime ) );
+		ActiveTimerHandle = RegisterActiveTimer( 0.f, FWidgetActiveTimerDelegate::CreateSP( this, &SEditorViewport::EnsureTick ) );
 	}
 
 	CommandList = MakeShareable( new FUICommandList );
@@ -340,7 +341,7 @@ void SEditorViewport::OnToggleRealtime()
 	else
 	{
 		Client->SetRealtime( true );
-		ActiveTimerHandle = RegisterActiveTimer( 0.f, FWidgetActiveTimerDelegate::CreateSP( this, &SEditorViewport::EnsureTickWhileRealtime ) );
+		ActiveTimerHandle = RegisterActiveTimer( 0.f, FWidgetActiveTimerDelegate::CreateSP( this, &SEditorViewport::EnsureTick ) );
 	}
 }
 
@@ -355,7 +356,7 @@ void SEditorViewport::OnToggleStats()
 		if ( !Client->IsRealtime() )
 		{
 			Client->SetRealtime( true );
-			ActiveTimerHandle = RegisterActiveTimer( 0.f, FWidgetActiveTimerDelegate::CreateSP( this, &SEditorViewport::EnsureTickWhileRealtime ) );
+			ActiveTimerHandle = RegisterActiveTimer( 0.f, FWidgetActiveTimerDelegate::CreateSP( this, &SEditorViewport::EnsureTick ) );
 		}
 
 		 // let the user know how they can enable stats via the console
@@ -415,6 +416,14 @@ bool SEditorViewport::IsExposureSettingSelected( int32 ID ) const
 	}
 }
 
+void SEditorViewport::Invalidate()
+{
+	bInvalidated = true;
+	if (!ActiveTimerHandle.IsValid())
+	{
+		ActiveTimerHandle = RegisterActiveTimer(0.f, FWidgetActiveTimerDelegate::CreateSP(this, &SEditorViewport::EnsureTick));
+	}
+}
 
 bool SEditorViewport::IsRealtime() const
 {
@@ -518,9 +527,12 @@ bool SEditorViewport::OnIsSurfaceSnapEnabled()
 	return GetDefault<ULevelEditorViewportSettings>()->SnapToSurface.bEnabled;
 }
 
-EActiveTimerReturnType SEditorViewport::EnsureTickWhileRealtime( double InCurrentTime, float InDeltaTime )
+EActiveTimerReturnType SEditorViewport::EnsureTick( double InCurrentTime, float InDeltaTime )
 {
-	return EActiveTimerReturnType::Continue;
+	// Keep the timer going if we're realtime or were invalidated this frame
+	const bool bShouldContinue = Client->IsRealtime() || bInvalidated;
+	bInvalidated = false;
+	return bShouldContinue ? EActiveTimerReturnType::Continue : EActiveTimerReturnType::Stop;
 }
 
 #undef LOCTEXT_NAMESPACE
diff --git a/Engine/Source/Editor/UnrealEd/Private/SSkeletonWidget.cpp b/Engine/Source/Editor/UnrealEd/Private/SSkeletonWidget.cpp
--- a/Engine/Source/Editor/UnrealEd/Private/SSkeletonWidget.cpp
+++ b/Engine/Source/Editor/UnrealEd/Private/SSkeletonWidget.cpp
@@ -613,16 +613,48 @@ void SAnimationRemapSkeleton::Construct( const FArguments& InArgs )
 
 					+SHorizontalBox::Slot()
 					[
-						SAssignNew(SourceViewport, SBasePoseViewport)
+						SNew(SVerticalBox)
+						+ SVerticalBox::Slot()
+						.AutoHeight()
+						[
+							SNew(STextBlock)
+							.Text(LOCTEXT("SourceSkeleteonTitle", "[Source]"))
+							.Font(FEditorStyle::GetFontStyle("Persona.RetargetManager.FilterFont"))
+							.AutoWrapText(true)
+						]
+
+						+ SVerticalBox::Slot()
+						.AutoHeight()
+						[
+							SAssignNew(SourceViewport, SBasePoseViewport)
+							.Skeleton(OldSkeleton)
+						]
+						/*SAssignNew(SourceViewport, SBasePoseViewport)
 						.Title(TEXT("[Source]"))
-						.Skeleton(OldSkeleton)
+						.Skeleton(OldSkeleton)*/
 					]
 
 					+SHorizontalBox::Slot()
 					[
-						SAssignNew(TargetViewport, SBasePoseViewport)
+						SNew(SVerticalBox)
+						+ SVerticalBox::Slot()
+						.AutoHeight()
+						[
+							SNew(STextBlock)
+							.Text(LOCTEXT("TargetSkeleteonTitle", "[Target]"))
+							.Font(FEditorStyle::GetFontStyle("Persona.RetargetManager.FilterFont"))
+							.AutoWrapText(true)
+						]
+
+						+ SVerticalBox::Slot()
+						.AutoHeight()
+						[
+							SAssignNew(TargetViewport, SBasePoseViewport)
+							.Skeleton(nullptr)
+						]
+						/*SAssignNew(TargetViewport, SBasePoseViewport)
 						.Title(TEXT("[Target]"))
-						.Skeleton(NULL)
+						.Skeleton(NULL)*/
 					]
 				]
 			]
@@ -1068,8 +1100,8 @@ bool SSkeletonBoneRemoval::ShowModal(const TArray<FName> BonesToRemove, const FT
 class FBasePoseViewportClient: public FEditorViewportClient
 {
 public:
-	FBasePoseViewportClient(FPreviewScene& InPreviewScene)
-			: FEditorViewportClient(nullptr, &InPreviewScene)
+	FBasePoseViewportClient(FPreviewScene& InPreviewScene, const TSharedRef<SBasePoseViewport>& InBasePoseViewport)
+		: FEditorViewportClient(nullptr, &InPreviewScene, StaticCastSharedRef<SEditorViewport>(InBasePoseViewport))
 	{
 		SetViewMode(VMI_Lit);
 
@@ -1129,46 +1161,7 @@ public:
 // SBasePoseViewport
 void SBasePoseViewport::Construct(const FArguments& InArgs)
 {
-	//RegisterActiveTimer( 0.f, FTickWidgetDelegate::CreateSP( this, &SBasePoseViewport::TickViewport ) );
-
-	this->ChildSlot
-	[
-		SNew(SVerticalBox)
-		+SVerticalBox::Slot()
-		.AutoHeight()
-		[
-			SNew(STextBlock)
-			.Text(FText::FromString(InArgs._Title))
-			.Font(FEditorStyle::GetFontStyle("Persona.RetargetManager.FilterFont"))
-			.AutoWrapText(true)
-		]
-		+SVerticalBox::Slot()
-		.AutoHeight()
-		[
-			SNew(SBorder)
-			.HAlign(HAlign_Center)
-			[
-				SAssignNew(ViewportWidget, SViewport)
-				.EnableGammaCorrection(false)
-			]
-		]
-	];
-
-	// Create a viewport client
-	LevelViewportClient	= MakeShareable(new FBasePoseViewportClient(PreviewScene));
-
-	LevelViewportClient->ViewportType = LVT_Perspective;
-	LevelViewportClient->bSetListenerPosition = false;
-	LevelViewportClient->SetViewLocation(EditorViewportDefs::DefaultPerspectiveViewLocation);
-	LevelViewportClient->SetViewRotation(EditorViewportDefs::DefaultPerspectiveViewRotation);
-
-	SceneViewport = MakeShareable(new FSceneViewport(LevelViewportClient.Get(), ViewportWidget));
-	LevelViewportClient->Viewport = SceneViewport.Get();
-	LevelViewportClient->SetRealtime(true);
-	LevelViewportClient->VisibilityDelegate.BindSP(this, &SBasePoseViewport::IsVisible);
-	LevelViewportClient->SetViewMode(VMI_Lit);
-
-	ViewportWidget->SetViewportInterface(SceneViewport.ToSharedRef());
+	SEditorViewport::Construct(SEditorViewport::FArguments());
 
 	PreviewComponent = ConstructObject<UDebugSkelMeshComponent>(UDebugSkelMeshComponent::StaticClass());
 	PreviewComponent->MeshComponentUpdateFlag = EMeshComponentUpdateFlag::AlwaysTickPoseAndRefreshBones;
@@ -1195,13 +1188,13 @@ void SBasePoseViewport::SetSkeleton(USkeleton* Skeleton)
 				PreviewComponent->RefreshBoneTransforms(NULL);
 
 				//Place the camera at a good viewer position
-				FVector NewPosition = LevelViewportClient->GetViewLocation();
+				FVector NewPosition = Client->GetViewLocation();
 				NewPosition.Normalize();
 				if(PreviewSkeletalMesh)
 				{
 					NewPosition *= (PreviewSkeletalMesh->Bounds.SphereRadius*1.5f);
 				}
-				LevelViewportClient->SetViewLocation(NewPosition);
+				Client->SetViewLocation(NewPosition);
 			}
 			else
 			{
@@ -1213,16 +1206,7 @@ void SBasePoseViewport::SetSkeleton(USkeleton* Skeleton)
 			PreviewComponent->SetSkeletalMesh(NULL);
 		}
 
-		LevelViewportClient->Invalidate();
-	}
-}
-
-SBasePoseViewport::~SBasePoseViewport()
-{
-	// Close viewport
-	if(LevelViewportClient.IsValid())
-	{
-		LevelViewportClient->Viewport = NULL;
+		Client->Invalidate();
 	}
 }
 
@@ -1235,5 +1219,26 @@ bool SBasePoseViewport::IsVisible() const
 {
 	return true;
 }
+
+TSharedRef<FEditorViewportClient> SBasePoseViewport::MakeEditorViewportClient()
+{
+	TSharedPtr<FEditorViewportClient> EditorViewportClient = MakeShareable(new FBasePoseViewportClient(PreviewScene, SharedThis(this)));
+
+	EditorViewportClient->ViewportType = LVT_Perspective;
+	EditorViewportClient->bSetListenerPosition = false;
+	EditorViewportClient->SetViewLocation(EditorViewportDefs::DefaultPerspectiveViewLocation);
+	EditorViewportClient->SetViewRotation(EditorViewportDefs::DefaultPerspectiveViewRotation);
+
+	EditorViewportClient->SetRealtime(false);
+	EditorViewportClient->VisibilityDelegate.BindSP(this, &SBasePoseViewport::IsVisible);
+	EditorViewportClient->SetViewMode(VMI_Lit);
+
+	return EditorViewportClient.ToSharedRef();
+}
+
+TSharedPtr<SWidget> SBasePoseViewport::MakeViewportToolbar()
+{
+	return nullptr;
+}
 #undef LOCTEXT_NAMESPACE 
 
diff --git a/Engine/Source/Editor/UnrealEd/Public/EditorViewportClient.h b/Engine/Source/Editor/UnrealEd/Public/EditorViewportClient.h
--- a/Engine/Source/Editor/UnrealEd/Public/EditorViewportClient.h
+++ b/Engine/Source/Editor/UnrealEd/Public/EditorViewportClient.h
@@ -157,6 +157,9 @@ struct FDropQuery
 	FText HintText;
 };
 
+// Forward declare
+class SEditorViewport;
+
 /** 
  * Stores the transformation data for the viewport camera
  */
@@ -205,7 +208,7 @@ public:
 	 * @param InDesiredLocation	The location to transition to
 	 * @param bInstant			If the desired location should be set instantly rather than transitioned to over time
 	 */
-	void TransitionToLocation( const FVector& InDesiredLocation, bool bInstant );
+	void TransitionToLocation( const FVector& InDesiredLocation, TWeakPtr<SWidget> EditorViewportWidget, bool bInstant );
 
 	/**
 	 * Updates any current location transitions
@@ -220,7 +223,10 @@ public:
 	FMatrix ComputeOrbitMatrix() const;
 private:
 	/** The time when a transition to the desired location began */
-	double TransitionStartTime;
+	//double TransitionStartTime;
+	
+	/** Curve for animating between locations */
+	TSharedPtr<struct FCurveSequence> TransitionCurve;
 	/** Current viewport Position. */
 	FVector	ViewLocation;
 	/** Current Viewport orientation; valid only for perspective projections. */
@@ -235,12 +241,14 @@ private:
 	float OrthoZoom;
 };
 
+/** Viewport client for editor viewports. Contains common functionality for camera movement, rendering debug information, etc. */
 class UNREALED_API FEditorViewportClient : public FCommonViewportClient, public FViewElementDrawer, public FGCObject
 {
 public:
 	friend class FMouseDeltaTracker;
 
-	FEditorViewportClient(FEditorModeTools* InModeTools, FPreviewScene* InPreviewScene = nullptr);
+	FEditorViewportClient(FEditorModeTools* InModeTools, FPreviewScene* InPreviewScene = nullptr, const TWeakPtr<SEditorViewport>& InEditorViewportWidget = nullptr);
+
 	virtual ~FEditorViewportClient();
 
 	/**
@@ -1400,6 +1408,9 @@ protected:
 	/** Draw helper for rendering common editor functionality like the grid */
 	FEditorCommonDrawHelper DrawHelper;
 
+	/** The editor viewport widget this client is attached to */
+	TWeakPtr<SEditorViewport> EditorViewportWidget;
+
 	/** The scene used for the viewport. Owned externally */
 	FPreviewScene* PreviewScene;
 
diff --git a/Engine/Source/Editor/UnrealEd/Public/SEditorViewport.h b/Engine/Source/Editor/UnrealEd/Public/SEditorViewport.h
--- a/Engine/Source/Editor/UnrealEd/Public/SEditorViewport.h
+++ b/Engine/Source/Editor/UnrealEd/Public/SEditorViewport.h
@@ -27,6 +27,12 @@ public:
 	/** @return True if the viewport is currently visible */
 	virtual bool IsVisible() const;
 
+	/** 
+	 * Invalidates the viewport to ensure it is redrawn during the next tick. 
+	 * This is implied every frame while the viewport IsRealtime().
+	 */
+	void Invalidate();
+
 	/**
 	 * @return true if the specified coordinate system the active one active
 	 */
@@ -102,7 +108,6 @@ protected:
 	 */
 	bool IsExposureSettingSelected( int32 ID ) const;
 	
-
 	virtual void OnScreenCapture();
 	virtual void OnScreenCaptureForProjectThumbnail();
 	virtual bool DoesAllowScreenCapture() { return true; }
@@ -172,9 +177,12 @@ protected:
 	double LastTickTime;
 
 private:
-	/** Ensures the viewport is ticked actively when realtime is enabled */
-	EActiveTimerReturnType EnsureTickWhileRealtime( double InCurrentTime, float InDeltaTime );
+	/** Ensures a Slate tick/paint pass when the viewport is realtime or was invalidated this frame */
+	EActiveTimerReturnType EnsureTick( double InCurrentTime, float InDeltaTime );
 
 	/** The handle to the active timer */
 	TWeakPtr<FActiveTimerHandle> ActiveTimerHandle;
+
+	/** Whether the viewport needs to update, even without input or realtime (e.g. inertial camera movement) */
+	bool bInvalidated;
 };
diff --git a/Engine/Source/Editor/UnrealEd/Public/SSkeletonWidget.h b/Engine/Source/Editor/UnrealEd/Public/SSkeletonWidget.h
--- a/Engine/Source/Editor/UnrealEd/Public/SSkeletonWidget.h
+++ b/Engine/Source/Editor/UnrealEd/Public/SSkeletonWidget.h
@@ -6,6 +6,7 @@
 
 #include "SlateBasics.h"
 #include "PreviewScene.h"
+#include "SEditorViewport.h"
 
 class UAnimSet;
 
@@ -251,30 +252,27 @@ private:
 
 
 //////////////////////////
-class SBasePoseViewport: public SCompoundWidget
+class SBasePoseViewport: public /*SCompoundWidget*/SEditorViewport
 {
 public:
 	SLATE_BEGIN_ARGS(SBasePoseViewport)
 	{}
 
-	SLATE_ARGUMENT(FString, Title)
 	SLATE_ARGUMENT(USkeleton*, Skeleton)
 SLATE_END_ARGS()
 
 public:
 	SBasePoseViewport();
-	virtual ~SBasePoseViewport();
 
 	void Construct(const FArguments& InArgs);
 	void SetSkeleton(USkeleton* Skeleton);
-private:
-	TSharedPtr<FEditorViewportClient> LevelViewportClient;
-
-	/** Slate viewport for rendering and I/O */
-	TSharedPtr<SViewport> ViewportWidget;
 
-	TSharedPtr<class FSceneViewport> SceneViewport;
+protected:
+	/** SEditorViewport interface */
+	virtual TSharedRef<FEditorViewportClient> MakeEditorViewportClient() override;
+	virtual TSharedPtr<SWidget> MakeViewportToolbar() override;
 
+private:
 	/** Skeleton */
 	USkeleton* TargetSkeleton;
 
