Flipped FixLevelScriptActorBindings() conditional typo, so it uses the target delegate (rather than ignoring it).

diff --git a/Engine/Source/Editor/UnrealEd/Private/Kismet2/BlueprintEditorUtils.cpp b/Engine/Source/Editor/UnrealEd/Private/Kismet2/BlueprintEditorUtils.cpp
--- a/Engine/Source/Editor/UnrealEd/Private/Kismet2/BlueprintEditorUtils.cpp
+++ b/Engine/Source/Editor/UnrealEd/Private/Kismet2/BlueprintEditorUtils.cpp
@@ -5653,9 +5653,8 @@ bool FBlueprintEditorUtils::FixLevelScriptActorBindings(ALevelScriptActor* Level
 				{
 					// Grab the MC delegate we need to add to
 					FMulticastScriptDelegate* TargetDelegate = EventNode->GetTargetDelegate();
-					if( TargetDelegate == nullptr )
+					if( TargetDelegate != nullptr )
 					{
-
 						// Create the delegate, and add it if it doesn't already exist
 						FScriptDelegate Delegate;
 						Delegate.BindUFunction(LevelScriptActor, TargetFunction);
