Restructure ReplaceSelectedActors

diff --git a/Engine/Source/Editor/LevelEditor/Private/LevelEditorActions.cpp b/Engine/Source/Editor/LevelEditor/Private/LevelEditorActions.cpp
--- a/Engine/Source/Editor/LevelEditor/Private/LevelEditorActions.cpp
+++ b/Engine/Source/Editor/LevelEditor/Private/LevelEditorActions.cpp
@@ -1213,7 +1213,7 @@ AActor* FLevelEditorActionCallbacks::ReplaceActors( UActorFactory* ActorFactory,
 	if( ActorFactory->CanCreateActorFrom( AssetData, ErrorMessage ) )
 	{
 		// Replace all selected actors with actors created from the specified factory
-		GEditor->ReplaceSelectedActors( ActorFactory, AssetData, NULL );
+		GEditor->ReplaceSelectedActors( ActorFactory, AssetData );
 
 		if ( IPlacementModeModule::IsAvailable() )
 		{
@@ -1253,12 +1253,12 @@ void FLevelEditorActionCallbacks::ReplaceActorsFromClass_Clicked( UClass* ActorC
 			if( ActorFactory->CanCreateActorFrom( TargetAssetData, ErrorMessage ) )
 			{
 				// Replace all selected actors with actors created from the specified factory
-				GEditor->ReplaceSelectedActors( ActorFactory, TargetAssetData, NULL );
+				GEditor->ReplaceSelectedActors( ActorFactory, TargetAssetData );
 			}	
 			else if ( ActorFactory->CanCreateActorFrom( NoAssetData, UnusedErrorMessage ) )
 			{
 				// Replace all selected actors with actors created from the specified factory
-				GEditor->ReplaceSelectedActors( ActorFactory, NoAssetData, NULL );
+				GEditor->ReplaceSelectedActors( ActorFactory, NoAssetData );
 			}
 			else
 			{
diff --git a/Engine/Source/Editor/UnrealEd/Classes/Editor/EditorEngine.h b/Engine/Source/Editor/UnrealEd/Classes/Editor/EditorEngine.h
--- a/Engine/Source/Editor/UnrealEd/Classes/Editor/EditorEngine.h
+++ b/Engine/Source/Editor/UnrealEd/Classes/Editor/EditorEngine.h
@@ -1058,13 +1058,20 @@ public:
 	AActor* UseActorFactory( UActorFactory* Factory, const FAssetData& AssetData, const FTransform* ActorLocation, EObjectFlags ObjectFlags = RF_Transactional );
 
 	/**
-	 * Replaces the selected Actors with the same number of a different kind of Actor
-	 * if a Factory is specified, it is used to spawn the requested Actors, otherwise NewActorClass is used (one or the other must be specified)
+	 * Replaces the selected Actors with the same number of a different kind of Actor using the specified factory to spawn the new Actors
 	 * note that only Location, Rotation, Drawscale, Drawscale3D, Tag, and Group are copied from the old Actors
 	 * 
 	 * @param Factory - the Factory to use to create Actors
 	 */
-	void ReplaceSelectedActors(UActorFactory* Factory, const FAssetData& AssetData, UClass* NewActorClass);
+	void ReplaceSelectedActors(UActorFactory* Factory, const FAssetData& AssetData);
+
+	/**
+	 * Replaces specified Actors with the same number of a different kind of Actor using the specified factory to spawn the new Actors
+	 * note that only Location, Rotation, Drawscale, Drawscale3D, Tag, and Group are copied from the old Actors
+	 * 
+	 * @param Factory - the Factory to use to create Actors
+	 */
+	void ReplaceActors(UActorFactory* Factory, const FAssetData& AssetData, const TArray<AActor*> ActorsToReplace);
 
 	/**
 	 * Converts passed in brushes into a single static mesh actor. 
diff --git a/Engine/Source/Editor/UnrealEd/Private/Editor.cpp b/Engine/Source/Editor/UnrealEd/Private/Editor.cpp
--- a/Engine/Source/Editor/UnrealEd/Private/Editor.cpp
+++ b/Engine/Source/Editor/UnrealEd/Private/Editor.cpp
@@ -5023,7 +5023,7 @@ namespace ReattachActorsHelper
 	}
 }
 
-void UEditorEngine::ReplaceSelectedActors(UActorFactory* Factory, const FAssetData& AssetData, UClass* NewActorClass)
+void UEditorEngine::ReplaceSelectedActors(UActorFactory* Factory, const FAssetData& AssetData)
 {
 	UObject* ObjectForFactory = NULL;
 
@@ -5032,7 +5032,7 @@ void UEditorEngine::ReplaceSelectedActors(UActorFactory* Factory, const FAssetDa
 	{
 		return;
 	}
-	else if (Factory != NULL)
+	else if (Factory != nullptr)
 	{
 		FText ActorErrorMsg;
 		if (!Factory->CanCreateActorFrom( AssetData, ActorErrorMsg))
@@ -5041,7 +5041,7 @@ void UEditorEngine::ReplaceSelectedActors(UActorFactory* Factory, const FAssetDa
 			return;
 		}
 	}
-	else if (NewActorClass == NULL)
+	else
 	{
 		UE_LOG(LogEditor, Error, TEXT("UEditorEngine::ReplaceSelectedActors() called with NULL parameters!"));
 		return;
@@ -5060,6 +5060,11 @@ void UEditorEngine::ReplaceSelectedActors(UActorFactory* Factory, const FAssetDa
 		}
 	}
 
+	ReplaceActors(Factory, AssetData, ActorsToReplace);
+}
+
+void UEditorEngine::ReplaceActors(UActorFactory* Factory, const FAssetData& AssetData, const TArray<AActor*> ActorsToReplace)
+{
 	// Cache for attachment info of all actors being converted.
 	TArray<ReattachActorsHelper::FActorAttachmentCache> AttachmentInfo;
 
@@ -5082,35 +5087,27 @@ void UEditorEngine::ReplaceSelectedActors(UActorFactory* Factory, const FAssetDa
 		ULevel* Level = OldActor->GetLevel();
 		AActor* NewActor = NULL;
 
+		const FName OldActorName = OldActor->GetFName();
+		OldActor->Rename(*FString::Printf(TEXT("%s_REPLACED"), *OldActorName.ToString()));
+
 		const FTransform OldTransform = OldActor->ActorToWorld();
 
 		// create the actor
-		if (Factory != NULL)
+		NewActor = Factory->CreateActor( Asset, Level, OldTransform, RF_Transactional, OldActorName );
+		// For blueprints, try to copy over properties
+		if (Factory->IsA(UActorFactoryBlueprint::StaticClass()))
 		{
-			NewActor = Factory->CreateActor( Asset, Level, OldTransform);
-			// For blueprints, try to copy over properties
-			if (Factory->IsA(UActorFactoryBlueprint::StaticClass()))
+			UBlueprint* Blueprint = CastChecked<UBlueprint>(Asset);
+			// Only try to copy properties if this blueprint is based on the actor
+			UClass* OldActorClass = OldActor->GetClass();
+			if (Blueprint->GeneratedClass->IsChildOf(OldActorClass))
 			{
-				UBlueprint* Blueprint = CastChecked<UBlueprint>(Asset);
-				// Only try to copy properties if this blueprint is based on the actor
-				UClass* OldActorClass = OldActor->GetClass();
-				if (Blueprint->GeneratedClass->IsChildOf(OldActorClass))
-				{
-					NewActor->UnregisterAllComponents();
-					UEditorEngine::CopyPropertiesForUnrelatedObjects(OldActor, NewActor);
-					NewActor->RegisterAllComponents();
-				}
+				NewActor->UnregisterAllComponents();
+				UEditorEngine::CopyPropertiesForUnrelatedObjects(OldActor, NewActor);
+				NewActor->RegisterAllComponents();
 			}
 		}
-		else
-		{
-			FActorSpawnParameters SpawnInfo;
-			SpawnInfo.OverrideLevel = Level;
 
-			const auto Rotation = OldTransform.GetRotation().Rotator();
-			const auto Translation = OldTransform.GetTranslation();
-			NewActor = World->SpawnActor( NewActorClass, &Translation, &Rotation, SpawnInfo );
-		}
 		if ( NewActor != NULL )
 		{
 			// The new actor might not have a root component
@@ -5140,8 +5137,11 @@ void UEditorEngine::ReplaceSelectedActors(UActorFactory* Factory, const FAssetDa
 			// Caches information for finding the new actor using the pre-converted actor.
 			ReattachActorsHelper::CacheActorConvert(OldActor, NewActor, ConvertedMap, AttachmentInfo[ActorIdx]);
 
-			SelectActor(OldActor, false, true);
-			SelectActor(NewActor, true, true);
+			if (SelectedActors->IsSelected(OldActor))
+			{
+				SelectActor(OldActor, false, true);
+				SelectActor(NewActor, true, true);
+			}
 
 			// Find compatible static mesh components and copy instance colors between them.
 			UStaticMeshComponent* NewActorStaticMeshComponent = NewActor->FindComponentByClass<UStaticMeshComponent>();
@@ -5158,6 +5158,11 @@ void UEditorEngine::ReplaceSelectedActors(UActorFactory* Factory, const FAssetDa
 			GEditor->Layers->DisassociateActorFromLayers( OldActor );
 			World->EditorDestroyActor(OldActor, true);
 		}
+		else
+		{
+			// If creating the new Actor failed, put the old Actor's name back
+			OldActor->Rename(*OldActorName.ToString());
+		}
 	}
 
 	SelectedActors->EndBatchSelectOperation();
