Changed console variable callback to use a variable sink instead as per current recommendations

diff --git a/Engine/Source/Editor/UnrealEd/Private/PerformanceMonitor.cpp b/Engine/Source/Editor/UnrealEd/Private/PerformanceMonitor.cpp
--- a/Engine/Source/Editor/UnrealEd/Private/PerformanceMonitor.cpp
+++ b/Engine/Source/Editor/UnrealEd/Private/PerformanceMonitor.cpp
@@ -91,35 +91,46 @@ public:
 	}
 };
 
+static TAutoConsoleVariable<int32> CVarFineSampleTime(
+	TEXT("PerfWarn.FineSampleTime"), 30, TEXT("How many seconds we sample the percentage for the fine-grained minimum FPS."), ECVF_Default);
+static TAutoConsoleVariable<int32> CVarCoarseSampleTime(
+	TEXT("PerfWarn.CoarseSampleTime"), 600, TEXT("How many seconds we sample the percentage for the coarse-grained minimum FPS."), ECVF_Default);
+static TAutoConsoleVariable<int32> CVarFineMinFPS(
+	TEXT("PerfWarn.FineMinFPS"), 10, TEXT("The FPS threshold below which we warn about for fine-grained sampling."), ECVF_Default);
+static TAutoConsoleVariable<int32> CVarCoarseMinFPS(
+	TEXT("PerfWarn.CoarseMinFPS"), 20, TEXT("The FPS threshold below which we warn about for coarse-grained sampling."), ECVF_Default);
+static TAutoConsoleVariable<int32> CVarFinePercentThreshold(
+	TEXT("PerfWarn.FinePercentThreshold"), 80, TEXT("The percentage of samples that fall below min FPS above which we warn for."), ECVF_Default);
+static TAutoConsoleVariable<int32> CVarCoarsePercentThreshold(
+	TEXT("PerfWarn.CoarsePercentThreshold"), 80, TEXT("The percentage of samples that fall below min FPS above which we warn for."), ECVF_Default);
+
 FPerformanceMonitor::FPerformanceMonitor()
 {
 	LastEnableTime = 0;
 	NotificationTimeout = AutoApplyScalabilityTimeout;
 	bIsNotificationAllowed = true;
 
-	// Register the console variables we need
-	IConsoleVariable* CVar = nullptr;
-
-	CVar = IConsoleManager::Get().RegisterConsoleVariable(TEXT("PerfWarn.FineSampleTime"), 30, TEXT("How many seconds we sample the percentage for the fine-grained minimum FPS."), ECVF_Default);
-	CVar->SetOnChangedCallback(FConsoleVariableDelegate::CreateRaw(this, &FPerformanceMonitor::UpdateMovingAverageSamplers));
-	CVar = IConsoleManager::Get().RegisterConsoleVariable(TEXT("PerfWarn.CoarseSampleTime"), 600, TEXT("How many seconds we sample the percentage for the coarse-grained minimum FPS."), ECVF_Default);
-	CVar->SetOnChangedCallback(FConsoleVariableDelegate::CreateRaw(this, &FPerformanceMonitor::UpdateMovingAverageSamplers));
+	CVarDelegate = FConsoleCommandDelegate::CreateLambda([this]{
+		static const auto CVarFine = IConsoleManager::Get().FindConsoleVariable(TEXT("PerfWarn.FineSampleTime"));
+		static const auto CVarCoarse = IConsoleManager::Get().FindConsoleVariable(TEXT("PerfWarn.CoarseSampleTime"));
 
-	IConsoleManager::Get().RegisterConsoleVariable(TEXT("PerfWarn.FineMinFPS"), 10, TEXT("The FPS threshold below which we warn about for fine-grained sampling."), ECVF_Default);
-	IConsoleManager::Get().RegisterConsoleVariable(TEXT("PerfWarn.CoarseMinFPS"), 20, TEXT("The FPS threshold below which we warn about for coarse-grained sampling."), ECVF_Default);
+		static int32 LastTickFine = -1, LastTickCoarse = -1;
 
-	IConsoleManager::Get().RegisterConsoleVariable(TEXT("PerfWarn.FinePercentThreshold"), 80, TEXT("The percentage of samples that fall below min FPS above which we warn for."), ECVF_Default);
-	IConsoleManager::Get().RegisterConsoleVariable(TEXT("PerfWarn.CoarsePercentThreshold"), 80, TEXT("The percentage of samples that fall below min FPS above which we warn for."), ECVF_Default);
+		if (CVarFine->GetInt() != LastTickFine || CVarCoarse->GetInt() != LastTickCoarse)
+		{
+			LastTickFine = CVarFine->GetInt();
+			LastTickCoarse = CVarCoarse->GetInt();
+			
+			UpdateMovingAverageSamplers();
+		}
+	});
 
-	// Set up the moving average samplers
-	UpdateMovingAverageSamplers();
+	IConsoleManager::Get().RegisterConsoleVariableSink(CVarDelegate);
 }
 
 FPerformanceMonitor::~FPerformanceMonitor()
 {
-	IConsoleManager& Console = IConsoleManager::Get();
-	Console.FindConsoleVariable(TEXT("PerfWarn.FineSampleTime"))->SetOnChangedCallback(FConsoleVariableDelegate());
-	Console.FindConsoleVariable(TEXT("PerfWarn.CoarseSampleTime"))->SetOnChangedCallback(FConsoleVariableDelegate());
+	IConsoleManager::Get().UnregisterConsoleVariableSink(CVarDelegate);	
 }
 
 bool FPerformanceMonitor::WillAutoScalabilityHelp() const
@@ -374,7 +385,7 @@ void FPerformanceMonitor::Reset()
 	bIsNotificationAllowed = true;
 }
 
-void FPerformanceMonitor::UpdateMovingAverageSamplers(IConsoleVariable* Unused)
+void FPerformanceMonitor::UpdateMovingAverageSamplers()
 {
 	static const int NumberOfSamples = 50;
 	IConsoleManager& Console = IConsoleManager::Get();
diff --git a/Engine/Source/Editor/UnrealEd/Private/PerformanceMonitor.h b/Engine/Source/Editor/UnrealEd/Private/PerformanceMonitor.h
--- a/Engine/Source/Editor/UnrealEd/Private/PerformanceMonitor.h
+++ b/Engine/Source/Editor/UnrealEd/Private/PerformanceMonitor.h
@@ -175,7 +175,7 @@ private:
 	void Reset();
 
 	/** Update the moving average samplers to match the settings specified in the console variables */
-	void UpdateMovingAverageSamplers(IConsoleVariable* Unused = nullptr);
+	void UpdateMovingAverageSamplers();
 
 	/** Moving average data for the fine and coarse-grained moving averages */
 	FMovingAverage FineMovingAverage, CoarseMovingAverage;
@@ -194,4 +194,7 @@ private:
 
 	/** Whether the notification is allowed to pop up this session */
 	bool bIsNotificationAllowed;
+
+	/** Console variable delegate for checking when the console variables have changed */
+	FConsoleCommandDelegate CVarDelegate;
 };
\ No newline at end of file
