Provide the correct ExportRootScope so that PPF_ExportsNotFullyQualifiedPaths works correctly when exporting Actors as part of a Level T3D Export (copy)

diff --git a/Engine/Source/Editor/UnrealEd/Private/EditorExporters.cpp b/Engine/Source/Editor/UnrealEd/Private/EditorExporters.cpp
--- a/Engine/Source/Editor/UnrealEd/Private/EditorExporters.cpp
+++ b/Engine/Source/Editor/UnrealEd/Private/EditorExporters.cpp
@@ -462,7 +462,9 @@ bool ULevelExporterT3D::ExportText( const FExportObjectInnerContext* Context, UO
 					FCString::Spc(TextIndent), *Actor->GetClass()->GetName(), *Actor->GetName(),
 					*Actor->GetArchetype()->GetClass()->GetName(), *Actor->GetArchetype()->GetPathName(), *ParentActorString, *SocketNameString, *GroupActor );
 
+				ExportRootScope = Actor;
 				ExportObjectInner( Context, Actor, Ar, PortFlags | PPF_ExportsNotFullyQualified );
+				ExportRootScope = nullptr;
 
 				Ar.Logf( TEXT("%sEnd Actor\r\n"), FCString::Spc(TextIndent) );
 				Actor->AttachRootComponentToActor(ParentActor, SocketName, EAttachLocation::KeepWorldPosition);
