Store current feature level preview setting so that subsequent new maps and map loads stay at the same feature level. (UE-5453)

diff --git a/Engine/Source/Editor/LevelEditor/Private/LevelEditorActions.cpp b/Engine/Source/Editor/LevelEditor/Private/LevelEditorActions.cpp
--- a/Engine/Source/Editor/LevelEditor/Private/LevelEditorActions.cpp
+++ b/Engine/Source/Editor/LevelEditor/Private/LevelEditorActions.cpp
@@ -471,6 +471,9 @@ bool FLevelEditorActionCallbacks::IsMaterialQualityLevelChecked( EMaterialQualit
 
 void FLevelEditorActionCallbacks::SetFeatureLevelPreview(ERHIFeatureLevel::Type InPreviewFeatureLevel)
 {
+	// Record this feature level as we want to use it for all subsequent level creation and loading
+	GEditor->DefaultWorldFeatureLevel = InPreviewFeatureLevel;
+
 	GetWorld()->ChangeFeatureLevel(InPreviewFeatureLevel);
 
 	// Update any currently running PIE sessions.
diff --git a/Engine/Source/Editor/UnrealEd/Classes/Editor/EditorEngine.h b/Engine/Source/Editor/UnrealEd/Classes/Editor/EditorEngine.h
--- a/Engine/Source/Editor/UnrealEd/Classes/Editor/EditorEngine.h
+++ b/Engine/Source/Editor/UnrealEd/Classes/Editor/EditorEngine.h
@@ -524,6 +524,9 @@ class UNREALED_API UEditorEngine : public UEngine
 	/** Cached state of whether or not we have a worldsettings actor in the selection */
 	bool bIsWorldSettingsSelected;
 
+	/** The feature level we should use when loading or creating a new world */
+	ERHIFeatureLevel::Type DefaultWorldFeatureLevel;
+
 public:
 
 	/** The "manager" of all the layers for the UWorld currently being edited */
diff --git a/Engine/Source/Editor/UnrealEd/Classes/Factories/WorldFactory.h b/Engine/Source/Editor/UnrealEd/Classes/Factories/WorldFactory.h
--- a/Engine/Source/Editor/UnrealEd/Classes/Factories/WorldFactory.h
+++ b/Engine/Source/Editor/UnrealEd/Classes/Factories/WorldFactory.h
@@ -15,6 +15,7 @@ class UWorldFactory : public UFactory
 
 	TEnumAsByte<EWorldType::Type> WorldType;
 	bool bInformEngineOfWorld;
+	ERHIFeatureLevel::Type FeatureLevel;
 
 	// Begin UFactory Interface
 	virtual bool ConfigureProperties() override;
diff --git a/Engine/Source/Editor/UnrealEd/Private/Editor.cpp b/Engine/Source/Editor/UnrealEd/Private/Editor.cpp
--- a/Engine/Source/Editor/UnrealEd/Private/Editor.cpp
+++ b/Engine/Source/Editor/UnrealEd/Private/Editor.cpp
@@ -274,6 +274,7 @@ UEditorEngine::UEditorEngine(const FObjectInitializer& ObjectInitializer)
 	bPlayOnLocalPcSession = false;
 	bAllowMultiplePIEWorlds = true;
 	NumOnlinePIEInstances = 0;
+	DefaultWorldFeatureLevel = GMaxRHIFeatureLevel;
 }
 
 
diff --git a/Engine/Source/Editor/UnrealEd/Private/EditorServer.cpp b/Engine/Source/Editor/UnrealEd/Private/EditorServer.cpp
--- a/Engine/Source/Editor/UnrealEd/Private/EditorServer.cpp
+++ b/Engine/Source/Editor/UnrealEd/Private/EditorServer.cpp
@@ -1876,6 +1876,7 @@ UWorld* UEditorEngine::NewMap()
 	UWorldFactory* Factory = ConstructObject<UWorldFactory>(UWorldFactory::StaticClass());
 	Factory->WorldType = EWorldType::Editor;
 	Factory->bInformEngineOfWorld = true;
+	Factory->FeatureLevel = GEditor->DefaultWorldFeatureLevel;
 	UPackage* Pkg = CreatePackage( NULL, NULL );
 	EObjectFlags Flags = RF_Public | RF_Standalone;
 	UWorld* NewWorld = CastChecked<UWorld>(Factory->FactoryCreateNew(UWorld::StaticClass(), Pkg, TEXT("NewWorld"), Flags, NULL, GWarn));
@@ -2187,24 +2188,37 @@ bool UEditorEngine::Map_Load(const TCHAR* Str, FOutputDevice& Ar)
 
 				World->WorldType = EWorldType::Editor;
 
+				// Parse requested feature level if supplied
+				int32 FeatureLevelIndex = (int32)GMaxRHIFeatureLevel;
+				FParse::Value(Str, TEXT("FEATURELEVEL="), FeatureLevelIndex);
+				FeatureLevelIndex = FMath::Clamp(FeatureLevelIndex, 0, (int32)ERHIFeatureLevel::Num);
+
 				if (World->bIsWorldInitialized)
 				{
 					// If we are using a previously initialized world, make sure it has a physics scene and FXSystem.
 					// Inactive worlds are already initialized but lack these two objects for memory reasons.
 					World->ClearWorldComponents();
 
-					if ( World->GetPhysicsScene() == nullptr )
+					if (World->FeatureLevel == FeatureLevelIndex)
 					{
-						World->CreatePhysicsScene();
-					}
+						if (World->GetPhysicsScene() == nullptr)
+						{
+							World->CreatePhysicsScene();
+						}
 
-					if ( World->FXSystem == nullptr )
+						if (World->FXSystem == nullptr)
+						{
+							World->CreateFXSystem();
+						}
+					}
+					else
 					{
-						World->CreateFXSystem();
+						World->ChangeFeatureLevel((ERHIFeatureLevel::Type)FeatureLevelIndex);
 					}
 				}
 				else
 				{
+					World->FeatureLevel = (ERHIFeatureLevel::Type)FeatureLevelIndex;
 					World->InitWorld( GetEditorWorldInitializationValues() );
 				}
 
diff --git a/Engine/Source/Editor/UnrealEd/Private/Factories/WorldFactory.cpp b/Engine/Source/Editor/UnrealEd/Private/Factories/WorldFactory.cpp
--- a/Engine/Source/Editor/UnrealEd/Private/Factories/WorldFactory.cpp
+++ b/Engine/Source/Editor/UnrealEd/Private/Factories/WorldFactory.cpp
@@ -11,6 +11,7 @@ UWorldFactory::UWorldFactory(const FObjectInitializer& ObjectInitializer)
 	SupportedClass = UWorld::StaticClass();
 	WorldType = EWorldType::Inactive;
 	bInformEngineOfWorld = false;
+	FeatureLevel = ERHIFeatureLevel::Num;
 }
 
 bool UWorldFactory::ConfigureProperties()
@@ -22,7 +23,7 @@ UObject* UWorldFactory::FactoryCreateNew(UClass* Class, UObject* InParent, FName
 {
 	// Create a new world.
 	const bool bAddToRoot = false;
-	UWorld* NewWorld = UWorld::CreateWorld(WorldType, bInformEngineOfWorld, Name, Cast<UPackage>(InParent), bAddToRoot);
+	UWorld* NewWorld = UWorld::CreateWorld(WorldType, bInformEngineOfWorld, Name, Cast<UPackage>(InParent), bAddToRoot, FeatureLevel);
 	NewWorld->SetFlags(Flags);
 	NewWorld->ThumbnailInfo = ConstructObject<UWorldThumbnailInfo>(UWorldThumbnailInfo::StaticClass(), NewWorld);
 
diff --git a/Engine/Source/Editor/UnrealEd/Private/FileHelpers.cpp b/Engine/Source/Editor/UnrealEd/Private/FileHelpers.cpp
--- a/Engine/Source/Editor/UnrealEd/Private/FileHelpers.cpp
+++ b/Engine/Source/Editor/UnrealEd/Private/FileHelpers.cpp
@@ -1903,7 +1903,7 @@ void FEditorFileUtils::LoadMap(const FString& InFilename, bool LoadAsTemplate, b
 	// Deactivate any editor modes when loading a new map
 	GLevelEditorModeTools().DeactivateAllModes();
 
-	FString LoadCommand = FString::Printf( TEXT("MAP LOAD FILE=\"%s\" TEMPLATE=%d SHOWPROGRESS=%d"), *Filename, LoadAsTemplate, bShowProgress );
+	FString LoadCommand = FString::Printf(TEXT("MAP LOAD FILE=\"%s\" TEMPLATE=%d SHOWPROGRESS=%d FEATURELEVEL=%d"), *Filename, LoadAsTemplate, bShowProgress, (int32)GEditor->DefaultWorldFeatureLevel);
 	bool bResult = GUnrealEd->Exec( NULL, *LoadCommand );
 
 	UWorld* World = GWorld;
diff --git a/Engine/Source/Runtime/Engine/Classes/Engine/World.h b/Engine/Source/Runtime/Engine/Classes/Engine/World.h
--- a/Engine/Source/Runtime/Engine/Classes/Engine/World.h
+++ b/Engine/Source/Runtime/Engine/Classes/Engine/World.h
@@ -1977,7 +1977,7 @@ public:
 	/**
 	 * Static function that creates a new UWorld and returns a pointer to it
 	 */
-	static UWorld* CreateWorld( const EWorldType::Type InWorldType, bool bInformEngineOfWorld, FName WorldName = NAME_None, UPackage* InWorldPackage = NULL, bool bAddToRoot = true );
+	static UWorld* CreateWorld( const EWorldType::Type InWorldType, bool bInformEngineOfWorld, FName WorldName = NAME_None, UPackage* InWorldPackage = NULL, bool bAddToRoot = true, ERHIFeatureLevel::Type InFeatureLevel = ERHIFeatureLevel::Num );
 
 	/** 
 	 * Destroy this World instance. If destroying the world to load a different world, supply it here to prevent GC of the new world or it's sublevels.
diff --git a/Engine/Source/Runtime/Engine/Private/World.cpp b/Engine/Source/Runtime/Engine/Private/World.cpp
--- a/Engine/Source/Runtime/Engine/Private/World.cpp
+++ b/Engine/Source/Runtime/Engine/Private/World.cpp
@@ -1027,8 +1027,13 @@ void UWorld::DestroyWorld( bool bInformEngineOfWorld, UWorld* NewWorld )
 	}
 }
 
-UWorld* UWorld::CreateWorld( const EWorldType::Type InWorldType, bool bInformEngineOfWorld, FName WorldName, UPackage* InWorldPackage, bool bAddToRoot )
+UWorld* UWorld::CreateWorld(const EWorldType::Type InWorldType, bool bInformEngineOfWorld, FName WorldName, UPackage* InWorldPackage, bool bAddToRoot, ERHIFeatureLevel::Type InFeatureLevel)
 {
+	if (InFeatureLevel >= ERHIFeatureLevel::Num)
+	{
+		InFeatureLevel = GMaxRHIFeatureLevel;
+	}
+
 	// Create a new package unless we're a commandlet in which case we keep the dummy world in the transient package.
 	UPackage* WorldPackage = InWorldPackage;
 	if ( !WorldPackage )
@@ -1050,8 +1055,9 @@ UWorld* UWorld::CreateWorld( const EWorldType::Type InWorldType, bool bInformEng
 
 	// Create new UWorld, ULevel and UModel.
 	const FString WorldNameString = (WorldName != NAME_None) ? WorldName.ToString() : TEXT("NewWorld");
-	UWorld* NewWorld = new( WorldPackage				, *WorldNameString			) UWorld(FObjectInitializer(),FURL(NULL));
+	UWorld* NewWorld = new(WorldPackage, *WorldNameString) UWorld(FObjectInitializer(),FURL(NULL));
 	NewWorld->WorldType = InWorldType;
+	NewWorld->FeatureLevel = InFeatureLevel;
 	NewWorld->InitializeNewWorld(UWorld::InitializationValues().ShouldSimulatePhysics(false).EnableTraceCollision(true).CreateNavigation(InWorldType == EWorldType::Editor).CreateAISystem(InWorldType == EWorldType::Editor));
 
 	// Clear the dirty flag set during SpawnActor and UpdateLevelComponents
