Added code to ensure that actor folder information for stale worlds is cleaned up during PIE

diff --git a/Engine/Source/Editor/UnrealEd/Private/EditorActorFolders.cpp b/Engine/Source/Editor/UnrealEd/Private/EditorActorFolders.cpp
--- a/Engine/Source/Editor/UnrealEd/Private/EditorActorFolders.cpp
+++ b/Engine/Source/Editor/UnrealEd/Private/EditorActorFolders.cpp
@@ -81,13 +81,11 @@ void FActorFolders::Cleanup()
 
 void FActorFolders::Housekeeping()
 {
-	TArray<TWeakObjectPtr<UWorld>> Worlds;
-	TemporaryWorldFolders.GenerateKeyArray(Worlds);
-	for (const auto& World : Worlds)
+	for (auto It = TemporaryWorldFolders.CreateIterator(); It; ++It)
 	{
-		if (!World.Get())
+		if (!It.Key().Get())
 		{
-			TemporaryWorldFolders.Remove(World);
+			It.RemoveCurrent();
 		}
 	}
 }
@@ -214,6 +212,9 @@ UEditorActorFolders& FActorFolders::GetOrCreateFoldersForWorld(UWorld& InWorld)
 
 UEditorActorFolders& FActorFolders::InitializeForWorld(UWorld& InWorld)
 {
+	// Clean up any stale worlds
+	Housekeeping();
+
 	// We intentionally don't pass RF_Transactional to ConstructObject so that we don't record the creation of the object into the undo buffer
 	// (to stop it getting deleted on undo as we manage its lifetime), but we still want it to be RF_Transactional so we can record any changes later
 	UEditorActorFolders* Folders = ConstructObject<UEditorActorFolders>(UEditorActorFolders::StaticClass(), GetTransientPackage(), NAME_None, RF_NoFlags);
