Fix FEdMode::PostUndo so it's actually called!

diff --git a/Engine/Source/Editor/UnrealEd/Private/Editor.cpp b/Engine/Source/Editor/UnrealEd/Private/Editor.cpp
--- a/Engine/Source/Editor/UnrealEd/Private/Editor.cpp
+++ b/Engine/Source/Editor/UnrealEd/Private/Editor.cpp
@@ -117,7 +117,6 @@ FSimpleMulticastDelegate								FEditorDelegates::LoadSelectedAssetsIfNeeded;
 FSimpleMulticastDelegate								FEditorDelegates::DisplayLoadErrors;
 FEditorDelegates::FOnEditorModeTransitioned				FEditorDelegates::EditorModeEnter;
 FEditorDelegates::FOnEditorModeTransitioned				FEditorDelegates::EditorModeExit;
-FSimpleMulticastDelegate								FEditorDelegates::Undo;
 FEditorDelegates::FOnPIEEvent							FEditorDelegates::BeginPIE;
 FEditorDelegates::FOnPIEEvent							FEditorDelegates::EndPIE;
 FEditorDelegates::FOnPIEEvent							FEditorDelegates::PausePIE;
diff --git a/Engine/Source/Editor/UnrealEd/Private/EditorModeManager.cpp b/Engine/Source/Editor/UnrealEd/Private/EditorModeManager.cpp
--- a/Engine/Source/Editor/UnrealEd/Private/EditorModeManager.cpp
+++ b/Engine/Source/Editor/UnrealEd/Private/EditorModeManager.cpp
@@ -52,6 +52,9 @@ FEditorModeTools::FEditorModeTools()
 	USelection::SelectNoneEvent.AddRaw(this, &FEditorModeTools::OnEditorSelectNone);
 	USelection::SelectionChangedEvent.AddRaw(this, &FEditorModeTools::OnEditorSelectionChanged);
 	USelection::SelectObjectEvent.AddRaw(this, &FEditorModeTools::OnEditorSelectionChanged);
+
+	// Register our callback for undo/redo
+	GEditor->RegisterForUndo(this);
 }
 
 FEditorModeTools::~FEditorModeTools()
@@ -60,6 +63,8 @@ FEditorModeTools::~FEditorModeTools()
 	USelection::SelectionChangedEvent.RemoveAll(this);
 	USelection::SelectNoneEvent.RemoveAll(this);
 	USelection::SelectObjectEvent.RemoveAll(this);
+
+	GEditor->UnregisterForUndo(this);
 }
 
 /**
@@ -738,14 +743,21 @@ void FEditorModeTools::DrawHUD( FEditorViewportClient* InViewportClient,FViewpor
 	}
 }
 
-/** Calls PostUndo on all active components */
-void FEditorModeTools::PostUndo()
+/** Calls PostUndo on all active modes */
+void FEditorModeTools::PostUndo(bool bSuccess)
 {
-	for( const auto& Mode : Modes)
+	if (bSuccess)
 	{
-		Mode->PostUndo();
+		for (const auto& Mode : Modes)
+		{
+			Mode->PostUndo();
+		}
 	}
 }
+void FEditorModeTools::PostRedo(bool bSuccess)
+{
+	PostUndo(bSuccess);
+}
 
 /** true if we should allow widget move */
 bool FEditorModeTools::AllowWidgetMove() const
diff --git a/Engine/Source/Editor/UnrealEd/Private/UnrealEdMisc.cpp b/Engine/Source/Editor/UnrealEd/Private/UnrealEdMisc.cpp
--- a/Engine/Source/Editor/UnrealEd/Private/UnrealEdMisc.cpp
+++ b/Engine/Source/Editor/UnrealEd/Private/UnrealEdMisc.cpp
@@ -176,7 +176,6 @@ void FUnrealEdMisc::OnInit()
 	// Register all callback notifications
 	FEditorDelegates::SelectedProps.AddRaw(this, &FUnrealEdMisc::CB_SelectedProps);
 	FEditorDelegates::DisplayLoadErrors.AddRaw(this, &FUnrealEdMisc::CB_DisplayLoadErrors);
-	FEditorDelegates::Undo.AddRaw(this, &FUnrealEdMisc::CB_Undo);
 	FEditorDelegates::MapChange.AddRaw(this, &FUnrealEdMisc::CB_MapChange);
 	FEditorDelegates::RefreshEditor.AddRaw(this, &FUnrealEdMisc::CB_RefreshEditor);
 	FEditorDelegates::PreSaveWorld.AddRaw(this, &FUnrealEdMisc::PreSaveWorld);
@@ -698,7 +697,6 @@ void FUnrealEdMisc::OnExit()
 	// Unregister all events
 	FEditorDelegates::SelectedProps.RemoveAll(this);
 	FEditorDelegates::DisplayLoadErrors.RemoveAll(this);
-	FEditorDelegates::Undo.RemoveAll(this);
 	FEditorDelegates::MapChange.RemoveAll(this);
 	FEditorDelegates::RefreshEditor.RemoveAll(this);
 	FEditorDelegates::PreSaveWorld.RemoveAll(this);
@@ -890,11 +888,6 @@ void FUnrealEdMisc::CB_LevelActorsAdded(AActor* InActor)
 	}
 }
 
-void FUnrealEdMisc::CB_Undo()
-{
-	GLevelEditorModeTools().PostUndo();
-}
-
 void FUnrealEdMisc::CB_PreAutomationTesting()
 {
 	// Shut down SCC if it's enabled, as unit tests shouldn't be allowed to make any modifications to source control
diff --git a/Engine/Source/Editor/UnrealEd/Public/Editor.h b/Engine/Source/Editor/UnrealEd/Public/Editor.h
--- a/Engine/Source/Editor/UnrealEd/Public/Editor.h
+++ b/Engine/Source/Editor/UnrealEd/Public/Editor.h
@@ -131,8 +131,6 @@ struct UNREALED_API FEditorDelegates
 	static FOnEditorModeTransitioned EditorModeEnter;
 	/** Called when an editor mode is being exited */
 	static FOnEditorModeTransitioned EditorModeExit;
-	/** Sent after an undo/redo operation takes place */
-	static FSimpleMulticastDelegate Undo;
 	/** Sent when a PIE session is beginning */
 	static FOnPIEEvent BeginPIE;
 	/** Sent when a PIE session is ending */
diff --git a/Engine/Source/Editor/UnrealEd/Public/EditorModeManager.h b/Engine/Source/Editor/UnrealEd/Public/EditorModeManager.h
--- a/Engine/Source/Editor/UnrealEd/Public/EditorModeManager.h
+++ b/Engine/Source/Editor/UnrealEd/Public/EditorModeManager.h
@@ -5,7 +5,7 @@
 /**
  * A helper class to store the state of the various editor modes.
  */
-class UNREALED_API FEditorModeTools : public FGCObject
+class UNREALED_API FEditorModeTools : public FGCObject, public FEditorUndoClient
 {
 public:
 	FEditorModeTools();
@@ -198,8 +198,11 @@ public:
 	/** Draws the HUD for all active modes */
 	void DrawHUD( FEditorViewportClient* InViewportClient,FViewport* Viewport,const FSceneView* View,FCanvas* Canvas );
 
-	/** Calls PostUndo on all active components */
-	void PostUndo();
+	/** Calls PostUndo on all active modes */
+	// Begin FEditorUndoClient
+	virtual void PostUndo(bool bSuccess) override;
+	virtual void PostRedo(bool bSuccess) override;
+	// End of FEditorUndoClient
 
 	/** True if we should allow widget move */
 	bool AllowWidgetMove() const;
diff --git a/Engine/Source/Editor/UnrealEd/Public/UnrealEdMisc.h b/Engine/Source/Editor/UnrealEd/Public/UnrealEdMisc.h
--- a/Engine/Source/Editor/UnrealEd/Public/UnrealEdMisc.h
+++ b/Engine/Source/Editor/UnrealEd/Public/UnrealEdMisc.h
@@ -216,7 +216,6 @@ private:
 	/** Tells the editor that something has been done to change the map.  Can be anything from loading a whole new map to changing the BSP. */
 	void CB_MapChange( uint32 InFlags );
 	void CB_RedrawAllViewports();
-	void CB_Undo();
 	void CB_EditorModeWindowClosed(const TSharedRef<SWindow>&);
 	void CB_LevelActorsAdded(class AActor* InActor);
 
