Added the Project Setting's GameInstanceClass to the list of cook files.

diff --git a/Engine/Source/Editor/UnrealEd/Private/Commandlets/CookCommandlet.cpp b/Engine/Source/Editor/UnrealEd/Private/Commandlets/CookCommandlet.cpp
--- a/Engine/Source/Editor/UnrealEd/Private/Commandlets/CookCommandlet.cpp
+++ b/Engine/Source/Editor/UnrealEd/Private/Commandlets/CookCommandlet.cpp
@@ -935,6 +935,13 @@ void UCookCommandlet::CollectFilesToCook(TArray<FString>& FilesInPath)
 				AddFileToCook(FilesInPath, Obj);
 			}
 		}
+		if (PlatformEngineIni.GetString(TEXT("/Script/EngineSettings.GameMapsSettings"), TEXT("GameInstanceClass"), Obj))
+		{
+			if (Obj != FName(NAME_None).ToString())
+			{
+				AddFileToCook(FilesInPath, Obj);
+			}
+		}
 	}
 
 	// make sure we cook any extra assets for the default touch interface
diff --git a/Engine/Source/Editor/UnrealEd/Private/CookOnTheFlyServer.cpp b/Engine/Source/Editor/UnrealEd/Private/CookOnTheFlyServer.cpp
--- a/Engine/Source/Editor/UnrealEd/Private/CookOnTheFlyServer.cpp
+++ b/Engine/Source/Editor/UnrealEd/Private/CookOnTheFlyServer.cpp
@@ -2403,6 +2403,13 @@ void UCookOnTheFlyServer::CollectFilesToCook(TArray<FString>& FilesInPath, const
 				AddFileToCook( FilesInPath, Obj);
 			}
 		}
+		if (PlatformEngineIni.GetString(TEXT("/Script/EngineSettings.GameMapsSettings"), TEXT("GameInstanceClass"), Obj))
+		{
+			if (Obj != FName(NAME_None).ToString())
+			{
+				AddFileToCook( FilesInPath, Obj);
+			}
+		}
 	}
 
 	// make sure we cook any extra assets for the default touch interface
