Localizable text from UI_COMMAND macros now appear under the namespace specified by a LOCTEXT_NAMESPACE macro region, if applicable, when gathering text from source. This namespace is nested under the UICommands namespace.

diff --git a/Engine/Source/Editor/UnrealEd/Private/Commandlets/GatherTextFromSourceCommandlet.cpp b/Engine/Source/Editor/UnrealEd/Private/Commandlets/GatherTextFromSourceCommandlet.cpp
--- a/Engine/Source/Editor/UnrealEd/Private/Commandlets/GatherTextFromSourceCommandlet.cpp
+++ b/Engine/Source/Editor/UnrealEd/Private/Commandlets/GatherTextFromSourceCommandlet.cpp
@@ -723,7 +723,8 @@ void UGatherTextFromSourceCommandlet::FCommandMacroDescriptor::TryParse(const FS
 			else
 			{
 				FString Identifier = Arguments[0].Trim();
-				FString Namespace = TEXT("UICommands");
+				const FString UICommandRootNamespace = TEXT("UICommands");
+				FString Namespace = Context.WithinNamespaceDefine && !Context.Namespace.IsEmpty() ? FString::Printf( TEXT("%s.%s"), *UICommandRootNamespace, *Context.Namespace) : UICommandRootNamespace;
 				FString SourceLocation = FString( Context.Filename + TEXT(" - line ") + FString::FromInt(Context.LineNumber) );
 				FString SourceText = Arguments[1].Trim();
 
