Fix non-unity compile errors

diff --git a/Engine/Source/Editor/LevelEditor/Private/LevelEditorContextMenu.cpp b/Engine/Source/Editor/LevelEditor/Private/LevelEditorContextMenu.cpp
--- a/Engine/Source/Editor/LevelEditor/Private/LevelEditorContextMenu.cpp
+++ b/Engine/Source/Editor/LevelEditor/Private/LevelEditorContextMenu.cpp
@@ -30,6 +30,7 @@
 #include "EditorClassUtils.h"
 #include "GenericCommands.h"
 #include "Engine/Selection.h"
+#include "Kismet2/BlueprintEditorUtils.h"
 
 #define LOCTEXT_NAMESPACE "LevelViewportContextMenu"
 
diff --git a/Engine/Source/Editor/UnrealEd/Private/EditorHook.cpp b/Engine/Source/Editor/UnrealEd/Private/EditorHook.cpp
--- a/Engine/Source/Editor/UnrealEd/Private/EditorHook.cpp
+++ b/Engine/Source/Editor/UnrealEd/Private/EditorHook.cpp
@@ -7,6 +7,8 @@
 #include "Editor/LevelEditor/Public/LevelEditor.h"
 #include "Engine/Selection.h"
 
+#include "LandscapeProxy.h"
+
 uint32			EngineThreadId;
 const TCHAR*	GItem;
 const TCHAR*	GValue;
diff --git a/Engine/Source/Editor/UnrealEd/Private/Factories/EditorFactories.cpp b/Engine/Source/Editor/UnrealEd/Private/Factories/EditorFactories.cpp
--- a/Engine/Source/Editor/UnrealEd/Private/Factories/EditorFactories.cpp
+++ b/Engine/Source/Editor/UnrealEd/Private/Factories/EditorFactories.cpp
@@ -103,6 +103,8 @@
 #include "Engine/DataTable.h"
 #include "DataTableEditorUtils.h"
 
+#include "Kismet2/BlueprintEditorUtils.h"
+
 DEFINE_LOG_CATEGORY_STATIC(LogEditorFactories, Log, All);
 
 #define LOCTEXT_NAMESPACE "EditorFactories"
