Add UNREAL_API to FbxFactory.  PR #601.

diff --git a/Engine/Source/Editor/UnrealEd/Classes/Factories/FbxFactory.h b/Engine/Source/Editor/UnrealEd/Classes/Factories/FbxFactory.h
--- a/Engine/Source/Editor/UnrealEd/Classes/Factories/FbxFactory.h
+++ b/Engine/Source/Editor/UnrealEd/Classes/Factories/FbxFactory.h
@@ -4,7 +4,7 @@
 #include "FbxFactory.generated.h"
 
 UCLASS(hidecategories=Object)
-class UFbxFactory : public UFactory
+class UNREALED_API UFbxFactory : public UFactory
 {
 	GENERATED_UCLASS_BODY()
 
