Classes are not replaced while objects reinstancing. Classes should only be replaced in Reparent.. function.
Previously SuperClass could be implicitly replaced, without ensuring proper data layout.

diff --git a/Engine/Source/Editor/UnrealEd/Private/Kismet2/KismetReinstanceUtilities.cpp b/Engine/Source/Editor/UnrealEd/Private/Kismet2/KismetReinstanceUtilities.cpp
--- a/Engine/Source/Editor/UnrealEd/Private/Kismet2/KismetReinstanceUtilities.cpp
+++ b/Engine/Source/Editor/UnrealEd/Private/Kismet2/KismetReinstanceUtilities.cpp
@@ -137,6 +137,14 @@ FBlueprintCompileReinstancer::FBlueprintCompileReinstancer(UClass* InClassToRein
 		ClassToReinstance->ClassFlags &= ~CLASS_NewerVersionExists;
 		GIsDuplicatingClassForReinstancing = false;
 
+		auto BPGDuplicatedClass = Cast<UBlueprintGeneratedClass>(DuplicatedClass);
+		auto DuplicatedClassUberGraphFunction = BPGDuplicatedClass ? BPGDuplicatedClass->UberGraphFunction : nullptr;
+		if (DuplicatedClassUberGraphFunction)
+		{
+			DuplicatedClassUberGraphFunction->Bind();
+			DuplicatedClassUberGraphFunction->StaticLink(true);
+		}
+
 		// Bind and link the duplicate class, so that it has the proper duplicate property offsets
 		DuplicatedClass->Bind();
 		DuplicatedClass->StaticLink(true);
@@ -766,12 +774,6 @@ void FBlueprintCompileReinstancer::ReplaceInstancesOfClass(UClass* OldClass, UCl
 
 	SourceObjects.Append(DstObjects);
 
-	// Add the old->new class mapping into the fixup map
-	OldToNewInstanceMap.Add(OldClass, NewClass);
-
-	// Add in the old class to this pass, so class references are fixed up
-	SourceObjects.Add(OldClass);
-
 	FReplaceReferenceHelper::IncludeCDO(OldClass, NewClass, OldToNewInstanceMap, SourceObjects, OriginalCDO);
 	FReplaceReferenceHelper::FindAndReplaceReferences(SourceObjects, ObjectsThatShouldUseOldStuff, ObjectsToReplace, OldToNewInstanceMap, ReinstancedObjectsWeakReferenceMap);
 
