Added method to set default visibility of Unreal Widget in viewport clients.

diff --git a/Engine/Source/Editor/MaterialEditor/Private/SMaterialEditorViewport.cpp b/Engine/Source/Editor/MaterialEditor/Private/SMaterialEditorViewport.cpp
--- a/Engine/Source/Editor/MaterialEditor/Private/SMaterialEditorViewport.cpp
+++ b/Engine/Source/Editor/MaterialEditor/Private/SMaterialEditorViewport.cpp
@@ -65,6 +65,9 @@ FMaterialEditorViewportClient::FMaterialEditorViewportClient(TWeakPtr<IMaterialE
 
 	OverrideNearClipPlane(1.0f);
 	bUsingOrbitCamera = true;
+
+	// Don't want to display the widget in this viewport
+	Widget->SetDefaultVisibility(false);
 }
 
 
diff --git a/Engine/Source/Editor/UnrealEd/Private/UnrealWidget.cpp b/Engine/Source/Editor/UnrealEd/Private/UnrealWidget.cpp
--- a/Engine/Source/Editor/UnrealEd/Private/UnrealWidget.cpp
+++ b/Engine/Source/Editor/UnrealEd/Private/UnrealWidget.cpp
@@ -64,6 +64,7 @@ FWidget::FWidget()
 
 	bDragging = false;
 	bSnapEnabled = false;
+	bDefaultVisibility = true;
 }
 
 extern ENGINE_API void StringSize(UFont* Font,int32& XL,int32& YL,const TCHAR* Text, FCanvas* Canvas);
@@ -118,9 +119,9 @@ void FWidget::Render( const FSceneView* View,FPrimitiveDrawInterface* PDI, FEdit
 	//reset HUD text
 	HUDString.Empty();
 
-	bool bDrawModeSupportsWidgetDrawing = true;
+	bool bDrawModeSupportsWidgetDrawing = bDefaultVisibility;
 
-	if( EditorModeTools )
+	if( EditorModeTools && ActiveModes.Num() > 0)
 	{
 		bDrawModeSupportsWidgetDrawing = false;
 		// Check to see of any active modes support widget drawing
diff --git a/Engine/Source/Editor/UnrealEd/Public/UnrealWidget.h b/Engine/Source/Editor/UnrealEd/Public/UnrealWidget.h
--- a/Engine/Source/Editor/UnrealEd/Public/UnrealWidget.h
+++ b/Engine/Source/Editor/UnrealEd/Public/UnrealWidget.h
@@ -92,6 +92,16 @@ public:
 	/** Only some modes support Absolute Translation Movement.  Check current mode */
 	static bool AllowsAbsoluteTranslationMovement(FWidget::EWidgetMode WidgetMode);
 
+	/**
+	 * Sets the default visibility of the widget, if it is not overridden by an active editor mode tool.
+	 *
+	 * @param	bInVisibility	true for visible
+	 */
+	void SetDefaultVisibility(bool bInDefaultVisibility)
+	{
+		bDefaultVisibility = bInDefaultVisibility;
+	}
+
 	/**
 	 * Sets the axis currently being moused over.  Typically called by FMouseDeltaTracker or FLevelEditorViewportClient.
 	 *
@@ -377,6 +387,8 @@ private:
 	bool bDragging;
 	/** Whether or not snapping is enabled for all actors */
 	bool bSnapEnabled;
+	/** Default visibility for the widget if an Editor Mode Tool doesn't override it */
+	bool bDefaultVisibility;
 
 	/** Total delta rotation applied since the widget was dragged */
 	float TotalDeltaRotation;
