Localization commandlets attempting source control operations now properly handle not being able to add files to source control.

diff --git a/Engine/Source/Editor/UnrealEd/Private/Commandlets/GatherTextCommandletBase.cpp b/Engine/Source/Editor/UnrealEd/Private/Commandlets/GatherTextCommandletBase.cpp
--- a/Engine/Source/Editor/UnrealEd/Private/Commandlets/GatherTextCommandletBase.cpp
+++ b/Engine/Source/Editor/UnrealEd/Private/Commandlets/GatherTextCommandletBase.cpp
@@ -777,7 +777,7 @@ bool FGatherTextSCC::CheckOutFile(const FString& InFile, FText& OutError)
 				OutError = FText::Format(NSLOCTEXT("GatherTextCmdlet", "FailedToCheckOutFile", "Failed to check out file '{Filepath}'."), Args);
 			}
 		}
-		else if(!SourceControlState->IsSourceControlled())
+		else if(!SourceControlState->IsSourceControlled() && SourceControlState->CanAdd())
 		{
 			bSuccessfullyCheckedOut = (SourceControlProvider.Execute( ISourceControlOperation::Create<FMarkForAdd>(), FilesToBeCheckedOut ) == ECommandResult::Succeeded);
 			if (!bSuccessfullyCheckedOut)
