The text localization report generating commandlet no longer fails entirely if a destination file path isn't checked out for editing when using source control, but will still fail if it is not writable.

diff --git a/Engine/Source/Editor/UnrealEd/Private/Commandlets/GenerateTextLocalizationReportCommandlet.cpp b/Engine/Source/Editor/UnrealEd/Private/Commandlets/GenerateTextLocalizationReportCommandlet.cpp
--- a/Engine/Source/Editor/UnrealEd/Private/Commandlets/GenerateTextLocalizationReportCommandlet.cpp
+++ b/Engine/Source/Editor/UnrealEd/Private/Commandlets/GenerateTextLocalizationReportCommandlet.cpp
@@ -345,8 +345,7 @@ bool UGenerateTextLocalizationReportCommandlet::ProcessWordCountReport(const FSt
 		FText SCCErrorText;
 		if (!SourceControlInfo->CheckOutFile(ReportFilePath, SCCErrorText))
 		{
-			UE_LOG(LogGenerateTextLocalizationReportCommandlet, Error, TEXT("Check out of file %s failed: %s"), *ReportFilePath, *SCCErrorText.ToString());
-			return false;
+			UE_LOG(LogGenerateTextLocalizationReportCommandlet, Warning, TEXT("Check out of file %s failed: %s"), *ReportFilePath, *SCCErrorText.ToString());
 		}
 	}
 
