In-world component selection - outlines show up around only the selected component

diff --git a/Engine/Source/Editor/LevelEditor/Private/ComponentsTree.cpp b/Engine/Source/Editor/LevelEditor/Private/ComponentsTree.cpp
--- a/Engine/Source/Editor/LevelEditor/Private/ComponentsTree.cpp
+++ b/Engine/Source/Editor/LevelEditor/Private/ComponentsTree.cpp
@@ -230,6 +230,12 @@ void SComponentsTree::OnTreeSelectionChanged(FComponentTreeNodePtrType InSelecte
 		PropertyView->SetObjects(Objects, false);
 
 		SelectedComponents->EndBatchSelectOperation();
+
+		if (Actor.IsValid())
+		{
+			GUnrealEd->SetActorSelectionFlags(Actor.Get());
+		}
+		
 		GUnrealEd->UpdatePivotLocationForSelection(true);
 		GEditor->RedrawLevelEditingViewports();
 	}
diff --git a/Engine/Source/Editor/UnrealEd/Classes/Editor/UnrealEdEngine.h b/Engine/Source/Editor/UnrealEd/Classes/Editor/UnrealEdEngine.h
--- a/Engine/Source/Editor/UnrealEd/Classes/Editor/UnrealEdEngine.h
+++ b/Engine/Source/Editor/UnrealEd/Classes/Editor/UnrealEdEngine.h
@@ -735,6 +735,7 @@ class UNREALED_API UUnrealEdEngine : public UEditorEngine, public FNotifyHook
 	
 protected:
 	EWriteDisallowedWarningState GetWarningStateForWritePermission(const FString& PackageName) const;
+	bool IsComponentSelected(const UPrimitiveComponent* PrimComponent);
 
 	/** The package auto-saver instance used by the editor */
 	TUniquePtr<IPackageAutoSaver> PackageAutoSaver;
diff --git a/Engine/Source/Editor/UnrealEd/Private/EditorSelectUtils.cpp b/Engine/Source/Editor/UnrealEd/Private/EditorSelectUtils.cpp
--- a/Engine/Source/Editor/UnrealEd/Private/EditorSelectUtils.cpp
+++ b/Engine/Source/Editor/UnrealEd/Private/EditorSelectUtils.cpp
@@ -281,7 +281,10 @@ void UUnrealEdEngine::UpdatePivotLocationForSelection( bool bOnChange )
 			UActorComponent* Component = CastChecked<UActorComponent>(*It);
 			AActor* ComponentOwner = Component->GetOwner();
 			check(ComponentOwner);
-			check(GetSelectedActors()->IsSelected(ComponentOwner));
+
+			auto SelectedActors = GetSelectedActors();
+			const bool bIsOwnerSelected = SelectedActors->IsSelected(ComponentOwner);
+			check(bIsOwnerSelected);
 
 			if (ComponentOwner->GetWorld() == GWorld)
 			{
@@ -381,7 +384,6 @@ void UUnrealEdEngine::NoteSelectionChange(bool bComponentSelectionChanged)
 	RedrawLevelEditingViewports();
 }
 
-
 void UUnrealEdEngine::SelectGroup(AGroupActor* InGroupActor, bool bForceSelection/*=false*/, bool bInSelected/*=true*/, bool bNotify/*=true*/)
 {
 	USelection* SelectedActors = GetSelectedActors();
@@ -557,9 +559,28 @@ void UUnrealEdEngine::SelectActor(AActor* Actor, bool bInSelected, bool bNotify,
 				for (UActorComponent* Component : Actor->GetComponents())
 				{
 					GetSelectedComponents()->Deselect( Component );
+
+					// Remove the selection override delegates from the deselected components
+					auto PrimComponent = Cast<UPrimitiveComponent>(Component);
+					if (PrimComponent && PrimComponent->SelectionOverrideDelegate.IsBound())
+					{
+						PrimComponent->SelectionOverrideDelegate.Unbind();
+					}
+				}
+			}
+			else
+			{
+				for (UActorComponent* Component : Actor->GetComponents())
+				{
+					// Establish the selection override delegate for each of the actor's components
+					auto PrimComponent = Cast<UPrimitiveComponent>(Component);
+					if (PrimComponent && !PrimComponent->SelectionOverrideDelegate.IsBound())
+					{
+						PrimComponent->SelectionOverrideDelegate = UPrimitiveComponent::FSelectionOverride::CreateUObject(this, &UUnrealEdEngine::IsComponentSelected);
+					}
 				}
 			}
-			
+
 			//A fast path to mark selection rather than reconnecting ALL components for ALL actors that have changed state
 			SetActorSelectionFlags (Actor);
 
@@ -602,6 +623,12 @@ void UUnrealEdEngine::SelectComponent(UActorComponent* Component, bool bInSelect
 
 		GetSelectedComponents()->Select(Component, bInSelected);
 
+		auto PrimComponent = Cast<UPrimitiveComponent>(Component);
+		if (PrimComponent && !PrimComponent->SelectionOverrideDelegate.IsBound())
+		{
+			PrimComponent->SelectionOverrideDelegate = UPrimitiveComponent::FSelectionOverride::CreateUObject(this, &UUnrealEdEngine::IsComponentSelected);
+		}
+
 		// Update the selection visualization
 		AActor* ComponentOwner = Component->GetOwner();
 		if (ComponentOwner != nullptr)
@@ -631,6 +658,16 @@ void UUnrealEdEngine::SelectComponent(UActorComponent* Component, bool bInSelect
 	}
 }
 
+bool UUnrealEdEngine::IsComponentSelected(const UPrimitiveComponent* PrimComponent)
+{
+	if (GetSelectedComponentCount() > 0)
+	{
+		return GetSelectedComponents()->IsSelected(PrimComponent);
+	}
+
+	return GetSelectedActors()->IsSelected(PrimComponent->GetOwner());
+}
+
 void UUnrealEdEngine::SelectBSPSurf(UModel* InModel, int32 iSurf, bool bSelected, bool bNoteSelectionChange)
 {
 	if( GEdSelectionLock )
