Empty actors now create anonymous blueprints

diff --git a/Engine/Source/Editor/UnrealEd/Classes/ActorFactories/ActorFactoryEmptyActor.h b/Engine/Source/Editor/UnrealEd/Classes/ActorFactories/ActorFactoryEmptyActor.h
--- a/Engine/Source/Editor/UnrealEd/Classes/ActorFactories/ActorFactoryEmptyActor.h
+++ b/Engine/Source/Editor/UnrealEd/Classes/ActorFactories/ActorFactoryEmptyActor.h
@@ -14,4 +14,7 @@ public:
 
 protected:
 	virtual AActor* SpawnActor( UObject* Asset, ULevel* InLevel, const FVector& Location, const FRotator& Rotation, EObjectFlags ObjectFlags, const FName& Name ) override;
+private:
+	AActor* SpawnActorForDragPreview( UObject* Asset, ULevel* InLevel, const FVector& Location, const FRotator& Rotation, EObjectFlags ObjectFlags, const FName& Name );
+	void SetupEditorOnlyBillboardComponent( class UBillboardComponent* BillboardComponent );
 };
diff --git a/Engine/Source/Editor/UnrealEd/Private/Factories/ActorFactory.cpp b/Engine/Source/Editor/UnrealEd/Private/Factories/ActorFactory.cpp
--- a/Engine/Source/Editor/UnrealEd/Private/Factories/ActorFactory.cpp
+++ b/Engine/Source/Editor/UnrealEd/Private/Factories/ActorFactory.cpp
@@ -1132,33 +1132,94 @@ AActor* UActorFactoryEmptyActor::SpawnActor( UObject* Asset, ULevel* InLevel, co
 	AActor* NewActor = nullptr;
 	if(GetDefault<UEditorExperimentalSettings>()->bInWorldBPEditing) 
 	{
-		NewActor = Super::SpawnActor(Asset, InLevel, Location, Rotation, ObjectFlags, Name);
 
-		USceneComponent* RootComponent = ConstructObject<USceneComponent>(USceneComponent::StaticClass(), NewActor, FName("Root"), RF_Transactional);
-		RootComponent->Mobility = EComponentMobility::Static;
-		RootComponent->SetWorldLocationAndRotation(Location, Rotation);
-		NewActor->SetRootComponent(RootComponent);
+		if( ObjectFlags == RF_Transient )
+		{
+			NewActor = SpawnActorForDragPreview( Asset, InLevel, Location, Rotation, ObjectFlags, Name );
+		}
+		else
+		{
+			FName ActorName(TEXT("Actor"));
+			FName UniqueActorName = MakeUniqueObjectName( InLevel, AActor::StaticClass(), ActorName );
 
+			// Spawn a blank actor with an anonymous blueprint
+			// Create and init a new Blueprint
+			const FName BPName = *FString::Printf(TEXT("%s_BPClass"), *UniqueActorName.ToString());
 
-		// Create a new billboard component to serve as a visualization of the actor until there is another primitive component
-		UBillboardComponent* BillboardComponent = ConstructObject<UBillboardComponent>(UBillboardComponent::StaticClass(), NewActor, NAME_None, RF_Transactional);
+			const FName UniqueBPName = MakeUniqueObjectName( InLevel, UBlueprint::StaticClass(), BPName );
+
+			
+			UBlueprint* NewBP = FKismetEditorUtilities::CreateBlueprint(AActor::StaticClass(), InLevel, UniqueBPName, BPTYPE_Normal, UBlueprint::StaticClass(), UBlueprintGeneratedClass::StaticClass());
+			if(NewBP)
+			{
+				NewBP->ClearFlags(RF_Standalone);
 
-		BillboardComponent->Sprite = LoadObject<UTexture2D>(nullptr, TEXT("/Engine/EditorResources/S_Actor.S_Actor"));
-		BillboardComponent->RelativeScale3D = FVector(0.5f, 0.5f, 0.5f);
-		BillboardComponent->SpriteInfo.Category = TEXT("EmptyActors");
-		BillboardComponent->SpriteInfo.DisplayName = LOCTEXT("EmptyActorDisplayName", "Empty Actor");
-		BillboardComponent->AttachTo(RootComponent);
-		BillboardComponent->Mobility = EComponentMobility::Static;
-		BillboardComponent->AlwaysLoadOnClient = false;
-		BillboardComponent->AlwaysLoadOnServer = false;
+				USCS_Node* NewSCSNode = NewBP->SimpleConstructionScript->CreateNode( UBillboardComponent::StaticClass() );
+				
+				UBillboardComponent* BillboardComponent = CastChecked<UBillboardComponent>( NewSCSNode->ComponentTemplate );
 
-		NewActor->SerializedComponents.Add(RootComponent);
-		NewActor->SerializedComponents.Add(BillboardComponent);
+				SetupEditorOnlyBillboardComponent( BillboardComponent );
+				
+				// Set parameters to parent this node to the "inherited" SCS node
+				NewSCSNode->SetParent( NewBP->SimpleConstructionScript->GetDefaultSceneRootNode() );
+
+				NewBP->SimpleConstructionScript->GetDefaultSceneRootNode()->AddChildNode(NewSCSNode);
+
+				FKismetEditorUtilities::CompileBlueprint(NewBP);
+
+				FAssetData AssetData(NewBP);
+				const FString* GeneratedClassPath = AssetData.TagsAndValues.Find("GeneratedClass");
+				if(GeneratedClassPath != NULL && !GeneratedClassPath->IsEmpty())
+				{
+					UClass* GeneratedClass = Cast<UClass>(StaticLoadObject(UClass::StaticClass(), NULL, **GeneratedClassPath, NULL, LOAD_NoWarn, NULL));
+
+					if(GeneratedClass != NULL)
+					{
+						FActorSpawnParameters SpawnInfo;
+						SpawnInfo.OverrideLevel = InLevel;
+						SpawnInfo.ObjectFlags = ObjectFlags;
+						SpawnInfo.Name = Name;
+						NewActor = InLevel->OwningWorld->SpawnActor(GeneratedClass, &Location, &Rotation, SpawnInfo);
+					}
+				}
+
+			}
+		}
 	}
 
 	return NewActor;
 }
 
+AActor* UActorFactoryEmptyActor::SpawnActorForDragPreview( UObject* Asset, ULevel* InLevel, const FVector& Location, const FRotator& Rotation, EObjectFlags ObjectFlags, const FName& Name )
+{
+	// Spawn a temporary actor for dragging around
+	AActor* NewActor = Super::SpawnActor(Asset, InLevel, Location, Rotation, ObjectFlags, Name);
+
+	USceneComponent* RootComponent = ConstructObject<USceneComponent>(USceneComponent::StaticClass(), NewActor, FName("Root"), RF_Transactional);
+	RootComponent->Mobility = EComponentMobility::Static;
+	RootComponent->SetWorldLocationAndRotation(Location, Rotation);
+	NewActor->SetRootComponent(RootComponent);
+
+	// Create a new billboard component to serve as a visualization of the actor until there is another primitive component
+	UBillboardComponent* BillboardComponent = ConstructObject<UBillboardComponent>(UBillboardComponent::StaticClass(), NewActor, NAME_None, RF_Transactional);
+
+	SetupEditorOnlyBillboardComponent( BillboardComponent );
+
+	BillboardComponent->AttachTo(RootComponent);
+	NewActor->SerializedComponents.Add(RootComponent);
+	NewActor->SerializedComponents.Add(BillboardComponent);
+
+	return NewActor;
+}
+
+void UActorFactoryEmptyActor::SetupEditorOnlyBillboardComponent( UBillboardComponent* BillboardComponent )
+{
+	BillboardComponent->Sprite = LoadObject<UTexture2D>(nullptr, TEXT("/Engine/EditorResources/S_Actor.S_Actor"));
+	BillboardComponent->RelativeScale3D = FVector(0.5f, 0.5f, 0.5f);
+	BillboardComponent->Mobility = EComponentMobility::Movable;
+	BillboardComponent->AlwaysLoadOnClient = false;
+	BillboardComponent->AlwaysLoadOnServer = false;
+}
 
 /*-----------------------------------------------------------------------------
 UActorFactoryAmbientSound
