Removed a failure condition from the commandlet to generate text localization reports. Reports should still be generated for the native culture even if there are no other cultures specified.

diff --git a/Engine/Source/Editor/UnrealEd/Private/Commandlets/GenerateTextLocalizationReportCommandlet.cpp b/Engine/Source/Editor/UnrealEd/Private/Commandlets/GenerateTextLocalizationReportCommandlet.cpp
--- a/Engine/Source/Editor/UnrealEd/Private/Commandlets/GenerateTextLocalizationReportCommandlet.cpp
+++ b/Engine/Source/Editor/UnrealEd/Private/Commandlets/GenerateTextLocalizationReportCommandlet.cpp
@@ -154,12 +154,6 @@ bool UGenerateTextLocalizationReportCommandlet::ProcessWordCountReport(const FSt
 	// Get cultures to generate.
 	GetStringArrayFromConfig( *SectionName, TEXT("CulturesToGenerate"), CulturesToGenerate, GatherTextConfigPath );
 
-	if( CulturesToGenerate.Num() == 0 )
-	{
-		UE_LOG(LogGenerateTextLocalizationReportCommandlet, Error, TEXT("No cultures specified for generation."));
-		return false;
-	}
-
 	for(int32 i = 0; i < CulturesToGenerate.Num(); ++i)
 	{
 		if( FInternationalization::Get().GetCulture( CulturesToGenerate[i] ).IsValid() )
