couple of missing files from last submission

diff --git a/Engine/Source/Editor/UnrealEd/Private/Kismet2/BlueprintEditorUtils.cpp b/Engine/Source/Editor/UnrealEd/Private/Kismet2/BlueprintEditorUtils.cpp
--- a/Engine/Source/Editor/UnrealEd/Private/Kismet2/BlueprintEditorUtils.cpp
+++ b/Engine/Source/Editor/UnrealEd/Private/Kismet2/BlueprintEditorUtils.cpp
@@ -2455,6 +2455,22 @@ bool FBlueprintEditorUtils::SupportsConstructionScript(const UBlueprint* Bluepri
 			!(Blueprint->BlueprintType == BPTYPE_FunctionLibrary);
 }
 
+bool FBlueprintEditorUtils::CanClassGenerateEvents(const UClass* InClass)
+{
+	if( InClass )
+	{
+		for( TFieldIterator<UMulticastDelegateProperty> PropertyIt( InClass, EFieldIteratorFlags::IncludeSuper ); PropertyIt; ++PropertyIt )
+		{
+			UProperty* Property = *PropertyIt;
+			if( !Property->HasAnyPropertyFlags( CPF_Parm ) && Property->HasAllPropertyFlags( CPF_BlueprintAssignable ))
+			{
+				return true;
+			}
+		}
+	}
+	return false;
+}
+
 UEdGraph* FBlueprintEditorUtils::FindUserConstructionScript(const UBlueprint* Blueprint)
 {
 	const UEdGraphSchema_K2* Schema = GetDefault<UEdGraphSchema_K2>();
diff --git a/Engine/Source/Editor/UnrealEd/Public/Kismet2/BlueprintEditorUtils.h b/Engine/Source/Editor/UnrealEd/Public/Kismet2/BlueprintEditorUtils.h
--- a/Engine/Source/Editor/UnrealEd/Public/Kismet2/BlueprintEditorUtils.h
+++ b/Engine/Source/Editor/UnrealEd/Public/Kismet2/BlueprintEditorUtils.h
@@ -378,6 +378,12 @@ public:
 	 */
 	static bool IsAnonymousBlueprintClass(const UClass* Class);
 
+	/**
+	 * Checks for events in the argument class
+	 * @param Class	The class to check for events.
+	 */
+	static bool CanClassGenerateEvents(const UClass* Class);
+
 	/**
 	 * If a blueprint is directly tied to a level (level script and anonymous blueprints), this will return a pointer to that level
 	 *
