Use existing (with slightly modified API) Create Blueprint From Actor functionality to create the Anonymous Blueprint

diff --git a/Engine/Source/Editor/LevelEditor/Private/LevelEditorActions.cpp b/Engine/Source/Editor/LevelEditor/Private/LevelEditorActions.cpp
--- a/Engine/Source/Editor/LevelEditor/Private/LevelEditorActions.cpp
+++ b/Engine/Source/Editor/LevelEditor/Private/LevelEditorActions.cpp
@@ -955,28 +955,7 @@ void FLevelEditorActionCallbacks::AddScriptBehavior_Clicked()
 	if (SelectedActor)
 	{
 		const FName Name = *FString::Printf(TEXT("%s_BPClass"), *SelectedActor->GetName());
-
-		// Create and init a new Blueprint
-		UBlueprint* NewBP = FKismetEditorUtilities::CreateBlueprint(SelectedActor->GetClass(), SelectedActor->GetLevel(), Name, BPTYPE_Normal, UBlueprint::StaticClass(), UBlueprintGeneratedClass::StaticClass());
-		if (NewBP)
-		{
-			NewBP->ClearFlags(RF_Standalone);
-			ULevel::LevelDirtiedEvent.Broadcast();
-
-			UEngine::CopyPropertiesForUnrelatedObjects(SelectedActor, NewBP->GeneratedClass->GetDefaultObject<AActor>());
-
-			FKismetEditorUtilities::CompileBlueprint(NewBP);
-
-			// the editor should be opened AFTER being added to the asset 
-			// registry (some systems could queue off of the asset registry 
-			// add event, and already having that blueprint open can be odd)
-			FAssetEditorManager::Get().OpenEditorForAsset(NewBP);
-
-			GEditor->ReplaceSelectedActors(GEditor->FindActorFactoryByClass(UActorFactoryBlueprint::StaticClass()), FAssetData(NewBP));
-
-			// Mark the package dirty...
-			SelectedActor->MarkPackageDirty();
-		}
+		FKismetEditorUtilities::CreateBlueprintFromActor(Name, SelectedActor->GetLevel(), SelectedActor, true);
 	}
 }
 
diff --git a/Engine/Source/Editor/UnrealEd/Private/Kismet2/Kismet2.cpp b/Engine/Source/Editor/UnrealEd/Private/Kismet2/Kismet2.cpp
--- a/Engine/Source/Editor/UnrealEd/Private/Kismet2/Kismet2.cpp
+++ b/Engine/Source/Editor/UnrealEd/Private/Kismet2/Kismet2.cpp
@@ -870,44 +870,54 @@ bool FKismetEditorUtilities::CanCreateBlueprintOfClass(const UClass* Class)
 	return bCanCreateBlueprint && bIsValidClass;
 }
 
-UBlueprint* FKismetEditorUtilities::CreateBlueprintFromActor(const FString& Path, UObject* Object, bool bReplaceActor )
+UBlueprint* FKismetEditorUtilities::CreateBlueprintFromActor(const FString& Path, AActor* Actor, const bool bReplaceActor )
 {
-	UBlueprint* NewBlueprint = NULL;
+	UBlueprint* NewBlueprint = nullptr;
 
-	if (AActor* Actor = Cast<AActor>(Object))
+	// Create a blueprint
+	FString PackageName = Path;
+	FString AssetName = FPackageName::GetLongPackageAssetName(Path);
+
+	// If no AssetName was found, generate a unique asset name.
+	if(AssetName.Len() == 0)
 	{
-		// Create a blueprint
-		FString PackageName = Path;
-		FString AssetName = FPackageName::GetLongPackageAssetName(Path);
+		PackageName = FPackageName::GetLongPackagePath(Path);
+		FString BasePath = PackageName + TEXT("/") + LOCTEXT("BlueprintName_Default", "NewBlueprint").ToString();
+		FAssetToolsModule& AssetToolsModule = FModuleManager::LoadModuleChecked<FAssetToolsModule>("AssetTools");
+		AssetToolsModule.Get().CreateUniqueAssetName(BasePath, TEXT(""), PackageName, AssetName);
+	}
 
+	UPackage* Package = CreatePackage(NULL, *PackageName);
 
-		// If no AssetName was found, generate a unique asset name.
-		if(AssetName.Len() == 0)
-		{
-			PackageName = FPackageName::GetLongPackagePath(Path);
-			FString BasePath = PackageName + TEXT("/") + LOCTEXT("BlueprintName_Default", "NewBlueprint").ToString();
-			FAssetToolsModule& AssetToolsModule = FModuleManager::LoadModuleChecked<FAssetToolsModule>("AssetTools");
-			AssetToolsModule.Get().CreateUniqueAssetName(BasePath, TEXT(""), PackageName, AssetName);
-		}
+	if(Package)
+	{
+		NewBlueprint = CreateBlueprintFromActor(FName(*AssetName), Package, Actor, bReplaceActor);
+	}
 
-		UPackage* Package = CreatePackage(NULL, *PackageName);
+	return NewBlueprint;
+}
+
+UBlueprint* FKismetEditorUtilities::CreateBlueprintFromActor(const FName BlueprintName, UObject* Outer, AActor* Actor, const bool bReplaceActor )
+{
+	UBlueprint* NewBlueprint = nullptr;
 
-		if(Package)
+	if (Actor)
+	{
+		if (Outer)
 		{
 			UActorFactory* FactoryToUse = GEditor->FindActorFactoryForActorClass( Actor->GetClass() );
-			const FName BlueprintName = FName(*AssetName);
 
 			if( FactoryToUse != NULL )
 			{
 				// Create the blueprint
 				UObject* Asset = FactoryToUse->GetAssetFromActorInstance(Actor);
 				// For Actors that don't have an asset associated with them, Asset will be null
-				NewBlueprint = FactoryToUse->CreateBlueprint( Asset, Package, BlueprintName, FName("CreateFromActor") );
+				NewBlueprint = FactoryToUse->CreateBlueprint( Asset, Outer, BlueprintName, FName("CreateFromActor") );
 			}
 			else 
 			{
 				// We don't have a factory, but we can still try to create a blueprint for this actor class
-				NewBlueprint = FKismetEditorUtilities::CreateBlueprint( Object->GetClass(), Package, BlueprintName, EBlueprintType::BPTYPE_Normal, UBlueprint::StaticClass(), UBlueprintGeneratedClass::StaticClass(), FName("CreateFromActor") );
+				NewBlueprint = FKismetEditorUtilities::CreateBlueprint( Actor->GetClass(), Outer, BlueprintName, EBlueprintType::BPTYPE_Normal, UBlueprint::StaticClass(), UBlueprintGeneratedClass::StaticClass(), FName("CreateFromActor") );
 			}
 		}
 
@@ -917,12 +927,12 @@ UBlueprint* FKismetEditorUtilities::CreateBlueprintFromActor(const FString& Path
 			FAssetRegistryModule::AssetCreated(NewBlueprint);
 
 			// Mark the package dirty
-			Package->MarkPackageDirty();
+			Outer->MarkPackageDirty();
 
 			if(NewBlueprint->GeneratedClass)
 			{
 				UObject* CDO = NewBlueprint->GeneratedClass->GetDefaultObject();
-				UEditorEngine::CopyPropertiesForUnrelatedObjects(Object, CDO);
+				UEditorEngine::CopyPropertiesForUnrelatedObjects(Actor, CDO);
 				if(AActor* CDOAsActor = Cast<AActor>(CDO))
 				{
 					if(USceneComponent* Scene = CDOAsActor->GetRootComponent())
diff --git a/Engine/Source/Editor/UnrealEd/Public/Kismet2/KismetEditorUtilities.h b/Engine/Source/Editor/UnrealEd/Public/Kismet2/KismetEditorUtilities.h
--- a/Engine/Source/Editor/UnrealEd/Public/Kismet2/KismetEditorUtilities.h
+++ b/Engine/Source/Editor/UnrealEd/Public/Kismet2/KismetEditorUtilities.h
@@ -93,7 +93,17 @@ public:
 	 * @param bReplaceActor			If true, replace the actor in the scene with one based on the created blueprint
 	 * @return The blueprint created from the actor
 	 */
-	static UBlueprint* CreateBlueprintFromActor(const FString& Path, UObject* Actor, bool bReplaceActor );
+	static UBlueprint* CreateBlueprintFromActor(const FString& Path, AActor* Actor, bool bReplaceActor );
+
+	/** 
+	 * Take an Actor and generate a blueprint based on it. Uses the Actors type as the parent class. 
+	 * @param BlueprintName			The name to use for the Blueprint
+	 * @param Outer					The outer object to create the blueprint within
+	 * @param Actor					The actor to use as the template for the blueprint
+	 * @param bReplaceActor			If true, replace the actor in the scene with one based on the created blueprint
+	 * @return The blueprint created from the actor
+	 */
+	static UBlueprint* CreateBlueprintFromActor(const FName BlueprintName, UObject* Outer, AActor* Actor, bool bReplaceActor );
 
 	/** 
 	 * Take a list of Actors and generate a blueprint  by harvesting the components they have. Uses AActor as parent class type as the parent class. 
