Introduced an "auto checkout" mode to the editor preferences.  Warns when assets which cannot be checked out are edited. Changes to the SPackagesDialog to make it more verbose and include tooltips.

diff --git a/Engine/Config/BaseEditorUserSettings.ini b/Engine/Config/BaseEditorUserSettings.ini
--- a/Engine/Config/BaseEditorUserSettings.ini
+++ b/Engine/Config/BaseEditorUserSettings.ini
@@ -146,6 +146,8 @@ AutoSaveTimeMinutes=10
 ; The number of seconds warning before an autosave
 AutoSaveWarningInSeconds=10
 
+; Whether to automatically perform an SCC checkout on asset modification
+bPromptForCheckoutOnAssetModification=False
 ; Whether to automatically prompt for SCC checkout on asset modification
 bPromptForCheckoutOnAssetModification=True
 ; Source control
diff --git a/Engine/Source/Editor/LevelEditor/Public/LevelEditorActions.h b/Engine/Source/Editor/LevelEditor/Public/LevelEditorActions.h
--- a/Engine/Source/Editor/LevelEditor/Public/LevelEditorActions.h
+++ b/Engine/Source/Editor/LevelEditor/Public/LevelEditorActions.h
@@ -364,12 +364,6 @@ public:
 	/** Selects all subtractive brushes */
 	TSharedPtr< FUICommandInfo > SelectAllSubtractiveBrushes;
 
-	/** Selects all semi-solid brushes */
-	TSharedPtr< FUICommandInfo > SelectAllSemiSolidBrushes;
-
-	/** Selects all non-solid brushes */
-	TSharedPtr< FUICommandInfo > SelectAllNonSolidBrushes;
-
 	/**
 	 * Surface commands                   
 	 */
diff --git a/Engine/Source/Editor/PackagesDialog/Private/PackagesDialog.cpp b/Engine/Source/Editor/PackagesDialog/Private/PackagesDialog.cpp
--- a/Engine/Source/Editor/PackagesDialog/Private/PackagesDialog.cpp
+++ b/Engine/Source/Editor/PackagesDialog/Private/PackagesDialog.cpp
@@ -7,6 +7,7 @@
 IMPLEMENT_MODULE( FPackagesDialogModule, PackagesDialog );
 
 const FVector2D FPackagesDialogModule::DEFAULT_WINDOW_SIZE = FVector2D(600, 400);	// TODO: MIN SIZE SHOULD BE 270, 330 FOR AESTHETICS
+const FVector2D FPackagesDialogModule::EXTRA_WINDOW_WIDTH = FVector2D(150, 0);
 
 /**
  * Called right after the module's DLL has been loaded and the module object has been created
@@ -60,7 +61,8 @@ EDialogReturnType FPackagesDialogModule::ShowPackagesDialog()
 EDialogReturnType FPackagesDialogModule::ShowPackagesDialog(OUT TSet<FString>& InOutIgnoredPackages)
 {
 	/** Ensure the package dialog window was not created already */
-	check( !EditorPackagesDialogWindow.IsValid() );
+	check(!EditorPackagesDialogWindow.IsValid());
+	check(PackagesDialogWidget.IsValid());
 
 	/** Reset the widget, as it may have been re-used */
 	PackagesDialogWidget->Reset();
@@ -69,9 +71,16 @@ EDialogReturnType FPackagesDialogModule::ShowPackagesDialog(OUT TSet<FString>& I
 	PackagesDialogWidget->PopulateIgnoreForSaveItems(InOutIgnoredPackages);
 
 	/** Create the window to host our package dialog widget */
+	FVector2D WindowSize = DEFAULT_WINDOW_SIZE;
+	if (PackagesDialogWidget->IsSourceControlConnectionAllowed())
+	{
+		// Make the window a little wider if there's an additional column for source control status
+		WindowSize += EXTRA_WINDOW_WIDTH;
+	}
+
 	TSharedRef< SWindow > EditorPackagesDialogWindowRef = SNew(SWindow)
 		.Title(PackageDialogTitle)
-		.ClientSize(FPackagesDialogModule::DEFAULT_WINDOW_SIZE);
+		.ClientSize(WindowSize);
 	
 	/** Store weak pointers to the package dialog Slate window and widget */
 	EditorPackagesDialogWindow = EditorPackagesDialogWindowRef;
@@ -121,6 +130,12 @@ void FPackagesDialogModule::SetMessage(const FText& InMessage)
 	PackagesDialogWidget->SetMessage(InMessage);
 }
 
+void FPackagesDialogModule::SetWarning(const FText& InMessage)
+{
+	check(PackagesDialogWidget.IsValid());
+	PackagesDialogWidget->SetWarning(InMessage);
+}
+
 /**
  * Populates the passed in array with the desired packages
  *
diff --git a/Engine/Source/Editor/PackagesDialog/Private/SPackagesDialog.cpp b/Engine/Source/Editor/PackagesDialog/Private/SPackagesDialog.cpp
--- a/Engine/Source/Editor/PackagesDialog/Private/SPackagesDialog.cpp
+++ b/Engine/Source/Editor/PackagesDialog/Private/SPackagesDialog.cpp
@@ -5,7 +5,6 @@
 #include "SPackagesDialog.h"
 #include "Editor/UnrealEd/Public/PackageTools.h"
 #include "AssetToolsModule.h"
-#include "ISourceControlModule.h"
 
 #define LOCTEXT_NAMESPACE "SPackagesDialog"
 
@@ -15,6 +14,7 @@ namespace SPackagesDialogDefs
 	const FName ColumnID_IconLabel( "Icon" );
 	const FName ColumnID_FileLabel( "File" );
 	const FName ColumnID_TypeLabel( "Type" );
+	const FName ColumnID_CheckedOutByLabel( "CheckedOutBy" );
 
 	const float CheckBoxColumnWidth = 23.0f;
 	const float IconColumnWidth = 21.0f;
@@ -23,8 +23,7 @@ namespace SPackagesDialogDefs
 
 UObject* FPackageItem::GetPackageObject() const
 {
-	const bool bIsLegacyOrMapPackage = !PackageTools::IsSingleAssetPackage(EntryName);
-	if ( !bIsLegacyOrMapPackage && !EntryName.StartsWith(TEXT("/Temp/Untitled")) )
+	if ( !EntryName.StartsWith(TEXT("/Temp/Untitled")) )
 	{
 		// Get the object which belongs in this package (there has to be a quicker function than GetObjectsInPackages!)
 		TArray<UPackage*> Packages;
@@ -71,6 +70,7 @@ void SPackagesDialog::Construct(const FArguments& InArgs)
 	bReadOnly = InArgs._ReadOnly.Get();
 	bAllowSourceControlConnection = InArgs._AllowSourceControlConnection.Get();
 	Message = InArgs._Message;
+	Warning = InArgs._Warning;
 	SortByColumn = SPackagesDialogDefs::ColumnID_FileLabel;
 	SortMode = EColumnSortMode::Ascending;
 
@@ -133,6 +133,17 @@ void SPackagesDialog::Construct(const FArguments& InArgs)
 		.FillWidth( 2.0f )
 		);
 
+	if (bAllowSourceControlConnection)
+	{
+		HeaderRowWidget->AddColumn(
+			SHeaderRow::Column(SPackagesDialogDefs::ColumnID_CheckedOutByLabel)
+			.DefaultLabel(LOCTEXT("CheckedOutByColumnLabel", "Checked Out By"))
+			.SortMode(this, &SPackagesDialog::GetColumnSortMode, SPackagesDialogDefs::ColumnID_CheckedOutByLabel)
+			.OnSort(this, &SPackagesDialog::OnColumnSortModeChanged)
+			.FillWidth(4.0f)
+			);
+	}
+
 	this->ChildSlot
 	[
 		SNew(SBorder)
@@ -141,13 +152,17 @@ void SPackagesDialog::Construct(const FArguments& InArgs)
 			SNew(SVerticalBox)
 			+SVerticalBox::Slot() .Padding(10) .AutoHeight()
 			[
-				SNew( SHorizontalBox )
-				+SHorizontalBox::Slot()
-				.AutoWidth()
-				[
 					SNew( STextBlock )
 					.Text( this, &SPackagesDialog::GetMessage )
+				.AutoWrapText( true )
 				]
+			+SVerticalBox::Slot() .Padding(FMargin(10, 0, 10, 10)) .AutoHeight()
+			[
+				SNew( STextBlock )
+				.Text( this, &SPackagesDialog::GetWarning )
+				.ColorAndOpacity( FLinearColor::Yellow )
+				.AutoWrapText( true )
+				.Visibility( this, &SPackagesDialog::GetWarningVisibility )
 			]
 			+SVerticalBox::Slot()  .FillHeight(0.8)
 			[
@@ -213,6 +228,16 @@ void SPackagesDialog::SetMessage(const FText& InMessage)
 	Message = InMessage;
 }
 
+/**
+* Sets the warning message of the widget
+*
+* @param	InMessage	The string that the warning message should be set to
+*/
+void SPackagesDialog::SetWarning(const FText& InWarning)
+{
+	Warning = InWarning;
+}
+
 /**
  * Gets the return type of the dialog and it populates the package array results
  *
@@ -370,6 +395,7 @@ TSharedRef<SWidget> SPackagesDialog::GenerateWidgetForItemAndColumn( TSharedPtr<
 			[
 				SNew(SImage)
 				.Image( IconBrush )
+				.ToolTipText(FText::FromString(Item->GetToolTip()))
 				.IsEnabled(!Item->IsDisabled())
 			];
 	}
@@ -381,6 +407,7 @@ TSharedRef<SWidget> SPackagesDialog::GenerateWidgetForItemAndColumn( TSharedPtr<
 			[
 				SNew(STextBlock)
 				.Text(FText::FromString(PackageName))
+				.ToolTipText(FText::FromString(PackageName))
 				.IsEnabled(!Item->IsDisabled())
 			];
 	}
@@ -392,6 +419,24 @@ TSharedRef<SWidget> SPackagesDialog::GenerateWidgetForItemAndColumn( TSharedPtr<
 			[
 				SNew(STextBlock)
 				.Text(FText::FromString(PackageType))
+				.ToolTipText(FText::FromString(PackageType))
+				.IsEnabled(!Item->IsDisabled())
+				.ColorAndOpacity(PackageColor)
+			];
+	}
+	else if (ColumnID == SPackagesDialogDefs::ColumnID_CheckedOutByLabel)
+	{
+		check(bAllowSourceControlConnection);
+
+		FString CheckedOutByString = Item->GetCheckedOutByString();
+
+		ItemContentWidget = SNew(SHorizontalBox)
+			+ SHorizontalBox::Slot()
+			.Padding(RowPadding)
+			[
+				SNew(STextBlock)
+				.Text(FText::FromString(CheckedOutByString))
+				.ToolTipText(FText::FromString(CheckedOutByString))
 				.IsEnabled(!Item->IsDisabled())
 				.ColorAndOpacity(PackageColor)
 			];
@@ -590,6 +635,16 @@ FText SPackagesDialog::GetMessage() const
 	return Message;
 }
 
+FText SPackagesDialog::GetWarning() const
+{
+	return Warning;
+}
+
+EVisibility SPackagesDialog::GetWarningVisibility() const
+{
+	return (Warning.IsEmpty()) ? EVisibility::Collapsed : EVisibility::Visible;
+}
+
 EColumnSortMode::Type SPackagesDialog::GetColumnSortMode( const FName ColumnId ) const
 {
 	if ( SortByColumn != ColumnId )
@@ -657,6 +712,19 @@ void SPackagesDialog::SortTree()
 				return A->GetIconName() >= B->GetIconName(); } );
 		}
 	}
+	else if (SortByColumn == SPackagesDialogDefs::ColumnID_CheckedOutByLabel)
+	{
+		if (SortMode == EColumnSortMode::Ascending)
+		{
+			Items.Sort([](const TSharedPtr<FPackageItem>& A, const TSharedPtr<FPackageItem>& B) {
+				return A->GetCheckedOutByString() < B->GetCheckedOutByString(); });
+		}
+		else if (SortMode == EColumnSortMode::Descending)
+		{
+			Items.Sort([](const TSharedPtr<FPackageItem>& A, const TSharedPtr<FPackageItem>& B) {
+				return A->GetCheckedOutByString() >= B->GetCheckedOutByString(); });
+		}
+	}
 }
 
 void SPackageItemsListRow::Construct( const FArguments& InArgs, const TSharedRef<STableViewBase>& InOwnerTableView )
diff --git a/Engine/Source/Editor/PackagesDialog/Private/SPackagesDialog.h b/Engine/Source/Editor/PackagesDialog/Private/SPackagesDialog.h
--- a/Engine/Source/Editor/PackagesDialog/Private/SPackagesDialog.h
+++ b/Engine/Source/Editor/PackagesDialog/Private/SPackagesDialog.h
@@ -9,6 +9,7 @@
 #endif
 
 #include "PackagesDialog.h"
+#include "ISourceControlModule.h"
 
 
 /**
@@ -201,6 +202,22 @@ public:
 	 */
 	FString GetIconName() const { return IconName; }
 
+	/**
+	 * Get a string containing the name(s) of other users who have the file checked out
+	 *
+	 * @return names of other users who have the file checked out
+	 */
+	FString GetCheckedOutByString() const
+	{
+		ISourceControlProvider& SourceControlProvider = ISourceControlModule::Get().GetProvider();
+		FSourceControlStatePtr SourceControlState = SourceControlProvider.GetState(Package, EStateCacheUsage::Use);
+
+		FString CheckedOutBy;
+		SourceControlState->IsCheckedOutOther(&CheckedOutBy);
+
+		return CheckedOutBy;
+	}
+
 	/**
 	 * Gets the type name and color of the package item
 	 *
@@ -258,6 +275,7 @@ public:
 		: _ReadOnly(false)
 		, _AllowSourceControlConnection(false)
 		, _Message()
+		, _Warning()
 		{}
 
 		/** When true, this dialog only shows a list of packages without the ability to filter */
@@ -269,6 +287,9 @@ public:
 		/** The message of the widget */
 		SLATE_ARGUMENT(FText, Message)
 
+		/** The warning message of the widget */
+		SLATE_ARGUMENT(FText, Warning)
+
 	SLATE_END_ARGS()
 	
 	/**
@@ -309,6 +330,13 @@ public:
 	 */
 	void SetMessage(const FText& InMessage);
 
+	/**
+	 * Sets the warning message of the widget
+	 *
+	 * @param	InMessage	The string that the warning message should be set to
+	 */
+	void SetWarning(const FText& InMessage);
+
 	/**
 	 * Gets the return type of the dialog and it populates the package array results
 	 *
@@ -354,6 +382,11 @@ public:
 	 */
 	void Reset();
 
+	/**
+	 * Whether the dialog allows a source control connection
+	 */
+	bool IsSourceControlConnectionAllowed() const { return bAllowSourceControlConnection; }
+
 private:
 
 	// Begin SWidget Interface
@@ -408,6 +441,12 @@ private:
 	/** Delegate used to supply message text to the widget */
 	FText GetMessage() const;
 
+	/** Delegate used to supply warning text to the widget */
+	FText GetWarning() const;
+
+	/** Delegate used to determine visibility of the warning */
+	EVisibility GetWarningVisibility() const;
+
 	/**
 	 * Returns the current column sort mode (ascending or descending) if the ColumnId parameter matches the current
 	 * column to be sorted by, otherwise returns EColumnSortMode_None.
@@ -464,9 +503,15 @@ private:
 	/** When true, this dialog displays a 'connect to source control' button */
 	bool bAllowSourceControlConnection;
 
+	/** When true, the warning message is displayed in the widget */
+	bool bShowWarning;
+
 	/** The message to display */
 	FText Message;
 
+	/** The warning to display */
+	FText Warning;
+
 	/** Specify which column to sort with */
 	FName SortByColumn;
 
diff --git a/Engine/Source/Editor/PackagesDialog/Public/PackagesDialog.h b/Engine/Source/Editor/PackagesDialog/Public/PackagesDialog.h
--- a/Engine/Source/Editor/PackagesDialog/Public/PackagesDialog.h
+++ b/Engine/Source/Editor/PackagesDialog/Public/PackagesDialog.h
@@ -72,6 +72,12 @@ public:
 	 */
 	virtual void SetMessage(const FText& InMessage);
 
+	/**
+	 * Sets the warning message displayed in the package dialog
+	 * @param InMessage		The warning to display
+	 */
+	virtual void SetWarning(const FText& InMessage);
+
 	/**
 	 * Populates the passed in array with the desired packages
 	 *
@@ -130,6 +136,9 @@ private:
 	/** A default window size for the package dialog */
 	static const FVector2D DEFAULT_WINDOW_SIZE;
 
+	/** Extra window width if source control connection is allowed */
+	static const FVector2D EXTRA_WINDOW_WIDTH;
+
 	/** Editor package dialog window */
 	TWeakPtr<SWindow> EditorPackagesDialogWindow;
 
diff --git a/Engine/Source/Editor/UnrealEd/Classes/Editor/UnrealEdEngine.h b/Engine/Source/Editor/UnrealEd/Classes/Editor/UnrealEdEngine.h
--- a/Engine/Source/Editor/UnrealEd/Classes/Editor/UnrealEdEngine.h
+++ b/Engine/Source/Editor/UnrealEd/Classes/Editor/UnrealEdEngine.h
@@ -3,6 +3,7 @@
 #pragma once
 #include "UniquePtr.h"
 #include "IPackageAutoSaver.h"
+#include "ISourceControlModule.h"
 #include "ComponentVisualizerManager.h"
 #include "UnrealEdEngine.generated.h"
 
@@ -15,6 +16,8 @@ enum EPackageNotifyState
 	NS_DialogPrompted,
 	// The package has been marked dirty and is pending a balloon prompt
 	NS_PendingPrompt,
+	// The package has been marked dirty but cannot be checked out, and is pending a modal warning dialog
+	NS_PendingWarning,
 	NS_MAX,
 };
 
@@ -214,6 +217,10 @@ class UNREALED_API UUnrealEdEngine : public UEditorEngine, public FNotifyHook
 
 	/** called when a package has has its dirty state updated */
 	void OnPackageDirtyStateUpdated( UPackage* Pkg);
+	/** called when a package's source control state is updated */
+	void OnSourceControlStateUpdated(const FSourceControlOperationRef& SourceControlOp, ECommandResult::Type ResultType, TWeakObjectPtr<UPackage> Package);
+	/** called when a package is automatically checked out from source control */
+	void OnPackageCheckedOut(const FSourceControlOperationRef& SourceControlOp, ECommandResult::Type ResultType, TWeakObjectPtr<UPackage> Package);
 	/** caled by FCoreDelegate::PostGarbageCollect */
 	void OnPostGarbageCollect();
 	/** called by color picker change event */
diff --git a/Engine/Source/Editor/UnrealEd/Classes/Settings/EditorLoadingSavingSettings.h b/Engine/Source/Editor/UnrealEd/Classes/Settings/EditorLoadingSavingSettings.h
--- a/Engine/Source/Editor/UnrealEd/Classes/Settings/EditorLoadingSavingSettings.h
+++ b/Engine/Source/Editor/UnrealEd/Classes/Settings/EditorLoadingSavingSettings.h
@@ -69,6 +69,10 @@ public:
 
 public:
 
+	/** Whether to automatically checkout on asset modification */
+	UPROPERTY(EditAnywhere, config, Category=SourceControl)
+	uint32 bAutomaticallyCheckoutOnAssetModification:1;
+
 	/** Whether to automatically prompt for SCC checkout on asset modification */
 	UPROPERTY(EditAnywhere, config, Category=SourceControl)
 	uint32 bPromptForCheckoutOnAssetModification:1;
diff --git a/Engine/Source/Editor/UnrealEd/Private/FileHelpers.cpp b/Engine/Source/Editor/UnrealEd/Private/FileHelpers.cpp
--- a/Engine/Source/Editor/UnrealEd/Private/FileHelpers.cpp
+++ b/Engine/Source/Editor/UnrealEd/Private/FileHelpers.cpp
@@ -1020,8 +1020,11 @@ bool FEditorFileUtils::PromptToCheckoutPackages(bool bCheckDirty, const TArray<U
 	}
 	
 	// The checkout dialog to show users if any packages need to be checked out
+	const FText DialogTitle = NSLOCTEXT("PackagesDialogModule", "CheckoutPackagesDialogTitle", "Check Out Assets");
+	const FText DialogHeading = NSLOCTEXT("PackagesDialogModule", "CheckoutPackagesDialogMessage", "Select assets to check out.");
+
 	FPackagesDialogModule& CheckoutPackagesDialogModule = FModuleManager::LoadModuleChecked<FPackagesDialogModule>( TEXT("PackagesDialog") );
-	CheckoutPackagesDialogModule.CreatePackagesDialog(NSLOCTEXT("PackagesDialogModule", "CheckoutPackagesDialogTitle", "Check Out Assets"), NSLOCTEXT("PackagesDialogModule", "CheckoutPackagesDialogMessage", "Select assets to check out, right-click assets for more options."), false, true);
+	CheckoutPackagesDialogModule.CreatePackagesDialog(DialogTitle, DialogHeading, false, true);
 
 	// Add any of the packages which do not report as editable by source control, yet are currently in the source control depot
 	// If the user has specified to check for dirty packages, only add those which are dirty
@@ -1037,6 +1040,8 @@ bool FEditorFileUtils::PromptToCheckoutPackages(bool bCheckDirty, const TArray<U
 		PackagesNotSavedDuringSaveAll.Remove(PackageName);
 	}
 
+	bool bShowWarning = false;
+
 	// Iterate through all the packages and add them to the dialog if necessary.
 	for ( TArray<UPackage*>::TConstIterator PackageIter( PackagesToCheckOut ); PackageIter; ++PackageIter )
 	{
@@ -1058,18 +1063,20 @@ bool FEditorFileUtils::PromptToCheckoutPackages(bool bCheckDirty, const TArray<U
 		// Package does not need to be checked out if its already checked out or we are ignoring it for source control
 		bool bSCCCanEdit = !SourceControlState.IsValid() || SourceControlState->CanCheckIn() || SourceControlState->IsIgnored() || SourceControlState->IsUnknown();
 		bool bIsSourceControlled = SourceControlState.IsValid() && SourceControlState->IsSourceControlled();
-		
+
 		if ( !bSCCCanEdit && (bIsSourceControlled && ( !bCheckDirty || ( bCheckDirty && CurPackage->IsDirty() ) ) ) && !SourceControlState->IsCheckedOut() )
 		{
 			if( SourceControlState.IsValid() && !SourceControlState->IsCurrent() )
 			{				
 				// This package is not at the head revision and it should be ghosted as a result
 				CheckoutPackagesDialogModule.AddPackageItem(CurPackage, CurPackage->GetName(), ECheckBoxState::Unchecked, true, TEXT("SavePackages.SCC_DlgNotCurrent"), SourceControlState->GetDisplayTooltip().ToString());
+				bShowWarning = true;
 			}
 			else if( SourceControlState.IsValid() && SourceControlState->IsCheckedOutOther() )
 			{
 				// This package is checked out by someone else so it should be ghosted
 				CheckoutPackagesDialogModule.AddPackageItem(CurPackage, CurPackage->GetName(), ECheckBoxState::Unchecked, true, TEXT("SavePackages.SCC_DlgCheckedOutOther"), SourceControlState->GetDisplayTooltip().ToString());
+				bShowWarning = true;
 			}
 			else
 			{
@@ -1105,6 +1112,12 @@ bool FEditorFileUtils::PromptToCheckoutPackages(bool bCheckDirty, const TArray<U
 	// If any packages were added to the dialog, show the dialog to the user and allow them to select which files to check out
 	if ( bPackagesAdded )
 	{
+		if (bShowWarning)
+		{
+			CheckoutPackagesDialogModule.SetWarning(
+				NSLOCTEXT("PackagesDialogModule", "CheckoutPackagesWarnMessage", "Warning: There are modified assets which you will not be able to check out as they are locked or not at the head revision. You may lose your changes if you continue, as you will be unable to submit them to source control."));
+		}
+
 		TAttribute<bool> CheckOutSelectedDisabledAttrib;
 		if( !bHavePackageToCheckOut && !IsCheckOutSelectedDisabled() )
 		{
diff --git a/Engine/Source/Editor/UnrealEd/Private/UnrealEdEngine.cpp b/Engine/Source/Editor/UnrealEd/Private/UnrealEdEngine.cpp
--- a/Engine/Source/Editor/UnrealEd/Private/UnrealEdEngine.cpp
+++ b/Engine/Source/Editor/UnrealEd/Private/UnrealEdEngine.cpp
@@ -20,7 +20,6 @@
 #include "GameMapsSettingsCustomization.h"
 #include "LevelEditorPlaySettingsCustomization.h"
 #include "ProjectPackagingSettingsCustomization.h"
-#include "ISourceControlModule.h"
 #include "Editor/StatsViewer/Public/StatsViewerModule.h"
 #include "SnappingUtils.h"
 #include "PackageAutoSaver.h"
@@ -29,6 +28,8 @@
 #include "ComponentVisualizer.h"
 #include "Editor/EditorLiveStreaming/Public/IEditorLiveStreaming.h"
 #include "SourceCodeNavigation.h"
+#include "NotificationManager.h"
+#include "SNotificationList.h"
 
 DEFINE_LOG_CATEGORY_STATIC(LogUnrealEdEngine, Log, All);
 
@@ -417,29 +418,24 @@ void UUnrealEdEngine::OnPackageDirtyStateUpdated( UPackage* Pkg)
 		const uint8* PromptState = PackageToNotifyState.Find( Package );
 		const bool bAlreadyAsked = PromptState != NULL;
 
-		// Get the source control state of the package
-		ISourceControlProvider& SourceControlProvider = ISourceControlModule::Get().GetProvider();
-		FSourceControlStatePtr SourceControlState = SourceControlProvider.GetState(Package, EStateCacheUsage::Use);
-
 		// During an autosave, packages are saved in the autosave directory which switches off their dirty flags.
 		// To preserve the pre-autosave state, any saved package is then remarked as dirty because it wasn't saved in the normal location where it would be picked up by source control.
-		// Any callback that happens during an autosave is bogus since a package wasnt  marked dirty due to a user modification.
+		// Any callback that happens during an autosave is bogus since a package wasn't marked dirty due to a user modification.
 		const bool bIsAutoSaving = PackageAutoSaver.Get() && PackageAutoSaver->IsAutoSaving();
+
+		const UEditorLoadingSavingSettings* Settings = GetDefault<UEditorLoadingSavingSettings>();
+
 		if( !bIsAutoSaving && 
 			!GIsEditorLoadingPackage && // Don't ask if the package was modified as a result of a load
 			!bAlreadyAsked && // Don't ask if we already asked once!
-			GetDefault<UEditorLoadingSavingSettings>()->bPromptForCheckoutOnAssetModification && // Only prompt if the user has specified to be prompted on modification
-			SourceControlState.IsValid() && 
-			(SourceControlState->CanCheckout() || !SourceControlState->IsCurrent() || SourceControlState->IsCheckedOutOther()) )
+			(Settings->bPromptForCheckoutOnAssetModification || Settings->bAutomaticallyCheckoutOnAssetModification) )
 		{
-			// Allow packages that are not checked out to pass through.
-			// Allow packages that are not current or checked out by others pass through.  
-			// The user wont be able to checkout these packages but the checkout dialog will show up with a special icon 
-			// to let the user know they wont be able to checkout the package they are modifying.
-
-			PackageToNotifyState.Add( Package, NS_PendingPrompt );
-			// We need to prompt since a new package was added
-			bNeedToPromptForCheckout = true;
+			// Force source control state to be updated
+			ISourceControlProvider& SourceControlProvider = ISourceControlModule::Get().GetProvider();
+
+			TArray<FString> Files;
+			Files.Add(SourceControlHelpers::PackageFilename(Package));
+			SourceControlProvider.Execute(ISourceControlOperation::Create<FUpdateStatus>(), SourceControlHelpers::AbsoluteFilenames(Files), EConcurrency::Asynchronous, FSourceControlOperationComplete::CreateUObject(this, &UUnrealEdEngine::OnSourceControlStateUpdated, TWeakObjectPtr<UPackage>(Package)));
 		}
 	}
 	else
@@ -450,6 +446,88 @@ void UUnrealEdEngine::OnPackageDirtyStateUpdated( UPackage* Pkg)
 }
 
 
+void UUnrealEdEngine::OnSourceControlStateUpdated(const FSourceControlOperationRef& SourceControlOp, ECommandResult::Type ResultType, TWeakObjectPtr<UPackage> Package)
+{
+	if (ResultType == ECommandResult::Succeeded && Package.IsValid())
+	{
+		// Get the source control state of the package
+		ISourceControlProvider& SourceControlProvider = ISourceControlModule::Get().GetProvider();
+		FSourceControlStatePtr SourceControlState = SourceControlProvider.GetState(Package.Get(), EStateCacheUsage::Use);
+
+		if (SourceControlState.IsValid())
+		{
+			const UEditorLoadingSavingSettings* Settings = GetDefault<UEditorLoadingSavingSettings>();
+			check(Settings->bPromptForCheckoutOnAssetModification || Settings->bAutomaticallyCheckoutOnAssetModification);
+
+			if (Settings->bAutomaticallyCheckoutOnAssetModification && SourceControlState->CanCheckout())
+			{
+				// Automatically check out asset
+				TArray<FString> Files;
+				Files.Add(SourceControlHelpers::PackageFilename(Package.Get()));
+				SourceControlProvider.Execute(ISourceControlOperation::Create<FCheckOut>(), SourceControlHelpers::AbsoluteFilenames(Files), EConcurrency::Asynchronous, FSourceControlOperationComplete::CreateUObject(this, &UUnrealEdEngine::OnPackageCheckedOut, TWeakObjectPtr<UPackage>(Package)));
+			}
+			else
+			{
+				if (SourceControlState->CanCheckout() || !SourceControlState->IsCurrent() || SourceControlState->IsCheckedOutOther())
+				{
+					// To get here, either "prompt for checkout on asset modification" is set, or "automatically checkout on asset modification"
+					// is set, but it failed.
+
+					// Allow packages that are not checked out to pass through.
+					// Allow packages that are not current or checked out by others pass through.  
+					// The user wont be able to checkout these packages but the checkout dialog will show up with a special icon 
+					// to let the user know they wont be able to checkout the package they are modifying.
+
+					PackageToNotifyState.Add(Package, SourceControlState->CanCheckout() ? NS_PendingPrompt : NS_PendingWarning);
+					// We need to prompt since a new package was added
+					bNeedToPromptForCheckout = true;
+				}
+			}
+		}
+	}
+}
+
+
+void UUnrealEdEngine::OnPackageCheckedOut(const FSourceControlOperationRef& SourceControlOp, ECommandResult::Type ResultType, TWeakObjectPtr<UPackage> Package)
+{
+	if (Package.IsValid())
+	{
+		// Get the source control state of the package
+		ISourceControlProvider& SourceControlProvider = ISourceControlModule::Get().GetProvider();
+		FSourceControlStatePtr SourceControlState = SourceControlProvider.GetState(Package.Get(), EStateCacheUsage::Use);
+
+		FFormatNamedArguments Arguments;
+		Arguments.Add(TEXT("Package"), FText::FromString(Package->GetName()));
+
+		if (ResultType == ECommandResult::Succeeded)
+		{
+			if (SourceControlState.IsValid() && SourceControlState->IsCheckedOut())
+			{
+				FNotificationInfo Notification(FText::Format(NSLOCTEXT("SourceControl", "AutoCheckOutNotification", "Package '{Package}' automatically checked out."), Arguments));
+				Notification.bFireAndForget = true;
+				Notification.ExpireDuration = 4.0f;
+				Notification.bUseThrobber = true;
+
+				FSlateNotificationManager::Get().AddNotification(Notification);
+
+				return;
+			}
+		}
+
+		FNotificationInfo ErrorNotification(FText::Format(NSLOCTEXT("SourceControl", "AutoCheckOutFailedNotification", "Unable to automatically check out Package '{Package}'."), Arguments));
+		ErrorNotification.bFireAndForget = true;
+		ErrorNotification.ExpireDuration = 4.0f;
+		ErrorNotification.bUseThrobber = true;
+
+		FSlateNotificationManager::Get().AddNotification(ErrorNotification);
+
+		// Automatic checkout failed - pop up the notification for manual checkout
+		PackageToNotifyState.Add(Package, SourceControlState->CanCheckout() ? NS_PendingPrompt : NS_PendingWarning);
+		bNeedToPromptForCheckout = true;
+	}
+}
+
+
 void UUnrealEdEngine::OnPostGarbageCollect()
 {
 	// Refresh Editor browsers after GC in case objects where removed.  Note that if the user is currently
diff --git a/Engine/Source/Editor/UnrealEd/Private/UnrealEdSrv.cpp b/Engine/Source/Editor/UnrealEd/Private/UnrealEdSrv.cpp
--- a/Engine/Source/Editor/UnrealEd/Private/UnrealEdSrv.cpp
+++ b/Engine/Source/Editor/UnrealEd/Private/UnrealEdSrv.cpp
@@ -1451,24 +1451,45 @@ void UUnrealEdEngine::AttemptModifiedPackageNotification()
 
 		if( bCanPrompt )
 		{
+			check(GUnrealEd != nullptr);
+			const TMap<TWeakObjectPtr<UPackage>, uint8>& PackageToNotifyState = GUnrealEd->PackageToNotifyState;
+
+			bool bNeedWarningDialog = false;
+			for (const auto& Entry : PackageToNotifyState)
+			{
+				if (Entry.Value == NS_PendingWarning)
+				{
+					bNeedWarningDialog = true;
+					break;
+				}
+			}
+
 			// The user is not interacting with anything, prompt to checkout packages that have been modified
 			
 			struct Local
 			{
-				static void OpenMessageLog()
+				static void OpenCheckOutDialog()
 				{
 					GUnrealEd->PromptToCheckoutModifiedPackages();
 				}
 			};
-			FNotificationInfo ErrorNotification( NSLOCTEXT("SourceControl", "CheckOutNotification", "Files need check-out!") );
-			ErrorNotification.bFireAndForget = true;
-			ErrorNotification.Hyperlink = FSimpleDelegate::CreateStatic(&Local::OpenMessageLog);
-			ErrorNotification.HyperlinkText = NSLOCTEXT("SourceControl", "CheckOutHyperlinkText", "Check-Out");
-			ErrorNotification.ExpireDuration = 3.0f; // Need this message to last a little longer than normal since the user may want to "Show Log"
-			ErrorNotification.bUseThrobber = true;
 
-			// For adding notifications.
-			FSlateNotificationManager::Get().AddNotification(ErrorNotification);
+			if (bNeedWarningDialog)
+			{
+				Local::OpenCheckOutDialog();
+			}
+			else
+			{
+				FNotificationInfo ErrorNotification(NSLOCTEXT("SourceControl", "CheckOutNotification", "Files need check-out!"));
+				ErrorNotification.bFireAndForget = true;
+				ErrorNotification.Hyperlink = FSimpleDelegate::CreateStatic(&Local::OpenCheckOutDialog);
+				ErrorNotification.HyperlinkText = NSLOCTEXT("SourceControl", "CheckOutHyperlinkText", "Check-Out");
+				ErrorNotification.ExpireDuration = 10.0f; // Need this message to last a little longer than normal since the user will probably want to click the hyperlink to check out files
+				ErrorNotification.bUseThrobber = true;
+
+				// For adding notifications.
+				FSlateNotificationManager::Get().AddNotification(ErrorNotification);
+			}
 
 			// No longer have a pending prompt.
 			bNeedToPromptForCheckout = false;
@@ -1548,7 +1569,7 @@ void UUnrealEdEngine::PromptToCheckoutModifiedPackages( bool bPromptAll )
 	{
 		for( TMap<TWeakObjectPtr<UPackage>,uint8>::TIterator It(PackageToNotifyState); It; ++It )
 		{
-			if( It.Key().IsValid() && (It.Value() == NS_BalloonPrompted || It.Value() == NS_PendingPrompt) )
+			if( It.Key().IsValid() && (It.Value() == NS_PendingWarning || It.Value() == NS_PendingPrompt) )
 			{
 				PackagesToCheckout.Add( It.Key().Get() );
 				It.Value() = NS_DialogPrompted;
@@ -1556,7 +1577,9 @@ void UUnrealEdEngine::PromptToCheckoutModifiedPackages( bool bPromptAll )
 		}
 	}
 
-	FEditorFileUtils::PromptToCheckoutPackages( true, PackagesToCheckout, NULL, NULL, true );
+	const bool bCheckDirty = true;
+	const bool bPromptingAfterModify = true;
+	FEditorFileUtils::PromptToCheckoutPackages( bCheckDirty, PackagesToCheckout, NULL, NULL, bPromptingAfterModify );
 }
 
 
diff --git a/Engine/Source/Runtime/CoreUObject/Private/UObject/Obj.cpp b/Engine/Source/Runtime/CoreUObject/Private/UObject/Obj.cpp
--- a/Engine/Source/Runtime/CoreUObject/Private/UObject/Obj.cpp
+++ b/Engine/Source/Runtime/CoreUObject/Private/UObject/Obj.cpp
@@ -741,7 +741,7 @@ bool UObject::Modify( bool bAlwaysMarkDirty/*=true*/ )
 
 			// If we failed to save to the transaction buffer, but the user requested the package
 			// marked dirty anyway, do so
-			if (!bSavedToTransactionBuffer || bAlwaysMarkDirty)
+			if (!bSavedToTransactionBuffer && bAlwaysMarkDirty)
 			{
 				MarkPackageDirty();
 			}
