#vr Uncapping the FPS rate in VRPIE mode (UE-7463)

diff --git a/Engine/Source/Editor/UnrealEd/Private/Editor.cpp b/Engine/Source/Editor/UnrealEd/Private/Editor.cpp
--- a/Engine/Source/Editor/UnrealEd/Private/Editor.cpp
+++ b/Engine/Source/Editor/UnrealEd/Private/Editor.cpp
@@ -1579,6 +1579,11 @@ float UEditorEngine::GetMaxTickRate( float DeltaTime, bool bAllowFrameRateSmooth
 	float MaxTickRate = 0.0f;
 	if( !ShouldThrottleCPUUsage() )
 	{
+		// do not limit fps in VR Preview mode
+		if (bUseVRPreviewForPlayWorld)
+		{
+			return 0.0f;
+		}
 		const float SuperMaxTickRate = Super::GetMaxTickRate( DeltaTime, bAllowFrameRateSmoothing );
 		if( SuperMaxTickRate != 0.0f )
 		{
diff --git a/Engine/Source/Editor/UnrealEd/Private/PlayLevel.cpp b/Engine/Source/Editor/UnrealEd/Private/PlayLevel.cpp
--- a/Engine/Source/Editor/UnrealEd/Private/PlayLevel.cpp
+++ b/Engine/Source/Editor/UnrealEd/Private/PlayLevel.cpp
@@ -302,6 +302,7 @@ void UEditorEngine::EndPlayMap()
 	bIsPlayWorldQueued = false;
 	bIsSimulateInEditorQueued = false;
 	bRequestEndPlayMapQueued = false;
+	bUseVRPreviewForPlayWorld = false;
 
 	// display any info if required.
 	FMessageLog("PIE").Notify( LOCTEXT("PIEErrorsPresent", "Errors/warnings reported while playing in editor."), EMessageSeverity::Warning );
