Editor: Ensure that modes are exited when an asset editor is closed if it uses standalone mode tools instances

diff --git a/Engine/Source/Editor/UnrealEd/Private/EditorViewportClient.cpp b/Engine/Source/Editor/UnrealEd/Private/EditorViewportClient.cpp
--- a/Engine/Source/Editor/UnrealEd/Private/EditorViewportClient.cpp
+++ b/Engine/Source/Editor/UnrealEd/Private/EditorViewportClient.cpp
@@ -309,6 +309,12 @@ FEditorViewportClient::FEditorViewportClient(FEditorModeTools* InModeTools, FPre
 
 FEditorViewportClient::~FEditorViewportClient()
 {
+	if (bOwnsModeTools)
+	{
+		ModeTools->SetDefaultMode(FBuiltinEditorModes::EM_Default);
+		ModeTools->DeactivateAllModes(); // this also activates the default mode
+	}
+
 	ModeTools->OnEditorModeChanged().RemoveAll(this);
 
 	delete Widget;
