Replacing actors in the level viewport will automatically replace all references in the level Blueprint.

Replacing Actor Reference nodes and Actor Bound Event nodes with another actor reference will properly refresh the node and it's title.

diff --git a/Engine/Source/Editor/BlueprintGraph/Private/EdGraphSchema_K2.cpp b/Engine/Source/Editor/BlueprintGraph/Private/EdGraphSchema_K2.cpp
--- a/Engine/Source/Editor/BlueprintGraph/Private/EdGraphSchema_K2.cpp
+++ b/Engine/Source/Editor/BlueprintGraph/Private/EdGraphSchema_K2.cpp
@@ -898,6 +898,7 @@ void UEdGraphSchema_K2::ReplaceSelectedNode(UEdGraphNode* SourceNode, AActor* Ta
 
 			LiteralNode->Modify();
 			LiteralNode->SetObjectRef( TargetActor );
+			LiteralNode->ReconstructNode();
 			UBlueprint* Blueprint = FBlueprintEditorUtils::FindBlueprintForGraphChecked(CastChecked<UEdGraph>(SourceNode->GetOuter()));
 			FBlueprintEditorUtils::MarkBlueprintAsModified(Blueprint);
 		}
diff --git a/Engine/Source/Editor/BlueprintGraph/Private/K2Node_ActorBoundEvent.cpp b/Engine/Source/Editor/BlueprintGraph/Private/K2Node_ActorBoundEvent.cpp
--- a/Engine/Source/Editor/BlueprintGraph/Private/K2Node_ActorBoundEvent.cpp
+++ b/Engine/Source/Editor/BlueprintGraph/Private/K2Node_ActorBoundEvent.cpp
@@ -48,6 +48,8 @@ void UK2Node_ActorBoundEvent::ReconstructNode()
 	// Ensure that we try to update the delegate we're bound to if its moved
 	GetTargetDelegate();
 
+	CachedNodeTitle.MarkDirty();
+
 	Super::ReconstructNode();
 }
 
diff --git a/Engine/Source/Editor/BlueprintGraph/Private/K2Node_Literal.cpp b/Engine/Source/Editor/BlueprintGraph/Private/K2Node_Literal.cpp
--- a/Engine/Source/Editor/BlueprintGraph/Private/K2Node_Literal.cpp
+++ b/Engine/Source/Editor/BlueprintGraph/Private/K2Node_Literal.cpp
@@ -258,6 +258,7 @@ void UK2Node_Literal::SetObjectRef(UObject* NewValue)
 		// Set the pin type to reflect the object we're referencing
 		if( ValuePin )
 		{
+			ValuePin->Modify();
 			ValuePin->PinType.PinCategory = Schema->PC_Object;
 			ValuePin->PinType.PinSubCategory = TEXT("");
 			ValuePin->PinType.PinSubCategoryObject = ObjectRef->GetClass();
diff --git a/Engine/Source/Editor/Kismet/Private/BlueprintEditor.cpp b/Engine/Source/Editor/Kismet/Private/BlueprintEditor.cpp
--- a/Engine/Source/Editor/Kismet/Private/BlueprintEditor.cpp
+++ b/Engine/Source/Editor/Kismet/Private/BlueprintEditor.cpp
@@ -4969,7 +4969,9 @@ void FBlueprintEditor::OnAssignReferencedActor()
 						// Store the node's current state and replace the referenced actor
 						CurrentEvent->Modify();
 						CurrentEvent->EventOwner = SelectedActor;
+						CurrentEvent->ReconstructNode();
 					}
+					FBlueprintEditorUtils::MarkBlueprintAsModified(GetBlueprintObj());
 				}
 			}
 		}
diff --git a/Engine/Source/Editor/UnrealEd/Private/Editor.cpp b/Engine/Source/Editor/UnrealEd/Private/Editor.cpp
--- a/Engine/Source/Editor/UnrealEd/Private/Editor.cpp
+++ b/Engine/Source/Editor/UnrealEd/Private/Editor.cpp
@@ -5169,6 +5169,10 @@ void UEditorEngine::ReplaceActors(UActorFactory* Factory, const FAssetData& Asse
 			NewActor->PostEditMove(true);
 			NewActor->MarkPackageDirty();
 
+			// Replace references in the level script Blueprint with the new Actor
+			ULevelScriptBlueprint* LSB = NewActor->GetLevel()->GetLevelScriptBlueprint(true);
+			FBlueprintEditorUtils::ReplaceAllActorRefrences(LSB, OldActor, NewActor);
+
 			GEditor->Layers->DisassociateActorFromLayers( OldActor );
 			World->EditorDestroyActor(OldActor, true);
 		}
diff --git a/Engine/Source/Editor/UnrealEd/Private/Kismet2/BlueprintEditorUtils.cpp b/Engine/Source/Editor/UnrealEd/Private/Kismet2/BlueprintEditorUtils.cpp
--- a/Engine/Source/Editor/UnrealEd/Private/Kismet2/BlueprintEditorUtils.cpp
+++ b/Engine/Source/Editor/UnrealEd/Private/Kismet2/BlueprintEditorUtils.cpp
@@ -5568,6 +5568,39 @@ int32 FBlueprintEditorUtils::FindNumReferencesToActorFromLevelScript(ULevelScrip
 	return RefCount;
 }
 
+void FBlueprintEditorUtils::ReplaceAllActorRefrences(ULevelScriptBlueprint* InLevelScriptBlueprint, AActor* InOldActor, AActor* InNewActor)
+{
+	InLevelScriptBlueprint->Modify();
+	FBlueprintEditorUtils::MarkBlueprintAsModified(InLevelScriptBlueprint);
+
+	// Literal nodes are the common "get" type nodes and need to be updated with the new object reference
+	TArray< UK2Node_Literal* > LiteralNodes;
+	FBlueprintEditorUtils::GetAllNodesOfClass(InLevelScriptBlueprint, LiteralNodes);
+
+	for( UK2Node_Literal* LiteralNode : LiteralNodes )
+	{
+		if(LiteralNode->GetObjectRef() == InOldActor)
+		{
+			LiteralNode->Modify();
+			LiteralNode->SetObjectRef(InNewActor);
+			LiteralNode->ReconstructNode();
+		}
+	}
+
+	// Actor Bound Events reference the actors as well and need to be updated
+	TArray< UK2Node_ActorBoundEvent* > ActorEventNodes;
+	FBlueprintEditorUtils::GetAllNodesOfClass(InLevelScriptBlueprint, ActorEventNodes);
+
+	for( UK2Node_ActorBoundEvent* ActorEventNode : ActorEventNodes )
+	{
+		if(ActorEventNode->GetReferencedLevelActor() == InOldActor)
+		{
+			ActorEventNode->Modify();
+			ActorEventNode->EventOwner = InNewActor;
+			ActorEventNode->ReconstructNode();
+		}
+	}
+}
 
 void  FBlueprintEditorUtils::ModifyActorReferencedGraphNodes(ULevelScriptBlueprint* LevelScriptBlueprint, const AActor* InActor)
 {
diff --git a/Engine/Source/Editor/UnrealEd/Public/Kismet2/BlueprintEditorUtils.h b/Engine/Source/Editor/UnrealEd/Public/Kismet2/BlueprintEditorUtils.h
--- a/Engine/Source/Editor/UnrealEd/Public/Kismet2/BlueprintEditorUtils.h
+++ b/Engine/Source/Editor/UnrealEd/Public/Kismet2/BlueprintEditorUtils.h
@@ -981,6 +981,9 @@ public:
 	/** Find how many nodes reference the supplied actor */
 	static int32 FindNumReferencesToActorFromLevelScript(ULevelScriptBlueprint* LevelScriptBlueprint, AActor* InActor);
 
+	/** Replace all references of the old actor with the new actor */
+	static void ReplaceAllActorRefrences(ULevelScriptBlueprint* InLevelScriptBlueprint, AActor* InOldActor, AActor* InNewActor);
+
 	/** Function to call modify() on all graph nodes which reference this actor */
 	static void  ModifyActorReferencedGraphNodes(ULevelScriptBlueprint* LevelScriptBlueprint, const AActor* InActor);
 
