Added the ability to create blank characters and pawns

diff --git a/Engine/Source/Editor/PlacementMode/Private/SPlacementModeTools.cpp b/Engine/Source/Editor/PlacementMode/Private/SPlacementModeTools.cpp
--- a/Engine/Source/Editor/PlacementMode/Private/SPlacementModeTools.cpp
+++ b/Engine/Source/Editor/PlacementMode/Private/SPlacementModeTools.cpp
@@ -378,7 +378,7 @@ TSharedRef< SWidget > SPlacementModeTools::CreateStandardPanel()
 			SNew( SBox )
 			.AddMetaData<FTagMetaData>(FTagMetaData(TEXT("PMVisual")))
 			[
-				CreatePlacementGroupTab( (int32)EPlacementTab::Visual, NSLOCTEXT( "PlacementMode", "Visual", "Visual" ), false )
+				CreatePlacementGroupTab( (int32)EPlacementTab::Visual, NSLOCTEXT( "PlacementMode", "VisualEffects", "Visual Effects" ), false )
 			]
 		]
 
@@ -464,7 +464,7 @@ TSharedRef< SWidget > SPlacementModeTools::CreateStandardPanel()
 					]
 				]
 
-				// Visual
+				// Visual Effects
 				+ SWidgetSwitcher::Slot()
 				[
 					SNew( SScrollBox )
@@ -669,10 +669,11 @@ TSharedRef< SWidget > SPlacementModeTools::BuildBasicWidget()
 	TSharedRef<SVerticalBox> VerticalBox = SNew( SVerticalBox )
 
 	// Basics
-	+ SVerticalBox::Slot()
+	
+	+SVerticalBox::Slot()
 	.AutoHeight()
 	[
-		BuildDraggableAssetWidget( UActorFactoryCameraActor::StaticClass() )
+		BuildDraggableAssetWidget(UActorFactoryPointLight::StaticClass())
 	]
 	+ SVerticalBox::Slot()
 	.AutoHeight()
@@ -681,11 +682,6 @@ TSharedRef< SWidget > SPlacementModeTools::BuildBasicWidget()
 	]
 	+ SVerticalBox::Slot()
 	.AutoHeight()
-	[
-		BuildDraggableAssetWidget( UActorFactoryPointLight::StaticClass() )
-	]
-	+ SVerticalBox::Slot()
-	.AutoHeight()
 	[
 		BuildDraggableAssetWidget( UActorFactoryTriggerBox::StaticClass() )
 	]
@@ -693,11 +689,6 @@ TSharedRef< SWidget > SPlacementModeTools::BuildBasicWidget()
 	.AutoHeight()
 	[
 		BuildDraggableAssetWidget( UActorFactoryTriggerSphere::StaticClass() )
-	]
-	+ SVerticalBox::Slot()
-	.AutoHeight()
-	[
-		BuildDraggableAssetWidget( UActorFactoryTriggerCapsule::StaticClass() )
 	];
 
 	if( GetDefault<UEditorExperimentalSettings>()->bInWorldBPEditing )
@@ -707,6 +698,18 @@ TSharedRef< SWidget > SPlacementModeTools::BuildBasicWidget()
 		[
 			BuildDraggableAssetWidget(UActorFactoryEmptyActor::StaticClass())
 		];
+
+		VerticalBox->InsertSlot(1)
+		.AutoHeight()
+		[
+			BuildDraggableAssetWidget(UActorFactoryCharacter::StaticClass())
+		];
+
+		VerticalBox->InsertSlot(2)
+		.AutoHeight()
+		[
+			BuildDraggableAssetWidget(UActorFactoryPawn::StaticClass())
+		];
 	}
 
 	return VerticalBox;
diff --git a/Engine/Source/Editor/UnrealEd/Classes/ActorFactories/ActorFactoryCharacter.h b/Engine/Source/Editor/UnrealEd/Classes/ActorFactories/ActorFactoryCharacter.h
new file mode 100644
--- /dev/null
+++ b/Engine/Source/Editor/UnrealEd/Classes/ActorFactories/ActorFactoryCharacter.h
@@ -0,0 +1,17 @@
+// Copyright 1998-2015 Epic Games, Inc. All Rights Reserved.
+
+
+#pragma once
+#include "ActorFactoryCharacter.generated.h"
+
+UCLASS(MinimalAPI,config=Editor)
+class UActorFactoryCharacter : public UActorFactory
+{
+	GENERATED_UCLASS_BODY()
+
+public:
+	virtual bool CanCreateActorFrom( const FAssetData& AssetData, FText& OutErrorMsg ) override;
+
+protected:
+	virtual AActor* SpawnActor( UObject* Asset, ULevel* InLevel, const FVector& Location, const FRotator& Rotation, EObjectFlags ObjectFlags, const FName& Name ) override;
+};
diff --git a/Engine/Source/Editor/UnrealEd/Classes/ActorFactories/ActorFactoryEmptyActor.h b/Engine/Source/Editor/UnrealEd/Classes/ActorFactories/ActorFactoryEmptyActor.h
--- a/Engine/Source/Editor/UnrealEd/Classes/ActorFactories/ActorFactoryEmptyActor.h
+++ b/Engine/Source/Editor/UnrealEd/Classes/ActorFactories/ActorFactoryEmptyActor.h
@@ -14,7 +14,9 @@ public:
 
 protected:
 	virtual AActor* SpawnActor( UObject* Asset, ULevel* InLevel, const FVector& Location, const FRotator& Rotation, EObjectFlags ObjectFlags, const FName& Name ) override;
+	void SetupEditorOnlyBillboardComponent(class UBillboardComponent* BillboardComponent);
+	void OnBlueprintActorPreSpawned(class USimpleConstructionScript* SimpleConstructionScript);
 private:
 	AActor* SpawnActorForDragPreview( UObject* Asset, ULevel* InLevel, const FVector& Location, const FRotator& Rotation, EObjectFlags ObjectFlags, const FName& Name );
-	void SetupEditorOnlyBillboardComponent( class UBillboardComponent* BillboardComponent );
+
 };
diff --git a/Engine/Source/Editor/UnrealEd/Classes/ActorFactories/ActorFactoryPawn.h b/Engine/Source/Editor/UnrealEd/Classes/ActorFactories/ActorFactoryPawn.h
new file mode 100644
--- /dev/null
+++ b/Engine/Source/Editor/UnrealEd/Classes/ActorFactories/ActorFactoryPawn.h
@@ -0,0 +1,19 @@
+// Copyright 1998-2015 Epic Games, Inc. All Rights Reserved.
+
+
+#pragma once
+#include "ActorFactoryPawn.generated.h"
+
+UCLASS(MinimalAPI,config=Editor)
+class UActorFactoryPawn : public UActorFactoryEmptyActor
+{
+	GENERATED_UCLASS_BODY()
+
+public:
+	virtual bool CanCreateActorFrom( const FAssetData& AssetData, FText& OutErrorMsg ) override;
+
+protected:
+	virtual AActor* SpawnActor( UObject* Asset, ULevel* InLevel, const FVector& Location, const FRotator& Rotation, EObjectFlags ObjectFlags, const FName& Name ) override;
+private:
+	AActor* SpawnActorForDragPreview(UObject* Asset, ULevel* InLevel, const FVector& Location, const FRotator& Rotation, EObjectFlags ObjectFlags, const FName& Name);
+};
diff --git a/Engine/Source/Editor/UnrealEd/Private/Factories/ActorFactory.cpp b/Engine/Source/Editor/UnrealEd/Private/Factories/ActorFactory.cpp
--- a/Engine/Source/Editor/UnrealEd/Private/Factories/ActorFactory.cpp
+++ b/Engine/Source/Editor/UnrealEd/Private/Factories/ActorFactory.cpp
@@ -1112,13 +1112,59 @@ UActorFactoryCameraActor::UActorFactoryCameraActor(const FObjectInitializer& Obj
 	NewActorClass = ACameraActor::StaticClass();
 }
 
+DECLARE_DELEGATE_OneParam( FOnBlueprintActorPreSpawned, USimpleConstructionScript* );
+
+template <typename ActorClassType>
+static AActor* CreateAnonymousBlueprintAndSpawnActor(ULevel* InLevel, const FVector& Location, const FRotator& Rotation, EObjectFlags ObjectFlags, const FName& Name, FOnBlueprintActorPreSpawned OnBlueprintActorPreSpawned = FOnBlueprintActorPreSpawned() )
+{
+	UClass* ActorClass = ActorClassType::StaticClass();
+	FName UniqueActorName = MakeUniqueObjectName(InLevel, ActorClass, Name);
+
+	// Spawn a blank actor with an anonymous blueprint
+	// Create and init a new Blueprint
+	const FName BPName = *FString::Printf(TEXT("%s_BPClass"), *UniqueActorName.ToString());
+
+	const FName UniqueBPName = MakeUniqueObjectName(InLevel, UBlueprint::StaticClass(), BPName);
+
+
+	UBlueprint* NewBP = FKismetEditorUtilities::CreateBlueprint(ActorClass, InLevel, UniqueBPName, BPTYPE_Normal, UBlueprint::StaticClass(), UBlueprintGeneratedClass::StaticClass());
+	if(NewBP)
+	{
+		NewBP->ClearFlags(RF_Standalone);
+
+		OnBlueprintActorPreSpawned.ExecuteIfBound(NewBP->SimpleConstructionScript);
+
+		FKismetEditorUtilities::CompileBlueprint(NewBP);
+
+		FAssetData AssetData(NewBP);
+		const FString* GeneratedClassPath = AssetData.TagsAndValues.Find("GeneratedClass");
+		if(GeneratedClassPath != NULL && !GeneratedClassPath->IsEmpty())
+		{
+			UClass* GeneratedClass = Cast<UClass>(StaticLoadObject(UClass::StaticClass(), NULL, **GeneratedClassPath, NULL, LOAD_NoWarn, NULL));
+
+			if(GeneratedClass != NULL)
+			{
+				FActorSpawnParameters SpawnInfo;
+				SpawnInfo.OverrideLevel = InLevel;
+				SpawnInfo.ObjectFlags = ObjectFlags;
+				SpawnInfo.Name = UniqueActorName;
+				return InLevel->OwningWorld->SpawnActor(GeneratedClass, &Location, &Rotation, SpawnInfo);
+			}
+		}
+
+	}
+
+	return nullptr;
+}
+
+
 /*-----------------------------------------------------------------------------
 UActorFactoryEmptyActor
 -----------------------------------------------------------------------------*/
 UActorFactoryEmptyActor::UActorFactoryEmptyActor(const FObjectInitializer& ObjectInitializer)
 	: Super(ObjectInitializer)
 {
-	DisplayName = LOCTEXT("EmptyActorDisplayName", "Empty Actor");
+	DisplayName = LOCTEXT("ActorFactoryEmptyActorDisplayName", "Empty Actor");
 	NewActorClass = AActor::StaticClass();
 }
 
@@ -1132,58 +1178,20 @@ AActor* UActorFactoryEmptyActor::SpawnActor( UObject* Asset, ULevel* InLevel, co
 	AActor* NewActor = nullptr;
 	if(GetDefault<UEditorExperimentalSettings>()->bInWorldBPEditing) 
 	{
-
 		if( ObjectFlags == RF_Transient )
 		{
 			NewActor = SpawnActorForDragPreview( Asset, InLevel, Location, Rotation, ObjectFlags, Name );
 		}
 		else
 		{
-			FName ActorName(TEXT("Actor"));
-			FName UniqueActorName = MakeUniqueObjectName( InLevel, AActor::StaticClass(), ActorName );
-
-			// Spawn a blank actor with an anonymous blueprint
-			// Create and init a new Blueprint
-			const FName BPName = *FString::Printf(TEXT("%s_BPClass"), *UniqueActorName.ToString());
-
-			const FName UniqueBPName = MakeUniqueObjectName( InLevel, UBlueprint::StaticClass(), BPName );
-
-			
-			UBlueprint* NewBP = FKismetEditorUtilities::CreateBlueprint(AActor::StaticClass(), InLevel, UniqueBPName, BPTYPE_Normal, UBlueprint::StaticClass(), UBlueprintGeneratedClass::StaticClass());
-			if(NewBP)
-			{
-				NewBP->ClearFlags(RF_Standalone);
-
-				USCS_Node* NewSCSNode = NewBP->SimpleConstructionScript->CreateNode( UBillboardComponent::StaticClass() );
-				
-				UBillboardComponent* BillboardComponent = CastChecked<UBillboardComponent>( NewSCSNode->ComponentTemplate );
-
-				SetupEditorOnlyBillboardComponent( BillboardComponent );
-				
-				// Set parameters to parent this node to the "inherited" SCS node
-				NewSCSNode->SetParent( NewBP->SimpleConstructionScript->GetDefaultSceneRootNode() );
-
-				NewBP->SimpleConstructionScript->GetDefaultSceneRootNode()->AddChildNode(NewSCSNode);
-
-				FKismetEditorUtilities::CompileBlueprint(NewBP);
-
-				FAssetData AssetData(NewBP);
-				const FString* GeneratedClassPath = AssetData.TagsAndValues.Find("GeneratedClass");
-				if(GeneratedClassPath != NULL && !GeneratedClassPath->IsEmpty())
-				{
-					UClass* GeneratedClass = Cast<UClass>(StaticLoadObject(UClass::StaticClass(), NULL, **GeneratedClassPath, NULL, LOAD_NoWarn, NULL));
-
-					if(GeneratedClass != NULL)
-					{
-						FActorSpawnParameters SpawnInfo;
-						SpawnInfo.OverrideLevel = InLevel;
-						SpawnInfo.ObjectFlags = ObjectFlags;
-						SpawnInfo.Name = Name;
-						NewActor = InLevel->OwningWorld->SpawnActor(GeneratedClass, &Location, &Rotation, SpawnInfo);
-					}
-				}
-
-			}
+			NewActor = CreateAnonymousBlueprintAndSpawnActor<AActor>(
+				InLevel, 
+				Location, 
+				Rotation, 
+				ObjectFlags, 
+				FName("Actor"),
+				FOnBlueprintActorPreSpawned::CreateUObject( this, &UActorFactoryEmptyActor::OnBlueprintActorPreSpawned )
+				);
 		}
 	}
 
@@ -1196,7 +1204,7 @@ AActor* UActorFactoryEmptyActor::SpawnActorForDragPreview( UObject* Asset, ULeve
 	AActor* NewActor = Super::SpawnActor(Asset, InLevel, Location, Rotation, ObjectFlags, Name);
 
 	USceneComponent* RootComponent = ConstructObject<USceneComponent>(USceneComponent::StaticClass(), NewActor, FName("Root"), RF_Transactional);
-	RootComponent->Mobility = EComponentMobility::Static;
+	RootComponent->Mobility = EComponentMobility::Movable;
 	RootComponent->SetWorldLocationAndRotation(Location, Rotation);
 	NewActor->SetRootComponent(RootComponent);
 
@@ -1221,6 +1229,96 @@ void UActorFactoryEmptyActor::SetupEditorOnlyBillboardComponent( UBillboardCompo
 	BillboardComponent->AlwaysLoadOnServer = false;
 }
 
+void UActorFactoryEmptyActor::OnBlueprintActorPreSpawned( USimpleConstructionScript* SimpleConstructionScript )
+{
+	USCS_Node* NewSCSNode = SimpleConstructionScript->CreateNode(UBillboardComponent::StaticClass());
+
+	UBillboardComponent* BillboardComponent = CastChecked<UBillboardComponent>(NewSCSNode->ComponentTemplate);
+
+	SetupEditorOnlyBillboardComponent(BillboardComponent);
+
+	// Set parameters to parent this node to the "inherited" SCS node
+	NewSCSNode->SetParent(SimpleConstructionScript->GetDefaultSceneRootNode());
+
+	SimpleConstructionScript->GetDefaultSceneRootNode()->AddChildNode(NewSCSNode);
+}
+
+/*-----------------------------------------------------------------------------
+UActorFactoryCharacter
+-----------------------------------------------------------------------------*/
+UActorFactoryCharacter::UActorFactoryCharacter(const FObjectInitializer& ObjectInitializer)
+	: Super(ObjectInitializer)
+{
+	DisplayName = LOCTEXT("ActorFactoryCharacterDisplayName", "Character");
+	NewActorClass = ACharacter::StaticClass();
+}
+
+bool UActorFactoryCharacter::CanCreateActorFrom(const FAssetData& AssetData, FText& OutErrorMsg)
+{
+	return GetDefault<UEditorExperimentalSettings>()->bInWorldBPEditing ? AssetData.ObjectPath == FName(*ACharacter::StaticClass()->GetPathName()) : false;
+}
+
+AActor* UActorFactoryCharacter::SpawnActor(UObject* Asset, ULevel* InLevel, const FVector& Location, const FRotator& Rotation, EObjectFlags ObjectFlags, const FName& Name)
+{
+	AActor* NewActor = nullptr;
+	if(GetDefault<UEditorExperimentalSettings>()->bInWorldBPEditing)
+	{
+
+		if(ObjectFlags == RF_Transient)
+		{
+			NewActor = Super::SpawnActor( Asset, InLevel, Location, Rotation, ObjectFlags, Name );
+		}
+		else
+		{
+			NewActor = CreateAnonymousBlueprintAndSpawnActor<ACharacter>( InLevel, Location, Rotation, ObjectFlags, FName("Character") );
+		}
+
+	}
+
+	return NewActor;
+}
+
+
+/*-----------------------------------------------------------------------------
+UActorFactoryPawn
+-----------------------------------------------------------------------------*/
+UActorFactoryPawn::UActorFactoryPawn(const FObjectInitializer& ObjectInitializer)
+	: Super(ObjectInitializer)
+{
+	DisplayName = LOCTEXT("ActorFactoryPawnDisplayName", "Pawn");
+	NewActorClass = APawn::StaticClass();
+}
+
+bool UActorFactoryPawn::CanCreateActorFrom(const FAssetData& AssetData, FText& OutErrorMsg)
+{
+	return GetDefault<UEditorExperimentalSettings>()->bInWorldBPEditing ? AssetData.ObjectPath == FName(*APawn::StaticClass()->GetPathName()) : false;
+}
+
+AActor* UActorFactoryPawn::SpawnActor(UObject* Asset, ULevel* InLevel, const FVector& Location, const FRotator& Rotation, EObjectFlags ObjectFlags, const FName& Name)
+{
+	AActor* NewActor = nullptr;
+	if(GetDefault<UEditorExperimentalSettings>()->bInWorldBPEditing)
+	{
+		if(ObjectFlags == RF_Transient)
+		{
+			NewActor = Super::SpawnActor(Asset, InLevel, Location, Rotation, ObjectFlags, Name);
+		}
+		else
+		{
+			NewActor = CreateAnonymousBlueprintAndSpawnActor<APawn>(
+				InLevel,
+				Location,
+				Rotation,
+				ObjectFlags,
+				FName("Pawn"),
+				FOnBlueprintActorPreSpawned::CreateUObject(this, &UActorFactoryPawn::OnBlueprintActorPreSpawned)
+				);
+		}
+	}
+
+	return NewActor;
+}
+
 /*-----------------------------------------------------------------------------
 UActorFactoryAmbientSound
 -----------------------------------------------------------------------------*/
diff --git a/Engine/Source/Runtime/Engine/Classes/GameFramework/Character.h b/Engine/Source/Runtime/Engine/Classes/GameFramework/Character.h
--- a/Engine/Source/Runtime/Engine/Classes/GameFramework/Character.h
+++ b/Engine/Source/Runtime/Engine/Classes/GameFramework/Character.h
@@ -196,7 +196,7 @@ typedef FBasedMovementInfo FRepRelativeMovement;
  * @see https://docs.unrealengine.com/latest/INT/Gameplay/Framework/Pawn/Character/
  */ 
 
-UCLASS(abstract, config=Game, BlueprintType, hidecategories=("Pawn|Character|InternalEvents"))
+UCLASS(config=Game, BlueprintType, hidecategories=("Pawn|Character|InternalEvents"))
 class ENGINE_API ACharacter : public APawn
 {
 	GENERATED_BODY()
diff --git a/Engine/Source/Runtime/Engine/Classes/GameFramework/Pawn.h b/Engine/Source/Runtime/Engine/Classes/GameFramework/Pawn.h
--- a/Engine/Source/Runtime/Engine/Classes/GameFramework/Pawn.h
+++ b/Engine/Source/Runtime/Engine/Classes/GameFramework/Pawn.h
@@ -27,7 +27,7 @@ struct FNavAgentProperties;
  *
  * @see https://docs.unrealengine.com/latest/INT/Gameplay/Framework/Pawn/
  */
-UCLASS(abstract, config=Game, BlueprintType, Blueprintable, hidecategories=(Navigation))
+UCLASS(config=Game, BlueprintType, Blueprintable, hidecategories=(Navigation))
 class ENGINE_API APawn : public AActor, public INavAgentInterface
 {
 	GENERATED_BODY()
