Modal dialog error messages during PIE caused by an infinite loop in a Blueprint will now be posted to the "PIE" message log, complete with a link to the Blueprint, the node it was on, and a popup callstack.

diff --git a/Engine/Source/Editor/Kismet/Private/BlueprintEditor.cpp b/Engine/Source/Editor/Kismet/Private/BlueprintEditor.cpp
--- a/Engine/Source/Editor/Kismet/Private/BlueprintEditor.cpp
+++ b/Engine/Source/Editor/Kismet/Private/BlueprintEditor.cpp
@@ -2902,6 +2902,14 @@ void FBlueprintEditor::JumpToHyperlink(const UObject* ObjectReference, bool bReq
 		// Point the camera at it
 		GUnrealEd->Exec( ReferencedActor->GetWorld(), TEXT("CAMERA ALIGN ACTIVEVIEWPORTONLY"));
 	}
+	else if(const UFunction* Function = Cast<const UFunction>(ObjectReference))
+	{
+		UEdGraph* FunctionGraph = FBlueprintEditorUtils::FindScopeGraph(GetBlueprintObj(), Function);
+		if(FunctionGraph)
+		{
+			OpenDocument(const_cast<UEdGraph*>(FunctionGraph), FDocumentTracker::OpenNewDocument);
+		}
+	}
 	else
 	{
 		UE_LOG(LogBlueprint, Warning, TEXT("Unknown type of hyperlinked object"));
diff --git a/Engine/Source/Editor/Kismet/Private/BlueprintEditorModule.cpp b/Engine/Source/Editor/Kismet/Private/BlueprintEditorModule.cpp
--- a/Engine/Source/Editor/Kismet/Private/BlueprintEditorModule.cpp
+++ b/Engine/Source/Editor/Kismet/Private/BlueprintEditorModule.cpp
@@ -99,6 +99,9 @@ void FBlueprintEditorModule::StartupModule()
 
 	// Listen for clicks in log so we can focus on the object, might have to restart K2 if the K2 tab has been closed
 	MessageLogModule.GetLogListing("BlueprintLog")->OnMessageTokenClicked().AddStatic( &FocusBlueprintEditorOnObject );
+	
+	// Also listen for clicks in the PIE log, runtime errors with Blueprints may post clickable links there
+	MessageLogModule.GetLogListing("PIE")->OnMessageTokenClicked().AddStatic( &FocusBlueprintEditorOnObject );
 
 	// Add a page for pre-loading of the editor
 	MessageLogModule.GetLogListing("BlueprintLog")->NewPage(LOCTEXT("PreloadLogPageLabel", "Editor Load"));
diff --git a/Engine/Source/Editor/UnrealEd/Private/Kismet2/Kismet2.cpp b/Engine/Source/Editor/UnrealEd/Private/Kismet2/Kismet2.cpp
--- a/Engine/Source/Editor/UnrealEd/Private/Kismet2/Kismet2.cpp
+++ b/Engine/Source/Editor/UnrealEd/Private/Kismet2/Kismet2.cpp
@@ -1403,6 +1403,15 @@ TSharedPtr<class IBlueprintEditor> FKismetEditorUtilities::GetIBlueprintEditorFo
 		for (UObject* TestOuter = ObjectToFocusOn->GetOuter(); TestOuter; TestOuter = TestOuter->GetOuter())
 		{
 			TargetBP = Cast<UBlueprint>(TestOuter);
+
+			if(TargetBP == nullptr)
+			{
+				if(UBlueprintGeneratedClass* BPGeneratedClass = Cast<UBlueprintGeneratedClass>(TestOuter))
+				{
+					TargetBP = Cast<UBlueprint>(BPGeneratedClass->ClassGeneratedBy);
+				}
+			}
+
 			if (TargetBP != NULL)
 			{
 				break;
diff --git a/Engine/Source/Editor/UnrealEd/Private/Kismet2/KismetDebugUtilities.cpp b/Engine/Source/Editor/UnrealEd/Private/Kismet2/KismetDebugUtilities.cpp
--- a/Engine/Source/Editor/UnrealEd/Private/Kismet2/KismetDebugUtilities.cpp
+++ b/Engine/Source/Editor/UnrealEd/Private/Kismet2/KismetDebugUtilities.cpp
@@ -272,13 +272,70 @@ void FKismetDebugUtilities::OnScriptException(const UObject* ActiveObject, const
 				{
 					GEditor->RequestEndPlayMap();
 					FSlateApplication::Get().LeaveDebuggingMode();
-					
 					// Launch a message box notifying the user why they have been booted
 					{
+						// Callback to display a pop-up showing the Callstack, the user can highlight and copy this if needed
+						auto DisplayCallStackLambda = [](const FText CallStack)
+						{
+							TSharedPtr<SMultiLineEditableText> TextBlock;
+							TSharedRef<SWidget> DisplayWidget = 
+								SNew(SBox)
+									.MaxDesiredHeight(512)
+									.MaxDesiredWidth(512)
+									.Content()
+									[
+										SNew(SBorder)
+											.BorderImage(FEditorStyle::GetBrush("ToolPanel.GroupBorder"))
+											[
+												SNew(SScrollBox)
+												+SScrollBox::Slot()
+												[
+													SAssignNew(TextBlock, SMultiLineEditableText)
+														.AutoWrapText(true)
+														.IsReadOnly(true)
+														.Text(CallStack)
+												]
+											]
+									];
+
+							FSlateApplication::Get().PushMenu(
+								FSlateApplication::Get().GetActiveTopLevelWindow().ToSharedRef(),
+								DisplayWidget,
+								FSlateApplication::Get().GetCursorPos(),
+								FPopupTransitionEffect(FPopupTransitionEffect::TypeInPopup)
+								);	
+
+							FSlateApplication::Get().SetKeyboardFocus(TextBlock);
+						};
+
+						TSharedRef<FTokenizedMessage> Message = FTokenizedMessage::Create(EMessageSeverity::Error);
+
+						// Display a UObject link to the Blueprint that is the source of the failure
+						Message->AddToken(FTextToken::Create(LOCTEXT( "InfiniteLoopWarning_Blueprint", "Infinite Loop detected in ")));
 						FString BlueprintName;
 						BlueprintObj->GetName(BlueprintName);
+						Message->AddToken(FUObjectToken::Create(BlueprintObj, FText::FromString(BlueprintName)));
+
+						// Display a UObject link to the UFunction that is crashing. Will open the Blueprint if able and focus on the function's graph
+						Message->AddToken(FTextToken::Create(LOCTEXT( "InfiniteLoopWarning_Function", ", asserted during ")));
+						const int32 BreakpointOffset = StackFrame.Code - StackFrame.Node->Script.GetData() - 1; //@TODO: Might want to make this a parameter of Info
+						UEdGraphNode* SourceNode = FindSourceNodeForCodeLocation(ActiveObject, StackFrame.Node, BreakpointOffset, /*bAllowImpreciseHit=*/ true);
+
+						// If a source node is found, that's the token we want to link, otherwise settle with the UFunction
+						if(SourceNode)
+						{
+							Message->AddToken(FUObjectToken::Create(SourceNode, SourceNode->GetNodeTitle(ENodeTitleType::ListView)));
+						}
+						else
+						{
+							Message->AddToken(FUObjectToken::Create(StackFrame.Node, StackFrame.Node->GetDisplayNameText()));
+						}
 
-						FMessageDialog::Open( EAppMsgType::Ok, FText::Format(LOCTEXT("InfiniteLoopWarning", "Infinite Loop detected in blueprint: {0}"), FText::FromString(BlueprintName)));
+						// Add an action token to display a pop-up that will display the complete script callstack
+						Message->AddToken(FTextToken::Create(LOCTEXT( "InfiniteLoopWarning_CallStack", " with the following ")));
+						const FText CallStackAsText = FText::FromString(StackFrame.GetStackTrace());
+						Message->AddToken(FActionToken::Create(LOCTEXT( "InfiniteLoopWarning_CallStackLink", "Call Stack" ), LOCTEXT( "InfiniteLoopWarning_CallStackDesc", "Displays the underlying callstack, tracing what function calls led to the assert occuring." ), FOnActionTokenExecuted::CreateStatic(DisplayCallStackLambda, CallStackAsText)));
+						FMessageLog("PIE").AddMessage(Message);
 					}
 				}
 			}
