Issue UE-6715
Fixed issue in cook on the side, slow loading levels will take a long time to cook and not finish processing before the time slice.

diff --git a/Engine/Source/Editor/UnrealEd/Classes/CookOnTheSide/CookOnTheFlyServer.h b/Engine/Source/Editor/UnrealEd/Classes/CookOnTheSide/CookOnTheFlyServer.h
--- a/Engine/Source/Editor/UnrealEd/Classes/CookOnTheSide/CookOnTheFlyServer.h
+++ b/Engine/Source/Editor/UnrealEd/Classes/CookOnTheSide/CookOnTheFlyServer.h
@@ -623,6 +623,26 @@ private:
 	TAutoPtr<class FSandboxPlatformFile> SandboxFile;
 	bool bIsSavingPackage; // used to stop recursive mark package dirty functions
 
+	// data about the current package being processed
+	struct FReentryData
+	{
+		FName FileName;
+		bool bBeginCacheFinished;
+		int BeginCacheCount;
+
+		FReentryData() : FileName(NAME_None), bBeginCacheFinished(false), BeginCacheCount(0)
+		{ }
+
+		void Reset( const FName& InFilename )
+		{
+			FileName = InFilename;
+			bBeginCacheFinished = false;
+			BeginCacheCount = 0;
+		}
+	};
+
+	FReentryData CurrentReentryData;
+
 	FThreadSafeQueue<struct FRecompileRequest*> RecompileRequests;
 	FFilenameQueue CookRequests; // list of requested files
 	FThreadSafeUnsolicitedPackagesList UnsolicitedCookedPackages;
diff --git a/Engine/Source/Editor/UnrealEd/Private/CookOnTheFlyServer.cpp b/Engine/Source/Editor/UnrealEd/Private/CookOnTheFlyServer.cpp
--- a/Engine/Source/Editor/UnrealEd/Private/CookOnTheFlyServer.cpp
+++ b/Engine/Source/Editor/UnrealEd/Private/CookOnTheFlyServer.cpp
@@ -931,6 +931,11 @@ uint32 UCookOnTheFlyServer::TickCookOnTheSide( const float TimeSlice, uint32 &Co
 
 		const FString BuildFilename = ToBuild.GetFilename().ToString();
 
+		if ( ToBuild.GetFilename() != CurrentReentryData.FileName )
+		{
+			CurrentReentryData.Reset( ToBuild.GetFilename() );
+		}
+
 		// if we have no target platforms then we want to cook because this will cook for all target platforms in that case
 		bool bShouldCook = TargetPlatformNames.Num() > 0 ? false : ShouldCook( BuildFilename, NAME_None );
 		{
@@ -1074,17 +1079,6 @@ uint32 UCookOnTheFlyServer::TickCookOnTheSide( const float TimeSlice, uint32 &Co
 		// this will load more packages :)
 		if ( IsCookByTheBookMode() )
 		{
-
-			// if we are doing cook by the book then we are not waiting on this single package to be cooked
-			// so we could just load some more packages and then save them all at once
-			/*if ( bLastLoadWasMap && !bNeedsGC )
-			{
-			// early out and load some more packages to save
-			continue;
-			}*/
-
-
-
 			SCOPE_TIMER(ResolveRedirectors);
 			GRedirectCollector.ResolveStringAssetReference();
 		}
@@ -1100,28 +1094,51 @@ uint32 UCookOnTheFlyServer::TickCookOnTheSide( const float TimeSlice, uint32 &Co
 			GetObjectsWithOuter( PackagesToSave[0], ObjectsInPackage );
 			for ( const auto& TargetPlatformName : AllTargetPlatformNames )
 			{
+				// avoid hitches in the editor by calling BeginCacheForCookedPlatformData on each object once using time slicing, 
+				// then check if the objects are finished caching by not using time slicing
 				const ITargetPlatform* TargetPlatform = TPM.FindTargetPlatform( TargetPlatformName.ToString() );
-				
-				for ( const auto& Obj : ObjectsInPackage )
+				if ( CurrentReentryData.bBeginCacheFinished == false )
 				{
-					Obj->BeginCacheForCookedPlatformData( TargetPlatform );
-					if ( Obj->IsCachedCookedPlatformDataLoaded(TargetPlatform) == false )
+					for ( int I = 0; I < ObjectsInPackage.Num(); ++I )
 					{
+						if ( CurrentReentryData.BeginCacheCount >= I )
+						{
+							const auto& Obj = ObjectsInPackage[I];
+							Obj->BeginCacheForCookedPlatformData( TargetPlatform );
+							
+							if ( Timer.IsTimeUp() && IsRealtimeMode() )
+							{
 #if DEBUG_COOKONTHEFLY
-						UE_LOG(LogCook, Display, TEXT("Object %s isn't cached yet"), *Obj->GetFullName());
+								UE_LOG(LogCook, Display, TEXT("Object %s took too long to cache"), *Obj->GetFullName());
 #endif
-						bIsAllDataCached = false;
+								bIsAllDataCached = false;
+								break;
+							}
+						}
 					}
-
-					if ( Timer.IsTimeUp() && IsRealtimeMode() )
+				}
+				
+				if ( bIsAllDataCached )
+				{
+					// if we get here and bIsAllDataCached is true then BeginCacheFOrCookedPlatformData has been called once on every object
+					// so now 
+					CurrentReentryData.bBeginCacheFinished = true;
+					for ( const auto& Obj : ObjectsInPackage )
 					{
+						// These begin cache calls should be quick 
+						// because they will just be checking that the data is already cached and kicking off new multithreaded requests if not
+						// all sync requests should have been caught in the first begincache call above
+						Obj->BeginCacheForCookedPlatformData( TargetPlatform );
+						if ( Obj->IsCachedCookedPlatformDataLoaded(TargetPlatform) == false )
+						{
 #if DEBUG_COOKONTHEFLY
-						UE_LOG(LogCook, Display, TEXT("Object %s took too long to cache"), *Obj->GetFullName());
+							UE_LOG(LogCook, Display, TEXT("Object %s isn't cached yet"), *Obj->GetFullName());
 #endif
-						bIsAllDataCached = false;
-						break;
+							bIsAllDataCached = false;
+						}
 					}
 				}
+				
 
 				if ( Timer.IsTimeUp() && IsRealtimeMode() && (bIsAllDataCached == false) )
 				{
