Fixed massive slowdown when selecting a landscape gizmo actor

diff --git a/Engine/Source/Editor/UnrealEd/Private/EdMode.cpp b/Engine/Source/Editor/UnrealEd/Private/EdMode.cpp
--- a/Engine/Source/Editor/UnrealEd/Private/EdMode.cpp
+++ b/Engine/Source/Editor/UnrealEd/Private/EdMode.cpp
@@ -959,28 +959,52 @@ void FEdMode::GetPropertyWidgetInfos(const UStruct* InStruct, const void* InCont
 	{
 		UProperty* CurrentProp = *PropertyIt;
 		check(CurrentProp);
-		FString DisplayName = CurrentProp->GetMetaData(TEXT("DisplayName"));
-		if (!PropertyNamePrefix.IsEmpty() && DisplayNamePrefix.IsEmpty()) // Display name is already invalid
+		FString DisplayName;
+		if (PropertyNamePrefix.IsEmpty() || !DisplayNamePrefix.IsEmpty())
 		{
-			DisplayName.Empty();
+			DisplayName = CurrentProp->GetMetaData(TEXT("DisplayName"));
 		}
-		if (!DisplayName.IsEmpty()) //Display name cannot be only the prefix.
+		if (!DisplayName.IsEmpty() && !DisplayNamePrefix.IsEmpty()) //Display name cannot be only the prefix.
 		{
 			DisplayName = DisplayNamePrefix + DisplayName;
 		}
 
-		if(	ShouldCreateWidgetForProperty(CurrentProp) )
+		if (CanCreateWidgetForProperty(CurrentProp))
 		{
-			if( UArrayProperty* ArrayProp = Cast<UArrayProperty>(CurrentProp) )
+			if (ShouldCreateWidgetForProperty(CurrentProp))
 			{
-				check(InContainer != NULL);
+				if (UArrayProperty* ArrayProp = Cast<UArrayProperty>(CurrentProp))
+				{
+					check(InContainer != NULL);
+
+					FScriptArrayHelper_InContainer ArrayHelper(ArrayProp, InContainer);
+
+					// See how many widgets we need to make for the array property
+					uint32 ArrayDim = ArrayHelper.Num();
+					for (uint32 i = 0; i < ArrayDim; i++)
+					{
+						//create a new widget info struct
+						FPropertyWidgetInfo WidgetInfo;
+
+						//fill it in with the struct name
+						WidgetInfo.PropertyName = PropertyNamePrefix + CurrentProp->GetFName().ToString();
+						WidgetInfo.DisplayName = DisplayName.IsEmpty() ? WidgetInfo.PropertyName : DisplayName;
+
+						//And see if we have any meta data that matches the MD_ValidateWidgetUsing name
+						WidgetInfo.PropertyValidationName = FName(*CurrentProp->GetMetaData(MD_ValidateWidgetUsing));
 
-				FScriptArrayHelper_InContainer ArrayHelper(ArrayProp, InContainer);
+						WidgetInfo.PropertyIndex = i;
 
-				// See how many widgets we need to make for the array property
-				uint32 ArrayDim = ArrayHelper.Num();
-				for( uint32 i = 0; i < ArrayDim; i++ )
+						// See if its a transform
+						WidgetInfo.bIsTransform = IsTransformProperty(ArrayProp->Inner);
+
+						//Add it to our out array
+						OutInfos.Add(WidgetInfo);
+					}
+				}
+				else
 				{
+
 					//create a new widget info struct
 					FPropertyWidgetInfo WidgetInfo;
 
@@ -991,35 +1015,13 @@ void FEdMode::GetPropertyWidgetInfos(const UStruct* InStruct, const void* InCont
 					//And see if we have any meta data that matches the MD_ValidateWidgetUsing name
 					WidgetInfo.PropertyValidationName = FName(*CurrentProp->GetMetaData(MD_ValidateWidgetUsing));
 
-					WidgetInfo.PropertyIndex = i;
-
 					// See if its a transform
-					WidgetInfo.bIsTransform = IsTransformProperty(ArrayProp->Inner);
+					WidgetInfo.bIsTransform = IsTransformProperty(CurrentProp);
 
 					//Add it to our out array
 					OutInfos.Add(WidgetInfo);
 				}
 			}
-			else
-			{
-
-				//create a new widget info struct
-				FPropertyWidgetInfo WidgetInfo;
-
-				//fill it in with the struct name
-				WidgetInfo.PropertyName = PropertyNamePrefix + CurrentProp->GetFName().ToString();
-				WidgetInfo.DisplayName = DisplayName.IsEmpty() ? WidgetInfo.PropertyName : DisplayName;
-
-				//And see if we have any meta data that matches the MD_ValidateWidgetUsing name
-				WidgetInfo.PropertyValidationName = FName(*CurrentProp->GetMetaData(MD_ValidateWidgetUsing));
-
-				// See if its a transform
-				WidgetInfo.bIsTransform = IsTransformProperty(CurrentProp);
-
-				//Add it to our out array
-				OutInfos.Add(WidgetInfo);
-			}
-
 		}
 		else
 		{
@@ -1037,23 +1039,21 @@ void FEdMode::GetPropertyWidgetInfos(const UStruct* InStruct, const void* InCont
 			{
 				// Recursively traverse into arrays of structures, looking for additional vector properties to expose
 				UArrayProperty* ArrayProp = Cast<UArrayProperty>(CurrentProp);
-				if(ArrayProp != NULL)
+				if (ArrayProp != NULL)
 				{
 					StructProp = Cast<UStructProperty>(ArrayProp->Inner);
-					if(StructProp != NULL)
+					if (StructProp != NULL)
 					{
 						FScriptArrayHelper_InContainer ArrayHelper(ArrayProp, InContainer);
-						for(int32 ArrayIndex = 0; ArrayIndex < ArrayHelper.Num(); ++ArrayIndex)
+
+						for (int32 ArrayIndex = 0; ArrayIndex < ArrayHelper.Num(); ++ArrayIndex)
 						{
-							if(ArrayHelper.IsValidIndex(ArrayIndex))
-							{
-								const FString ArrayPostfix = FString::Printf(TEXT("[%d]"), ArrayIndex) + TEXT(".");
-								GetPropertyWidgetInfos(StructProp->Struct
-									, ArrayHelper.GetRawPtr(ArrayIndex)
-									, OutInfos
-									, PropertyNamePrefix + ArrayProp->GetFName().ToString() + ArrayPostfix
-									, !DisplayName.IsEmpty() ? (DisplayName + ArrayPostfix) : FString());
-							}
+							const FString ArrayPostfix = FString::Printf(TEXT("[%d]."), ArrayIndex);
+							GetPropertyWidgetInfos(StructProp->Struct
+								, ArrayHelper.GetRawPtr(ArrayIndex)
+								, OutInfos
+								, PropertyNamePrefix + ArrayProp->GetFName().ToString() + ArrayPostfix
+								, !DisplayName.IsEmpty() ? (DisplayName + ArrayPostfix) : FString());
 						}
 					}
 				}
