[UE-7296] Circular dependency test assets

diff --git a/Engine/Source/Editor/UnrealEd/Classes/Tests/CircularDependencyTestActor.h b/Engine/Source/Editor/UnrealEd/Classes/Tests/CircularDependencyTestActor.h
new file mode 100644
--- /dev/null
+++ b/Engine/Source/Editor/UnrealEd/Classes/Tests/CircularDependencyTestActor.h
@@ -0,0 +1,55 @@
+// Copyright 1998-2015 Epic Games, Inc. All Rights Reserved.
+
+#pragma once
+
+#include "Gameframework/Actor.h"
+#include "CircularDependencyTestActor.generated.h"
+
+UENUM(BlueprintType)
+namespace ETestResult
+{
+	enum Type
+	{
+		Unknown,
+		Failed,
+		Succeeded
+	};
+}
+
+UCLASS(Abstract, Blueprintable, BlueprintType, meta=(ChildCanTick))
+class ACircularDependencyTestActor : public AActor
+{
+	GENERATED_UCLASS_BODY()
+public:
+
+	UFUNCTION(BlueprintCallable, Category="Development|Tests")
+	bool SetTestState(ETestResult::Type NewState);
+
+	UFUNCTION(BlueprintImplementableEvent, Category="Development|Tests")
+	void VisualizeNewTestState(ETestResult::Type OldState, ETestResult::Type NewState);
+
+	UFUNCTION(BlueprintCallable, Category="Development|Tests")
+	bool TestVerifyClass(bool bCheckPropertyType = true);
+
+	UFUNCTION(BlueprintCallable, Category="Development|Tests")
+	bool TestVerifyStructMember(UScriptStruct* Struct, bool bCheckPropertyType = true);
+
+	UFUNCTION(BlueprintCallable, BlueprintNativeEvent, Category="Development|Tests")
+	bool RunVerificationTests();
+
+	UFUNCTION(BlueprintCallable, Category="Development|Tests")
+	bool TestVerifyIsBlueprintTypeVar(FName VarName, bool bCheckPropertyType);
+
+	UFUNCTION(BlueprintCallable, Category="Development|Tests")
+	void ForceResetFailure();
+
+	UFUNCTION(BlueprintCallable, Category="Development|Tests")
+	bool CheckObjectValIsBlueprint(UObject* BPObject);
+
+protected:
+	UPROPERTY(BlueprintReadOnly, Category="Development|Tests")
+	TEnumAsByte<ETestResult::Type> TestState;
+
+private: 
+	bool TestVerifyProperty(UProperty* Property, bool bCheckPropertyType);
+};
\ No newline at end of file
diff --git a/Engine/Source/Editor/UnrealEd/Private/Tests/CircularDependencyTestActor.cpp b/Engine/Source/Editor/UnrealEd/Private/Tests/CircularDependencyTestActor.cpp
new file mode 100644
--- /dev/null
+++ b/Engine/Source/Editor/UnrealEd/Private/Tests/CircularDependencyTestActor.cpp
@@ -0,0 +1,171 @@
+// Copyright 1998-2015 Epic Games, Inc. All Rights Reserved.
+
+#include "UnrealEd.h"
+#include "Tests/CircularDependencyTestActor.h"
+
+ACircularDependencyTestActor::ACircularDependencyTestActor(const FObjectInitializer& ObjectInitializer)
+	: Super(ObjectInitializer)
+	, TestState(ETestResult::Unknown)
+{
+	VisualizeNewTestState(TestState, TestState);
+}
+
+bool ACircularDependencyTestActor::SetTestState(ETestResult::Type NewState)
+{
+	VisualizeNewTestState(TestState, NewState);
+	if (TestState == ETestResult::Failed)
+	{
+		return (NewState == ETestResult::Failed);
+	}
+	
+	TestState = NewState;
+	return true;
+}
+
+static bool IsPlaceholderClass(UClass* TestClass)
+{
+	//return Cast<ULinkerPlaceholderClass>(TestClass) != nullptr;
+	return TestClass->GetName().Contains(TEXT("PLACEHOLDER"));
+}
+
+bool ACircularDependencyTestActor::TestVerifyClass(bool bCheckPropertyType)
+{
+	UClass* BPClass = GetClass();
+
+	for (UProperty* ClassProp = BPClass->PropertyLink; ClassProp != nullptr; ClassProp = ClassProp->PropertyLinkNext)
+	{
+		if (!TestVerifyProperty(ClassProp, bCheckPropertyType))
+		{
+			SetTestState(ETestResult::Failed);
+			return false;
+		}
+
+		if (UStructProperty* StructProp = Cast<UStructProperty>(ClassProp))
+		{
+			if ((StructProp->Struct == nullptr) || !TestVerifyStructMember(StructProp->Struct, bCheckPropertyType))
+			{
+				SetTestState(ETestResult::Failed);
+				return false;
+			}
+		}
+	}
+	return true;
+}
+
+bool ACircularDependencyTestActor::TestVerifyStructMember(UScriptStruct* Struct, bool bCheckPropertyType)
+{
+	for (UProperty* StructProp = Struct->PropertyLink; StructProp != nullptr; StructProp = StructProp->NextRef)
+	{
+		if (!TestVerifyProperty(StructProp, bCheckPropertyType))
+		{
+			SetTestState(ETestResult::Failed);
+			return false;
+		}
+
+// 		if (UStructProperty* StructProp2 = Cast<UStructProperty>(StructProp))
+// 		{
+// 			if ((StructProp2->Struct == nullptr) || !TestVerifyStructMember(StructProp2->Struct, bCheckPropertyType))
+// 			{
+// 				SetTestState(ETestResult::Failed);
+// 				return false;
+// 			}
+// 		}
+	}
+	return true;
+}
+
+bool ACircularDependencyTestActor::TestVerifyIsBlueprintTypeVar(FName VarName, bool bCheckPropertyType)
+{
+	FMemberReference VarRef;
+	VarRef.SetExternalMember(VarName, GetClass());
+
+	UProperty* VarProperty = VarRef.ResolveMember<UProperty>(GetClass());
+	if (UObjectProperty* ObjProp = Cast<UObjectProperty>(VarProperty))
+	{
+		UClass* ClassToVerify = ObjProp->PropertyClass;
+		if (UClassProperty* ClassProp = Cast<UClassProperty>(ObjProp))
+		{
+			ClassToVerify = Cast<UClass>(ClassProp->GetPropertyValue_InContainer(this));
+		}
+		else
+		{
+			UObject* ObjVal = ObjProp->GetPropertyValue_InContainer(this);
+			if (UChildActorComponent* AsChildComponent = Cast<UChildActorComponent>(ObjVal))
+			{
+				ObjVal = AsChildComponent->ChildActor;
+			}
+			if ((ObjVal != nullptr) && !CheckObjectValIsBlueprint(ObjVal))
+			{
+				return false;
+			}
+		}
+
+		bCheckPropertyType &= !ClassToVerify->IsChildOf<UChildActorComponent>();
+		if (bCheckPropertyType && Cast<UBlueprintGeneratedClass>(ClassToVerify) != nullptr)
+		{
+			return true;
+		}
+	}
+	return !bCheckPropertyType;
+}
+
+bool ACircularDependencyTestActor::RunVerificationTests_Implementation()
+{
+	return TestVerifyClass();
+}
+
+bool ACircularDependencyTestActor::TestVerifyProperty(UProperty* Property, bool bCheckPropertyType)
+{
+	if (UObjectProperty* ObjProp = Cast<UObjectProperty>(Property))
+	{
+		UClass* ClassToVerify = ObjProp->PropertyClass;
+		if (UClassProperty* ClassProp = Cast<UClassProperty>(ObjProp))
+		{
+			//ClassToVerify = Cast<UClass>(ClassProp->GetPropertyValue_InContainer(this));
+		}
+
+		if (ClassToVerify == nullptr || IsPlaceholderClass(ClassToVerify))
+		{
+			return false;
+		}
+
+		//if (bVerifyObjTypes)
+		{
+			bool bIsBlueprintSubClassProperty = false;
+			if (UClass* OwnerClass = Property->GetOwnerClass())
+			{
+				bIsBlueprintSubClassProperty = (OwnerClass == GetClass()) && (Cast<UBlueprintGeneratedClass>(OwnerClass) != nullptr);
+				if (UClass* SuperClass = OwnerClass->GetSuperClass())
+				{
+					FMemberReference VarRef;
+					VarRef.SetExternalMember(Property->GetFName(), SuperClass);
+					// make sure the parent doesn't have this variable
+					bIsBlueprintSubClassProperty &= (VarRef.ResolveMember<UProperty>(GetClass()) == nullptr);
+				}
+			}
+
+			if (bIsBlueprintSubClassProperty)
+			{
+				// assume any object properties that we've added to the BP 
+				// directly should be circular refs (and thus, pointers to other blueprint types) 
+				return TestVerifyIsBlueprintTypeVar(Property->GetFName(), bCheckPropertyType);
+			}
+		}
+	}
+
+	return true;
+}
+
+void ACircularDependencyTestActor::ForceResetFailure()
+{
+	if (TestState == ETestResult::Failed)
+	{
+		VisualizeNewTestState(TestState, ETestResult::Unknown);
+		TestState = ETestResult::Unknown;
+	}
+}
+
+bool ACircularDependencyTestActor::CheckObjectValIsBlueprint(UObject* BPObject)
+{
+	return (BPObject != nullptr) && (Cast<UBlueprintGeneratedClass>(BPObject->GetClass()) != nullptr);
+}
\ No newline at end of file
