Fix non-unity build.

diff --git a/Engine/Source/Editor/UnrealEd/Private/Tests/CircularDependencyTestActor.cpp b/Engine/Source/Editor/UnrealEd/Private/Tests/CircularDependencyTestActor.cpp
--- a/Engine/Source/Editor/UnrealEd/Private/Tests/CircularDependencyTestActor.cpp
+++ b/Engine/Source/Editor/UnrealEd/Private/Tests/CircularDependencyTestActor.cpp
@@ -2,6 +2,7 @@
 
 #include "UnrealEd.h"
 #include "Tests/CircularDependencyTestActor.h"
+#include "K2Node.h"
 
 ACircularDependencyTestActor::ACircularDependencyTestActor(const FObjectInitializer& ObjectInitializer)
 	: Super(ObjectInitializer)
