[UE-7296] Updating my circular dependency test class so it properly verifies struct class variables.

diff --git a/Engine/Source/Editor/UnrealEd/Classes/Tests/CircularDependencyTestActor.h b/Engine/Source/Editor/UnrealEd/Classes/Tests/CircularDependencyTestActor.h
--- a/Engine/Source/Editor/UnrealEd/Classes/Tests/CircularDependencyTestActor.h
+++ b/Engine/Source/Editor/UnrealEd/Classes/Tests/CircularDependencyTestActor.h
@@ -31,8 +31,8 @@ public:
 	UFUNCTION(BlueprintCallable, Category="Development|Tests")
 	bool TestVerifyClass(bool bCheckPropertyType = true);
 
-	UFUNCTION(BlueprintCallable, Category="Development|Tests")
-	bool TestVerifyStructMember(UScriptStruct* Struct, bool bCheckPropertyType = true);
+	//UFUNCTION(BlueprintCallable, Category="Development|Tests")
+	bool TestVerifyStructMember(UScriptStruct* Struct, uint8* Container, bool bCheckPropertyType = true);
 
 	UFUNCTION(BlueprintCallable, BlueprintNativeEvent, Category="Development|Tests")
 	bool RunVerificationTests();
@@ -51,5 +51,5 @@ protected:
 	TEnumAsByte<ETestResult::Type> TestState;
 
 private: 
-	bool TestVerifyProperty(UProperty* Property, bool bCheckPropertyType);
+	bool TestVerifyProperty(UProperty* Property, uint8* Container, bool bCheckPropertyType);
 };
\ No newline at end of file
diff --git a/Engine/Source/Editor/UnrealEd/Private/Tests/CircularDependencyTestActor.cpp b/Engine/Source/Editor/UnrealEd/Private/Tests/CircularDependencyTestActor.cpp
--- a/Engine/Source/Editor/UnrealEd/Private/Tests/CircularDependencyTestActor.cpp
+++ b/Engine/Source/Editor/UnrealEd/Private/Tests/CircularDependencyTestActor.cpp
@@ -8,7 +8,6 @@ ACircularDependencyTestActor::ACircularDependencyTestActor(const FObjectInitiali
 	: Super(ObjectInitializer)
 	, TestState(ETestResult::Unknown)
 {
-	VisualizeNewTestState(TestState, TestState);
 }
 
 bool ACircularDependencyTestActor::SetTestState(ETestResult::Type NewState)
@@ -35,7 +34,7 @@ bool ACircularDependencyTestActor::TestVerifyClass(bool bCheckPropertyType)
 
 	for (UProperty* ClassProp = BPClass->PropertyLink; ClassProp != nullptr; ClassProp = ClassProp->PropertyLinkNext)
 	{
-		if (!TestVerifyProperty(ClassProp, bCheckPropertyType))
+		if (!TestVerifyProperty(ClassProp, (uint8*)this, bCheckPropertyType))
 		{
 			SetTestState(ETestResult::Failed);
 			return false;
@@ -43,7 +42,8 @@ bool ACircularDependencyTestActor::TestVerifyClass(bool bCheckPropertyType)
 
 		if (UStructProperty* StructProp = Cast<UStructProperty>(ClassProp))
 		{
-			if ((StructProp->Struct == nullptr) || !TestVerifyStructMember(StructProp->Struct, bCheckPropertyType))
+			uint8* StructInst = StructProp->ContainerPtrToValuePtr<uint8>(this);
+			if ((StructProp->Struct == nullptr) || !TestVerifyStructMember(StructProp->Struct, (uint8*)StructInst, bCheckPropertyType))
 			{
 				SetTestState(ETestResult::Failed);
 				return false;
@@ -53,24 +53,25 @@ bool ACircularDependencyTestActor::TestVerifyClass(bool bCheckPropertyType)
 	return true;
 }
 
-bool ACircularDependencyTestActor::TestVerifyStructMember(UScriptStruct* Struct, bool bCheckPropertyType)
+bool ACircularDependencyTestActor::TestVerifyStructMember(UScriptStruct* Struct, uint8* StructInst, bool bCheckPropertyType)
 {
 	for (UProperty* StructProp = Struct->PropertyLink; StructProp != nullptr; StructProp = StructProp->NextRef)
 	{
-		if (!TestVerifyProperty(StructProp, bCheckPropertyType))
+		if (!TestVerifyProperty(StructProp, StructInst, bCheckPropertyType))
 		{
 			SetTestState(ETestResult::Failed);
 			return false;
 		}
 
-// 		if (UStructProperty* StructProp2 = Cast<UStructProperty>(StructProp))
-// 		{
-// 			if ((StructProp2->Struct == nullptr) || !TestVerifyStructMember(StructProp2->Struct, bCheckPropertyType))
-// 			{
-// 				SetTestState(ETestResult::Failed);
-// 				return false;
-// 			}
-// 		}
+		if (UStructProperty* StructProp2 = Cast<UStructProperty>(StructProp))
+		{
+			uint8* StructInst2 = StructProp2->ContainerPtrToValuePtr<uint8>(StructInst);
+			if ((StructProp2->Struct == nullptr) || !TestVerifyStructMember(StructProp2->Struct, StructInst2, bCheckPropertyType))
+			{
+				SetTestState(ETestResult::Failed);
+				return false;
+			}
+		}
 	}
 	return true;
 }
@@ -115,14 +116,14 @@ bool ACircularDependencyTestActor::RunVerificationTests_Implementation()
 	return TestVerifyClass();
 }
 
-bool ACircularDependencyTestActor::TestVerifyProperty(UProperty* Property, bool bCheckPropertyType)
+bool ACircularDependencyTestActor::TestVerifyProperty(UProperty* Property, uint8* Container, bool bCheckPropertyType)
 {
 	if (UObjectProperty* ObjProp = Cast<UObjectProperty>(Property))
 	{
 		UClass* ClassToVerify = ObjProp->PropertyClass;
 		if (UClassProperty* ClassProp = Cast<UClassProperty>(ObjProp))
 		{
-			//ClassToVerify = Cast<UClass>(ClassProp->GetPropertyValue_InContainer(this));
+			ClassToVerify = Cast<UClass>(ClassProp->GetPropertyValue_InContainer(Container));
 		}
 
 		if (ClassToVerify == nullptr || IsPlaceholderClass(ClassToVerify))
@@ -144,6 +145,10 @@ bool ACircularDependencyTestActor::TestVerifyProperty(UProperty* Property, bool
 					bIsBlueprintSubClassProperty &= (VarRef.ResolveMember<UProperty>(GetClass()) == nullptr);
 				}
 			}
+			else if (UStruct* OwnerStruct = Property->GetOwnerStruct())
+			{
+
+			}
 
 			if (bIsBlueprintSubClassProperty)
 			{
