Integrating from UE4-UT: Change 2407307 by [REDACTED_PERSON][REDACTED_HANDLE] on 2015/01/15 10:55:22

	Working around crash during lighting build in UT (UE-7559)

	Seems that the material interface and corresponding resource are accessed without null checks in BeginCompiling. The cause for the problem seems to lie with certain materials that end up being invalid when compiling for lightmass, so those materials go full white emissive in the build lighting, but at least it's not crashing anymore.

diff --git a/Engine/Source/Editor/UnrealEd/Private/Lightmass/LightmassRender.cpp b/Engine/Source/Editor/UnrealEd/Private/Lightmass/LightmassRender.cpp
--- a/Engine/Source/Editor/UnrealEd/Private/Lightmass/LightmassRender.cpp
+++ b/Engine/Source/Editor/UnrealEd/Private/Lightmass/LightmassRender.cpp
@@ -183,20 +183,25 @@ public:
 	/** Initializes the material proxy and kicks off async shader compiling. */
 	void BeginCompiling(UMaterialInterface* InMaterialInterface, EMaterialProperty InPropertyToCompile, EMaterialShaderMapUsage::Type InUsage)
 	{
-		MaterialInterface = InMaterialInterface;
-		Material = MaterialInterface ? MaterialInterface->GetMaterial() : NULL;
-		PropertyToCompile = InPropertyToCompile;
-		Usage = InUsage;
-
-		Material->AppendReferencedTextures(ReferencedTextures);
+		if (InMaterialInterface)
+		{
+			MaterialInterface = InMaterialInterface;
+			Material = MaterialInterface ? MaterialInterface->GetMaterial() : NULL;
+			PropertyToCompile = InPropertyToCompile;
+			Usage = InUsage;
 
-		FMaterialResource* Resource = InMaterialInterface->GetMaterialResource(GMaxRHIFeatureLevel);
+			Material->AppendReferencedTextures(ReferencedTextures);
 
-		FMaterialShaderMapId ResourceId;
-		Resource->GetShaderMapId(GMaxRHIShaderPlatform, ResourceId);
-		// Override with a special usage so we won't re-use the shader map used by the material for rendering
-		ResourceId.Usage = GetShaderMapUsage();
-		CacheShaders(ResourceId, GMaxRHIShaderPlatform, true);
+			FMaterialResource* Resource = InMaterialInterface->GetMaterialResource(GMaxRHIFeatureLevel);
+			if (Resource)
+			{
+				FMaterialShaderMapId ResourceId;
+				Resource->GetShaderMapId(GMaxRHIShaderPlatform, ResourceId);
+				// Override with a special usage so we won't re-use the shader map used by the material for rendering
+				ResourceId.Usage = GetShaderMapUsage();
+				CacheShaders(ResourceId, GMaxRHIShaderPlatform, true);
+			}
+		}
 	}
 
 	virtual const TArray<UTexture*>& GetReferencedTextures() const override
