FBlurpint::OnChanged split into OnCompiled and OnChanged. Only difference is that we don't call RefreshEditors in the case of OnCompiled (we manually update displayed error information on the graph nodes)
Added Stats
Fixed bad delegate logic in BlueprintActionDatabase

diff --git a/Engine/Source/Editor/BlueprintGraph/Private/BlueprintActionDatabase.cpp b/Engine/Source/Editor/BlueprintGraph/Private/BlueprintActionDatabase.cpp
--- a/Engine/Source/Editor/BlueprintGraph/Private/BlueprintActionDatabase.cpp
+++ b/Engine/Source/Editor/BlueprintGraph/Private/BlueprintActionDatabase.cpp
@@ -906,7 +906,11 @@ static void BlueprintActionDatabaseImpl::OnWorldDestroyed(UWorld* DestroyedWorld
 static bool BlueprintActionDatabaseImpl::IsObjectValidForDatabase(UObject const* Object)
 {
 	bool bReturn = false;
-	if(Object->GetOutermost()->PackageFlags & PKG_PlayInEditor)
+	if( Object == nullptr )
+	{
+		bReturn = false;
+	}
+	else if(Object->GetOutermost()->PackageFlags & PKG_PlayInEditor)
 	{
 		// Do not keep track of any PIE objects as we may prevent them from being cleaned up when ending PIE
 		bReturn = false;
@@ -1170,7 +1174,8 @@ void FBlueprintActionDatabase::RefreshAssetActions(UObject* const AssetObject)
 		// blueprint
 		if (!bHadExistingEntry)
 		{
-			OnBlueprintChangedDelegateHandle = BlueprintAsset->OnChanged().AddStatic(&BlueprintActionDatabaseImpl::OnBlueprintChanged);
+			BlueprintAsset->OnChanged().AddRaw(this, &FBlueprintActionDatabase::OnBlueprintChanged);
+			BlueprintAsset->OnCompiled().AddRaw(this, &FBlueprintActionDatabase::OnBlueprintChanged);
 		}
 	}
 
@@ -1226,7 +1231,8 @@ void FBlueprintActionDatabase::ClearAssetActions(UObject* const AssetObject)
 
 	if (UBlueprint* BlueprintAsset = Cast<UBlueprint>(AssetObject))
 	{
-		BlueprintAsset->OnChanged().Remove(OnBlueprintChangedDelegateHandle);
+		BlueprintAsset->OnChanged().RemoveAll(this);
+		BlueprintAsset->OnCompiled().RemoveAll(this);
 	}
 
 	if (bHasEntry && (ActionList->Num() > 0) && !BlueprintActionDatabaseImpl::bIsInitializing)
@@ -1299,4 +1305,9 @@ void FBlueprintActionDatabase::RegisterAllNodeActions(FBlueprintActionDatabaseRe
 	}
 }
 
+void FBlueprintActionDatabase::OnBlueprintChanged(UBlueprint* InBlueprint)
+{
+	BlueprintActionDatabaseImpl::OnBlueprintChanged(InBlueprint);
+}
+
 #undef LOCTEXT_NAMESPACE
diff --git a/Engine/Source/Editor/BlueprintGraph/Public/BlueprintActionDatabase.h b/Engine/Source/Editor/BlueprintGraph/Public/BlueprintActionDatabase.h
--- a/Engine/Source/Editor/BlueprintGraph/Public/BlueprintActionDatabase.h
+++ b/Engine/Source/Editor/BlueprintGraph/Public/BlueprintActionDatabase.h
@@ -126,6 +126,10 @@ private:
 	 */
 	void RegisterAllNodeActions(FBlueprintActionDatabaseRegistrar& Registrar);
 
+	/**
+	 * This exists only because we need a pointer to associate our delegates with
+	 */
+	void OnBlueprintChanged( UBlueprint* );
 private:
 	/** 
 	 * A map of associated node-spawners for each class/asset. A spawner that 
diff --git a/Engine/Source/Editor/GraphEditor/Private/KismetNodes/SGraphNodeK2ArrayFunction.cpp b/Engine/Source/Editor/GraphEditor/Private/KismetNodes/SGraphNodeK2ArrayFunction.cpp
--- a/Engine/Source/Editor/GraphEditor/Private/KismetNodes/SGraphNodeK2ArrayFunction.cpp
+++ b/Engine/Source/Editor/GraphEditor/Private/KismetNodes/SGraphNodeK2ArrayFunction.cpp
@@ -51,7 +51,7 @@ void SGraphNodeK2ArrayFunction::UpdateGraphNode()
 		OutputPins.Empty();
 
 		// error handling set-up
-		TSharedPtr<SWidget> ErrorText = SetupErrorReporting();
+		SetupErrorReporting();
 
 		// Reset variables that are going to be exposed, in case we are refreshing an already setup node.
 		RightNodeBox.Reset();
@@ -160,7 +160,7 @@ void SGraphNodeK2ArrayFunction::UpdateGraphNode()
 					.AutoHeight()
 					.Padding( FMargin(5.0f, 1.0f) )
 					[
-						ErrorText->AsShared()
+						ErrorReporting->AsWidget()
 					]
 			];
 
diff --git a/Engine/Source/Editor/GraphEditor/Private/KismetNodes/SGraphNodeK2Base.cpp b/Engine/Source/Editor/GraphEditor/Private/KismetNodes/SGraphNodeK2Base.cpp
--- a/Engine/Source/Editor/GraphEditor/Private/KismetNodes/SGraphNodeK2Base.cpp
+++ b/Engine/Source/Editor/GraphEditor/Private/KismetNodes/SGraphNodeK2Base.cpp
@@ -33,7 +33,7 @@ void SGraphNodeK2Base::UpdateCompactNode()
 	OutputPins.Empty();
 
 	// error handling set-up
-	TSharedPtr<SWidget> ErrorText = SetupErrorReporting();
+	SetupErrorReporting();
 
 	// Reset variables that are going to be exposed, in case we are refreshing an already setup node.
 	RightNodeBox.Reset();
@@ -69,7 +69,7 @@ void SGraphNodeK2Base::UpdateCompactNode()
 		.AutoHeight()
 		.Padding( FMargin(5.0f, 1.0f) )
 		[
-			ErrorText->AsShared()
+			ErrorReporting->AsWidget()
 		]
 		+SVerticalBox::Slot()
 		[
diff --git a/Engine/Source/Editor/GraphEditor/Private/KismetNodes/SGraphNodeK2Composite.cpp b/Engine/Source/Editor/GraphEditor/Private/KismetNodes/SGraphNodeK2Composite.cpp
--- a/Engine/Source/Editor/GraphEditor/Private/KismetNodes/SGraphNodeK2Composite.cpp
+++ b/Engine/Source/Editor/GraphEditor/Private/KismetNodes/SGraphNodeK2Composite.cpp
@@ -29,7 +29,7 @@ void SGraphNodeK2Composite::UpdateGraphNode()
 	RightNodeBox.Reset();
 	LeftNodeBox.Reset();
 
-	TSharedPtr<SWidget> ErrorText = SetupErrorReporting();
+	SetupErrorReporting();
 	TSharedPtr<SNodeTitle> NodeTitle = SNew(SNodeTitle, GraphNode);
 
 	//
@@ -104,7 +104,7 @@ void SGraphNodeK2Composite::UpdateGraphNode()
 								.AutoHeight()
 								.Padding(1.0f)
 								[
-									ErrorText->AsShared()
+									ErrorReporting->AsWidget()
 								]
 							]
 						]
diff --git a/Engine/Source/Editor/GraphEditor/Private/KismetNodes/SGraphNodeK2Var.cpp b/Engine/Source/Editor/GraphEditor/Private/KismetNodes/SGraphNodeK2Var.cpp
--- a/Engine/Source/Editor/GraphEditor/Private/KismetNodes/SGraphNodeK2Var.cpp
+++ b/Engine/Source/Editor/GraphEditor/Private/KismetNodes/SGraphNodeK2Var.cpp
@@ -158,7 +158,7 @@ void SGraphNodeK2Var::UpdateGraphNode()
 						.Text(TitleText);
 	}
 
-	TSharedPtr<SWidget> ErrorText = SetupErrorReporting();
+	SetupErrorReporting();
 
 	// Setup a meta tag for this node
 	FGraphNodeMetaData TagMeta(TEXT("Graphnode"));
@@ -183,7 +183,7 @@ void SGraphNodeK2Var::UpdateGraphNode()
 		.AutoHeight() 
 		.Padding( FMargin(5.0f, 1.0f) )
 		[
-			ErrorText->AsShared()
+			ErrorReporting->AsWidget()
 		]
 		+SVerticalBox::Slot()
 		[
diff --git a/Engine/Source/Editor/GraphEditor/Private/SGraphNode.cpp b/Engine/Source/Editor/GraphEditor/Private/SGraphNode.cpp
--- a/Engine/Source/Editor/GraphEditor/Private/SGraphNode.cpp
+++ b/Engine/Source/Editor/GraphEditor/Private/SGraphNode.cpp
@@ -595,21 +595,22 @@ void SGraphNode::UpdateErrorInfo()
 	}
 }
 
-TSharedPtr<SWidget>	SGraphNode::SetupErrorReporting()
+void SGraphNode::SetupErrorReporting()
 {
-	TSharedPtr<SErrorText> ErrorText;
-
 	UpdateErrorInfo();
 
-	// generate widget
-	SAssignNew(ErrorText, SErrorText)
-		.BackgroundColor( this, &SGraphNode::GetErrorColor )
-		.ToolTipText( this, &SGraphNode::GetErrorMsgToolTip );
+	if( !ErrorReporting.IsValid() )
+	{
+		TSharedPtr<SErrorText> ErrorTextWidget;
 
-	ErrorReporting = ErrorText;
-	ErrorReporting->SetError(ErrorMsg);
+		// generate widget
+		SAssignNew(ErrorTextWidget, SErrorText)
+			.BackgroundColor( this, &SGraphNode::GetErrorColor )
+			.ToolTipText( this, &SGraphNode::GetErrorMsgToolTip );
 
-	return ErrorText;
+		ErrorReporting = ErrorTextWidget;
+	}
+	ErrorReporting->SetError(ErrorMsg);
 }
 
 TSharedRef<SWidget> SGraphNode::CreateTitleWidget(TSharedPtr<SNodeTitle> NodeTitle)
@@ -648,7 +649,7 @@ void SGraphNode::UpdateGraphNode()
 	//            |_______|______|_______|
 	//
 	TSharedPtr<SVerticalBox> MainVerticalBox;
-	TSharedPtr<SWidget> ErrorText = SetupErrorReporting();
+	SetupErrorReporting();
 
 	TSharedPtr<SNodeTitle> NodeTitle = SNew(SNodeTitle, GraphNode);
 
@@ -787,7 +788,7 @@ void SGraphNode::UpdateGraphNode()
 					.AutoHeight()
 					.Padding(Settings->GetNonPinNodeBodyPadding())
 					[
-						ErrorText->AsShared()
+						ErrorReporting->AsWidget()
 					]
 				]
 			]			
@@ -1364,4 +1365,9 @@ void SGraphNode::PopulateMetaTag(FGraphNodeMetaData* TagMeta) const
 EGraphRenderingLOD::Type SGraphNode::GetCurrentLOD() const
 {
 	return OwnerGraphPanelPtr.IsValid() ? OwnerGraphPanelPtr.Pin()->GetCurrentLOD() : EGraphRenderingLOD::DefaultDetail;
-}
\ No newline at end of file
+}
+
+void SGraphNode::RefreshErrorInfo()
+{
+	SetupErrorReporting();
+}
diff --git a/Engine/Source/Editor/GraphEditor/Private/SGraphNodeComment.cpp b/Engine/Source/Editor/GraphEditor/Private/SGraphNodeComment.cpp
--- a/Engine/Source/Editor/GraphEditor/Private/SGraphNodeComment.cpp
+++ b/Engine/Source/Editor/GraphEditor/Private/SGraphNodeComment.cpp
@@ -114,7 +114,7 @@ void SGraphNodeComment::UpdateGraphNode()
 		TagName = FString::Printf(TEXT("GraphNode,%s,%s"), *LastOuter->GetFullName(), *GraphNode->NodeGuid.ToString());
 	}
 
-	TSharedPtr<SWidget> ErrorText = SetupErrorReporting();
+	SetupErrorReporting();
 
 	// Setup a meta tag for this node
 	FGraphNodeMetaData TagMeta(TEXT("Graphnode"));
@@ -166,7 +166,7 @@ void SGraphNodeComment::UpdateGraphNode()
 				.AutoHeight()
 				.Padding(1.0f)
 				[
-					ErrorText->AsShared()
+					ErrorReporting->AsWidget()
 				]
 				+SVerticalBox::Slot()
 				.AutoHeight()
diff --git a/Engine/Source/Editor/GraphEditor/Private/SGraphNodeDocumentation.cpp b/Engine/Source/Editor/GraphEditor/Private/SGraphNodeDocumentation.cpp
--- a/Engine/Source/Editor/GraphEditor/Private/SGraphNodeDocumentation.cpp
+++ b/Engine/Source/Editor/GraphEditor/Private/SGraphNodeDocumentation.cpp
@@ -59,7 +59,7 @@ void SGraphNodeDocumentation::UpdateGraphNode()
 	RightNodeBox.Reset();
 	LeftNodeBox.Reset();
 
-	TSharedPtr<SWidget> ErrorText = SetupErrorReporting();
+	SetupErrorReporting();
 
 	// Create Node Title
 	TSharedPtr<SNodeTitle> NodeTitle = SNew( SNodeTitle, GraphNode );
diff --git a/Engine/Source/Editor/GraphEditor/Public/SGraphNode.h b/Engine/Source/Editor/GraphEditor/Public/SGraphNode.h
--- a/Engine/Source/Editor/GraphEditor/Public/SGraphNode.h
+++ b/Engine/Source/Editor/GraphEditor/Public/SGraphNode.h
@@ -2,13 +2,13 @@
 
 #pragma once
 
-class SGraphPanel;
-class SGraphNode;
+#include "SNodePanel.h"
+
 class SGraphPin;
+class SGraphPanel;
 class UEdGraphNode;
 class FActorDragDropOp;
 
-
 /////////////////////////////////////////////////////
 // SNodeTitle
 
@@ -195,6 +195,9 @@ public:
 	/** Returns the current Node LOD or Highest LOD if unable to query */
 	EGraphRenderingLOD::Type GetCurrentLOD() const;
 
+	/** Called when GraphNode changes its error information, may be called when no change has actually occurred: */
+	void RefreshErrorInfo();
+
 protected:
 	SGraphNode();
 
@@ -254,7 +257,7 @@ protected:
 	void UpdateErrorInfo();
 
 	/** Set-up the error reporting widget for the node */
-	TSharedPtr<SWidget>	SetupErrorReporting();
+	void SetupErrorReporting();
 
 	// Should we use low-detail node titles?
 	virtual bool UseLowDetailNodeTitles() const;
diff --git a/Engine/Source/Editor/Kismet/Private/BlueprintEditor.cpp b/Engine/Source/Editor/Kismet/Private/BlueprintEditor.cpp
--- a/Engine/Source/Editor/Kismet/Private/BlueprintEditor.cpp
+++ b/Engine/Source/Editor/Kismet/Private/BlueprintEditor.cpp
@@ -41,6 +41,7 @@
 #include "ClassIconFinder.h"
 
 // Core kismet tabs
+#include "SGraphNode.h"
 #include "SSCSEditor.h"
 #include "SSCSEditorViewport.h"
 #include "STimelineEditor.h"
@@ -1342,6 +1343,7 @@ void FBlueprintEditor::CommonInitialization(const TArray<UBlueprint*>& InitBluep
 
 		// When the blueprint that we are observing changes, it will notify this wrapper widget.
 		InitBlueprint->OnChanged().AddSP(this, &FBlueprintEditor::OnBlueprintChanged);
+		InitBlueprint->OnCompiled().AddSP(this, &FBlueprintEditor::OnBlueprintCompiled);
 	}
 
 	CreateDefaultCommands();
@@ -2706,12 +2708,15 @@ void FBlueprintEditor::OnSelectedNodesChanged(const FGraphPanelSelectionSet& New
 	Inspector->ShowDetailsForObjects(NewSelection.Array());
 }
 
-void FBlueprintEditor::OnBlueprintChanged(UBlueprint* InBlueprint)
+void FBlueprintEditor::OnBlueprintChangedImpl(UBlueprint* InBlueprint, bool bIsJustBeingCompiled )
 {
 	if (InBlueprint)
 	{
 		// Refresh the graphs
-		RefreshEditors();
+		if( !bIsJustBeingCompiled )
+		{
+			RefreshEditors();
+		}
 
 		// Notify that the blueprint has been changed (update Content browser, etc)
 		InBlueprint->PostEditChange();
@@ -2727,6 +2732,34 @@ void FBlueprintEditor::OnBlueprintChanged(UBlueprint* InBlueprint)
 	}
 }
 
+void FBlueprintEditor::OnBlueprintCompiled(UBlueprint* InBlueprint)
+{	
+	if( InBlueprint )
+	{
+		// This could be made more efficient by tracking which nodes change
+		// their bHasCompilerMessage flag, or immediately updating the error info
+		// when we assign the flag:
+		TArray<UEdGraph*> Graphs;
+		InBlueprint->GetAllGraphs(Graphs);
+		for (const auto Graph : Graphs)
+		{
+			for (const auto Node : Graph->Nodes)
+			{
+				if (Node)
+				{
+					auto Widget = Node->NodeWidget.Pin();
+					if (Widget.IsValid())
+					{
+						Widget->RefreshErrorInfo();
+					}
+				}
+			}
+		}
+	}
+
+	OnBlueprintChangedImpl( InBlueprint, true );
+}
+
 void FBlueprintEditor::OnBlueprintUnloaded(UBlueprint* InBlueprint)
 {
 	for (UObject* EditingObj : GetEditingObjects())
diff --git a/Engine/Source/Editor/Kismet/Public/BlueprintEditor.h b/Engine/Source/Editor/Kismet/Public/BlueprintEditor.h
--- a/Engine/Source/Editor/Kismet/Public/BlueprintEditor.h
+++ b/Engine/Source/Editor/Kismet/Public/BlueprintEditor.h
@@ -573,7 +573,12 @@ protected:
 	void OnGraphActionMenuClosed(bool bActionExecuted, bool bContextSensitiveChecked, bool bGraphPinContext);
 
 	/** Called when the Blueprint we are editing has changed */
-	virtual void OnBlueprintChanged(UBlueprint* InBlueprint);
+	virtual void OnBlueprintChangedImpl(UBlueprint* InBlueprint, bool bIsJustBeingCompiled = false);
+
+	/** Called when the Blueprint we are editing has changed, forwards to impl */
+	void OnBlueprintChanged(UBlueprint* InBlueprint) { return OnBlueprintChangedImpl(InBlueprint); }
+
+	void OnBlueprintCompiled(UBlueprint* InBlueprint);
 
 	/** Handles the unloading of Blueprints (by closing the editor, if it operating on the Blueprint being unloaded)*/
 	void OnBlueprintUnloaded(UBlueprint* InBlueprint);
diff --git a/Engine/Source/Editor/KismetCompiler/Private/KismetCompiler.cpp b/Engine/Source/Editor/KismetCompiler/Private/KismetCompiler.cpp
--- a/Engine/Source/Editor/KismetCompiler/Private/KismetCompiler.cpp
+++ b/Engine/Source/Editor/KismetCompiler/Private/KismetCompiler.cpp
@@ -32,25 +32,20 @@ static bool bDebugPropertyPropagation = false;
 
 //////////////////////////////////////////////////////////////////////////
 // Stats for this module
-
-
-DEFINE_STAT(EKismetCompilerStats_CompileTime);
-DEFINE_STAT(EKismetCompilerStats_CreateSchema);
-DEFINE_STAT(EKismetCompilerStats_CreateFunctionList);
-DEFINE_STAT(EKismetCompilerStats_Expansion);
-DEFINE_STAT(EKismetCompilerStats_ProcessUbergraph);
-DEFINE_STAT(EKismetCompilerStats_ProcessFunctionGraph);
-DEFINE_STAT(EKismetCompilerStats_PrecompileFunction);
-DEFINE_STAT(EKismetCompilerStats_CompileFunction);
-DEFINE_STAT(EKismetCompilerStats_PostcompileFunction);
-DEFINE_STAT(EKismetCompilerStats_FinalizationWork);
-DEFINE_STAT(EKismetCompilerStats_CodeGenerationTime);
-DEFINE_STAT(EKismetCompilerStats_ChooseTerminalScope);
-DEFINE_STAT(EKismetCompilerStats_CleanAndSanitizeClass);
-DEFINE_STAT(EKismetCompilerStats_CreateClassVariables);
-DEFINE_STAT(EKismetCompilerStats_BindAndLinkClass);
-DEFINE_STAT(EKismetCompilerStats_ChecksumCDO);
-DEFINE_STAT(EKismetCompilerStats_ResolveCompiledStatements);
+DECLARE_CYCLE_STAT(TEXT("Create Schema"), EKismetCompilerStats_CreateSchema, STATGROUP_KismetCompiler );
+DECLARE_CYCLE_STAT(TEXT("Create Function List"), EKismetCompilerStats_CreateFunctionList, STATGROUP_KismetCompiler );
+DECLARE_CYCLE_STAT(TEXT("Expansion"), EKismetCompilerStats_Expansion, STATGROUP_KismetCompiler )
+DECLARE_CYCLE_STAT(TEXT("Process uber"), EKismetCompilerStats_ProcessUbergraph, STATGROUP_KismetCompiler );
+DECLARE_CYCLE_STAT(TEXT("Process func"), EKismetCompilerStats_ProcessFunctionGraph, STATGROUP_KismetCompiler );
+DECLARE_CYCLE_STAT(TEXT("Precompile Function"), EKismetCompilerStats_PrecompileFunction, STATGROUP_KismetCompiler );
+DECLARE_CYCLE_STAT(TEXT("Compile Function"), EKismetCompilerStats_CompileFunction, STATGROUP_KismetCompiler );
+DECLARE_CYCLE_STAT(TEXT("Postcompile Function"), EKismetCompilerStats_PostcompileFunction, STATGROUP_KismetCompiler );
+DECLARE_CYCLE_STAT(TEXT("Finalization"), EKismetCompilerStats_FinalizationWork, STATGROUP_KismetCompiler );
+DECLARE_CYCLE_STAT(TEXT("Code Gen"), EKismetCompilerStats_CodeGenerationTime, STATGROUP_KismetCompiler);
+DECLARE_CYCLE_STAT(TEXT("Clean and Sanitize Class"), EKismetCompilerStats_CleanAndSanitizeClass, STATGROUP_KismetCompiler );
+DECLARE_CYCLE_STAT(TEXT("Create Class Properties"), EKismetCompilerStats_CreateClassVariables, STATGROUP_KismetCompiler );
+DECLARE_CYCLE_STAT(TEXT("Bind and Link Class"), EKismetCompilerStats_BindAndLinkClass, STATGROUP_KismetCompiler );
+DECLARE_CYCLE_STAT(TEXT("Calculate checksum of CDO"), EKismetCompilerStats_ChecksumCDO, STATGROUP_KismetCompiler );
 		
 //////////////////////////////////////////////////////////////////////////
 // FKismetCompilerContext
diff --git a/Engine/Source/Editor/KismetCompiler/Private/KismetCompilerMisc.cpp b/Engine/Source/Editor/KismetCompiler/Private/KismetCompilerMisc.cpp
--- a/Engine/Source/Editor/KismetCompiler/Private/KismetCompilerMisc.cpp
+++ b/Engine/Source/Editor/KismetCompiler/Private/KismetCompilerMisc.cpp
@@ -16,6 +16,10 @@
 
 #define LOCTEXT_NAMESPACE "KismetCompiler"
 
+DECLARE_CYCLE_STAT(TEXT("Choose Terminal Scope"), EKismetCompilerStats_ChooseTerminalScope, STATGROUP_KismetCompiler);
+DECLARE_CYCLE_STAT(TEXT("Resolve compiled statements"), EKismetCompilerStats_ResolveCompiledStatements, STATGROUP_KismetCompiler );
+
+
 //////////////////////////////////////////////////////////////////////////
 // FKismetCompilerUtilities
 
diff --git a/Engine/Source/Editor/KismetCompiler/Private/KismetCompilerModule.cpp b/Engine/Source/Editor/KismetCompiler/Private/KismetCompilerModule.cpp
--- a/Engine/Source/Editor/KismetCompiler/Private/KismetCompilerModule.cpp
+++ b/Engine/Source/Editor/KismetCompiler/Private/KismetCompilerModule.cpp
@@ -15,6 +15,9 @@
 #include "UserDefinedStructureCompilerUtils.h"
 
 DEFINE_LOG_CATEGORY(LogK2Compiler);
+DECLARE_CYCLE_STAT(TEXT("Compile Time"), EKismetCompilerStats_CompileTime, STATGROUP_KismetCompiler);
+DECLARE_CYCLE_STAT(TEXT("Compile Skeleton Class"), EKismetCompilerStats_CompileSkeletonClass, STATGROUP_KismetCompiler);
+DECLARE_CYCLE_STAT(TEXT("Compile Generated Class"), EKismetCompilerStats_CompileGeneratedClass, STATGROUP_KismetCompiler);
 
 IMPLEMENT_MODULE( FKismet2CompilerModule, KismetCompiler );
 
@@ -114,7 +117,7 @@ void FKismet2CompilerModule::CompileBlueprint(class UBlueprint* Blueprint, const
 
 	if ((CompileOptions.CompileType != EKismetCompileType::BytecodeOnly) && (CompileOptions.CompileType != EKismetCompileType::Cpp))
 	{
-		BP_SCOPED_COMPILER_EVENT_NAME(TEXT("Compile Skeleton Class"));
+		BP_SCOPED_COMPILER_EVENT_STAT(EKismetCompilerStats_CompileSkeletonClass);
 
 		FBlueprintCompileReinstancer SkeletonReinstancer(Blueprint->SkeletonGeneratedClass);
 
@@ -128,7 +131,7 @@ void FKismet2CompilerModule::CompileBlueprint(class UBlueprint* Blueprint, const
 	// If this was a full compile, take appropriate actions depending on the success of failure of the compile
 	if( CompileOptions.IsGeneratedClassCompileType() )
 	{
-		BP_SCOPED_COMPILER_EVENT_NAME(TEXT("Compile Generated Class"));
+		BP_SCOPED_COMPILER_EVENT_STAT(EKismetCompilerStats_CompileGeneratedClass);
 
 		// Perform the full compile
 		CompileBlueprintInner(Blueprint, CompileOptions, Results, ObjLoaded);
diff --git a/Engine/Source/Editor/KismetCompiler/Public/EdGraphCompilerUtilities.h b/Engine/Source/Editor/KismetCompiler/Public/EdGraphCompilerUtilities.h
--- a/Engine/Source/Editor/KismetCompiler/Public/EdGraphCompilerUtilities.h
+++ b/Engine/Source/Editor/KismetCompiler/Public/EdGraphCompilerUtilities.h
@@ -10,26 +10,6 @@
 #include "BlueprintUtilities.h"
 #include "EdGraphUtilities.h"
 
-DECLARE_STATS_GROUP(TEXT("Kismet Compiler"), STATGROUP_KismetCompiler, STATCAT_Advanced);
-
-DECLARE_CYCLE_STAT_EXTERN(TEXT("Compile Time"), EKismetCompilerStats_CompileTime, STATGROUP_KismetCompiler, );
-DECLARE_CYCLE_STAT_EXTERN(TEXT("Create Schema"), EKismetCompilerStats_CreateSchema, STATGROUP_KismetCompiler, );
-DECLARE_CYCLE_STAT_EXTERN(TEXT("Create Function List"), EKismetCompilerStats_CreateFunctionList, STATGROUP_KismetCompiler, );
-DECLARE_CYCLE_STAT_EXTERN(TEXT("Expansion"), EKismetCompilerStats_Expansion, STATGROUP_KismetCompiler, );
-DECLARE_CYCLE_STAT_EXTERN(TEXT("Process uber"), EKismetCompilerStats_ProcessUbergraph, STATGROUP_KismetCompiler, );
-DECLARE_CYCLE_STAT_EXTERN(TEXT("Process func"), EKismetCompilerStats_ProcessFunctionGraph, STATGROUP_KismetCompiler, );
-DECLARE_CYCLE_STAT_EXTERN(TEXT("Precompile Function"), EKismetCompilerStats_PrecompileFunction, STATGROUP_KismetCompiler, );
-DECLARE_CYCLE_STAT_EXTERN(TEXT("Compile Function"), EKismetCompilerStats_CompileFunction, STATGROUP_KismetCompiler, );
-DECLARE_CYCLE_STAT_EXTERN(TEXT("Postcompile Function"), EKismetCompilerStats_PostcompileFunction, STATGROUP_KismetCompiler, );
-DECLARE_CYCLE_STAT_EXTERN(TEXT("Finalization"), EKismetCompilerStats_FinalizationWork, STATGROUP_KismetCompiler, );
-DECLARE_CYCLE_STAT_EXTERN(TEXT("Code Gen"), EKismetCompilerStats_CodeGenerationTime, STATGROUP_KismetCompiler, );
-DECLARE_CYCLE_STAT_EXTERN(TEXT("Choose Terminal Scope"), EKismetCompilerStats_ChooseTerminalScope, STATGROUP_KismetCompiler, );
-DECLARE_CYCLE_STAT_EXTERN(TEXT("Clean and Sanitize Class"), EKismetCompilerStats_CleanAndSanitizeClass, STATGROUP_KismetCompiler, );
-DECLARE_CYCLE_STAT_EXTERN(TEXT("Create Class Properties"), EKismetCompilerStats_CreateClassVariables, STATGROUP_KismetCompiler, );
-DECLARE_CYCLE_STAT_EXTERN(TEXT("Bind and Link Class"), EKismetCompilerStats_BindAndLinkClass, STATGROUP_KismetCompiler, );
-DECLARE_CYCLE_STAT_EXTERN(TEXT("Calculate checksum of CDO"), EKismetCompilerStats_ChecksumCDO, STATGROUP_KismetCompiler, );
-DECLARE_CYCLE_STAT_EXTERN(TEXT("Resolve compiled statements"), EKismetCompilerStats_ResolveCompiledStatements, STATGROUP_KismetCompiler, );
-
 //////////////////////////////////////////////////////////////////////////
 
 class KISMETCOMPILER_API FGraphCompilerContext
diff --git a/Engine/Source/Editor/Persona/Private/SAnimBlueprintParentPlayerList.cpp b/Engine/Source/Editor/Persona/Private/SAnimBlueprintParentPlayerList.cpp
--- a/Engine/Source/Editor/Persona/Private/SAnimBlueprintParentPlayerList.cpp
+++ b/Engine/Source/Editor/Persona/Private/SAnimBlueprintParentPlayerList.cpp
@@ -18,6 +18,7 @@ SAnimBlueprintParentPlayerList::~SAnimBlueprintParentPlayerList()
 	if(RootBlueprint)
 	{
 		RootBlueprint->OnChanged().RemoveAll(this);
+		RootBlueprint->OnCompiled().RemoveAll(this);
 	}
 
 	if(CurrentBlueprint)
@@ -34,6 +35,7 @@ SAnimBlueprintParentPlayerList::~SAnimBlueprintParentPlayerList()
 		}
 
 		CurrentBlueprint->OnChanged().RemoveAll(this);
+		CurrentBlueprint->OnCompiled().RemoveAll(this);
 	}
 
 	TSharedPtr<FPersona> PersonaShared = PersonaPtr.Pin();
@@ -84,6 +86,7 @@ void SAnimBlueprintParentPlayerList::Construct(const FArguments& InArgs)
 		if(RootBlueprint != CurrentBlueprint)
 		{
 			RootBlueprint->OnChanged().AddSP(this, &SAnimBlueprintParentPlayerList::OnRootBlueprintChanged);
+			RootBlueprint->OnCompiled().AddSP(this, &SAnimBlueprintParentPlayerList::OnRootBlueprintChanged);
 		}
 		else
 		{
@@ -104,6 +107,7 @@ void SAnimBlueprintParentPlayerList::Construct(const FArguments& InArgs)
 
 		// Register a delegate for the current BP changing; so we can look for a root change.
 		CurrentBlueprint->OnChanged().AddSP(this, &SAnimBlueprintParentPlayerList::OnCurrentBlueprintChanged);
+		CurrentBlueprint->OnCompiled().AddSP(this, &SAnimBlueprintParentPlayerList::OnCurrentBlueprintChanged);
 	}
 	else
 	{
diff --git a/Engine/Source/Editor/UMGEditor/Private/Designer/SDesignerView.cpp b/Engine/Source/Editor/UMGEditor/Private/Designer/SDesignerView.cpp
--- a/Engine/Source/Editor/UMGEditor/Private/Designer/SDesignerView.cpp
+++ b/Engine/Source/Editor/UMGEditor/Private/Designer/SDesignerView.cpp
@@ -521,6 +521,7 @@ SDesignerView::~SDesignerView()
 	if ( Blueprint )
 	{
 		Blueprint->OnChanged().RemoveAll(this);
+		Blueprint->OnCompiled().RemoveAll(this);
 	}
 
 	if ( BlueprintEditor.IsValid() )
diff --git a/Engine/Source/Editor/UMGEditor/Private/Hierarchy/SHierarchyView.cpp b/Engine/Source/Editor/UMGEditor/Private/Hierarchy/SHierarchyView.cpp
--- a/Engine/Source/Editor/UMGEditor/Private/Hierarchy/SHierarchyView.cpp
+++ b/Engine/Source/Editor/UMGEditor/Private/Hierarchy/SHierarchyView.cpp
@@ -36,6 +36,7 @@ void SHierarchyView::Construct(const FArguments& InArgs, TSharedPtr<FWidgetBluep
 
 	UWidgetBlueprint* Blueprint = GetBlueprint();
 	Blueprint->OnChanged().AddRaw(this, &SHierarchyView::OnBlueprintChanged);
+	Blueprint->OnCompiled().AddRaw(this, &SHierarchyView::OnBlueprintChanged);
 
 	FilterHandler = MakeShareable(new TreeFilterHandler< TSharedPtr<FHierarchyModel> >());
 	FilterHandler->SetFilter(SearchBoxWidgetFilter.Get());
@@ -81,6 +82,7 @@ SHierarchyView::~SHierarchyView()
 	if ( Blueprint )
 	{
 		Blueprint->OnChanged().RemoveAll(this);
+		Blueprint->OnCompiled().RemoveAll(this);
 	}
 
 	if ( BlueprintEditor.IsValid() )
diff --git a/Engine/Source/Editor/UMGEditor/Private/WidgetBlueprintEditor.cpp b/Engine/Source/Editor/UMGEditor/Private/WidgetBlueprintEditor.cpp
--- a/Engine/Source/Editor/UMGEditor/Private/WidgetBlueprintEditor.cpp
+++ b/Engine/Source/Editor/UMGEditor/Private/WidgetBlueprintEditor.cpp
@@ -40,6 +40,7 @@ FWidgetBlueprintEditor::~FWidgetBlueprintEditor()
 	if ( Blueprint )
 	{
 		Blueprint->OnChanged().RemoveAll(this);
+		Blueprint->OnCompiled().RemoveAll(this);
 	}
 
 	GEditor->OnObjectsReplaced().RemoveAll(this);
@@ -222,9 +223,9 @@ void FWidgetBlueprintEditor::InvalidatePreview()
 	bPreviewInvalidated = true;
 }
 
-void FWidgetBlueprintEditor::OnBlueprintChanged(UBlueprint* InBlueprint)
+void FWidgetBlueprintEditor::OnBlueprintChangedImpl(UBlueprint* InBlueprint, bool bIsJustBeingCompiled )
 {
-	FBlueprintEditor::OnBlueprintChanged(InBlueprint);
+	FBlueprintEditor::OnBlueprintChangedImpl(InBlueprint, bIsJustBeingCompiled);
 
 	if ( InBlueprint )
 	{
diff --git a/Engine/Source/Editor/UMGEditor/Private/WidgetBlueprintEditor.h b/Engine/Source/Editor/UMGEditor/Private/WidgetBlueprintEditor.h
--- a/Engine/Source/Editor/UMGEditor/Private/WidgetBlueprintEditor.h
+++ b/Engine/Source/Editor/UMGEditor/Private/WidgetBlueprintEditor.h
@@ -161,7 +161,7 @@ private:
 
 private:
 	/** Called whenever the blueprint is structurally changed. */
-	virtual void OnBlueprintChanged(UBlueprint* InBlueprint) override;
+	virtual void OnBlueprintChangedImpl(UBlueprint* InBlueprint, bool bIsJustBeingCompiled = false ) override;
 
 	/** Called when objects need to be swapped out for new versions, like after a blueprint recompile. */
 	void OnObjectsReplaced(const TMap<UObject*, UObject*>& ReplacementMap);
diff --git a/Engine/Source/Editor/UnrealEd/Private/Kismet2/BlueprintEditorUtils.cpp b/Engine/Source/Editor/UnrealEd/Private/Kismet2/BlueprintEditorUtils.cpp
--- a/Engine/Source/Editor/UnrealEd/Private/Kismet2/BlueprintEditorUtils.cpp
+++ b/Engine/Source/Editor/UnrealEd/Private/Kismet2/BlueprintEditorUtils.cpp
@@ -30,6 +30,7 @@
 #include "Editor/Kismet/Public/FindInBlueprintManager.h"
 
 #include "BlueprintEditor.h"
+#include "BlueprintEditorSettings.h"
 #include "Editor/UnrealEd/Public/Kismet2/Kismet2NameValidators.h"
 
 #include "DefaultValueHelper.h"
@@ -51,6 +52,10 @@
 
 DEFINE_LOG_CATEGORY(LogBlueprintDebug);
 
+DECLARE_CYCLE_STAT(TEXT("Notify Blueprint Changed"), EKismetCompilerStats_NotifyBlueprintChanged, STATGROUP_KismetCompiler);
+DECLARE_CYCLE_STAT(TEXT("Mark Blueprint as Structurally Modified"), EKismetCompilerStats_MarkBlueprintasStructurallyModified, STATGROUP_KismetCompiler);
+DECLARE_CYCLE_STAT(TEXT("Refresh External DependencyNodes"), EKismetCompilerStats_RefreshExternalDependencyNodes, STATGROUP_KismetCompiler);
+
 FBlueprintEditorUtils::FFixLevelScriptActorBindingsEvent FBlueprintEditorUtils::FixLevelScriptActorBindingsEvent;
 
 struct FCompareNodePriority
@@ -320,7 +325,7 @@ void FBlueprintEditorUtils::RefreshAllNodes(UBlueprint* Blueprint)
 
 void FBlueprintEditorUtils::RefreshExternalBlueprintDependencyNodes(UBlueprint* Blueprint, UStruct* RefreshOnlyChild)
 {
-	BP_SCOPED_COMPILER_EVENT_NAME(TEXT("Refresh External Dependency Nodes"));
+	BP_SCOPED_COMPILER_EVENT_STAT(EKismetCompilerStats_RefreshExternalDependencyNodes);
 
 	if (!Blueprint || !Blueprint->HasAllFlags(RF_LoadCompleted))
 	{
@@ -1631,7 +1636,7 @@ void FBlueprintEditorUtils::MarkBlueprintAsStructurallyModified(UBlueprint* Blue
 		FCompilerResultsLog Results;
 		Results.bLogInfoOnly = Blueprint->bIsRegeneratingOnLoad;
 
-		BP_SCOPED_COMPILER_EVENT_NAME(TEXT("Mark Blueprint as Structurally Modified"));
+		BP_SCOPED_COMPILER_EVENT_STAT(EKismetCompilerStats_MarkBlueprintasStructurallyModified);
 
 		TArray<UEdGraph*> AllGraphs;
 		Blueprint->GetAllGraphs(AllGraphs);
@@ -1659,7 +1664,8 @@ void FBlueprintEditorUtils::MarkBlueprintAsStructurallyModified(UBlueprint* Blue
 		}
 		UpdateDelegatesInBlueprint(Blueprint);
 
-		{ BP_SCOPED_COMPILER_EVENT_NAME(TEXT("Notify Blueprint Changed"));
+		{
+			BP_SCOPED_COMPILER_EVENT_STAT(EKismetCompilerStats_NotifyBlueprintChanged);
 
 			// Notify any interested parties that the blueprint has changed
 			Blueprint->BroadcastChanged();
diff --git a/Engine/Source/Editor/UnrealEd/Private/Kismet2/Kismet2.cpp b/Engine/Source/Editor/UnrealEd/Private/Kismet2/Kismet2.cpp
--- a/Engine/Source/Editor/UnrealEd/Private/Kismet2/Kismet2.cpp
+++ b/Engine/Source/Editor/UnrealEd/Private/Kismet2/Kismet2.cpp
@@ -38,6 +38,14 @@
 #include "NotificationManager.h"
 #include "SNotificationList.h" // for FNotificationInfo
 
+DECLARE_CYCLE_STAT(TEXT("Compile Blueprint"), EKismetCompilerStats_CompileBlueprint, STATGROUP_KismetCompiler);
+DECLARE_CYCLE_STAT(TEXT("Broadcast Precompile"), EKismetCompilerStats_BroadcastPrecompile, STATGROUP_KismetCompiler);
+DECLARE_CYCLE_STAT(TEXT("Update Search Metadata"), EKismetCompilerStats_UpdateSearchMetaData, STATGROUP_KismetCompiler);
+DECLARE_CYCLE_STAT(TEXT("Garbage Collection"), EKismetCompilerStats_GarbageCollection, STATGROUP_KismetCompiler);
+DECLARE_CYCLE_STAT(TEXT("Notify Blueprint Changed"), EKismetCompilerStats_NotifyBlueprintChanged, STATGROUP_KismetCompiler);
+DECLARE_CYCLE_STAT(TEXT("Refresh Dependent Blueprints"), EKismetCompilerStats_RefreshDependentBlueprints, STATGROUP_KismetCompiler);
+DECLARE_CYCLE_STAT(TEXT("Validate Generated Class"), EKismetCompilerStats_ValidateGeneratedClass, STATGROUP_KismetCompiler);
+
 #define LOCTEXT_NAMESPACE "UnrealEd.Editor"
 
 //////////////////////////////////////////////////////////////////////////
@@ -568,11 +576,16 @@ bool FKismetEditorUtilities::IsReferencedByUndoBuffer(UBlueprint* Blueprint)
 
 void FKismetEditorUtilities::CompileBlueprint(UBlueprint* BlueprintObj, bool bIsRegeneratingOnLoad, bool bSkipGarbageCollection, bool bSaveIntermediateProducts, FCompilerResultsLog* pResults)
 {
+	BP_SCOPED_COMPILER_EVENT_STAT(EKismetCompilerStats_CompileBlueprint);
+
 	// Broadcast pre-compile
 #if WITH_EDITOR
-	if(GEditor && GIsEditor)
 	{
-		GEditor->BroadcastBlueprintPreCompile(BlueprintObj);
+		BP_SCOPED_COMPILER_EVENT_STAT(EKismetCompilerStats_BroadcastPrecompile);
+		if(GEditor && GIsEditor)
+		{
+			GEditor->BroadcastBlueprintPreCompile(BlueprintObj);
+		}
 	}
 #endif
 
@@ -587,6 +600,7 @@ void FKismetEditorUtilities::CompileBlueprint(UBlueprint* BlueprintObj, bool bIs
 	// Do not want to run this code without the editor present nor when running commandlets.
 	if (GEditor && GIsEditor)
 	{
+		BP_SCOPED_COMPILER_EVENT_STAT(EKismetCompilerStats_UpdateSearchMetaData);
 		// We do not want to regenerate a search Guid during loads, nothing has changed in the Blueprint and it is cached elsewhere
 		if (!bIsRegeneratingOnLoad)
 		{
@@ -608,8 +622,6 @@ void FKismetEditorUtilities::CompileBlueprint(UBlueprint* BlueprintObj, bool bIs
 	FCompilerResultsLog LocalResults;
 	FCompilerResultsLog& Results = (pResults != NULL) ? *pResults : LocalResults;
 
-	BP_SCOPED_COMPILER_EVENT_NAME(TEXT("Compile Blueprint"));
-
 	FBlueprintCompileReinstancer ReinstanceHelper(OldClass);
 
 	// Suppress errors/warnings in the log if we're recompiling on load on a build machine
@@ -657,7 +669,7 @@ void FKismetEditorUtilities::CompileBlueprint(UBlueprint* BlueprintObj, bool bIs
 
 		if (!bSkipGarbageCollection)
 		{
-			BP_SCOPED_COMPILER_EVENT_NAME(TEXT("Garbage Collection"));
+			BP_SCOPED_COMPILER_EVENT_STAT(EKismetCompilerStats_GarbageCollection);
 
 			// Garbage collect to make sure the old class and actors are disposed of
 			CollectGarbage(GARBAGE_COLLECTION_KEEPFLAGS);
@@ -668,10 +680,9 @@ void FKismetEditorUtilities::CompileBlueprint(UBlueprint* BlueprintObj, bool bIs
 	}
 
 	{ 
-		BP_SCOPED_COMPILER_EVENT_NAME(TEXT("Notify Blueprint Changed"));
+		BP_SCOPED_COMPILER_EVENT_STAT(EKismetCompilerStats_NotifyBlueprintChanged);
 
-		// Blueprint has changed, broadcast notify
-		BlueprintObj->BroadcastChanged();
+		BlueprintObj->BroadcastCompiled();
 
 		if(GEditor)
 		{
@@ -686,7 +697,7 @@ void FKismetEditorUtilities::CompileBlueprint(UBlueprint* BlueprintObj, bool bIs
 	}
 
 	{ 
-		BP_SCOPED_COMPILER_EVENT_NAME(TEXT("Refresh Dependent Blueprints"));
+		BP_SCOPED_COMPILER_EVENT_STAT(EKismetCompilerStats_RefreshDependentBlueprints);
 
 		TArray<UBlueprint*> DependentBPs;
 		FBlueprintEditorUtils::GetDependentBlueprints(BlueprintObj, DependentBPs);
@@ -713,7 +724,7 @@ void FKismetEditorUtilities::CompileBlueprint(UBlueprint* BlueprintObj, bool bIs
 
 	if(!bIsRegeneratingOnLoad && BlueprintObj->GeneratedClass)
 	{
-		BP_SCOPED_COMPILER_EVENT_NAME(TEXT("Validate Generated Class"));
+		BP_SCOPED_COMPILER_EVENT_STAT(EKismetCompilerStats_ValidateGeneratedClass);
 		UBlueprint::ValidateGeneratedClass(BlueprintObj->GeneratedClass);
 	}
 
diff --git a/Engine/Source/Editor/UnrealEd/Private/Kismet2/KismetReinstanceUtilities.cpp b/Engine/Source/Editor/UnrealEd/Private/Kismet2/KismetReinstanceUtilities.cpp
--- a/Engine/Source/Editor/UnrealEd/Private/Kismet2/KismetReinstanceUtilities.cpp
+++ b/Engine/Source/Editor/UnrealEd/Private/Kismet2/KismetReinstanceUtilities.cpp
@@ -9,13 +9,14 @@
 #include "ComponentInstanceDataCache.h"
 #include "Engine/DynamicBlueprintBinding.h"
 
-DEFINE_STAT(EKismetReinstancerStats_ReplaceInstancesOfClass);
-DEFINE_STAT(EKismetReinstancerStats_FindReferencers);
-DEFINE_STAT(EKismetReinstancerStats_ReplaceReferences);
-DEFINE_STAT(EKismetReinstancerStats_ReplacementConstruction);
-DEFINE_STAT(EKismetReinstancerStats_UpdateBytecodeReferences);
-DEFINE_STAT(EKismetReinstancerStats_RecompileChildClasses);
-DEFINE_STAT(EKismetReinstancerStats_ReplaceClassNoReinsancing);
+DECLARE_CYCLE_STAT(TEXT("Replace Instances"), EKismetReinstancerStats_ReplaceInstancesOfClass, STATGROUP_KismetReinstancer );
+DECLARE_CYCLE_STAT(TEXT("Find Referencers"), EKismetReinstancerStats_FindReferencers, STATGROUP_KismetReinstancer );
+DECLARE_CYCLE_STAT(TEXT("Replace References"), EKismetReinstancerStats_ReplaceReferences, STATGROUP_KismetReinstancer );
+DECLARE_CYCLE_STAT(TEXT("Construct Replacements"), EKismetReinstancerStats_ReplacementConstruction, STATGROUP_KismetReinstancer );
+DECLARE_CYCLE_STAT(TEXT("Update Bytecode References"), EKismetReinstancerStats_UpdateBytecodeReferences, STATGROUP_KismetReinstancer );
+DECLARE_CYCLE_STAT(TEXT("Recompile Child Classes"), EKismetReinstancerStats_RecompileChildClasses, STATGROUP_KismetReinstancer );
+DECLARE_CYCLE_STAT(TEXT("Replace Classes Without Reinstancing"), EKismetReinstancerStats_ReplaceClassNoReinsancing, STATGROUP_KismetReinstancer );
+DECLARE_CYCLE_STAT(TEXT("Reinstance Objects"), EKismetCompilerStats_ReinstanceObjects, STATGROUP_KismetCompiler);
 
 struct FReplaceReferenceHelper
 {
@@ -272,7 +273,7 @@ FBlueprintCompileReinstancer::~FBlueprintCompileReinstancer()
 
 void FBlueprintCompileReinstancer::ReinstanceObjects(bool bAlwaysReinstance)
 {
-	BP_SCOPED_COMPILER_EVENT_NAME(TEXT("Reinstance Objects"));
+	BP_SCOPED_COMPILER_EVENT_STAT(EKismetCompilerStats_ReinstanceObjects);
 
 	// Make sure we only reinstance classes once!
 	if( bHasReinstanced )
diff --git a/Engine/Source/Editor/UnrealEd/Public/Kismet2/CompilerResultsLog.h b/Engine/Source/Editor/UnrealEd/Public/Kismet2/CompilerResultsLog.h
--- a/Engine/Source/Editor/UnrealEd/Public/Kismet2/CompilerResultsLog.h
+++ b/Engine/Source/Editor/UnrealEd/Public/Kismet2/CompilerResultsLog.h
@@ -224,9 +224,6 @@ public:
 	}
 };
 
-#define BP_SCOPED_COMPILER_EVENT_NAME(Name) \
-	FScopedCompilerEvent ScopedCompilerEvent(Name)
-
 #if STATS
 #define BP_SCOPED_COMPILER_EVENT_STAT(Stat) \
 	SCOPE_CYCLE_COUNTER(Stat); \
diff --git a/Engine/Source/Editor/UnrealEd/Public/Kismet2/KismetReinstanceUtilities.h b/Engine/Source/Editor/UnrealEd/Public/Kismet2/KismetReinstanceUtilities.h
--- a/Engine/Source/Editor/UnrealEd/Public/Kismet2/KismetReinstanceUtilities.h
+++ b/Engine/Source/Editor/UnrealEd/Public/Kismet2/KismetReinstanceUtilities.h
@@ -10,14 +10,6 @@
 
 DECLARE_STATS_GROUP(TEXT("Kismet Reinstancer"), STATGROUP_KismetReinstancer, STATCAT_Advanced);
 
-DECLARE_CYCLE_STAT_EXTERN(TEXT("Replace Instances"), EKismetReinstancerStats_ReplaceInstancesOfClass, STATGROUP_KismetReinstancer, );
-DECLARE_CYCLE_STAT_EXTERN(TEXT("Find Referencers"), EKismetReinstancerStats_FindReferencers, STATGROUP_KismetReinstancer, );
-DECLARE_CYCLE_STAT_EXTERN(TEXT("Replace References"), EKismetReinstancerStats_ReplaceReferences, STATGROUP_KismetReinstancer, );
-DECLARE_CYCLE_STAT_EXTERN(TEXT("Construct Replacements"), EKismetReinstancerStats_ReplacementConstruction, STATGROUP_KismetReinstancer, );
-DECLARE_CYCLE_STAT_EXTERN(TEXT("Update Bytecode References"), EKismetReinstancerStats_UpdateBytecodeReferences, STATGROUP_KismetReinstancer, );
-DECLARE_CYCLE_STAT_EXTERN(TEXT("Recompile Child Classes"), EKismetReinstancerStats_RecompileChildClasses, STATGROUP_KismetReinstancer, );
-DECLARE_CYCLE_STAT_EXTERN(TEXT("Replace Classes Without Reinstancing"), EKismetReinstancerStats_ReplaceClassNoReinsancing, STATGROUP_KismetReinstancer, );
-
 class UNREALED_API FBlueprintCompileReinstancer
 {
 protected:
diff --git a/Engine/Source/Runtime/Engine/Classes/Engine/Blueprint.h b/Engine/Source/Runtime/Engine/Classes/Engine/Blueprint.h
--- a/Engine/Source/Runtime/Engine/Classes/Engine/Blueprint.h
+++ b/Engine/Source/Runtime/Engine/Classes/Engine/Blueprint.h
@@ -432,7 +432,12 @@ public:
 	FChangedEvent& OnChanged() { return ChangedEvent; }
 
 	/**	This should NOT be public */
-	void BroadcastChanged() { ChangedEvent.Broadcast( this ); }
+	void BroadcastChanged() { ChangedEvent.Broadcast(this); }
+
+	/** Broadcasts a notification whenever the blueprint has changed. */
+	DECLARE_EVENT_OneParam(UBlueprint, FCompiledEvent, class UBlueprint*);
+	FCompiledEvent& OnCompiled() { return CompiledEvent; }
+	void BroadcastCompiled() { CompiledEvent.Broadcast(this); }
 
 #if WITH_EDITORONLY_DATA
 protected:
@@ -655,6 +660,9 @@ private:
 	/** Broadcasts a notification whenever the blueprint has changed. */
 	FChangedEvent ChangedEvent;
 
+	/** Broadcasts a notification whenever the blueprint is compiled. */
+	FCompiledEvent CompiledEvent;
+
 #if WITH_EDITOR
 public:
 	/** If this blueprint is currently being compiled, the CurrentMessageLog will be the log currently being used to send logs to. */
diff --git a/Engine/Source/Runtime/Engine/Public/BlueprintUtilities.h b/Engine/Source/Runtime/Engine/Public/BlueprintUtilities.h
--- a/Engine/Source/Runtime/Engine/Public/BlueprintUtilities.h
+++ b/Engine/Source/Runtime/Engine/Public/BlueprintUtilities.h
@@ -10,4 +10,4 @@ DECLARE_DELEGATE_OneParam( FEdGraphEvent, class UEdGraph* );
 /** Delegate for notification when property changed */
 DECLARE_MULTICAST_DELEGATE_TwoParams( FOnPropertyChanged, const FPropertyChangedEvent&, const FString& );
 
-
+DECLARE_STATS_GROUP(TEXT("Kismet Compiler"), STATGROUP_KismetCompiler, STATCAT_Advanced);
