UE-7542 Framerate drop when selecting any actor

FEdMode::GetPropertyWidgetInfos generates labels for 3d widgets only when it's necessary

diff --git a/Engine/Source/Editor/UnrealEd/Private/EdMode.cpp b/Engine/Source/Editor/UnrealEd/Private/EdMode.cpp
--- a/Engine/Source/Editor/UnrealEd/Private/EdMode.cpp
+++ b/Engine/Source/Editor/UnrealEd/Private/EdMode.cpp
@@ -948,118 +948,145 @@ static bool IsTransformProperty(UProperty* InProp)
 
 }
 
-void FEdMode::GetPropertyWidgetInfos(const UStruct* InStruct, const void* InContainer, TArray<FPropertyWidgetInfo>& OutInfos, FString PropertyNamePrefix, FString DisplayNamePrefix) const
+struct FPropertyWidgetInfoChainElement
 {
-	if(PropertyNamePrefix.Len() == 0)
+	UProperty* Property;
+	int32 Index;
+
+	FPropertyWidgetInfoChainElement(UProperty* InProperty = nullptr, int32 InIndex = INDEX_NONE)
+		: Property(InProperty), Index(InIndex)
+	{}
+
+	static bool ShouldCreateWidgetSomwhereInBranch(UProperty* InProp)
 	{
-		OutInfos.Empty();
+		UStructProperty* StructProperty = Cast<UStructProperty>(InProp);
+		if (!StructProperty)
+		{
+			UArrayProperty* ArrayProperty = Cast<UArrayProperty>(InProp);
+			if (ArrayProperty)
+			{
+				StructProperty = Cast<UStructProperty>(ArrayProperty->Inner);
+			}
+		}
+
+		if (StructProperty)
+		{
+			if (FEdMode::CanCreateWidgetForStructure(StructProperty->Struct) && InProp->HasMetaData(FEdMode::MD_MakeEditWidget))
+			{
+				return true;
+			}
+
+			for (TFieldIterator<UProperty> PropertyIt(StructProperty->Struct, EFieldIteratorFlags::IncludeSuper); PropertyIt; ++PropertyIt)
+			{
+				if (ShouldCreateWidgetSomwhereInBranch(*PropertyIt))
+				{
+					return true;
+				}
+			}
+		}
+
+		return false;
 	}
 
-	for (TFieldIterator<UProperty> PropertyIt(InStruct, EFieldIteratorFlags::IncludeSuper); PropertyIt; ++PropertyIt)
+	static FEdMode::FPropertyWidgetInfo CreateWidgetInfo(const TArray<FPropertyWidgetInfoChainElement>& Chain, bool bIsTransform, UProperty* CurrentProp, int32 Index = INDEX_NONE)
 	{
-		UProperty* CurrentProp = *PropertyIt;
 		check(CurrentProp);
-		FString DisplayName;
-		if (PropertyNamePrefix.IsEmpty() || !DisplayNamePrefix.IsEmpty())
+		FEdMode::FPropertyWidgetInfo WidgetInfo;
+		WidgetInfo.PropertyValidationName = FName(*CurrentProp->GetMetaData(FEdMode::MD_ValidateWidgetUsing));
+		WidgetInfo.bIsTransform = bIsTransform;
+		WidgetInfo.PropertyIndex = Index;
+
+		const FString SimplePostFix(TEXT("."));
+		for (int32 ChainIndex = 0; ChainIndex < Chain.Num(); ++ChainIndex)
 		{
-			DisplayName = CurrentProp->GetMetaData(TEXT("DisplayName"));
+			const FPropertyWidgetInfoChainElement& Element = Chain[ChainIndex];
+			check(Element.Property);
+			const FString Postfix = (Element.Index != INDEX_NONE) ? FString::Printf(TEXT("[%d]."), Element.Index) : SimplePostFix;
+			const FString PropertyName = Element.Property->GetName() + Postfix;
+			const FString DisplayName = Element.Property->GetMetaData(TEXT("DisplayName"));
+
+			WidgetInfo.PropertyName += PropertyName;
+			WidgetInfo.DisplayName += (!DisplayName.IsEmpty()) ? (DisplayName + Postfix) : PropertyName;
 		}
-		if (!DisplayName.IsEmpty() && !DisplayNamePrefix.IsEmpty()) //Display name cannot be only the prefix.
+
 		{
-			DisplayName = DisplayNamePrefix + DisplayName;
+			const FString PropertyName = CurrentProp->GetName();
+			const FString DisplayName = CurrentProp->GetMetaData(TEXT("DisplayName"));
+
+			WidgetInfo.PropertyName += PropertyName;
+			WidgetInfo.DisplayName += (!DisplayName.IsEmpty()) ? DisplayName : PropertyName;
 		}
+		return WidgetInfo;
+	}
 
-		if (CanCreateWidgetForProperty(CurrentProp))
+	static void RecursiveGet(const FEdMode& EdMode, const UStruct* InStruct, const void* InContainer, TArray<FEdMode::FPropertyWidgetInfo>& OutInfos, TArray<FPropertyWidgetInfoChainElement>& Chain)
+	{
+		for (TFieldIterator<UProperty> PropertyIt(InStruct, EFieldIteratorFlags::IncludeSuper); PropertyIt; ++PropertyIt)
 		{
-			if (ShouldCreateWidgetForProperty(CurrentProp))
+			UProperty* CurrentProp = *PropertyIt;
+			check(CurrentProp);
+
+			if (EdMode.ShouldCreateWidgetForProperty(CurrentProp))
 			{
 				if (UArrayProperty* ArrayProp = Cast<UArrayProperty>(CurrentProp))
 				{
-					check(InContainer != NULL);
-
+					check(InContainer);
 					FScriptArrayHelper_InContainer ArrayHelper(ArrayProp, InContainer);
-
 					// See how many widgets we need to make for the array property
-					uint32 ArrayDim = ArrayHelper.Num();
-					for (uint32 i = 0; i < ArrayDim; i++)
+					const uint32 ArrayDim = ArrayHelper.Num();
+					for (uint32 Index = 0; Index < ArrayDim; ++Index)
 					{
-						//create a new widget info struct
-						FPropertyWidgetInfo WidgetInfo;
-
-						//fill it in with the struct name
-						WidgetInfo.PropertyName = PropertyNamePrefix + CurrentProp->GetFName().ToString();
-						WidgetInfo.DisplayName = DisplayName.IsEmpty() ? WidgetInfo.PropertyName : DisplayName;
-
-						//And see if we have any meta data that matches the MD_ValidateWidgetUsing name
-						WidgetInfo.PropertyValidationName = FName(*CurrentProp->GetMetaData(MD_ValidateWidgetUsing));
-
-						WidgetInfo.PropertyIndex = i;
-
-						// See if its a transform
-						WidgetInfo.bIsTransform = IsTransformProperty(ArrayProp->Inner);
-
-						//Add it to our out array
-						OutInfos.Add(WidgetInfo);
+						OutInfos.Add(FPropertyWidgetInfoChainElement::CreateWidgetInfo(Chain, IsTransformProperty(ArrayProp->Inner), CurrentProp, Index));
 					}
 				}
 				else
 				{
-
-					//create a new widget info struct
-					FPropertyWidgetInfo WidgetInfo;
-
-					//fill it in with the struct name
-					WidgetInfo.PropertyName = PropertyNamePrefix + CurrentProp->GetFName().ToString();
-					WidgetInfo.DisplayName = DisplayName.IsEmpty() ? WidgetInfo.PropertyName : DisplayName;
-
-					//And see if we have any meta data that matches the MD_ValidateWidgetUsing name
-					WidgetInfo.PropertyValidationName = FName(*CurrentProp->GetMetaData(MD_ValidateWidgetUsing));
-
-					// See if its a transform
-					WidgetInfo.bIsTransform = IsTransformProperty(CurrentProp);
-
-					//Add it to our out array
-					OutInfos.Add(WidgetInfo);
+					OutInfos.Add(FPropertyWidgetInfoChainElement::CreateWidgetInfo(Chain, IsTransformProperty(CurrentProp), CurrentProp));
 				}
 			}
-		}
-		else
-		{
-			UStructProperty* StructProp = Cast<UStructProperty>(CurrentProp);
-			if(StructProp != NULL)
+			else if (UStructProperty* StructProp = Cast<UStructProperty>(CurrentProp))
 			{
 				// Recursively traverse into structures, looking for additional vector properties to expose
-				GetPropertyWidgetInfos(StructProp->Struct
+				Chain.Push(FPropertyWidgetInfoChainElement(StructProp));
+				RecursiveGet(EdMode
+					, StructProp->Struct
 					, StructProp->ContainerPtrToValuePtr<void>(InContainer)
 					, OutInfos
-					, PropertyNamePrefix + StructProp->GetFName().ToString() + TEXT(".")
-					, !DisplayName.IsEmpty() ? (DisplayName + TEXT(".")) : FString());
+					, Chain);
+				Chain.Pop(false);
 			}
-			else
+			else if (UArrayProperty* ArrayProp = Cast<UArrayProperty>(CurrentProp))
 			{
 				// Recursively traverse into arrays of structures, looking for additional vector properties to expose
-				UArrayProperty* ArrayProp = Cast<UArrayProperty>(CurrentProp);
-				if (ArrayProp != NULL)
+				UStructProperty* InnerStructProp = Cast<UStructProperty>(ArrayProp->Inner);
+				if (InnerStructProp)
 				{
-					StructProp = Cast<UStructProperty>(ArrayProp->Inner);
-					if (StructProp != NULL)
-					{
-						FScriptArrayHelper_InContainer ArrayHelper(ArrayProp, InContainer);
+					FScriptArrayHelper_InContainer ArrayHelper(ArrayProp, InContainer);
 
+					// If the array is not empty the do additional check to tell if iteration is necessary
+					if (ArrayHelper.Num() && ShouldCreateWidgetSomwhereInBranch(InnerStructProp))
+					{
 						for (int32 ArrayIndex = 0; ArrayIndex < ArrayHelper.Num(); ++ArrayIndex)
 						{
-							const FString ArrayPostfix = FString::Printf(TEXT("[%d]."), ArrayIndex);
-							GetPropertyWidgetInfos(StructProp->Struct
+							Chain.Push(FPropertyWidgetInfoChainElement(ArrayProp, ArrayIndex));
+							RecursiveGet(EdMode
+								, InnerStructProp->Struct
 								, ArrayHelper.GetRawPtr(ArrayIndex)
 								, OutInfos
-								, PropertyNamePrefix + ArrayProp->GetFName().ToString() + ArrayPostfix
-								, !DisplayName.IsEmpty() ? (DisplayName + ArrayPostfix) : FString());
+								, Chain);
+							Chain.Pop(false);
 						}
 					}
 				}
 			}
 		}
 	}
+};
+
+void FEdMode::GetPropertyWidgetInfos(const UStruct* InStruct, const void* InContainer, TArray<FPropertyWidgetInfo>& OutInfos) const
+{
+	TArray<FPropertyWidgetInfoChainElement> Chain;
+	FPropertyWidgetInfoChainElement::RecursiveGet(*this, InStruct, InContainer, OutInfos, Chain);
 }
 
 bool FEdMode::IsSnapRotationEnabled()
diff --git a/Engine/Source/Editor/UnrealEd/Public/EdMode.h b/Engine/Source/Editor/UnrealEd/Public/EdMode.h
--- a/Engine/Source/Editor/UnrealEd/Public/EdMode.h
+++ b/Engine/Source/Editor/UnrealEd/Public/EdMode.h
@@ -274,7 +274,7 @@ protected:
 	FEditorModeTools* Owner;
 
 	// Property Widgets
-
+public:
 	/** Structure that holds info about our optional property widget */
 	struct FPropertyWidgetInfo
 	{
@@ -291,6 +291,7 @@ protected:
 		}
 	};
 
+protected:
 	/**
 	 * Returns the first selected Actor, or NULL if there is no selection.
 	 */
@@ -302,9 +303,8 @@ protected:
 	 * @param InStruct The type of structure/class to access widget info structures for.
 	 * @param InContainer The container of the given type.
 	 * @param OutInfos An array of widget info structures (output).
-	 * @param PropertyNamePrefix Optional prefix to use to filter the output.
 	 */
-	void GetPropertyWidgetInfos(const UStruct* InStruct, const void* InContainer, TArray<FPropertyWidgetInfo>& OutInfos, FString PropertyNamePrefix = TEXT(""), FString DisplayNamePrefix = TEXT("")) const;
+	void GetPropertyWidgetInfos(const UStruct* InStruct, const void* InContainer, TArray<FPropertyWidgetInfo>& OutInfos) const;
 
 	/** Name of the property currently being edited */
 	FString EditedPropertyName;
