Avoid divide by zeros with when setting FOV to zero in Matinee

diff --git a/Engine/Source/Editor/UnrealEd/Private/EditorViewportClient.cpp b/Engine/Source/Editor/UnrealEd/Private/EditorViewportClient.cpp
--- a/Engine/Source/Editor/UnrealEd/Private/EditorViewportClient.cpp
+++ b/Engine/Source/Editor/UnrealEd/Private/EditorViewportClient.cpp
@@ -570,7 +570,8 @@ FSceneView* FEditorViewportClient::CalcSceneView(FSceneViewFamily* ViewFamily)
 
 		float MinZ = GetNearClipPlane();
 		float MaxZ = MinZ;
-		float MatrixFOV = ViewFOV * (float)PI / 360.0f;
+		// Avoid zero ViewFOV's which cause divide by zero's in projection matrix
+		float MatrixFOV = FMath::Max(0.001f, ViewFOV) * (float)PI / 360.0f;
 
 		if( bConstrainAspectRatio )
 		{
