Empty actor types no longer create anonymous blueprints

diff --git a/Engine/Source/Editor/UnrealEd/Classes/ActorFactories/ActorFactoryEmptyActor.h b/Engine/Source/Editor/UnrealEd/Classes/ActorFactories/ActorFactoryEmptyActor.h
--- a/Engine/Source/Editor/UnrealEd/Classes/ActorFactories/ActorFactoryEmptyActor.h
+++ b/Engine/Source/Editor/UnrealEd/Classes/ActorFactories/ActorFactoryEmptyActor.h
@@ -14,9 +14,5 @@ public:
 
 protected:
 	virtual AActor* SpawnActor( UObject* Asset, ULevel* InLevel, const FVector& Location, const FRotator& Rotation, EObjectFlags ObjectFlags, const FName& Name ) override;
-	void SetupEditorOnlyBillboardComponent(class UBillboardComponent* BillboardComponent);
-	void OnBlueprintActorPreSpawned(class USimpleConstructionScript* SimpleConstructionScript);
-private:
-	AActor* SpawnActorForDragPreview( UObject* Asset, ULevel* InLevel, const FVector& Location, const FRotator& Rotation, EObjectFlags ObjectFlags, const FName& Name );
 
 };
diff --git a/Engine/Source/Editor/UnrealEd/Classes/ActorFactories/ActorFactoryPawn.h b/Engine/Source/Editor/UnrealEd/Classes/ActorFactories/ActorFactoryPawn.h
--- a/Engine/Source/Editor/UnrealEd/Classes/ActorFactories/ActorFactoryPawn.h
+++ b/Engine/Source/Editor/UnrealEd/Classes/ActorFactories/ActorFactoryPawn.h
@@ -14,6 +14,4 @@ public:
 
 protected:
 	virtual AActor* SpawnActor( UObject* Asset, ULevel* InLevel, const FVector& Location, const FRotator& Rotation, EObjectFlags ObjectFlags, const FName& Name ) override;
-private:
-	AActor* SpawnActorForDragPreview(UObject* Asset, ULevel* InLevel, const FVector& Location, const FRotator& Rotation, EObjectFlags ObjectFlags, const FName& Name);
 };
diff --git a/Engine/Source/Editor/UnrealEd/Private/Factories/ActorFactory.cpp b/Engine/Source/Editor/UnrealEd/Private/Factories/ActorFactory.cpp
--- a/Engine/Source/Editor/UnrealEd/Private/Factories/ActorFactory.cpp
+++ b/Engine/Source/Editor/UnrealEd/Private/Factories/ActorFactory.cpp
@@ -1112,52 +1112,22 @@ UActorFactoryCameraActor::UActorFactoryCameraActor(const FObjectInitializer& Obj
 	NewActorClass = ACameraActor::StaticClass();
 }
 
-DECLARE_DELEGATE_OneParam( FOnBlueprintActorPreSpawned, USimpleConstructionScript* );
-
-template <typename ActorClassType>
-static AActor* CreateAnonymousBlueprintAndSpawnActor(ULevel* InLevel, const FVector& Location, const FRotator& Rotation, EObjectFlags ObjectFlags, const FName& Name, FOnBlueprintActorPreSpawned OnBlueprintActorPreSpawned = FOnBlueprintActorPreSpawned() )
+static UBillboardComponent* CreateEditorOnlyBillboardComponent(AActor* ActorOwner, USceneComponent* AttachParent)
 {
-	UClass* ActorClass = ActorClassType::StaticClass();
-	FName UniqueActorName = MakeUniqueObjectName(InLevel, ActorClass, Name);
-
-	// Spawn a blank actor with an anonymous blueprint
-	// Create and init a new Blueprint
-	const FName BPName = *FString::Printf(TEXT("%s_BPClass"), *UniqueActorName.ToString());
-
-	const FName UniqueBPName = MakeUniqueObjectName(InLevel, UBlueprint::StaticClass(), BPName);
-
-
-	UBlueprint* NewBP = FKismetEditorUtilities::CreateBlueprint(ActorClass, InLevel, UniqueBPName, BPTYPE_Normal, UBlueprint::StaticClass(), UBlueprintGeneratedClass::StaticClass());
-	if(NewBP)
-	{
-		NewBP->ClearFlags(RF_Standalone);
-
-		OnBlueprintActorPreSpawned.ExecuteIfBound(NewBP->SimpleConstructionScript);
-
-		FKismetEditorUtilities::CompileBlueprint(NewBP);
-
-		FAssetData AssetData(NewBP);
-		const FString* GeneratedClassPath = AssetData.TagsAndValues.Find("GeneratedClass");
-		if(GeneratedClassPath != NULL && !GeneratedClassPath->IsEmpty())
-		{
-			UClass* GeneratedClass = Cast<UClass>(StaticLoadObject(UClass::StaticClass(), NULL, **GeneratedClassPath, NULL, LOAD_NoWarn, NULL));
+	// Create a new billboard component to serve as a visualization of the actor until there is another primitive component
+	UBillboardComponent* BillboardComponent = ConstructObject<UBillboardComponent>(UBillboardComponent::StaticClass(), ActorOwner, NAME_None, RF_Transactional);
 
-			if(GeneratedClass != NULL)
-			{
-				FActorSpawnParameters SpawnInfo;
-				SpawnInfo.OverrideLevel = InLevel;
-				SpawnInfo.ObjectFlags = ObjectFlags;
-				SpawnInfo.Name = UniqueActorName;
-				return InLevel->OwningWorld->SpawnActor(GeneratedClass, &Location, &Rotation, SpawnInfo);
-			}
-		}
+	BillboardComponent->Sprite = LoadObject<UTexture2D>(nullptr, TEXT("/Engine/EditorResources/S_Actor.S_Actor"));
+	BillboardComponent->RelativeScale3D = FVector(0.5f, 0.5f, 0.5f);
+	BillboardComponent->Mobility = EComponentMobility::Movable;
+	BillboardComponent->AlwaysLoadOnClient = false;
+	BillboardComponent->AlwaysLoadOnServer = false;
 
-	}
+	BillboardComponent->AttachTo(AttachParent);
 
-	return nullptr;
+	return BillboardComponent;
 }
 
-
 /*-----------------------------------------------------------------------------
 UActorFactoryEmptyActor
 -----------------------------------------------------------------------------*/
@@ -1178,70 +1148,25 @@ AActor* UActorFactoryEmptyActor::SpawnActor( UObject* Asset, ULevel* InLevel, co
 	AActor* NewActor = nullptr;
 	if(GetDefault<UEditorExperimentalSettings>()->bInWorldBPEditing) 
 	{
-		if( ObjectFlags == RF_Transient )
-		{
-			NewActor = SpawnActorForDragPreview( Asset, InLevel, Location, Rotation, ObjectFlags, Name );
-		}
-		else
-		{
-			NewActor = CreateAnonymousBlueprintAndSpawnActor<AActor>(
-				InLevel, 
-				Location, 
-				Rotation, 
-				ObjectFlags, 
-				FName("Actor"),
-				FOnBlueprintActorPreSpawned::CreateUObject( this, &UActorFactoryEmptyActor::OnBlueprintActorPreSpawned )
-				);
-		}
-	}
-
-	return NewActor;
-}
+		// Spawn a temporary actor for dragging around
+		AActor* NewActor = Super::SpawnActor(Asset, InLevel, Location, Rotation, ObjectFlags, Name);
 
-AActor* UActorFactoryEmptyActor::SpawnActorForDragPreview( UObject* Asset, ULevel* InLevel, const FVector& Location, const FRotator& Rotation, EObjectFlags ObjectFlags, const FName& Name )
-{
-	// Spawn a temporary actor for dragging around
-	AActor* NewActor = Super::SpawnActor(Asset, InLevel, Location, Rotation, ObjectFlags, Name);
-
-	USceneComponent* RootComponent = ConstructObject<USceneComponent>(USceneComponent::StaticClass(), NewActor, FName("Root"), RF_Transactional);
-	RootComponent->Mobility = EComponentMobility::Movable;
-	RootComponent->SetWorldLocationAndRotation(Location, Rotation);
-	NewActor->SetRootComponent(RootComponent);
+		USceneComponent* RootComponent = ConstructObject<USceneComponent>(USceneComponent::StaticClass(), NewActor, FName("Root"), RF_Transactional);
+		RootComponent->Mobility = EComponentMobility::Movable;
+		RootComponent->SetWorldLocationAndRotation(Location, Rotation);
+		NewActor->SetRootComponent(RootComponent);
 
-	// Create a new billboard component to serve as a visualization of the actor until there is another primitive component
-	UBillboardComponent* BillboardComponent = ConstructObject<UBillboardComponent>(UBillboardComponent::StaticClass(), NewActor, NAME_None, RF_Transactional);
+		UBillboardComponent* BillboardComponent	= CreateEditorOnlyBillboardComponent(NewActor, RootComponent);
 
-	SetupEditorOnlyBillboardComponent( BillboardComponent );
+		NewActor->SerializedComponents.Add(RootComponent);
+		NewActor->SerializedComponents.Add(BillboardComponent);
 
-	BillboardComponent->AttachTo(RootComponent);
-	NewActor->SerializedComponents.Add(RootComponent);
-	NewActor->SerializedComponents.Add(BillboardComponent);
+		return NewActor;
+	}
 
 	return NewActor;
 }
 
-void UActorFactoryEmptyActor::SetupEditorOnlyBillboardComponent( UBillboardComponent* BillboardComponent )
-{
-	BillboardComponent->Sprite = LoadObject<UTexture2D>(nullptr, TEXT("/Engine/EditorResources/S_Actor.S_Actor"));
-	BillboardComponent->RelativeScale3D = FVector(0.5f, 0.5f, 0.5f);
-	BillboardComponent->Mobility = EComponentMobility::Movable;
-	BillboardComponent->AlwaysLoadOnClient = false;
-	BillboardComponent->AlwaysLoadOnServer = false;
-}
-
-void UActorFactoryEmptyActor::OnBlueprintActorPreSpawned( USimpleConstructionScript* SimpleConstructionScript )
-{
-	USCS_Node* NewSCSNode = SimpleConstructionScript->CreateNode(UBillboardComponent::StaticClass());
-
-	UBillboardComponent* BillboardComponent = CastChecked<UBillboardComponent>(NewSCSNode->ComponentTemplate);
-
-	SetupEditorOnlyBillboardComponent(BillboardComponent);
-
-	// Set parameters to parent this node to the "inherited" SCS node
-	NewSCSNode->SetParent(SimpleConstructionScript->GetDefaultSceneRootNode());
-
-	SimpleConstructionScript->GetDefaultSceneRootNode()->AddChildNode(NewSCSNode);
-}
 
 /*-----------------------------------------------------------------------------
 UActorFactoryCharacter
@@ -1263,16 +1188,7 @@ AActor* UActorFactoryCharacter::SpawnActor(UObject* Asset, ULevel* InLevel, cons
 	AActor* NewActor = nullptr;
 	if(GetDefault<UEditorExperimentalSettings>()->bInWorldBPEditing)
 	{
-
-		if(ObjectFlags == RF_Transient)
-		{
-			NewActor = Super::SpawnActor( Asset, InLevel, Location, Rotation, ObjectFlags, Name );
-		}
-		else
-		{
-			NewActor = CreateAnonymousBlueprintAndSpawnActor<ACharacter>( InLevel, Location, Rotation, ObjectFlags, FName("Character") );
-		}
-
+		NewActor = Super::SpawnActor( Asset, InLevel, Location, Rotation, ObjectFlags, Name );
 	}
 
 	return NewActor;
@@ -1299,21 +1215,7 @@ AActor* UActorFactoryPawn::SpawnActor(UObject* Asset, ULevel* InLevel, const FVe
 	AActor* NewActor = nullptr;
 	if(GetDefault<UEditorExperimentalSettings>()->bInWorldBPEditing)
 	{
-		if(ObjectFlags == RF_Transient)
-		{
-			NewActor = Super::SpawnActor(Asset, InLevel, Location, Rotation, ObjectFlags, Name);
-		}
-		else
-		{
-			NewActor = CreateAnonymousBlueprintAndSpawnActor<APawn>(
-				InLevel,
-				Location,
-				Rotation,
-				ObjectFlags,
-				FName("Pawn"),
-				FOnBlueprintActorPreSpawned::CreateUObject(this, &UActorFactoryPawn::OnBlueprintActorPreSpawned)
-				);
-		}
+		NewActor = Super::SpawnActor(Asset, InLevel, Location, Rotation, ObjectFlags, Name);
 	}
 
 	return NewActor;
