Replaced the "Get Visual Studio" hyperlink in the SNewProjectWizard and SNewClassDialog with an "Install Visual Studio" button on Windows
- Added IDesktopPlatform::GetUserTempPath() and implemented it for windows
- Created SGetSuggestedIDEWidget class that handles whether to show a "Download X" hyperlink vs. an "Install X" button (depending on whether the platform supports on-demand installation)
- Analytics event added ("Editor.Usage.InstalledIDE") that fires whenever the "Install X" button is clicked
- Changed SourceCodeIDEURL_Windows in BaseEditor.ini from the VSC 2013 web page link to the installer download link (need to replace with perma-link once we have it)

diff --git a/Engine/Config/BaseEditor.ini b/Engine/Config/BaseEditor.ini
--- a/Engine/Config/BaseEditor.ini
+++ b/Engine/Config/BaseEditor.ini
@@ -240,7 +240,7 @@ DocumentationURL="http://udn.unrealengine.com/docs/ue4/INT/index.html"
 AskAQuestionURL="http://udn.unrealengine.com/questions/ask.html"
 SearchForAnswersURL="http://udn.unrealengine.com/search.html"
 SupportWebsiteURL="http://udn.unrealengine.com/index.html"
-SourceCodeIDEURL_Windows="http://www.visualstudio.com/products/visual-studio-community-vs"
+SourceCodeIDEURL_Windows="http://go.microsoft.com/?linkid=9863608"
 SourceCodeIDEURL_Mac="https://developer.apple.com/xcode"
 SourceCodeIDEURL_Other="http://en.wikipedia.org/wiki/Integrated_development_environment"
 EpicGamesURL="http://UnrealEngine.com"
diff --git a/Engine/Source/Developer/DesktopPlatform/Private/Linux/DesktopPlatformLinux.cpp b/Engine/Source/Developer/DesktopPlatform/Private/Linux/DesktopPlatformLinux.cpp
--- a/Engine/Source/Developer/DesktopPlatform/Private/Linux/DesktopPlatformLinux.cpp
+++ b/Engine/Source/Developer/DesktopPlatform/Private/Linux/DesktopPlatformLinux.cpp
@@ -378,4 +378,9 @@ FFeedbackContext* FDesktopPlatformLinux::GetNativeFeedbackContext()
 	return GWarn;
 }
 
+FString FDesktopPlatformLinux::GetUserTempPath()
+{
+	//@todo - get the path to the user's temp dir
+}
+
 #undef LOCTEXT_NAMESPACE
diff --git a/Engine/Source/Developer/DesktopPlatform/Private/Linux/DesktopPlatformLinux.h b/Engine/Source/Developer/DesktopPlatform/Private/Linux/DesktopPlatformLinux.h
--- a/Engine/Source/Developer/DesktopPlatform/Private/Linux/DesktopPlatformLinux.h
+++ b/Engine/Source/Developer/DesktopPlatform/Private/Linux/DesktopPlatformLinux.h
@@ -31,6 +31,8 @@ public:
 
 	virtual FFeedbackContext* GetNativeFeedbackContext() override;
 
+	virtual FString GetUserTempPath() override;
+
 private:
 	bool FileDialogShared(bool bSave, const void* ParentWindowHandle, const FString& DialogTitle, const FString& DefaultPath, const FString& DefaultFile, const FString& FileTypes, uint32 Flags, TArray<FString>& OutFilenames, int32& OutFilterIndex);
 };
diff --git a/Engine/Source/Developer/DesktopPlatform/Private/Mac/DesktopPlatformMac.cpp b/Engine/Source/Developer/DesktopPlatform/Private/Mac/DesktopPlatformMac.cpp
--- a/Engine/Source/Developer/DesktopPlatform/Private/Mac/DesktopPlatformMac.cpp
+++ b/Engine/Source/Developer/DesktopPlatform/Private/Mac/DesktopPlatformMac.cpp
@@ -752,4 +752,9 @@ FFeedbackContext* FDesktopPlatformMac::GetNativeFeedbackContext()
 	return &Warn;
 }
 
+FString FDesktopPlatformMac::GetUserTempPath()
+{
+	//@todo - get the path to the user's temp dir
+}
+
 #undef LOCTEXT_NAMESPACE
diff --git a/Engine/Source/Developer/DesktopPlatform/Private/Mac/DesktopPlatformMac.h b/Engine/Source/Developer/DesktopPlatform/Private/Mac/DesktopPlatformMac.h
--- a/Engine/Source/Developer/DesktopPlatform/Private/Mac/DesktopPlatformMac.h
+++ b/Engine/Source/Developer/DesktopPlatform/Private/Mac/DesktopPlatformMac.h
@@ -27,6 +27,8 @@ public:
 
 	virtual FFeedbackContext* GetNativeFeedbackContext() override;
 
+	virtual FString GetUserTempPath() override;
+
 private:
 	bool GetLauncherPath(FString& OutLauncherPath) const;
 	bool GetLauncherInstallerPath(FString& OutInstallerPath) const;
diff --git a/Engine/Source/Developer/DesktopPlatform/Private/Windows/DesktopPlatformWindows.cpp b/Engine/Source/Developer/DesktopPlatform/Private/Windows/DesktopPlatformWindows.cpp
--- a/Engine/Source/Developer/DesktopPlatform/Private/Windows/DesktopPlatformWindows.cpp
+++ b/Engine/Source/Developer/DesktopPlatform/Private/Windows/DesktopPlatformWindows.cpp
@@ -578,6 +578,16 @@ FFeedbackContext* FDesktopPlatformWindows::GetNativeFeedbackContext()
 	return &FeedbackContext;
 }
 
+FString FDesktopPlatformWindows::GetUserTempPath()
+{
+	TCHAR TempPath[MAX_PATH];
+	ZeroMemory(TempPath, sizeof(TCHAR) * MAX_PATH);
+
+	::GetTempPath(MAX_PATH, TempPath);
+
+	return FString(TempPath).Replace(TEXT("\\"), TEXT("/"));
+}
+
 void FDesktopPlatformWindows::GetRequiredRegistrySettings(TIndirectArray<FRegistryRootedKey> &RootedKeys)
 {
 	// Get the path to VersionSelector.exe. If we're running from UnrealVersionSelector itself, try to stick with the current configuration.
diff --git a/Engine/Source/Developer/DesktopPlatform/Private/Windows/DesktopPlatformWindows.h b/Engine/Source/Developer/DesktopPlatform/Private/Windows/DesktopPlatformWindows.h
--- a/Engine/Source/Developer/DesktopPlatform/Private/Windows/DesktopPlatformWindows.h
+++ b/Engine/Source/Developer/DesktopPlatform/Private/Windows/DesktopPlatformWindows.h
@@ -32,6 +32,8 @@ public:
 
 	virtual FFeedbackContext* GetNativeFeedbackContext() override;
 
+	virtual FString GetUserTempPath() override;
+
 private:
 	bool GetLauncherPath(FString& OutLauncherPath);
 	bool GetLauncherInstallerPath(FString& OutInstallerPath);
diff --git a/Engine/Source/Developer/DesktopPlatform/Public/IDesktopPlatform.h b/Engine/Source/Developer/DesktopPlatform/Public/IDesktopPlatform.h
--- a/Engine/Source/Developer/DesktopPlatform/Public/IDesktopPlatform.h
+++ b/Engine/Source/Developer/DesktopPlatform/Public/IDesktopPlatform.h
@@ -377,6 +377,13 @@ public:
 	*/
 	virtual bool GetSolutionPath(FString& OutSolutionPath) = 0;
 
+	/**
+	 * Gets the path to the user's temporary directory
+	 *
+	 * @return The path to the user's temporary directory
+	 */
+	virtual FString GetUserTempPath() = 0;
+
 	/**
 	* Gets a feedback context which can display progress information using the native platform GUI.
 	*
diff --git a/Engine/Source/Editor/EditorStyle/Private/SlateEditorStyle.cpp b/Engine/Source/Editor/EditorStyle/Private/SlateEditorStyle.cpp
--- a/Engine/Source/Editor/EditorStyle/Private/SlateEditorStyle.cpp
+++ b/Engine/Source/Editor/EditorStyle/Private/SlateEditorStyle.cpp
@@ -1590,7 +1590,7 @@ void FSlateEditorStyle::FStyle::SetupGeneralStyles()
 			.SetShadowOffset( FVector2D::ZeroVector )
 		);
 
-		Set( "GameProjectDialog.ErrorLabelBorder", new FSlateColorBrush( FLinearColor(0.2, 0, 0) ) );
+		Set( "GameProjectDialog.ErrorLabelBorder", new FSlateColorBrush( FLinearColor(0.2f, 0.f, 0.f, 0.7f) ) );
 		Set( "GameProjectDialog.ErrorLabelCloseButton", new IMAGE_BRUSH( "Icons/Cross_12x", Icon12x12 ) );
 
 		Set( "GameProjectDialog.TemplateListView.TableRow", FTableRowStyle()
@@ -1675,7 +1675,7 @@ void FSlateEditorStyle::FStyle::SetupGeneralStyles()
 			.SetFont( TTF_CORE_FONT( "Fonts/Roboto-Regular", 10 ) )
 		 );
 	
-		Set( "NewClassDialog.ErrorLabelBorder", new FSlateColorBrush( FLinearColor(0.2, 0, 0) ) );
+		Set( "NewClassDialog.ErrorLabelBorder", new FSlateColorBrush( FLinearColor(0.2f, 0.f, 0.f, 0.7f) ) );
 		Set( "NewClassDialog.ErrorLabelCloseButton", new IMAGE_BRUSH( "Icons/Cross_12x", Icon12x12 ) );
 
 		Set( "NewClassDialog.ParentClassListView.TableRow", FTableRowStyle()
diff --git a/Engine/Source/Editor/GameProjectGeneration/Private/GameProjectGenerationPrivatePCH.h b/Engine/Source/Editor/GameProjectGeneration/Private/GameProjectGenerationPrivatePCH.h
--- a/Engine/Source/Editor/GameProjectGeneration/Private/GameProjectGenerationPrivatePCH.h
+++ b/Engine/Source/Editor/GameProjectGeneration/Private/GameProjectGenerationPrivatePCH.h
@@ -15,4 +15,5 @@ DECLARE_LOG_CATEGORY_EXTERN(LogGameProjectGeneration, Verbose, All);
 #include "SProjectBrowser.h"
 #include "SNewProjectWizard.h"
 #include "SGameProjectDialog.h"
-#include "SNewClassDialog.h"
\ No newline at end of file
+#include "SNewClassDialog.h"
+#include "SGetSuggestedIDEWidget.h"
\ No newline at end of file
diff --git a/Engine/Source/Editor/GameProjectGeneration/Private/SNewClassDialog.cpp b/Engine/Source/Editor/GameProjectGeneration/Private/SNewClassDialog.cpp
--- a/Engine/Source/Editor/GameProjectGeneration/Private/SNewClassDialog.cpp
+++ b/Engine/Source/Editor/GameProjectGeneration/Private/SNewClassDialog.cpp
@@ -222,7 +222,7 @@ void SNewClassDialog::Construct( const FArguments& InArgs )
 						.VAlign(VAlign_Center)
 						[
 							SNew(STextBlock)
-							.Text( FText::Format( LOCTEXT("ChooseParentClassDescription", "You are about to add a C++ source code file. To compile these files you must have {0} installed."), FSourceCodeNavigation::GetSuggestedSourceCodeIDE() ) )
+							.Text(LOCTEXT("ChooseParentClassDescription", "This will add a C++ header and source code file to your game project."))
 						]
 
 						// Full tree checkbox
@@ -241,6 +241,46 @@ void SNewClassDialog::Construct( const FArguments& InArgs )
 						]
 					]
 
+					// Get IDE information
+					+SVerticalBox::Slot()
+					.Padding(0, 5)
+					.AutoHeight()
+					[
+						SNew(SBorder)
+						.Visibility( this, &SNewClassDialog::GetGlobalErrorLabelVisibility )
+						.BorderImage( FEditorStyle::GetBrush("NewClassDialog.ErrorLabelBorder") )
+						.Content()
+						[
+							SNew(SHorizontalBox)
+
+							+SHorizontalBox::Slot()
+							.VAlign(VAlign_Center)
+							.Padding(2.f)
+							.AutoWidth()
+							[
+								SNew(SImage)
+								.Image(FEditorStyle::GetBrush("MessageLog.Warning"))
+							]
+
+							+SHorizontalBox::Slot()
+							.VAlign(VAlign_Center)
+							[
+								SNew(STextBlock)
+								.Text( this, &SNewClassDialog::GetGlobalErrorLabelText )
+								.TextStyle( FEditorStyle::Get(), "NewClassDialog.ErrorLabelFont" )
+							]
+
+							+SHorizontalBox::Slot()
+							.VAlign(VAlign_Center)
+							.HAlign(HAlign_Center)
+							.AutoWidth()
+							.Padding(5.f, 0.f)
+							[
+								SNew(SGetSuggestedIDEWidget)
+							]
+						]
+					]
+
 					// Add Code list
 					+SVerticalBox::Slot()
 					.FillHeight(1.f)
@@ -641,38 +681,6 @@ void SNewClassDialog::Construct( const FArguments& InArgs )
 					]
 				]
 			]
-
-			// IDE download information
-			+SVerticalBox::Slot()
-			.Padding(0, 5)
-			.AutoHeight()
-			[
-				SNew(SBorder)
-				.Visibility( this, &SNewClassDialog::GetGlobalErrorLabelVisibility )
-				.BorderImage( FEditorStyle::GetBrush("NewClassDialog.ErrorLabelBorder") )
-				.Content()
-				[
-					SNew(SHorizontalBox)
-
-					+SHorizontalBox::Slot()
-					.VAlign(VAlign_Center)
-					[
-						SNew(STextBlock)
-						.Text( this, &SNewClassDialog::GetGlobalErrorLabelText )
-						.TextStyle( FEditorStyle::Get(), "NewClassDialog.ErrorLabelFont" )
-					]
-
-					+SHorizontalBox::Slot()
-					.VAlign(VAlign_Center)
-					.AutoWidth()
-					[
-						SNew(SHyperlink)
-						.Text( FText::Format( LOCTEXT("IDEDownloadLinkText", "Download {0}"), FSourceCodeNavigation::GetSuggestedSourceCodeIDE() ) )
-						.OnNavigate( this, &SNewClassDialog::OnDownloadIDEClicked, FSourceCodeNavigation::GetSuggestedSourceCodeIDEDownloadURL() )
-						.Visibility( this, &SNewClassDialog::GetGlobalErrorLabelIDELinkVisibility )
-					]
-				]
-			]
 		]
 	];
 
@@ -878,12 +886,7 @@ FText SNewClassDialog::GetNameErrorLabelText() const
 
 EVisibility SNewClassDialog::GetGlobalErrorLabelVisibility() const
 {
-	return GetGlobalErrorLabelText().IsEmpty() ? EVisibility::Hidden : EVisibility::Visible;
-}
-
-EVisibility SNewClassDialog::GetGlobalErrorLabelIDELinkVisibility() const
-{
-	return FSourceCodeNavigation::IsCompilerAvailable() ? EVisibility::Collapsed : EVisibility::Visible;
+	return GetGlobalErrorLabelText().IsEmpty() ? EVisibility::Collapsed : EVisibility::Visible;
 }
 
 FText SNewClassDialog::GetGlobalErrorLabelText() const
@@ -1028,11 +1031,6 @@ void SNewClassDialog::FinishClicked()
 	}
 }
 
-void SNewClassDialog::OnDownloadIDEClicked(FString URL)
-{
-	FPlatformProcess::LaunchURL( *URL, NULL, NULL );
-}
-
 FReply SNewClassDialog::HandleChooseFolderButtonClicked()
 {
 	IDesktopPlatform* DesktopPlatform = FDesktopPlatformModule::Get();
diff --git a/Engine/Source/Editor/GameProjectGeneration/Private/SNewClassDialog.h b/Engine/Source/Editor/GameProjectGeneration/Private/SNewClassDialog.h
--- a/Engine/Source/Editor/GameProjectGeneration/Private/SNewClassDialog.h
+++ b/Engine/Source/Editor/GameProjectGeneration/Private/SNewClassDialog.h
@@ -77,9 +77,6 @@ private:
 	/** Gets the visibility of the global error label */
 	EVisibility GetGlobalErrorLabelVisibility() const;
 
-	/** Gets the visibility of the global error label IDE Link */
-	EVisibility GetGlobalErrorLabelIDELinkVisibility() const;
-
 	/** Gets the text to display in the global error label */
 	FText GetGlobalErrorLabelText() const;
 
@@ -116,9 +113,6 @@ private:
 	/** Handler for when finish is clicked */
 	void FinishClicked();
 
-	/** Handler for when the error label IDE hyperlink was clicked */
-	void OnDownloadIDEClicked(FString URL);
-
 	/** Handler for when the "Choose Folder" button is clicked */
 	FReply HandleChooseFolderButtonClicked( );
 
diff --git a/Engine/Source/Editor/GameProjectGeneration/Private/SNewProjectWizard.cpp b/Engine/Source/Editor/GameProjectGeneration/Private/SNewProjectWizard.cpp
--- a/Engine/Source/Editor/GameProjectGeneration/Private/SNewProjectWizard.cpp
+++ b/Engine/Source/Editor/GameProjectGeneration/Private/SNewProjectWizard.cpp
@@ -685,7 +685,16 @@ void SNewProjectWizard::Construct( const FArguments& InArgs )
 				.Padding( UniformPadding / 2 )
 				[
 					SNew(SHorizontalBox)
-										
+						
+					+ SHorizontalBox::Slot()
+					.VAlign(VAlign_Center)
+					.Padding(2.f)
+					.AutoWidth()
+					[
+						SNew(SImage)
+						.Image(FEditorStyle::GetBrush("MessageLog.Warning"))
+					]
+
 					+SHorizontalBox::Slot()
 					.VAlign(VAlign_Center)
 					.FillWidth(1.0f)
@@ -695,15 +704,14 @@ void SNewProjectWizard::Construct( const FArguments& InArgs )
 						.TextStyle(FEditorStyle::Get(), TEXT("GameProjectDialog.ErrorLabelFont"))
 					]
 
-					// A link to a platform-specific IDE, only shown when a compiler is not available
+					// Button/link to the suggested IDE
 					+SHorizontalBox::Slot()
 					.VAlign(VAlign_Center)
+					.HAlign(HAlign_Center)
 					.AutoWidth()
+					.Padding(5.f, 0.f)
 					[
-						SNew(SHyperlink)
-						.Text(FText::Format(LOCTEXT("IDEDownloadLinkText", "Download {0}"), FSourceCodeNavigation::GetSuggestedSourceCodeIDE()))
-						.OnNavigate(this, &SNewProjectWizard::OnDownloadIDEClicked, FSourceCodeNavigation::GetSuggestedSourceCodeIDEDownloadURL())
-						.Visibility(this, &SNewProjectWizard::GetGlobalErrorLabelIDELinkVisibility)
+						SNew(SGetSuggestedIDEWidget)
 					]
 									
 					// A button to close the persistent global error text
@@ -736,10 +744,19 @@ void SNewProjectWizard::Construct( const FArguments& InArgs )
 				.Padding(UniformPadding / 2)
 				[
 					SNew(SHorizontalBox)
-										
+					
+					+ SHorizontalBox::Slot()
+					.VAlign(VAlign_Center)
+					.Padding(2.f)
+					.AutoWidth()
+					[
+						SNew(SImage)
+						.Image(FEditorStyle::GetBrush("MessageLog.Warning"))
+					]
+
 					+SHorizontalBox::Slot()
 					.VAlign(VAlign_Center)
-					.FillWidth(1.0f)
+					.AutoWidth()
 					[
 						SNew(STextBlock)
 						.AutoWrapText(true)
@@ -1031,13 +1048,6 @@ FReply SNewProjectWizard::HandleBrowseButtonClicked()
 	return FReply::Handled();
 }
 
-
-void SNewProjectWizard::OnDownloadIDEClicked(FString URL)
-{
-	FPlatformProcess::LaunchURL( *URL, NULL, NULL );
-}
-
-
 void SNewProjectWizard::HandleTemplateListViewDoubleClick( TSharedPtr<FTemplateItem> TemplateItem )
 {
 	// Advance to the name/location page
@@ -1089,13 +1099,6 @@ EVisibility SNewProjectWizard::GetGlobalErrorLabelCloseButtonVisibility() const
 	return PersistentGlobalErrorLabelText.IsEmpty() ? EVisibility::Collapsed : EVisibility::Visible;
 }
 
-
-EVisibility SNewProjectWizard::GetGlobalErrorLabelIDELinkVisibility() const
-{
-	return (IsCompilerRequired() && !FSourceCodeNavigation::IsCompilerAvailable()) ? EVisibility::Visible : EVisibility::Collapsed;
-}
-
-
 FText SNewProjectWizard::GetGlobalErrorLabelText() const
 {
 	if ( !PersistentGlobalErrorLabelText.IsEmpty() )
diff --git a/Engine/Source/Editor/UnrealEd/Private/SourceCodeNavigation.cpp b/Engine/Source/Editor/UnrealEd/Private/SourceCodeNavigation.cpp
--- a/Engine/Source/Editor/UnrealEd/Private/SourceCodeNavigation.cpp
+++ b/Engine/Source/Editor/UnrealEd/Private/SourceCodeNavigation.cpp
@@ -5,6 +5,7 @@
 #include "SourceCodeNavigation.h"
 #include "MainFrame.h"
 #include "ISourceCodeAccessModule.h"
+#include "Http.h"
 
 #if PLATFORM_WINDOWS
 #include "AllowWindowsPlatformTypes.h"
@@ -22,11 +23,16 @@
 #include "SNotificationList.h"
 #include "NotificationManager.h"
 #include "GameFramework/Pawn.h"
+#include "DesktopPlatformModule.h"
 
 DEFINE_LOG_CATEGORY(LogSelectionDetails);
 
 #define LOCTEXT_NAMESPACE "SourceCodeNavigation"
 
+namespace SourceCodeNavigationDefs
+{
+	FString IDEInstallerFilename("UE4_SuggestedIDEInstaller");
+}
 
 /**
  * Caches information about source symbols for fast look-up
@@ -298,6 +304,15 @@ public:
 	/** The final symbol query in a batch completed */
 	void SymbolQueryFinished();
 
+	/** Handler called when the installer for the suggested IDE has finished downloading */
+	void OnSuggestedIDEInstallerDownloadComplete(FHttpRequestPtr Request, FHttpResponsePtr Response, bool bWasSuccessful, FOnIDEInstallerDownloadComplete OnDownloadComplete);
+
+	/** Launches the IDE installer process */
+	void LaunchIDEInstaller(const FString& Filepath);
+
+	/** @return The name of the IDE installer file for the platform */
+	FString GetSuggestedIDEInstallerFileName();
+
 protected:
 	/** FTickableEditorObject interface */
 	virtual void Tick( float DeltaTime );
@@ -1439,6 +1454,42 @@ FString FSourceCodeNavigation::GetSuggestedSourceCodeIDEDownloadURL()
 	return SourceCodeIDEURL;
 }
 
+bool FSourceCodeNavigation::GetCanDirectlyInstallSourceCodeIDE()
+{
+#if PLATFORM_WINDOWS
+	return true;
+#else
+	return false;
+#endif
+}
+
+void FSourceCodeNavigation::DownloadAndInstallSuggestedIDE(FOnIDEInstallerDownloadComplete OnDownloadComplete)
+{
+	FSourceCodeNavigationImpl& SourceCodeNavImpl = FSourceCodeNavigationImpl::Get();
+
+	// Check to see if the file exists first
+	auto UserTempDir = FPaths::ConvertRelativePathToFull(FDesktopPlatformModule::Get()->GetUserTempPath());
+	FString InstallerFullPath = FString::Printf(TEXT("%s%s"), *UserTempDir, *SourceCodeNavImpl.GetSuggestedIDEInstallerFileName());
+
+	if (!IPlatformFile::GetPlatformPhysical().FileExists(*InstallerFullPath))
+	{
+		FString DownloadURL;
+		TSharedRef<IHttpRequest> HttpRequest = FHttpModule::Get().CreateRequest();
+
+		// Download the installer for the suggested IDE
+		HttpRequest->OnProcessRequestComplete().BindRaw(&SourceCodeNavImpl, &FSourceCodeNavigationImpl::OnSuggestedIDEInstallerDownloadComplete, OnDownloadComplete);
+		HttpRequest->SetVerb(TEXT("GET"));
+
+		HttpRequest->SetURL(GetSuggestedSourceCodeIDEDownloadURL());
+		HttpRequest->ProcessRequest();
+	}
+	else
+	{
+		SourceCodeNavImpl.LaunchIDEInstaller(InstallerFullPath);
+		OnDownloadComplete.ExecuteIfBound(true);
+	}
+}
+
 bool FSourceCodeNavigation::IsCompilerAvailable()
 {
 	ISourceCodeAccessModule& SourceCodeAccessModule = FModuleManager::LoadModuleChecked<ISourceCodeAccessModule>("SourceCodeAccess");
@@ -1528,6 +1579,38 @@ bool FSourceCodeNavigation::FindClassHeaderPath( const UField *Field, FString &O
 	return false;
 }
 
+void FSourceCodeNavigationImpl::OnSuggestedIDEInstallerDownloadComplete(FHttpRequestPtr Request, FHttpResponsePtr Response, bool bWasSuccessful, FOnIDEInstallerDownloadComplete OnDownloadComplete)
+{
+	if (bWasSuccessful)
+	{
+		// Get the user's temp directory
+		auto UserTempDir = FDesktopPlatformModule::Get()->GetUserTempPath();
+
+		// Create the installer file in the temp dir
+		auto InstallerName = GetSuggestedIDEInstallerFileName();
+		FString Filepath = FString::Printf(TEXT("%s%s"), *UserTempDir, *InstallerName);
+		auto InstallerFileHandle = IPlatformFile::GetPlatformPhysical().OpenWrite(*Filepath);
+
+		// Copy the content from the response into the installer file
+		auto InstallerContent = Response->GetContent();
+
+		bool bWriteSucceeded = InstallerFileHandle->Write(InstallerContent.GetData(), InstallerContent.Num());
+		delete InstallerFileHandle;
+
+		if (bWriteSucceeded)
+		{
+			// Launch the created executable in a separate window to begin the installation
+			LaunchIDEInstaller(Filepath);
+		}
+		else
+		{
+			bWasSuccessful = false;
+		}
+	}
+
+	OnDownloadComplete.ExecuteIfBound(bWasSuccessful);
+}
+
 void FSourceCodeNavigationImpl::TryToGatherFunctions( const FString& ModuleName, const FString& ClassName, TArray< FString >& OutFunctionSymbolNames, bool& OutIsCompleteList )
 {
 	FScopeLock ScopeLock( &SynchronizationObject );
@@ -1583,6 +1666,26 @@ void FSourceCodeNavigationImpl::SymbolQueryFinished()
 	OnSymbolQueryFinished.Broadcast();
 }
 
+FString FSourceCodeNavigationImpl::GetSuggestedIDEInstallerFileName()
+{
+	FString Extension;
+#if PLATFORM_WINDOWS
+	Extension = "exe";
+#elif PLATFORM_MAC
+	Extension = "app";
+#endif
+
+	return FString::Printf(TEXT("%s.%s"), *SourceCodeNavigationDefs::IDEInstallerFilename, *Extension);
+}
+
+void FSourceCodeNavigationImpl::LaunchIDEInstaller(const FString& Filepath)
+{
+#if PLATFORM_WINDOWS
+	auto Params = TEXT("/PromptRestart");
+	FPlatformProcess::CreateProc(*Filepath, Params, true, false, false, nullptr, 0, nullptr, nullptr);
+#endif
+}
+
 void FSourceCodeNavigationImpl::Tick( float DeltaTime )
 {
 	const bool bAsyncWorkAvailable = ClassesToGatherSymbolsFor.Num() > 0;
diff --git a/Engine/Source/Editor/UnrealEd/Public/SourceCodeNavigation.h b/Engine/Source/Editor/UnrealEd/Public/SourceCodeNavigation.h
--- a/Engine/Source/Editor/UnrealEd/Public/SourceCodeNavigation.h
+++ b/Engine/Source/Editor/UnrealEd/Public/SourceCodeNavigation.h
@@ -5,9 +5,11 @@
 #define __SourceCodeNavigation_h__
 
 #pragma once
+#include "IHttpRequest.h"
 
 DECLARE_LOG_CATEGORY_EXTERN(LogSelectionDetails, Log, All);
 
+DECLARE_DELEGATE_OneParam(FOnIDEInstallerDownloadComplete, bool /*bWasSuccessful*/);
 
 /**
  * Singleton holding database of module names and disallowed header names in the engine and current project.
@@ -145,6 +147,12 @@ public:
 	/** Returns the url to the location where the suggested IDE can be downloaded */
 	UNREALED_API static FString GetSuggestedSourceCodeIDEDownloadURL();
 
+	/** Returns whether the suggested source code IDE for the current platform can be installed directly (vs. requiring that the user download it manually) */
+	UNREALED_API static bool GetCanDirectlyInstallSourceCodeIDE();
+
+	/** Downloads and installs the suggested IDE (currently only works for Windows) */
+	UNREALED_API static void DownloadAndInstallSuggestedIDE(FOnIDEInstallerDownloadComplete OnDownloadComplete);
+
 	/** Returns true if the compiler for the current platform is available for use */
 	UNREALED_API static bool IsCompilerAvailable();
 
diff --git a/Engine/Source/Editor/UnrealEd/UnrealEd.Build.cs b/Engine/Source/Editor/UnrealEd/UnrealEd.Build.cs
--- a/Engine/Source/Editor/UnrealEd/UnrealEd.Build.cs
+++ b/Engine/Source/Editor/UnrealEd/UnrealEd.Build.cs
@@ -72,6 +72,7 @@ public class UnrealEd : ModuleRules
 				"UnrealEdMessages",
                 "AIModule",
 				"BlueprintGraph",
+                "Http",
 			}
 		);
 
