Fix a case-sensitive error.

diff --git a/Engine/Source/Editor/UnrealEd/Classes/Tests/CircularDependencyTestActor.h b/Engine/Source/Editor/UnrealEd/Classes/Tests/CircularDependencyTestActor.h
--- a/Engine/Source/Editor/UnrealEd/Classes/Tests/CircularDependencyTestActor.h
+++ b/Engine/Source/Editor/UnrealEd/Classes/Tests/CircularDependencyTestActor.h
@@ -2,7 +2,7 @@
 
 #pragma once
 
-#include "Gameframework/Actor.h"
+#include "GameFramework/Actor.h"
 #include "CircularDependencyTestActor.generated.h"
 
 UENUM(BlueprintType)
@@ -52,4 +52,4 @@ protected:
 
 private: 
 	bool TestVerifyProperty(UProperty* Property, uint8* Container, bool bCheckPropertyType);
-};
\ No newline at end of file
+};
