Rename SerializedComponents to BlueprintCreatedComponents
Add InstanceComponents for per-instance added components

diff --git a/Engine/Source/Editor/Kismet/Private/SSCSEditor.cpp b/Engine/Source/Editor/Kismet/Private/SSCSEditor.cpp
--- a/Engine/Source/Editor/Kismet/Private/SSCSEditor.cpp
+++ b/Engine/Source/Editor/Kismet/Private/SSCSEditor.cpp
@@ -2834,7 +2834,7 @@ UActorComponent* SSCSEditor::AddNewComponent( UClass* NewComponentClass, UObject
 		check(NewComponentInstance != nullptr);
 
 		// Add to SerializedComponents array so it gets saved
-		ActorInstance->SerializedComponents.Add(NewComponentInstance);
+		ActorInstance->InstanceComponents.Add(NewComponentInstance);
 
 		return AddNewNode(NewComponentInstance, Asset, true);
 	}
@@ -3141,7 +3141,7 @@ void SSCSEditor::PasteNodes()
 			NewActorComponent->Rename(NULL, ActorInstance, REN_DontCreateRedirectors|REN_DoNotDirty);
 
 			// Add to SerializedComponents array so it gets saved
-			ActorInstance->SerializedComponents.Add(NewActorComponent);
+			ActorInstance->InstanceComponents.Add(NewActorComponent);
 
 			// Create a new node to contain the new component instance and add it to the tree
 			NewActorComponent = AddNewNode(NewActorComponent, NULL, false);
diff --git a/Engine/Source/Editor/LevelEditor/Private/ComponentsTree.cpp b/Engine/Source/Editor/LevelEditor/Private/ComponentsTree.cpp
--- a/Engine/Source/Editor/LevelEditor/Private/ComponentsTree.cpp
+++ b/Engine/Source/Editor/LevelEditor/Private/ComponentsTree.cpp
@@ -253,8 +253,8 @@ void SComponentsTree::OnSelectedCompClass(TSubclassOf<UActorComponent> CompClass
 		UActorComponent* NewComp = ConstructObject<UActorComponent>(CompClass, Actor.Get(), NAME_None, RF_Transactional);
 		check(NewComp);
 
-		// Add to SerializedComponents array so it gets saved
-		Actor->SerializedComponents.Add(NewComp);
+		// Add to InstanceComponents array so it gets saved
+		Actor->InstanceComponents.Add(NewComp);
 
 		// If a scene component, and actor has a root component, attach it to root
 		USceneComponent* NewSceneComp = Cast<USceneComponent>(NewComp);
diff --git a/Engine/Source/Editor/MeshPaint/Private/MeshPaintEdMode.cpp b/Engine/Source/Editor/MeshPaint/Private/MeshPaintEdMode.cpp
--- a/Engine/Source/Editor/MeshPaint/Private/MeshPaintEdMode.cpp
+++ b/Engine/Source/Editor/MeshPaint/Private/MeshPaintEdMode.cpp
@@ -3575,7 +3575,7 @@ void FEdModeMeshPaint::CopyInstanceVertexColors()
 			{
 				if( StaticMeshComponent )
 				{
-					FPerComponentVertexColorData& PerComponentData = CopiedColorsByComponent[CopiedColorsByComponent.Add(FPerComponentVertexColorData(StaticMeshComponent->StaticMesh, StaticMeshComponent->GetSerializedComponentIndex()))];
+					FPerComponentVertexColorData& PerComponentData = CopiedColorsByComponent[CopiedColorsByComponent.Add(FPerComponentVertexColorData(StaticMeshComponent->StaticMesh, StaticMeshComponent->GetBlueprintCreatedComponentIndex()))];
 
 					int32 NumLODs = StaticMeshComponent->StaticMesh->GetNumLODs();
 					for( int32 CurLODIndex = 0; CurLODIndex < NumLODs; ++CurLODIndex )
@@ -3664,12 +3664,12 @@ void FEdModeMeshPaint::PasteInstanceVertexColors()
 					}
 
 					// attempt to find a matching component in our clipboard data
-					const int32 SerializedComponentIndex = StaticMeshComponent->GetSerializedComponentIndex();
+					const int32 BlueprintCreatedComponentIndex = StaticMeshComponent->GetBlueprintCreatedComponentIndex();
 					FPerComponentVertexColorData* FoundColors = NULL;
 					for(auto& CopiedColors : CopiedColorsByComponent)
 					{
 						if(CopiedColors.OriginalMesh.Get() == StaticMeshComponent->StaticMesh &&
-							CopiedColors.ComponentIndex == SerializedComponentIndex)
+							CopiedColors.ComponentIndex == BlueprintCreatedComponentIndex)
 						{
 							FoundColors = &CopiedColors;
 							break;
diff --git a/Engine/Source/Editor/UnrealEd/Private/Factories/ActorFactory.cpp b/Engine/Source/Editor/UnrealEd/Private/Factories/ActorFactory.cpp
--- a/Engine/Source/Editor/UnrealEd/Private/Factories/ActorFactory.cpp
+++ b/Engine/Source/Editor/UnrealEd/Private/Factories/ActorFactory.cpp
@@ -1148,21 +1148,21 @@ AActor* UActorFactoryEmptyActor::SpawnActor( UObject* Asset, ULevel* InLevel, co
 	AActor* NewActor = nullptr;
 	if(GetDefault<UEditorExperimentalSettings>()->bInWorldBPEditing) 
 	{
-		// Spawn a temporary actor for dragging around
-		AActor* NewActor = Super::SpawnActor(Asset, InLevel, Location, Rotation, ObjectFlags, Name);
+	// Spawn a temporary actor for dragging around
+	AActor* NewActor = Super::SpawnActor(Asset, InLevel, Location, Rotation, ObjectFlags, Name);
 
-		USceneComponent* RootComponent = ConstructObject<USceneComponent>(USceneComponent::StaticClass(), NewActor, FName("Root"), RF_Transactional);
-		RootComponent->Mobility = EComponentMobility::Movable;
-		RootComponent->SetWorldLocationAndRotation(Location, Rotation);
-		NewActor->SetRootComponent(RootComponent);
+	USceneComponent* RootComponent = ConstructObject<USceneComponent>(USceneComponent::StaticClass(), NewActor, FName("Root"), RF_Transactional);
+	RootComponent->Mobility = EComponentMobility::Movable;
+	RootComponent->SetWorldLocationAndRotation(Location, Rotation);
+	NewActor->SetRootComponent(RootComponent);
 
 		UBillboardComponent* BillboardComponent	= CreateEditorOnlyBillboardComponent(NewActor, RootComponent);
 
-		NewActor->SerializedComponents.Add(RootComponent);
-		NewActor->SerializedComponents.Add(BillboardComponent);
+	NewActor->InstanceComponents.Add(RootComponent);
+	NewActor->InstanceComponents.Add(BillboardComponent);
 
-		return NewActor;
-	}
+	return NewActor;
+}
 
 	return NewActor;
 }
@@ -1188,8 +1188,8 @@ AActor* UActorFactoryCharacter::SpawnActor(UObject* Asset, ULevel* InLevel, cons
 	AActor* NewActor = nullptr;
 	if(GetDefault<UEditorExperimentalSettings>()->bInWorldBPEditing)
 	{
-		NewActor = Super::SpawnActor( Asset, InLevel, Location, Rotation, ObjectFlags, Name );
-	}
+			NewActor = Super::SpawnActor( Asset, InLevel, Location, Rotation, ObjectFlags, Name );
+		}
 
 	return NewActor;
 }
@@ -1215,8 +1215,8 @@ AActor* UActorFactoryPawn::SpawnActor(UObject* Asset, ULevel* InLevel, const FVe
 	AActor* NewActor = nullptr;
 	if(GetDefault<UEditorExperimentalSettings>()->bInWorldBPEditing)
 	{
-		NewActor = Super::SpawnActor(Asset, InLevel, Location, Rotation, ObjectFlags, Name);
-	}
+			NewActor = Super::SpawnActor(Asset, InLevel, Location, Rotation, ObjectFlags, Name);
+		}
 
 	return NewActor;
 }
diff --git a/Engine/Source/Runtime/Engine/Classes/Components/StaticMeshComponent.h b/Engine/Source/Runtime/Engine/Classes/Components/StaticMeshComponent.h
--- a/Engine/Source/Runtime/Engine/Classes/Components/StaticMeshComponent.h
+++ b/Engine/Source/Runtime/Engine/Classes/Components/StaticMeshComponent.h
@@ -448,8 +448,8 @@ public:
 	/** Returns the wireframe color to use for this component. */
 	FColor GetWireframeColor() const;
 
-	/** Get this components index in its parents serialized components array (used for matching instance data) */
-	int32 GetSerializedComponentIndex() const;
+	/** Get this components index in its parents blueprint created components array (used for matching instance data) */
+	int32 GetBlueprintCreatedComponentIndex() const;
 };
 
 
diff --git a/Engine/Source/Runtime/Engine/Classes/GameFramework/Actor.h b/Engine/Source/Runtime/Engine/Classes/GameFramework/Actor.h
--- a/Engine/Source/Runtime/Engine/Classes/GameFramework/Actor.h
+++ b/Engine/Source/Runtime/Engine/Classes/GameFramework/Actor.h
@@ -2228,9 +2228,13 @@ private:
 
 public:
 
-	/** Array of ActorComponents that is actually serialized per-instance. */
+	/** Array of ActorComponents that are created by blueprints and serialized per-instance. */
 	UPROPERTY(TextExportTransient, NonTransactional)
-	TArray<UActorComponent*> SerializedComponents;
+	TArray<UActorComponent*> BlueprintCreatedComponents;
+
+	/** Array of ActorComponents that have been added by the user on a per-instance basis. */
+	UPROPERTY()
+	TArray<UActorComponent*> InstanceComponents;
 
 public:
 	//=============================================================================
diff --git a/Engine/Source/Runtime/Engine/Private/Actor.cpp b/Engine/Source/Runtime/Engine/Private/Actor.cpp
--- a/Engine/Source/Runtime/Engine/Private/Actor.cpp
+++ b/Engine/Source/Runtime/Engine/Private/Actor.cpp
@@ -158,9 +158,9 @@ bool AActor::CheckActorComponents()
 			}
 		}
 	}
-	for (int32 Index = 0; Index < SerializedComponents.Num(); Index++)
+	for (int32 Index = 0; Index < BlueprintCreatedComponents.Num(); Index++)
 	{
-		UActorComponent* Inner = SerializedComponents[Index];
+		UActorComponent* Inner = BlueprintCreatedComponents[Index];
 		if (!Inner)
 		{
 			continue;
diff --git a/Engine/Source/Runtime/Engine/Private/ActorComponent.cpp b/Engine/Source/Runtime/Engine/Private/ActorComponent.cpp
--- a/Engine/Source/Runtime/Engine/Private/ActorComponent.cpp
+++ b/Engine/Source/Runtime/Engine/Private/ActorComponent.cpp
@@ -698,7 +698,14 @@ void UActorComponent::DestroyComponent()
 	AActor* Owner = GetOwner();
 	if(Owner != NULL)
 	{
-		Owner->SerializedComponents.Remove(this);
+		if (bCreatedByConstructionScript)
+		{
+			Owner->BlueprintCreatedComponents.Remove(this);
+		}
+		else
+		{
+			Owner->InstanceComponents.Remove(this);
+		}
 		Owner->RemoveOwnedComponent(this);
 		if (Owner->GetRootComponent() == this)
 		{
diff --git a/Engine/Source/Runtime/Engine/Private/ActorConstruction.cpp b/Engine/Source/Runtime/Engine/Private/ActorConstruction.cpp
--- a/Engine/Source/Runtime/Engine/Private/ActorConstruction.cpp
+++ b/Engine/Source/Runtime/Engine/Private/ActorConstruction.cpp
@@ -542,7 +542,7 @@ void AActor::ProcessUserConstructionScript()
 void AActor::FinishAndRegisterComponent(UActorComponent* Component)
 {
 	Component->RegisterComponent();
-	SerializedComponents.Add(Component);
+	BlueprintCreatedComponents.Add(Component);
 }
 
 UActorComponent* AActor::CreateComponentFromTemplate(UActorComponent* Template, const FString& InName)
@@ -557,7 +557,7 @@ UActorComponent* AActor::CreateComponentFromTemplate(UActorComponent* Template,
 		NewActorComp->bCreatedByConstructionScript = true;
 
 		// Need to do this so component gets saved - Components array is not serialized
-		SerializedComponents.Add(NewActorComp);
+		BlueprintCreatedComponents.Add(NewActorComp);
 	}
 	return NewActorComp;
 }
diff --git a/Engine/Source/Runtime/Engine/Private/BlueprintGeneratedClass.cpp b/Engine/Source/Runtime/Engine/Private/BlueprintGeneratedClass.cpp
--- a/Engine/Source/Runtime/Engine/Private/BlueprintGeneratedClass.cpp
+++ b/Engine/Source/Runtime/Engine/Private/BlueprintGeneratedClass.cpp
@@ -306,10 +306,10 @@ void UBlueprintGeneratedClass::CreateComponentsForActor(AActor* Actor) const
 			continue;
 		}
 
-		FName NewName = *(FString::Printf(TEXT("TimelineComp__%d"), Actor->SerializedComponents.Num() ) );
+		FName NewName = *(FString::Printf(TEXT("TimelineComp__%d"), Actor->BlueprintCreatedComponents.Num() ) );
 		UTimelineComponent* NewTimeline = NewNamedObject<UTimelineComponent>(Actor, NewName);
 		NewTimeline->bCreatedByConstructionScript = true; // Indicate it comes from a blueprint so it gets cleared when we rerun construction scripts
-		Actor->SerializedComponents.Add(NewTimeline); // Add to array so it gets saved
+		Actor->BlueprintCreatedComponents.Add(NewTimeline); // Add to array so it gets saved
 		NewTimeline->SetNetAddressable();	// This component has a stable name that can be referenced for replication
 
 		NewTimeline->SetPropertySetObject(Actor); // Set which object the timeline should drive properties on
diff --git a/Engine/Source/Runtime/Engine/Private/ComponentInstanceDataCache.cpp b/Engine/Source/Runtime/Engine/Private/ComponentInstanceDataCache.cpp
--- a/Engine/Source/Runtime/Engine/Private/ComponentInstanceDataCache.cpp
+++ b/Engine/Source/Runtime/Engine/Private/ComponentInstanceDataCache.cpp
@@ -13,17 +13,17 @@ FComponentInstanceDataBase::FComponentInstanceDataBase(const UActorComponent* So
 	if (ComponentOwner)
 	{
 		bool bFound = false;
-		for (const UActorComponent* SerializedComponent : ComponentOwner->SerializedComponents)
+		for (const UActorComponent* BlueprintCreatedComponent : ComponentOwner->BlueprintCreatedComponents)
 		{
-			if (SerializedComponent)
+			if (BlueprintCreatedComponent)
 			{
-				if (SerializedComponent == SourceComponent)
+				if (BlueprintCreatedComponent == SourceComponent)
 				{
 					++SourceComponentTypeSerializedIndex;
 					bFound = true;
 					break;
 				}
-				else if (SerializedComponent->GetClass() == SourceComponent->GetClass())
+				else if (BlueprintCreatedComponent->GetClass() == SourceComponent->GetClass())
 				{
 					++SourceComponentTypeSerializedIndex;
 				}
@@ -51,13 +51,13 @@ bool FComponentInstanceDataBase::MatchesComponent(const UActorComponent* Compone
 			AActor* ComponentOwner = Component->GetOwner();
 			if (ComponentOwner)
 			{
-				for (const UActorComponent* SerializedComponent : ComponentOwner->SerializedComponents)
+				for (const UActorComponent* BlueprintCreatedComponent : ComponentOwner->BlueprintCreatedComponents)
 				{
-					if (   SerializedComponent
-						&& (SerializedComponent->GetClass() == Component->GetClass())
+					if (   BlueprintCreatedComponent
+						&& (BlueprintCreatedComponent->GetClass() == Component->GetClass())
 						&& (++FoundSerializedComponentsOfType == SourceComponentTypeSerializedIndex))
 					{
-						bMatches = (SerializedComponent == Component);
+						bMatches = (BlueprintCreatedComponent == Component);
 						break;
 					}
 				}
diff --git a/Engine/Source/Runtime/Engine/Private/StaticMeshComponent.cpp b/Engine/Source/Runtime/Engine/Private/StaticMeshComponent.cpp
--- a/Engine/Source/Runtime/Engine/Private/StaticMeshComponent.cpp
+++ b/Engine/Source/Runtime/Engine/Private/StaticMeshComponent.cpp
@@ -1504,10 +1504,10 @@ void UStaticMeshComponent::GetUsedMaterials(TArray<UMaterialInterface*>& OutMate
 	}
 }
 
-int32 UStaticMeshComponent::GetSerializedComponentIndex() const
+int32 UStaticMeshComponent::GetBlueprintCreatedComponentIndex() const
 {
 	int32 ComponentIndex = 0;
-	for(const auto& Component : GetOwner()->SerializedComponents)
+	for(const auto& Component : GetOwner()->BlueprintCreatedComponents)
 	{
 		if(Component == this)
 		{
