Added secondary key binding for Opening Assests

diff --git a/Engine/Source/Editor/UnrealEd/Private/Toolkits/GlobalEditorCommonCommands.cpp b/Engine/Source/Editor/UnrealEd/Private/Toolkits/GlobalEditorCommonCommands.cpp
--- a/Engine/Source/Editor/UnrealEd/Private/Toolkits/GlobalEditorCommonCommands.cpp
+++ b/Engine/Source/Editor/UnrealEd/Private/Toolkits/GlobalEditorCommonCommands.cpp
@@ -96,6 +96,7 @@ void FGlobalEditorCommonCommands::RegisterCommands()
 {
 	//UI_COMMAND( SummonControlTabNavigation, "Tab Navigation", "Summons a list of open assets and tabs", EUserInterfaceActionType::Button, FInputGesture(EModifierKey::Control, EKeys::Tab) );
 	UI_COMMAND( SummonOpenAssetDialog, "Open Asset...", "Summons an asset picker", EUserInterfaceActionType::Button, FInputGesture(EModifierKey::Control, EKeys::P) );
+	UI_COMMAND( SummonOpenAssetDialogAlternate, "Open Asset...", "Summons an asset picker", EUserInterfaceActionType::Button, FInputGesture(EModifierKey::Alt | EModifierKey::Shift, EKeys::O));
 	UI_COMMAND( FindInContentBrowser, "Find in Content Browser", "Summons the Content Browser and navigates to the selected asset", EUserInterfaceActionType::Button, FInputGesture(EModifierKey::Control, EKeys::B));
 	UI_COMMAND( ViewReferences, "View References", "Launches the reference viewer showing the selected assets' references", EUserInterfaceActionType::Button, FInputGesture(EModifierKey::Shift | EModifierKey::Alt, EKeys::R));
 	
@@ -116,6 +117,10 @@ void FGlobalEditorCommonCommands::MapActions(TSharedRef<FUICommandList>& Toolkit
 		Get().SummonOpenAssetDialog,
 		FExecuteAction::CreateStatic( &FGlobalEditorCommonCommands::OnSummonedAssetPicker ) );
 
+	ToolkitCommands->MapAction(
+		Get().SummonOpenAssetDialogAlternate,
+		FExecuteAction::CreateStatic(&FGlobalEditorCommonCommands::OnSummonedAssetPicker));
+
 	ToolkitCommands->MapAction(
 		Get().OpenConsoleCommandBox,
 		FExecuteAction::CreateStatic( &FGlobalEditorCommonCommands::OnSummonedConsoleCommandBox ) );
diff --git a/Engine/Source/Editor/UnrealEd/Public/Toolkits/GlobalEditorCommonCommands.h b/Engine/Source/Editor/UnrealEd/Public/Toolkits/GlobalEditorCommonCommands.h
--- a/Engine/Source/Editor/UnrealEd/Public/Toolkits/GlobalEditorCommonCommands.h
+++ b/Engine/Source/Editor/UnrealEd/Public/Toolkits/GlobalEditorCommonCommands.h
@@ -22,6 +22,7 @@ public:
 	TSharedPtr<FUICommandInfo> FindInContentBrowser;
 	TSharedPtr<FUICommandInfo> SummonControlTabNavigation;
 	TSharedPtr<FUICommandInfo> SummonOpenAssetDialog;
+	TSharedPtr<FUICommandInfo> SummonOpenAssetDialogAlternate;
 	TSharedPtr<FUICommandInfo> OpenDocumentation;
 	TSharedPtr<FUICommandInfo> ViewReferences;
 	TSharedPtr<FUICommandInfo> OpenConsoleCommandBox;
