Reworked the comment commit code to restore functionality and handle material expression synchronisation correctly.
#UE-7369 - Giving a custom event node a comment will rename that node instead
#UE-7639 - Node comments are cleared when compiling or closing the blueprint
#UE-5661 - New [redundant] CommentBubble in material editor has bugs
#Branch UE4
#Proj Engine, UnrealEd, GraphEditor

#Info This change adds UEdGraphNode::OnUpdateCommentText to allow derived GraphNodes to hook into comment changes where neccesary. I've implemented this on the Material Expressions to address the initial change that resulted in the bugs.

diff --git a/Engine/Source/Editor/GraphEditor/Private/SCommentBubble.cpp b/Engine/Source/Editor/GraphEditor/Private/SCommentBubble.cpp
--- a/Engine/Source/Editor/GraphEditor/Private/SCommentBubble.cpp
+++ b/Engine/Source/Editor/GraphEditor/Private/SCommentBubble.cpp
@@ -335,10 +335,7 @@ void SCommentBubble::OnCommentTextCommitted( const FText& NewText, ETextCommit::
 	const bool bValidText = !NewText.EqualTo( HintText.Get() );
 	if( GraphNode && bValidText )
 	{
-		const FScopedTransaction Transaction( NSLOCTEXT( "CommentBubble", "CommentCommitted", "Comment Changed" ) );
-		GraphNode->Modify();
-		GraphNode->OnRenameNode(NewText.ToString());
-//		GraphNode->NodeComment = NewText.ToString();
+		GraphNode->OnUpdateCommentText( NewText.ToString() );
 	}
 }
 
diff --git a/Engine/Source/Editor/UnrealEd/Classes/MaterialGraph/MaterialGraphNode.h b/Engine/Source/Editor/UnrealEd/Classes/MaterialGraph/MaterialGraphNode.h
--- a/Engine/Source/Editor/UnrealEd/Classes/MaterialGraph/MaterialGraphNode.h
+++ b/Engine/Source/Editor/UnrealEd/Classes/MaterialGraph/MaterialGraphNode.h
@@ -58,6 +58,7 @@ public:
 	virtual void GetContextMenuActions(const FGraphNodeContextMenuBuilder& Context) const override;
 	virtual void PostPlacedNewNode() override;
 	virtual void OnRenameNode(const FString& NewName) override;
+	virtual void OnUpdateCommentText( const FString& NewComment ) override;
 	virtual void GetPinHoverText(const UEdGraphPin& Pin, FString& HoverTextOut) const override;
 	virtual FString GetDocumentationExcerptName() const override;
 	// End UEdGraphNode interface.
diff --git a/Engine/Source/Editor/UnrealEd/Private/MaterialGraphNode.cpp b/Engine/Source/Editor/UnrealEd/Private/MaterialGraphNode.cpp
--- a/Engine/Source/Editor/UnrealEd/Private/MaterialGraphNode.cpp
+++ b/Engine/Source/Editor/UnrealEd/Private/MaterialGraphNode.cpp
@@ -604,6 +604,20 @@ void UMaterialGraphNode::OnRenameNode(const FString& NewName)
 	MaterialDirtyDelegate.ExecuteIfBound();
 }
 
+void UMaterialGraphNode::OnUpdateCommentText( const FString& NewComment )
+{
+	const FScopedTransaction Transaction( LOCTEXT( "CommentCommitted", "Comment Changed" ) );
+	// Update the Node comment
+	Modify();
+	NodeComment	= NewComment;
+	// Update the Material Expresssion desc to match the comment
+	if( MaterialExpression )
+	{
+		MaterialExpression->Modify();
+		MaterialExpression->Desc = NewComment;
+	}
+}
+
 void UMaterialGraphNode::GetPinHoverText(const UEdGraphPin& Pin, FString& HoverTextOut) const
 {
 	Super::GetPinHoverText(Pin, HoverTextOut);
diff --git a/Engine/Source/Runtime/Engine/Classes/EdGraph/EdGraphNode.h b/Engine/Source/Runtime/Engine/Classes/EdGraph/EdGraphNode.h
--- a/Engine/Source/Runtime/Engine/Classes/EdGraph/EdGraphNode.h
+++ b/Engine/Source/Runtime/Engine/Classes/EdGraph/EdGraphNode.h
@@ -404,6 +404,9 @@ public:
 	// called when this node is being renamed after a successful name validation
 	virtual void OnRenameNode(const FString& NewName) {}
 
+	// called to replace this nodes comment text
+	virtual void OnUpdateCommentText( const FString& NewComment );
+
 	/** Return whether to draw this node as a comment node */
 	virtual bool ShouldDrawNodeAsComment() const { return false; }
 
diff --git a/Engine/Source/Runtime/Engine/Private/EdGraph/EdGraphNode.cpp b/Engine/Source/Runtime/Engine/Private/EdGraph/EdGraphNode.cpp
--- a/Engine/Source/Runtime/Engine/Private/EdGraph/EdGraphNode.cpp
+++ b/Engine/Source/Runtime/Engine/Private/EdGraph/EdGraphNode.cpp
@@ -346,6 +346,12 @@ void UEdGraphNode::AddSearchMetaDataInfo(TArray<struct FSearchTagDataPair>& OutT
 	OutTaggedMetaData.Add(FSearchTagDataPair(FFindInBlueprintSearchTags::FiB_Comment, FText::FromString(NodeComment)));
 }
 
+void UEdGraphNode::OnUpdateCommentText( const FString& NewComment )
+{
+	const FScopedTransaction Transaction( LOCTEXT( "CommentCommitted", "Comment Changed" ) );
+	Modify();
+	NodeComment	= NewComment;
+}
 
 #endif	//#if WITH_EDITOR
 
