Unity Build fix

diff --git a/Engine/Source/Editor/UnrealEd/Private/MaterialGraphNode.cpp b/Engine/Source/Editor/UnrealEd/Private/MaterialGraphNode.cpp
--- a/Engine/Source/Editor/UnrealEd/Private/MaterialGraphNode.cpp
+++ b/Engine/Source/Editor/UnrealEd/Private/MaterialGraphNode.cpp
@@ -34,6 +34,7 @@
 #include "GraphEditorSettings.h"
 #include "EditorClassUtils.h"
 #include "GenericCommands.h"
+//#include "ScopedTransaction.h"
 
 #define LOCTEXT_NAMESPACE "MaterialGraphNode"
 
