actual unity build fix, instead of a commented out include :/

diff --git a/Engine/Source/Editor/UnrealEd/Private/MaterialGraphNode.cpp b/Engine/Source/Editor/UnrealEd/Private/MaterialGraphNode.cpp
--- a/Engine/Source/Editor/UnrealEd/Private/MaterialGraphNode.cpp
+++ b/Engine/Source/Editor/UnrealEd/Private/MaterialGraphNode.cpp
@@ -34,7 +34,7 @@
 #include "GraphEditorSettings.h"
 #include "EditorClassUtils.h"
 #include "GenericCommands.h"
-//#include "ScopedTransaction.h"
+#include "ScopedTransaction.h"
 
 #define LOCTEXT_NAMESPACE "MaterialGraphNode"
 
