Fixed: World composition can't be enabled when "Levels" tool is closed
UE-7692

diff --git a/Engine/Source/Editor/UnrealEd/Private/UnrealEdMisc.cpp b/Engine/Source/Editor/UnrealEd/Private/UnrealEdMisc.cpp
--- a/Engine/Source/Editor/UnrealEd/Private/UnrealEdMisc.cpp
+++ b/Engine/Source/Editor/UnrealEd/Private/UnrealEdMisc.cpp
@@ -36,6 +36,7 @@
 #include "NavigationBuildingNotification.h"
 #include "HotReloadInterface.h"
 #include "PerformanceMonitor.h"
+#include "Engine/WorldComposition.h"
 
 #define LOCTEXT_NAMESPACE "UnrealEd"
 
@@ -438,6 +439,9 @@ void FUnrealEdMisc::OnInit()
 
 	// add handler to notify about navmesh building process
 	NavigationBuildingNotificationHandler = MakeShareable(new FNavigationBuildingNotificationImpl());
+
+	// Handles "Enable World Composition" option in WorldSettings
+	UWorldComposition::EnableWorldCompositionEvent.BindRaw(this, &FUnrealEdMisc::EnableWorldComposition);
 }
 
 void FUnrealEdMisc::InitEngineAnalytics()
@@ -613,6 +617,73 @@ void FUnrealEdMisc::TickAssetAnalytics()
 	}
 }
 
+bool FUnrealEdMisc::EnableWorldComposition(UWorld* InWorld, bool bEnable)
+{
+	if (InWorld == nullptr || InWorld->WorldType != EWorldType::Editor)
+	{
+		return false;
+	}
+			
+	if (!bEnable)
+	{
+		if (InWorld->WorldComposition != nullptr)
+		{
+			InWorld->FlushLevelStreaming();
+			InWorld->WorldComposition->MarkPendingKill();
+			InWorld->WorldComposition = nullptr;
+			UWorldComposition::WorldCompositionChangedEvent.Broadcast(InWorld);
+		}
+
+		return false;
+	}
+	
+	if (InWorld->WorldComposition == nullptr)
+	{
+		FString RootPackageName = InWorld->GetOutermost()->GetName();
+
+		// Map should be saved to disk
+		if (!FPackageName::DoesPackageExist(RootPackageName))
+		{
+			FMessageDialog::Open(EAppMsgType::Ok, LOCTEXT("EnableWorldCompositionNotSaved_Message", "Please save your level to disk before enabling World Composition"));
+			return false;
+		}
+			
+		// All existing sub-levels on this map should be removed
+		int32 NumExistingSublevels = InWorld->StreamingLevels.Num();
+		if (NumExistingSublevels > 0)
+		{
+			FMessageDialog::Open(EAppMsgType::Ok, LOCTEXT("EnableWorldCompositionExistingSublevels_Message", "World Composition cannot be enabled because there are already sub-levels manually added to the persistent level. World Composition uses auto-discovery so you must first remove any manually added sub-levels from the Levels window"));
+			return false;
+		}
+			
+		UWorldComposition* WorldCompostion = ConstructObject<UWorldComposition>(UWorldComposition::StaticClass(), InWorld);
+		// All map files found in the same and folder and all sub-folders will be added ass sub-levels to this map
+		// Make sure user understands this
+		int32 NumFoundSublevels = WorldCompostion->GetTilesList().Num();
+		if (NumFoundSublevels)
+		{
+			FFormatNamedArguments Arguments;
+			Arguments.Add(TEXT("NumSubLevels"), NumFoundSublevels);
+			Arguments.Add(TEXT("FolderLocation"), FText::FromString(FPackageName::GetLongPackagePath(RootPackageName)));
+			const FText Message = FText::Format(LOCTEXT("EnableWorldCompositionPrompt_Message", "World Composition auto-discovers sub-levels by scanning the folder the level is saved in, and all sub-folders. {NumSubLevels} level files were found in {FolderLocation} and will be added as sub-levels. Do you want to continue?"), Arguments);
+			
+			auto AppResult = FMessageDialog::Open(EAppMsgType::OkCancel, Message);
+			if (AppResult != EAppReturnType::Ok)
+			{
+				WorldCompostion->MarkPendingKill();
+				return false;
+			}
+		}
+			
+		// 
+		InWorld->WorldComposition = WorldCompostion;
+		UWorldComposition::WorldCompositionChangedEvent.Broadcast(InWorld);
+	}
+	
+	return true;
+}
+
+
 /** Build and return the path to the current project (used for relaunching the editor.)	 */
 FString CreateProjectPath()
 {
@@ -720,7 +791,8 @@ void FUnrealEdMisc::OnExit()
 
 	ISourceControlModule::Get().GetProvider().Close();
 
-
+	UWorldComposition::EnableWorldCompositionEvent.Unbind();
+	
 	UnloadFBxLibraries();
 
 	const TMap<FString, FString>& IniRestoreFiles = GetConfigRestoreFilenames();
diff --git a/Engine/Source/Editor/UnrealEd/Public/UnrealEdMisc.h b/Engine/Source/Editor/UnrealEd/Public/UnrealEdMisc.h
--- a/Engine/Source/Editor/UnrealEd/Public/UnrealEdMisc.h
+++ b/Engine/Source/Editor/UnrealEd/Public/UnrealEdMisc.h
@@ -276,6 +276,9 @@ private:
 	/** Called when the heartbeat event should be sent to engine analytics */
 	void EditorAnalyticsHeartbeat();
 
+	/** Handles "Enable World Composition" option in WorldSettings */
+	bool EnableWorldComposition(UWorld* InWorld, bool bEnable);
+
 private:
 
 	/** The current state of the autosave */
diff --git a/Engine/Source/Editor/WorldBrowser/Private/WorldBrowserModule.cpp b/Engine/Source/Editor/WorldBrowser/Private/WorldBrowserModule.cpp
--- a/Engine/Source/Editor/WorldBrowser/Private/WorldBrowserModule.cpp
+++ b/Engine/Source/Editor/WorldBrowser/Private/WorldBrowserModule.cpp
@@ -30,7 +30,7 @@ void FWorldBrowserModule::StartupModule()
 		GEngine->OnWorldDestroyed().AddRaw(this, &FWorldBrowserModule::OnWorldDestroyed);
 	}
 
-	UWorldComposition::EnableWorldCompositionEvent.BindRaw(this, &FWorldBrowserModule::EnableWorldComposition);
+	UWorldComposition::WorldCompositionChangedEvent.AddRaw(this, &FWorldBrowserModule::OnWorldCompositionChanged);
 }
 
 void FWorldBrowserModule::ShutdownModule()
@@ -40,9 +40,9 @@ void FWorldBrowserModule::ShutdownModule()
 		GEngine->OnWorldAdded().RemoveAll(this);
 		GEngine->OnWorldDestroyed().RemoveAll(this);
 	}
-	
-	UWorldComposition::EnableWorldCompositionEvent.Unbind();
 
+	UWorldComposition::WorldCompositionChangedEvent.RemoveAll(this);
+	
 	FLevelCollectionCommands::Unregister();
 
 	// unregister the editor mode
@@ -128,70 +128,4 @@ TSharedPtr<FLevelCollectionModel> FWorldBrowserModule::SharedWorldModel(UWorld*
 	return SharedWorldModel;
 }
 
-bool FWorldBrowserModule::EnableWorldComposition(UWorld* InWorld, bool bEnable)
-{
-	if (InWorld == nullptr || InWorld->WorldType != EWorldType::Editor)
-	{
-		return false;
-	}
-		
-	if (!bEnable)
-	{
-		if (InWorld->WorldComposition != nullptr)
-		{
-			InWorld->FlushLevelStreaming();
-			InWorld->WorldComposition->MarkPendingKill();
-			InWorld->WorldComposition = nullptr;
-			OnWorldCompositionChanged(InWorld);
-		}
-
-		return false;
-	}
-	
-	if (InWorld->WorldComposition == nullptr)
-	{
-		FString RootPackageName = InWorld->GetOutermost()->GetName();
-
-		// Map should be saved to disk
-		if (!FPackageName::DoesPackageExist(RootPackageName))
-		{
-			FMessageDialog::Open(EAppMsgType::Ok, LOCTEXT("EnableWorldCompositionNotSaved_Message", "Please save your level to disk before enabling World Composition"));
-			return false;
-		}
-			
-		// All existing sub-levels on this map should be removed
-		int32 NumExistingSublevels = InWorld->StreamingLevels.Num();
-		if (NumExistingSublevels > 0)
-		{
-			FMessageDialog::Open(EAppMsgType::Ok, LOCTEXT("EnableWorldCompositionExistingSublevels_Message", "World Composition cannot be enabled because there are already sub-levels manually added to the persistent level. World Composition uses auto-discovery so you must first remove any manually added sub-levels from the Levels window"));
-			return false;
-		}
-			
-		UWorldComposition* WorldCompostion = ConstructObject<UWorldComposition>(UWorldComposition::StaticClass(), InWorld);
-		// All map files found in the same and folder and all sub-folders will be added ass sub-levels to this map
-		// Make sure user understands this
-		int32 NumFoundSublevels = WorldCompostion->GetTilesList().Num();
-		if (NumFoundSublevels)
-		{
-			FFormatNamedArguments Arguments;
-			Arguments.Add(TEXT("NumSubLevels"), NumFoundSublevels);
-			Arguments.Add(TEXT("FolderLocation"), FText::FromString(FPackageName::GetLongPackagePath(RootPackageName)));
-			const FText Message = FText::Format(LOCTEXT("EnableWorldCompositionPrompt_Message", "World Composition auto-discovers sub-levels by scanning the folder the level is saved in, and all sub-folders. {NumSubLevels} level files were found in {FolderLocation} and will be added as sub-levels. Do you want to continue?"), Arguments);
-			
-			auto AppResult = FMessageDialog::Open(EAppMsgType::OkCancel, Message);
-			if (AppResult != EAppReturnType::Ok)
-			{
-				WorldCompostion->MarkPendingKill();
-				return false;
-			}
-		}
-			
-		// 
-		InWorld->WorldComposition = WorldCompostion;
-		OnWorldCompositionChanged(InWorld);
-	}
-	
-	return true;
-}
-
 #undef LOCTEXT_NAMESPACE
diff --git a/Engine/Source/Editor/WorldBrowser/Public/WorldBrowserModule.h b/Engine/Source/Editor/WorldBrowser/Public/WorldBrowserModule.h
--- a/Engine/Source/Editor/WorldBrowser/Public/WorldBrowserModule.h
+++ b/Engine/Source/Editor/WorldBrowser/Public/WorldBrowserModule.h
@@ -46,13 +46,12 @@ public:
 	 */
 	DECLARE_EVENT_OneParam(FWorldBrowserModule, FOnBrowseWorld, UWorld*);
 	FOnBrowseWorld OnBrowseWorld;
-		
+				
 private:
 	void OnWorldCreated(UWorld* InWorld);
 	void OnWorldDestroyed(UWorld* InWorld);
 	void OnWorldCompositionChanged(UWorld* InWorld);
-	bool EnableWorldComposition(UWorld* InWorld, bool bEnable);
-			
+				
 private:
 	TWeakPtr<class FLevelCollectionModel>	WorldModel;
 };
diff --git a/Engine/Source/Runtime/Engine/Classes/Engine/WorldComposition.h b/Engine/Source/Runtime/Engine/Classes/Engine/WorldComposition.h
--- a/Engine/Source/Runtime/Engine/Classes/Engine/WorldComposition.h
+++ b/Engine/Source/Runtime/Engine/Classes/Engine/WorldComposition.h
@@ -142,9 +142,13 @@ class ENGINE_API UWorldComposition : public UObject
 	/** Collect tiles package names to cook  */
 	void CollectTilesToCook(TArray<FString>& PackageNames);
 
-	// Delegate to enable/disable world composition in the world
+	// Event to enable/disable world composition in the world
 	DECLARE_DELEGATE_RetVal_TwoParams(bool, FEnableWorldCompositionEvent, UWorld*, bool);
 	static FEnableWorldCompositionEvent EnableWorldCompositionEvent;
+	
+	// Event when world composition was successfully enabled/disabled in the world
+	DECLARE_MULTICAST_DELEGATE_OneParam(FWorldCompositionChangedEvent, UWorld*);
+	static FWorldCompositionChangedEvent WorldCompositionChangedEvent;
 
 #endif //WITH_EDITOR
 
diff --git a/Engine/Source/Runtime/Engine/Private/WorldComposition.cpp b/Engine/Source/Runtime/Engine/Private/WorldComposition.cpp
--- a/Engine/Source/Runtime/Engine/Private/WorldComposition.cpp
+++ b/Engine/Source/Runtime/Engine/Private/WorldComposition.cpp
@@ -12,6 +12,7 @@ DEFINE_LOG_CATEGORY_STATIC(LogWorldComposition, Log, All);
 
 #if WITH_EDITOR
 UWorldComposition::FEnableWorldCompositionEvent UWorldComposition::EnableWorldCompositionEvent;
+UWorldComposition::FWorldCompositionChangedEvent UWorldComposition::WorldCompositionChangedEvent;
 #endif // WITH_EDITOR
 
 UWorldComposition::UWorldComposition(const FObjectInitializer& ObjectInitializer)
