Fix for UE-7508.  New packages might not be picked up when pushing them using cook in editor to device.

diff --git a/Engine/Source/Editor/UnrealEd/Classes/CookOnTheSide/CookOnTheFlyServer.h b/Engine/Source/Editor/UnrealEd/Classes/CookOnTheSide/CookOnTheFlyServer.h
--- a/Engine/Source/Editor/UnrealEd/Classes/CookOnTheSide/CookOnTheFlyServer.h
+++ b/Engine/Source/Editor/UnrealEd/Classes/CookOnTheSide/CookOnTheFlyServer.h
@@ -651,12 +651,14 @@ private:
 	FString GetCachedPackageFilename( const FName& PackageName ) const;
 	FString GetCachedStandardPackageFilename( const FName& PackageName ) const;
 	FName GetCachedStandardPackageFileFName( const FName& PackageName ) const;
-	FString GetCachedPackageFilename( UPackage* Package ) const;
-	FString GetCachedStandardPackageFilename( UPackage* Package ) const;
-	FName GetCachedStandardPackageFileFName( UPackage* Package ) const;
-	const FString& GetCachedSandboxFilename( UPackage* Package, TAutoPtr<class FSandboxPlatformFile>& SandboxFile ) const;
+	FString GetCachedPackageFilename( const UPackage* Package ) const;
+	FString GetCachedStandardPackageFilename( const UPackage* Package ) const;
+	FName GetCachedStandardPackageFileFName( const UPackage* Package ) const;
+	const FString& GetCachedSandboxFilename( const UPackage* Package, TAutoPtr<class FSandboxPlatformFile>& SandboxFile ) const;
 	const FCachedPackageFilename& Cache(const FName& PackageName) const;
 	void ClearPackageFilenameCache() const;
+	bool ClearPackageFilenameCacheForPackage( const UPackage* Package ) const;
+	bool ClearPackageFilenameCacheForPackage( const FName& PackageName ) const;
 
 	// declared mutable as it's used to cache package filename strings and I don't want to declare all functions using it as non const
 	// used by GetCached * Filename functions
@@ -789,7 +791,13 @@ public:
 
 	uint32 NumConnections() const;
 
-	
+	/**
+	 * Is this cooker running in the editor
+	 *
+	 * @return true if we are running in the editor
+	 */
+	bool IsCookingInEditor() const;
+
 	/**
 	 * Is this cooker running in real time mode (where it needs to respect the timeslice) 
 	 * 
diff --git a/Engine/Source/Editor/UnrealEd/Private/CookOnTheFlyServer.cpp b/Engine/Source/Editor/UnrealEd/Private/CookOnTheFlyServer.cpp
--- a/Engine/Source/Editor/UnrealEd/Private/CookOnTheFlyServer.cpp
+++ b/Engine/Source/Editor/UnrealEd/Private/CookOnTheFlyServer.cpp
@@ -325,25 +325,36 @@ FName UCookOnTheFlyServer::GetCachedStandardPackageFileFName( const FName& Packa
 }
 
 
-FString UCookOnTheFlyServer::GetCachedPackageFilename( UPackage* Package ) const 
+FString UCookOnTheFlyServer::GetCachedPackageFilename( const UPackage* Package ) const 
 {
 	check( Package->GetName() == Package->GetFName().ToString() );
 	return Cache( Package->GetFName() ).PackageFilename;
 }
 
-FString UCookOnTheFlyServer::GetCachedStandardPackageFilename( UPackage* Package ) const 
+FString UCookOnTheFlyServer::GetCachedStandardPackageFilename( const UPackage* Package ) const 
 {
 	check( Package->GetName() == Package->GetFName().ToString() );
 	return Cache( Package->GetFName() ).StandardFilename;
 }
 
-FName UCookOnTheFlyServer::GetCachedStandardPackageFileFName( UPackage* Package ) const 
+FName UCookOnTheFlyServer::GetCachedStandardPackageFileFName( const UPackage* Package ) const 
 {
 	check( Package->GetName() == Package->GetFName().ToString() );
 	return Cache( Package->GetFName() ).StandardFileFName;
 }
 
-const FString& UCookOnTheFlyServer::GetCachedSandboxFilename( UPackage* Package, TAutoPtr<class FSandboxPlatformFile>& SandboxFile ) const 
+
+bool UCookOnTheFlyServer::ClearPackageFilenameCacheForPackage( const FName& PackageName ) const
+{
+	return PackageFilenameCache.Remove( PackageName ) >= 1;
+}
+
+bool UCookOnTheFlyServer::ClearPackageFilenameCacheForPackage( const UPackage* Package ) const
+{
+	return PackageFilenameCache.Remove( Package->GetFName() ) >= 1;
+}
+
+const FString& UCookOnTheFlyServer::GetCachedSandboxFilename( const UPackage* Package, TAutoPtr<class FSandboxPlatformFile>& SandboxFile ) const 
 {
 	FName PackageFName = Package->GetFName();
 	static TMap<FName, FString> CachedSandboxFilenames;
@@ -807,6 +818,11 @@ void UCookOnTheFlyServer::GenerateManifestInfo( UPackage* Package, const TArray<
 	}
 }
 
+bool UCookOnTheFlyServer::IsCookingInEditor() const 
+{
+	return CurrentCookMode == ECookMode::CookByTheBookFromTheEditor || CurrentCookMode == ECookMode::CookOnTheFlyFromTheEditor;;
+}
+
 bool UCookOnTheFlyServer::IsRealtimeMode() const 
 {
 	return CurrentCookMode == ECookMode::CookByTheBookFromTheEditor || CurrentCookMode == ECookMode::CookOnTheFlyFromTheEditor;
@@ -1106,7 +1122,7 @@ uint32 UCookOnTheFlyServer::TickCookOnTheSide( const float TimeSlice, uint32 &Co
 							const auto& Obj = ObjectsInPackage[I];
 							Obj->BeginCacheForCookedPlatformData( TargetPlatform );
 							
-							if ( Timer.IsTimeUp() && IsRealtimeMode() )
+							if ( Timer.IsTimeUp() )
 							{
 #if DEBUG_COOKONTHEFLY
 								UE_LOG(LogCook, Display, TEXT("Object %s took too long to cache"), *Obj->GetFullName());
@@ -1140,7 +1156,7 @@ uint32 UCookOnTheFlyServer::TickCookOnTheSide( const float TimeSlice, uint32 &Co
 				}
 				
 
-				if ( Timer.IsTimeUp() && IsRealtimeMode() && (bIsAllDataCached == false) )
+				if ( Timer.IsTimeUp() && (bIsAllDataCached == false) )
 				{
 					break; // break out of the target platform loop (can't break out of the main package tick loop here without requeing the request)
 				}
@@ -1243,7 +1259,8 @@ uint32 UCookOnTheFlyServer::TickCookOnTheSide( const float TimeSlice, uint32 &Co
 			{
 				// if we are processing unsolicited packages we can optionally not save these right now
 				// the unsolicited packages which we missed now will be picked up on next run
-				if ( (CurrentCookMode == ECookMode::CookOnTheFly || CurrentCookMode == ECookMode::CookByTheBookFromTheEditor) && (I >= FirstUnsolicitedPackage) )
+				// we want to do this in cook on the fly also, if there is a new network package request instead of saving unsolicited packages we can process the requested package
+				if ( (IsRealtimeMode() || IsCookOnTheFlyMode()) && (I >= FirstUnsolicitedPackage) )
 				{
 					bool bShouldFinishTick = false;
 
@@ -1425,6 +1442,10 @@ void UCookOnTheFlyServer::MarkPackageDirtyForCooker( UPackage *Package )
 		const FString Name = Package->GetPathName();
 		/*FString PackageFilename(GetPackageFilename(Package));
 		FPaths::MakeStandardFilename(PackageFilename);*/
+
+		ClearPackageFilenameCacheForPackage( Package );
+
+
 		FString PackageFilename = GetCachedStandardPackageFilename(Package);
 #if DEBUG_COOKONTHEFLY
 		UE_LOG(LogCook, Display, TEXT("Modification detected to package %s"), *PackageFilename);
@@ -2893,7 +2914,7 @@ void UCookOnTheFlyServer::StartCookByTheBook( const FCookByTheBookStartupOptions
 
 void UCookOnTheFlyServer::HandleNetworkFileServerFileRequest( const FString& Filename, const FString &Platformname, TArray<FString>& UnsolicitedFiles )
 {
-	check( CurrentCookMode == ECookMode::CookOnTheFly );	
+	check( IsCookOnTheFlyMode() );	
 	
 	bool bIsCookable = FPackageName::IsPackageExtension(*FPaths::GetExtension(Filename, true));
 
