Added support for launch on to use cook on the fly when experimental option is enabled.

diff --git a/Engine/Source/Developer/LauncherServices/Private/Launcher/LauncherWorker.cpp b/Engine/Source/Developer/LauncherServices/Private/Launcher/LauncherWorker.cpp
--- a/Engine/Source/Developer/LauncherServices/Private/Launcher/LauncherWorker.cpp
+++ b/Engine/Source/Developer/LauncherServices/Private/Launcher/LauncherWorker.cpp
@@ -355,11 +355,6 @@ FString FLauncherWorker::CreateUATCommand( const ILauncherProfileRef& InProfile,
 
 			UATCommand += MapList;
 
-			if (InProfile->IsCookingIncrementally())
-			{
-				UATCommand += TEXT(" -iterativecooking");
-			}
-
 			if (InProfile->IsCookingUnversioned())
 			{
 				UATCommand += TEXT(" -Unversioned");
@@ -411,17 +406,24 @@ FString FLauncherWorker::CreateUATCommand( const ILauncherProfileRef& InProfile,
 			}
 		}
 		break;
+	case ELauncherProfileCookModes::OnTheFlyInEditor:
+		UATCommand += MapList;
+		UATCommand += " -skipcook -cookonthefly";
+		break;
 	case ELauncherProfileCookModes::ByTheBookInEditor:
 		UATCommand += MapList;
-		if (InProfile->IsCookingIncrementally())
-		{
-			UATCommand += TEXT(" -iterativecooking");
-		}
+		UATCommand += TEXT(" -skipcook"); // don't cook anything the editor is doing it ;)
+		break;
 	case ELauncherProfileCookModes::DoNotCook:
 		UATCommand += TEXT(" -skipcook");
 		break;
 	}
 
+	if (InProfile->IsCookingIncrementally())
+	{
+		UATCommand += TEXT(" -iterativecooking");
+	}
+
 	// stage/package/deploy
 	if (InProfile->GetDeploymentMode() != ELauncherProfileDeploymentModes::DoNotDeploy)
 	{
@@ -583,6 +585,7 @@ void FLauncherWorker::CreateAndExecuteTasks( const ILauncherProfileRef& InProfil
 #if !WITH_EDITOR
 	// can't cook by the book in the editor if we are not in the editor...
 	check( InProfile->GetCookMode() != ELauncherProfileCookModes::ByTheBookInEditor );
+	check( InProfile->GetCookMode() != ELauncherProfileCookModes::OnTheFlyInEditor );
 #endif
 
 	TSharedPtr<FLauncherTask> NextTask = TaskChain;
diff --git a/Engine/Source/Developer/LauncherServices/Private/Profiles/LauncherProfile.h b/Engine/Source/Developer/LauncherServices/Private/Profiles/LauncherProfile.h
--- a/Engine/Source/Developer/LauncherServices/Private/Profiles/LauncherProfile.h
+++ b/Engine/Source/Developer/LauncherServices/Private/Profiles/LauncherProfile.h
@@ -407,7 +407,11 @@ public:
 
 	virtual bool IsCookingIncrementally( ) const override
 	{
-		return CookIncremental;
+		if ( CookMode != ELauncherProfileCookModes::DoNotCook )
+		{
+			return CookIncremental;
+		}
+		return false;
 	}
 
 	virtual bool IsCookingUnversioned( ) const override
@@ -1006,7 +1010,7 @@ protected:
 		}
 
 		// Launch: when launching, all devices that the build is launched on must have content cooked for their platform
-		if ((LaunchMode != ELauncherProfileLaunchModes::DoNotLaunch) && (CookMode != ELauncherProfileCookModes::OnTheFly))
+		if ((LaunchMode != ELauncherProfileLaunchModes::DoNotLaunch) && (CookMode != ELauncherProfileCookModes::OnTheFly || CookMode != ELauncherProfileCookModes::OnTheFlyInEditor))
 		{
 			// @todo ensure that launched devices have cooked content
 		}
diff --git a/Engine/Source/Developer/LauncherServices/Public/Interfaces/ILauncherProfile.h b/Engine/Source/Developer/LauncherServices/Public/Interfaces/ILauncherProfile.h
--- a/Engine/Source/Developer/LauncherServices/Public/Interfaces/ILauncherProfile.h
+++ b/Engine/Source/Developer/LauncherServices/Public/Interfaces/ILauncherProfile.h
@@ -20,7 +20,10 @@ namespace ELauncherProfileCookModes
 		OnTheFly,
 
 		/** Cook by the book in the editor process space */
-		ByTheBookInEditor
+		ByTheBookInEditor,
+
+		/** Cook on the fly in the editor process space */
+		OnTheFlyInEditor,
 	};
 }
 
diff --git a/Engine/Source/Editor/UnrealEd/Classes/Editor/EditorEngine.h b/Engine/Source/Editor/UnrealEd/Classes/Editor/EditorEngine.h
--- a/Engine/Source/Editor/UnrealEd/Classes/Editor/EditorEngine.h
+++ b/Engine/Source/Editor/UnrealEd/Classes/Editor/EditorEngine.h
@@ -1473,6 +1473,11 @@ public:
 	 */
 	virtual bool CanCookByTheBookInEditor() const { return false; }
 
+	/**
+	 * Can the editor act as a cook on the fly server
+	 */
+	virtual bool CanCookOnTheFlyInEditor() const { return false; }
+
 	/**
 	 * Start cook by the book in the editor process space
 	 */
diff --git a/Engine/Source/Editor/UnrealEd/Classes/Editor/UnrealEdEngine.h b/Engine/Source/Editor/UnrealEd/Classes/Editor/UnrealEdEngine.h
--- a/Engine/Source/Editor/UnrealEd/Classes/Editor/UnrealEdEngine.h
+++ b/Engine/Source/Editor/UnrealEd/Classes/Editor/UnrealEdEngine.h
@@ -549,7 +549,12 @@ class UNREALED_API UUnrealEdEngine : public UEditorEngine, public FNotifyHook
 	/**
 	 * Can the editor do cook by the book in the editor process space
 	 */
-	virtual bool CanCookByTheBookInEditor() const;
+	virtual bool CanCookByTheBookInEditor() const override;
+
+	/**
+	 * Can the editor act as a cook on the fly server
+	 */
+	virtual bool CanCookOnTheFlyInEditor() const override;
 
 	/**
 	 * Start cook by the book in the editor process space
diff --git a/Engine/Source/Editor/UnrealEd/Private/PlayLevel.cpp b/Engine/Source/Editor/UnrealEd/Private/PlayLevel.cpp
--- a/Engine/Source/Editor/UnrealEd/Private/PlayLevel.cpp
+++ b/Engine/Source/Editor/UnrealEd/Private/PlayLevel.cpp
@@ -1403,11 +1403,24 @@ void UEditorEngine::PlayUsingLauncher()
 		// Setup launch profile, keep the setting here to a minimum.
 		ILauncherProfileRef LauncherProfile = LauncherServicesModule.CreateProfile(TEXT("Play On Device"));
 		LauncherProfile->SetBuildGame(bHasCode && FSourceCodeNavigation::IsCompilerAvailable());
-		LauncherProfile->SetCookMode(CanCookByTheBookInEditor() ? ELauncherProfileCookModes::ByTheBookInEditor : ELauncherProfileCookModes::ByTheBook);
-		LauncherProfile->SetIncrementalCooking(true);
+
+		// select the quickest cook mode based on which in editor cook mode is enabled
+		bool bIncrimentalCooking = true;
+		ELauncherProfileCookModes::Type CurrentLauncherCookMode = ELauncherProfileCookModes::ByTheBook;
+		if ( CanCookByTheBookInEditor() )
+		{
+			CurrentLauncherCookMode = ELauncherProfileCookModes::ByTheBookInEditor;
+		}
+		if ( CanCookOnTheFlyInEditor() )
+		{
+			CurrentLauncherCookMode = ELauncherProfileCookModes::OnTheFlyInEditor;
+			bIncrimentalCooking = false;
+		}
+		LauncherProfile->SetCookMode( CurrentLauncherCookMode );
+		LauncherProfile->SetIncrementalCooking(bIncrimentalCooking);
 		LauncherProfile->AddCookedPlatform(PlayUsingLauncherDeviceId.Left(PlayUsingLauncherDeviceId.Find(TEXT("@"))));
 		LauncherProfile->SetDeployedDeviceGroup(DeviceGroup);
-		LauncherProfile->SetIncrementalDeploying(true);
+		LauncherProfile->SetIncrementalDeploying(bIncrimentalCooking);
 		LauncherProfile->SetEditorExe(FUnrealEdMisc::Get().GetExecutableForCommandlets());
 
 		const FString DummyDeviceName(FString::Printf(TEXT("All_iOS_On_%s"), FPlatformProcess::ComputerName()));
@@ -1416,9 +1429,14 @@ void UEditorEngine::PlayUsingLauncher()
 			LauncherProfile->SetLaunchMode(ELauncherProfileLaunchModes::DefaultRole);
 		}
 
+		if ( LauncherProfile->GetCookMode() == ELauncherProfileCookModes::OnTheFlyInEditor || LauncherProfile->GetCookMode() == ELauncherProfileCookModes::OnTheFly )
+		{
+			LauncherProfile->SetDeploymentMode(ELauncherProfileDeploymentModes::FileServer);
+		}
+
 		TArray<FString> MapNames;
 		FWorldContext & EditorContext = GetEditorWorldContext();
-		if (EditorContext.World()->WorldComposition || (LauncherProfile->GetCookMode() == ELauncherProfileCookModes::ByTheBookInEditor) )
+		if (EditorContext.World()->WorldComposition || (LauncherProfile->GetCookMode() == ELauncherProfileCookModes::ByTheBookInEditor) || (LauncherProfile->GetCookMode() == ELauncherProfileCookModes::OnTheFlyInEditor) )
 		{
 			// Open world composition from original folder
 			// Or if using by book in editor don't need to resave the package just cook it by the book 
@@ -1433,7 +1451,7 @@ void UEditorEngine::PlayUsingLauncher()
 			{
 				GEditor->CancelRequestPlaySession();
 				return;
-		}
+			}
 		}
 	
 		FString InitialMapName;
@@ -1452,8 +1470,6 @@ void UEditorEngine::PlayUsingLauncher()
 
 		if ( LauncherProfile->GetCookMode() == ELauncherProfileCookModes::ByTheBookInEditor )
 		{
-			// check( GEditor != NULL );
-
 			TArray<ITargetPlatform*> TargetPlatforms;
 			for ( const auto &PlatformName : LauncherProfile->GetCookedPlatforms() )
 			{
@@ -1462,8 +1478,6 @@ void UEditorEngine::PlayUsingLauncher()
 				// crashes if two requests are inflight but we can support having multiple platforms cooking at once
 				TargetPlatforms.Add( TargetPlatform ); 
 			}
-
-
 			const TArray<FString> &CookedMaps = LauncherProfile->GetCookedMaps();
 
 			// const TArray<FString>& CookedMaps = ChainState.Profile->GetCookedMaps();
diff --git a/Engine/Source/Editor/UnrealEd/Private/UnrealEdEngine.cpp b/Engine/Source/Editor/UnrealEd/Private/UnrealEdEngine.cpp
--- a/Engine/Source/Editor/UnrealEd/Private/UnrealEdEngine.cpp
+++ b/Engine/Source/Editor/UnrealEd/Private/UnrealEdEngine.cpp
@@ -117,20 +117,19 @@ void UUnrealEdEngine::Init(IEngineLoop* InEngineLoop)
 	UEditorExperimentalSettings const* ExperimentalSettings =  GetDefault<UEditorExperimentalSettings>();
 	ECookInitializationFlags BaseCookingFlags = ECookInitializationFlags::AutoTick | ECookInitializationFlags::AsyncSave;
 	BaseCookingFlags |= ExperimentalSettings->bIterativeCookingForLaunchOn ? ECookInitializationFlags::Iterative : ECookInitializationFlags::None;
-
-	if ( !ExperimentalSettings->bDisableCookInEditor)
+	if ( ExperimentalSettings->bCookOnTheSide )
 	{
 		CookServer = ConstructObject<UCookOnTheFlyServer>( UCookOnTheFlyServer::StaticClass() );
-		CookServer->Initialize( ECookMode::CookByTheBookFromTheEditor, BaseCookingFlags );
+		CookServer->Initialize( ECookMode::CookOnTheFlyFromTheEditor, BaseCookingFlags );
+		CookServer->StartNetworkFileServer( false );
 
 		FCoreUObjectDelegates::OnObjectPropertyChanged.AddUObject(CookServer, &UCookOnTheFlyServer::OnObjectPropertyChanged);
 		FCoreUObjectDelegates::OnObjectModified.AddUObject(CookServer, &UCookOnTheFlyServer::OnObjectModified);
 	}
-	else if ( ExperimentalSettings->bCookOnTheSide )
+	else if ( !ExperimentalSettings->bDisableCookInEditor)
 	{
 		CookServer = ConstructObject<UCookOnTheFlyServer>( UCookOnTheFlyServer::StaticClass() );
-		CookServer->Initialize( ECookMode::CookOnTheFlyFromTheEditor, BaseCookingFlags );
-		CookServer->StartNetworkFileServer( false );
+		CookServer->Initialize( ECookMode::CookByTheBookFromTheEditor, BaseCookingFlags );
 
 		FCoreUObjectDelegates::OnObjectPropertyChanged.AddUObject(CookServer, &UCookOnTheFlyServer::OnObjectPropertyChanged);
 		FCoreUObjectDelegates::OnObjectModified.AddUObject(CookServer, &UCookOnTheFlyServer::OnObjectModified);
@@ -146,6 +145,15 @@ bool UUnrealEdEngine::CanCookByTheBookInEditor() const
 	return false;
 }
 
+bool UUnrealEdEngine::CanCookOnTheFlyInEditor() const
+{
+	if ( CookServer )
+	{
+		return CookServer->GetCookMode() == ECookMode::CookOnTheFlyFromTheEditor;
+	}
+	return false;
+}
+
 void UUnrealEdEngine::StartCookByTheBookInEditor( const TArray<ITargetPlatform*> &TargetPlatforms, const TArray<FString> &CookMaps, const TArray<FString> &CookDirectories, const TArray<FString> &CookCultures, const TArray<FString> &IniMapSections )
 {
 	UCookOnTheFlyServer::FCookByTheBookStartupOptions StartupOptions;
