Editor heartbeat analytics event changed from 1m to 5m.
#UE-7592

diff --git a/Engine/Source/Editor/UnrealEd/Private/UnrealEdMisc.cpp b/Engine/Source/Editor/UnrealEd/Private/UnrealEdMisc.cpp
--- a/Engine/Source/Editor/UnrealEd/Private/UnrealEdMisc.cpp
+++ b/Engine/Source/Editor/UnrealEd/Private/UnrealEdMisc.cpp
@@ -431,7 +431,7 @@ void FUnrealEdMisc::OnInit()
 	InitEngineAnalytics();
 	
 	// Setup a timer for a heartbeat event to track if users are actually using the editor or it is idle.
-	float Seconds = 60.0f;
+	float Seconds = (float)FTimespan::FromMinutes(5.0).GetTotalSeconds();
 	FTimerDelegate Delegate;
 	Delegate.BindRaw( this, &FUnrealEdMisc::EditorAnalyticsHeartbeat );
 
