Merging DLC cooking safety change

//depot/UE4-UT/Engine/...

to //depot/UE4/Engine/...

diff --git a/Engine/Source/Editor/UnrealEd/Private/CookOnTheFlyServer.cpp b/Engine/Source/Editor/UnrealEd/Private/CookOnTheFlyServer.cpp
--- a/Engine/Source/Editor/UnrealEd/Private/CookOnTheFlyServer.cpp
+++ b/Engine/Source/Editor/UnrealEd/Private/CookOnTheFlyServer.cpp
@@ -3078,7 +3078,15 @@ bool UCookOnTheFlyServer::GetAllPackagesFromAssetRegistry( const FString& AssetR
 
 			UE_LOG(LogCook, Verbose, TEXT("Read package %s from %s"), *GetCachedStandardPackageFilename(NewAssetData->ObjectPath), *AssetRegistryPath);
 
-			OutPackageNames.Add( GetCachedStandardPackageFileFName(NewAssetData->ObjectPath) );
+			FName CachedPackageFileFName = GetCachedStandardPackageFileFName(NewAssetData->ObjectPath);
+			if (CachedPackageFileFName != NAME_None)
+			{
+				OutPackageNames.Add(CachedPackageFileFName);
+			}
+			else
+			{
+				UE_LOG(LogCook, Warning, TEXT("Could not resolve package %s from %s"), *NewAssetData->ObjectPath.ToString(), *AssetRegistryPath);
+			}
 		}
 
 		delete[] PreallocatedAssetDataBuffer;
