This change adds some functionality on end play map ( PIE exit ) to locate and close any asset editors that are currently editing objects which are about to be garbage collected.
#UE-6235 - Pressing Stop on PIE while a Dynamic Material Instance in a Blueprint is open will crash the editor
#Branch UE4
#Proj UnrealEd

#info this change is intended to walk through all objects that have an outer package that is flagged PKG_PlayInEditor and then close any associated asset editors.

diff --git a/Engine/Source/Editor/UnrealEd/Private/PlayLevel.cpp b/Engine/Source/Editor/UnrealEd/Private/PlayLevel.cpp
--- a/Engine/Source/Editor/UnrealEd/Private/PlayLevel.cpp
+++ b/Engine/Source/Editor/UnrealEd/Private/PlayLevel.cpp
@@ -154,10 +154,16 @@ void UEditorEngine::EndPlayMap()
 	for( FObjectIterator It; It; ++It )
 	{
 		UObject* Object = *It;
-		if( Object->HasAnyFlags(RF_Standalone) && (Object->GetOutermost()->PackageFlags & PKG_PlayInEditor)  )
+
+		if((Object->GetOutermost()->PackageFlags & PKG_PlayInEditor) != 0)
 		{
-			// Clear RF_Standalone flag from objects in the levels used for PIE so they get cleaned up.
-			Object->ClearFlags(RF_Standalone);
+			if(Object->HasAnyFlags(RF_Standalone))
+			{
+				// Clear RF_Standalone flag from objects in the levels used for PIE so they get cleaned up.
+				Object->ClearFlags(RF_Standalone);
+			}
+			// Close any asset editors that are currently editing this object
+			FAssetEditorManager::Get().CloseAllEditorsForAsset(Object);
 		}
 	}
 
diff --git a/Engine/Source/Editor/UnrealEd/Private/Toolkits/AssetEditorManager.cpp b/Engine/Source/Editor/UnrealEd/Private/Toolkits/AssetEditorManager.cpp
--- a/Engine/Source/Editor/UnrealEd/Private/Toolkits/AssetEditorManager.cpp
+++ b/Engine/Source/Editor/UnrealEd/Private/Toolkits/AssetEditorManager.cpp
@@ -114,6 +114,18 @@ TArray<IAssetEditorInstance*> FAssetEditorManager::FindEditorsForAsset(UObject*
 	return AssetEditors;
 }
 
+
+void FAssetEditorManager::CloseAllEditorsForAsset(UObject* Asset)
+{
+	TArray<IAssetEditorInstance*> EditorInstances = FindEditorsForAsset(Asset);
+
+	for( auto EditorIter : EditorInstances )
+	{
+		EditorIter->CloseWindow();
+	}
+}
+
+
 void FAssetEditorManager::CloseOtherEditors( UObject* Asset, IAssetEditorInstance* OnlyEditor)
 {
 	TArray<UObject*> AllAssets;
diff --git a/Engine/Source/Editor/UnrealEd/Public/Toolkits/AssetEditorManager.h b/Engine/Source/Editor/UnrealEd/Public/Toolkits/AssetEditorManager.h
--- a/Engine/Source/Editor/UnrealEd/Public/Toolkits/AssetEditorManager.h
+++ b/Engine/Source/Editor/UnrealEd/Public/Toolkits/AssetEditorManager.h
@@ -64,6 +64,9 @@ public:
 	/** Returns all editors currently opened for the specified asset */
 	TArray<IAssetEditorInstance*> FindEditorsForAsset(UObject* Asset);
 
+	/** Close all active editors for the supplied asset */
+	void CloseAllEditorsForAsset(UObject* Asset);
+
 	/** Close any editor which is not this one */
 	void CloseOtherEditors(UObject* Asset, IAssetEditorInstance* OnlyEditor);
 
