Fix FBX crash if analytics is disabled

diff --git a/Engine/Source/Editor/UnrealEd/Private/Fbx/FbxMainImport.cpp b/Engine/Source/Editor/UnrealEd/Private/Fbx/FbxMainImport.cpp
--- a/Engine/Source/Editor/UnrealEd/Private/Fbx/FbxMainImport.cpp
+++ b/Engine/Source/Editor/UnrealEd/Private/Fbx/FbxMainImport.cpp
@@ -769,26 +769,29 @@ bool FFbxImporter::ImportFromFile(const FString& Filename, const FString& Type)
 				FbxDocumentInfo* DocInfo = Scene->GetSceneInfo();
 				if (DocInfo)
 				{
-					TArray<FAnalyticsEventAttribute> Attribs;
+					if( FEngineAnalytics::IsAvailable() )
+					{
+						TArray<FAnalyticsEventAttribute> Attribs;
 
-					FString OriginalVendor(ANSI_TO_TCHAR(DocInfo->Original_ApplicationVendor.Get().Buffer()));
-					FString OriginalAppName(ANSI_TO_TCHAR(DocInfo->Original_ApplicationName.Get().Buffer()));
-					FString OriginalAppVersion(ANSI_TO_TCHAR(DocInfo->Original_ApplicationVersion.Get().Buffer()));
+						FString OriginalVendor(ANSI_TO_TCHAR(DocInfo->Original_ApplicationVendor.Get().Buffer()));
+						FString OriginalAppName(ANSI_TO_TCHAR(DocInfo->Original_ApplicationName.Get().Buffer()));
+						FString OriginalAppVersion(ANSI_TO_TCHAR(DocInfo->Original_ApplicationVersion.Get().Buffer()));
 
-					FString LastSavedVendor(ANSI_TO_TCHAR(DocInfo->LastSaved_ApplicationVendor.Get().Buffer()));
-					FString LastSavedAppName(ANSI_TO_TCHAR(DocInfo->LastSaved_ApplicationName.Get().Buffer()));
-					FString LastSavedAppVersion(ANSI_TO_TCHAR(DocInfo->LastSaved_ApplicationVersion.Get().Buffer()));
+						FString LastSavedVendor(ANSI_TO_TCHAR(DocInfo->LastSaved_ApplicationVendor.Get().Buffer()));
+						FString LastSavedAppName(ANSI_TO_TCHAR(DocInfo->LastSaved_ApplicationName.Get().Buffer()));
+						FString LastSavedAppVersion(ANSI_TO_TCHAR(DocInfo->LastSaved_ApplicationVersion.Get().Buffer()));
 
-					Attribs.Add(FAnalyticsEventAttribute(TEXT("Original Application Vendor"), OriginalVendor));
-					Attribs.Add(FAnalyticsEventAttribute(TEXT("Original Application Name"), OriginalAppName));
-					Attribs.Add(FAnalyticsEventAttribute(TEXT("Original Application Version"), OriginalAppVersion));
+						Attribs.Add(FAnalyticsEventAttribute(TEXT("Original Application Vendor"), OriginalVendor));
+						Attribs.Add(FAnalyticsEventAttribute(TEXT("Original Application Name"), OriginalAppName));
+						Attribs.Add(FAnalyticsEventAttribute(TEXT("Original Application Version"), OriginalAppVersion));
 
-					Attribs.Add(FAnalyticsEventAttribute(TEXT("LastSaved Application Vendor"), LastSavedVendor));
-					Attribs.Add(FAnalyticsEventAttribute(TEXT("LastSaved Application Name"), LastSavedAppName));
-					Attribs.Add(FAnalyticsEventAttribute(TEXT("LastSaved Application Version"), LastSavedAppVersion));
+						Attribs.Add(FAnalyticsEventAttribute(TEXT("LastSaved Application Vendor"), LastSavedVendor));
+						Attribs.Add(FAnalyticsEventAttribute(TEXT("LastSaved Application Name"), LastSavedAppName));
+						Attribs.Add(FAnalyticsEventAttribute(TEXT("LastSaved Application Version"), LastSavedAppVersion));
 
-					FString EventString = FString::Printf(TEXT("Editor.Usage.FBX.Import"));
-					FEngineAnalytics::GetProvider().RecordEvent(EventString, Attribs);
+						FString EventString = FString::Printf(TEXT("Editor.Usage.FBX.Import"));
+						FEngineAnalytics::GetProvider().RecordEvent(EventString, Attribs);
+					}
 				}
 			}
 
