Made the ability to enable or disable editor analytics project agnostic.

diff --git a/Engine/Source/Editor/UnrealEd/Classes/Editor/EditorEngine.h b/Engine/Source/Editor/UnrealEd/Classes/Editor/EditorEngine.h
--- a/Engine/Source/Editor/UnrealEd/Classes/Editor/EditorEngine.h
+++ b/Engine/Source/Editor/UnrealEd/Classes/Editor/EditorEngine.h
@@ -2577,6 +2577,9 @@ private:
 	/** Callback for finished undo transactions. */
 	void HandleTransactorUndo( FUndoSessionContext SessionContext, bool Succeeded );
 
+	/** UEngine interface */
+	virtual bool AreEditorAnalyticsEnabled() const override;
+
 private:
 
 	/** Delegate broadcast just before a blueprint is compiled */
diff --git a/Engine/Source/Editor/UnrealEd/Classes/Editor/EditorGameAgnosticSettings.h b/Engine/Source/Editor/UnrealEd/Classes/Editor/EditorGameAgnosticSettings.h
--- a/Engine/Source/Editor/UnrealEd/Classes/Editor/EditorGameAgnosticSettings.h
+++ b/Engine/Source/Editor/UnrealEd/Classes/Editor/EditorGameAgnosticSettings.h
@@ -10,9 +10,12 @@ class UEditorGameAgnosticSettings : public UObject
 	GENERATED_UCLASS_BODY()
 
 	/** When checked, the most recently loaded project will be auto-loaded at editor startup if no other project was specified on the command line */
-	UPROPERTY(EditAnywhere, Category=Startup)
+	UPROPERTY(EditAnywhere, Category=General)
 	bool bLoadTheMostRecentlyLoadedProjectAtStartup; // Note that this property is NOT config since it is not necessary to save the value to ini. It is determined at startup in UEditorEngine::InitEditor().
 
+	/** Can the editor report usage analytics (types of assets being spawned, etc...) back to Epic in order for us to improve the editor user experience?  Note: The editor must be restarted for changes to take effect. */
+	UPROPERTY(EditAnywhere, config, Category=General, AdvancedDisplay)
+	bool bEditorAnalyticsEnabled;
 
 	// =====================================================================
 	// The following options are NOT exposed in the preferences Editor
diff --git a/Engine/Source/Editor/UnrealEd/Private/Editor/EditorGameAgnosticSettings.cpp b/Engine/Source/Editor/UnrealEd/Private/Editor/EditorGameAgnosticSettings.cpp
--- a/Engine/Source/Editor/UnrealEd/Private/Editor/EditorGameAgnosticSettings.cpp
+++ b/Engine/Source/Editor/UnrealEd/Private/Editor/EditorGameAgnosticSettings.cpp
@@ -6,6 +6,7 @@ UEditorGameAgnosticSettings::UEditorGameAgnosticSettings(const FObjectInitialize
 	: Super(ObjectInitializer)
 {
 	bCopyStarterContentPreference = true;
+	bEditorAnalyticsEnabled = true;
 	AutoScalabilityWorkScaleAmount = 1;
 }
 
diff --git a/Engine/Source/Editor/UnrealEd/Private/EditorServer.cpp b/Engine/Source/Editor/UnrealEd/Private/EditorServer.cpp
--- a/Engine/Source/Editor/UnrealEd/Private/EditorServer.cpp
+++ b/Engine/Source/Editor/UnrealEd/Private/EditorServer.cpp
@@ -1081,6 +1081,11 @@ void UEditorEngine::HandleTransactorUndo( FUndoSessionContext SessionContext, bo
 	ShowUndoRedoNotification(FText::Format(NSLOCTEXT("UnrealEd", "UndoMessageFormat", "Undo: {0}"), SessionContext.Title), Succeeded);
 }
 
+bool UEditorEngine::AreEditorAnalyticsEnabled() const 
+{
+	return GetGameAgnosticSettings().bEditorAnalyticsEnabled;
+}
+
 UTransactor* UEditorEngine::CreateTrans()
 {
 	int32 UndoBufferSize;
diff --git a/Engine/Source/Editor/UnrealEd/Private/Settings/EditorLoadingSavingSettingsCustomization.h b/Engine/Source/Editor/UnrealEd/Private/Settings/EditorLoadingSavingSettingsCustomization.h
--- a/Engine/Source/Editor/UnrealEd/Private/Settings/EditorLoadingSavingSettingsCustomization.h
+++ b/Engine/Source/Editor/UnrealEd/Private/Settings/EditorLoadingSavingSettingsCustomization.h
@@ -58,6 +58,8 @@ protected:
 			TArray<UObject*> ObjectList;
 			ObjectList.Add(&GEditor->AccessGameAgnosticSettings());
 			StartupCategory.AddExternalProperty(ObjectList, "bLoadTheMostRecentlyLoadedProjectAtStartup");
+
+			StartupCategory.AddExternalProperty(ObjectList, "bEditorAnalyticsEnabled", EPropertyLocation::Advanced);
 		}
 	}
 };
diff --git a/Engine/Source/Runtime/Engine/Classes/Engine/Engine.h b/Engine/Source/Runtime/Engine/Classes/Engine/Engine.h
--- a/Engine/Source/Runtime/Engine/Classes/Engine/Engine.h
+++ b/Engine/Source/Runtime/Engine/Classes/Engine/Engine.h
@@ -1091,10 +1091,6 @@ public:
 	UPROPERTY(EditAnywhere, config, Category=Settings)
 	float NearClipPlane;
 
-	/** Can the editor report usage analytics (types of assets being spawned, etc...) back to Epic in order for us to improve the editor user experience?  Note: The editor must be restarted for changes to take effect. */
-	UPROPERTY(EditAnywhere, config, Category=Settings, AdvancedDisplay)
-	uint32 bEditorAnalyticsEnabled:1;
-
 	/** Can a runtime game/application report anonymous hardware survey statistics (such as display resolution and GPU model) back to Epic? */
 	UPROPERTY(EditAnywhere, config, Category=Settings, AdvancedDisplay)
 	uint32 bHardwareSurveyEnabled:1;
@@ -2462,6 +2458,8 @@ public:
 
 	UGameViewportClient* GameViewportForWorld(const UWorld *InWorld) const;
 
+	/** @return true if editor analytics are enabled */
+	virtual bool AreEditorAnalyticsEnabled() const { return false; }
 protected:
 
 	TIndirectArray<FWorldContext>	WorldList;
diff --git a/Engine/Source/Runtime/Engine/Private/EngineAnalytics.cpp b/Engine/Source/Runtime/Engine/Private/EngineAnalytics.cpp
--- a/Engine/Source/Runtime/Engine/Private/EngineAnalytics.cpp
+++ b/Engine/Source/Runtime/Engine/Private/EngineAnalytics.cpp
@@ -127,7 +127,7 @@ void FEngineAnalytics::Initialize()
 #if WITH_EDITORONLY_DATA
 	if (GIsEditor)
 	{
-		bShouldSendUsageEvents = GEngine->bEditorAnalyticsEnabled;
+		bShouldSendUsageEvents = GEngine->AreEditorAnalyticsEnabled();
 	}
 	else
 #endif
diff --git a/Engine/Source/Runtime/Engine/Private/GameEngine.cpp b/Engine/Source/Runtime/Engine/Private/GameEngine.cpp
--- a/Engine/Source/Runtime/Engine/Private/GameEngine.cpp
+++ b/Engine/Source/Runtime/Engine/Private/GameEngine.cpp
@@ -380,7 +380,7 @@ UEngine::UEngine(const FObjectInitializer& ObjectInitializer)
 	SelectionHighlightIntensityBillboards = 0.25f;
 
 	bUseSound = true;
-	bEditorAnalyticsEnabled = true;
+
 	bHardwareSurveyEnabled = true;
 	bPendingHardwareSurveyResults = false;
 	bIsInitialized = false;
