Add a startupsession event to detect if the launcher is installed

diff --git a/Engine/Source/Editor/UnrealEd/Classes/Editor/EditorEngine.h b/Engine/Source/Editor/UnrealEd/Classes/Editor/EditorEngine.h
--- a/Engine/Source/Editor/UnrealEd/Classes/Editor/EditorEngine.h
+++ b/Engine/Source/Editor/UnrealEd/Classes/Editor/EditorEngine.h
@@ -2579,7 +2579,7 @@ private:
 
 	/** UEngine interface */
 	virtual bool AreEditorAnalyticsEnabled() const override;
-
+	virtual void CreateStartupAnalyticsAttributes( TArray<FAnalyticsEventAttribute>& StartSessionAttributes ) const override;
 private:
 
 	/** Delegate broadcast just before a blueprint is compiled */
diff --git a/Engine/Source/Editor/UnrealEd/Private/EditorServer.cpp b/Engine/Source/Editor/UnrealEd/Private/EditorServer.cpp
--- a/Engine/Source/Editor/UnrealEd/Private/EditorServer.cpp
+++ b/Engine/Source/Editor/UnrealEd/Private/EditorServer.cpp
@@ -59,6 +59,7 @@
 #include "SNotificationList.h"
 #include "NotificationManager.h"
 #include "EditorUndoClient.h"
+#include "DesktopPlatformModule.h"
 
 DEFINE_LOG_CATEGORY_STATIC(LogEditorServer, Log, All);
 
@@ -1086,6 +1087,21 @@ bool UEditorEngine::AreEditorAnalyticsEnabled() const
 	return GetGameAgnosticSettings().bEditorAnalyticsEnabled;
 }
 
+void UEditorEngine::CreateStartupAnalyticsAttributes( TArray<FAnalyticsEventAttribute>& StartSessionAttributes ) const
+{
+	Super::CreateStartupAnalyticsAttributes( StartSessionAttributes );
+
+	IDesktopPlatform* DesktopPlatform = FDesktopPlatformModule::Get();
+	if(DesktopPlatform != nullptr)
+	{
+		// If this is false, CanOpenLauncher will only return true if the launcher is already installed on the users machine
+		const bool bIncludeLauncherInstaller = false;
+
+		bool bIsLauncherInstalled = DesktopPlatform->CanOpenLauncher(bIncludeLauncherInstaller);
+		StartSessionAttributes.Add(FAnalyticsEventAttribute(TEXT("IsLauncherInstalled"), bIsLauncherInstalled));
+	}
+}
+
 UTransactor* UEditorEngine::CreateTrans()
 {
 	int32 UndoBufferSize;
diff --git a/Engine/Source/Runtime/Engine/Classes/Engine/Engine.h b/Engine/Source/Runtime/Engine/Classes/Engine/Engine.h
--- a/Engine/Source/Runtime/Engine/Classes/Engine/Engine.h
+++ b/Engine/Source/Runtime/Engine/Classes/Engine/Engine.h
@@ -2460,6 +2460,8 @@ public:
 
 	/** @return true if editor analytics are enabled */
 	virtual bool AreEditorAnalyticsEnabled() const { return false; }
+	virtual void CreateStartupAnalyticsAttributes( TArray<struct FAnalyticsEventAttribute>& StartSessionAttributes ) const {}
+
 protected:
 
 	TIndirectArray<FWorldContext>	WorldList;
diff --git a/Engine/Source/Runtime/Engine/Private/EngineAnalytics.cpp b/Engine/Source/Runtime/Engine/Private/EngineAnalytics.cpp
--- a/Engine/Source/Runtime/Engine/Private/EngineAnalytics.cpp
+++ b/Engine/Source/Runtime/Engine/Private/EngineAnalytics.cpp
@@ -6,7 +6,6 @@
 #include "Runtime/Analytics/Analytics/Public/Analytics.h"
 #include "Runtime/Analytics/Analytics/Public/Interfaces/IAnalyticsProvider.h"
 
-
 bool FEngineAnalytics::bIsInitialized;
 TSharedPtr<IAnalyticsProvider> FEngineAnalytics::Analytics;
 bool FEngineAnalytics::bShouldSendUsageEvents;
@@ -205,7 +204,11 @@ void FEngineAnalytics::Initialize()
 			if (Analytics.IsValid())
 			{
 				Analytics->SetUserID(FString::Printf(TEXT("%s|%s"), *FPlatformMisc::GetMachineId().ToString(EGuidFormats::Digits).ToLower(), *FPlatformMisc::GetEpicAccountId()));
-				Analytics->StartSession();
+
+				TArray<FAnalyticsEventAttribute> StartSessionAttributes;
+				GEngine->CreateStartupAnalyticsAttributes( StartSessionAttributes );
+
+				Analytics->StartSession( StartSessionAttributes );
 			}
 		}
 #endif
