Fix incorrect case in generated header file name.

diff --git a/Engine/Source/Editor/UnrealEd/Classes/Commandlets/GatherTextFromMetadataCommandlet.h b/Engine/Source/Editor/UnrealEd/Classes/Commandlets/GatherTextFromMetadataCommandlet.h
--- a/Engine/Source/Editor/UnrealEd/Classes/Commandlets/GatherTextFromMetadataCommandlet.h
+++ b/Engine/Source/Editor/UnrealEd/Classes/Commandlets/GatherTextFromMetadataCommandlet.h
@@ -2,7 +2,7 @@
 
 #pragma once 
 #include "Commandlets/Commandlet.h"
-#include "GatherTextFromMetaDataCommandlet.generated.h"
+#include "GatherTextFromMetadataCommandlet.generated.h"
 
 /**
  *	UGatherTextFromMetaDataCommandlet: Localization commandlet that collects all text to be localized from generated metadata.
diff --git a/Engine/Source/Editor/UnrealEd/Classes/Factories/Texture2dFactoryNew.h b/Engine/Source/Editor/UnrealEd/Classes/Factories/Texture2dFactoryNew.h
--- a/Engine/Source/Editor/UnrealEd/Classes/Factories/Texture2dFactoryNew.h
+++ b/Engine/Source/Editor/UnrealEd/Classes/Factories/Texture2dFactoryNew.h
@@ -5,7 +5,7 @@
 //=============================================================================
 
 #pragma once
-#include "Texture2DFactoryNew.generated.h"
+#include "Texture2dFactoryNew.generated.h"
 
 UCLASS(hidecategories=Object, MinimalAPI)
 class UTexture2DFactoryNew : public UFactory
diff --git a/Engine/Source/Runtime/Engine/Classes/Slate/CheckboxStyleAsset.h b/Engine/Source/Runtime/Engine/Classes/Slate/CheckboxStyleAsset.h
--- a/Engine/Source/Runtime/Engine/Classes/Slate/CheckboxStyleAsset.h
+++ b/Engine/Source/Runtime/Engine/Classes/Slate/CheckboxStyleAsset.h
@@ -3,7 +3,7 @@
 #pragma once
 
 #include "SlateCore.h"
-#include "CheckBoxStyleAsset.generated.h"
+#include "CheckboxStyleAsset.generated.h"
 
 /**
  * An asset describing a CheckBox's appearance.
