#Hierarchical LOD generation system

diff --git a/Engine/Source/Editor/LevelEditor/Private/LevelEditor.cpp b/Engine/Source/Editor/LevelEditor/Private/LevelEditor.cpp
--- a/Engine/Source/Editor/LevelEditor/Private/LevelEditor.cpp
+++ b/Engine/Source/Editor/LevelEditor/Private/LevelEditor.cpp
@@ -1320,6 +1320,9 @@ void FLevelEditorModule::BindGlobalLevelEditorCommands()
 	ActionList.MapAction( Commands.BuildPathsOnly,
 		FExecuteAction::CreateStatic( &FLevelEditorActionCallbacks::BuildPathsOnly_Execute ) );
 
+	ActionList.MapAction(Commands.BuildLODsOnly,
+		FExecuteAction::CreateStatic(&FLevelEditorActionCallbacks::BuildLODsOnly_Execute));
+
 	ActionList.MapAction( 
 		Commands.LightingQuality_Production, 
 		FExecuteAction::CreateStatic( &FLevelEditorActionCallbacks::SetLightingQuality, (ELightingBuildQuality)Quality_Production ),
diff --git a/Engine/Source/Editor/LevelEditor/Private/LevelEditorActions.cpp b/Engine/Source/Editor/LevelEditor/Private/LevelEditorActions.cpp
--- a/Engine/Source/Editor/LevelEditor/Private/LevelEditorActions.cpp
+++ b/Engine/Source/Editor/LevelEditor/Private/LevelEditorActions.cpp
@@ -611,6 +611,16 @@ void FLevelEditorActionCallbacks::BuildPathsOnly_Execute()
 	FEditorBuildUtils::EditorBuild( GetWorld(), EBuildOptions::BuildAIPaths );
 }
 
+void FLevelEditorActionCallbacks::BuildLODsOnly_Execute()
+{
+	// Build paths!
+	FEditorBuildUtils::EditorBuild(GetWorld(), EBuildOptions::BuildHierarchicalLOD);
+
+	if (BuildLighting_CanExecute())
+	{
+		BuildLightingOnly_Execute();
+	}
+}
 
 bool FLevelEditorActionCallbacks::IsLightingQualityChecked( ELightingBuildQuality TestQuality )
 {
@@ -2692,6 +2702,7 @@ void FLevelEditorCommands::RegisterCommands()
 	UI_COMMAND( BuildGeometryOnly, "Build Geometry", "Only builds geometry (all levels.)", EUserInterfaceActionType::Button, FInputGesture() );
 	UI_COMMAND( BuildGeometryOnly_OnlyCurrentLevel, "Build Geometry (Current Level)", "Builds geometry, only for the current level", EUserInterfaceActionType::Button, FInputGesture() );
 	UI_COMMAND( BuildPathsOnly, "Build Paths", "Only builds paths (all levels.)", EUserInterfaceActionType::Button, FInputGesture() );
+	UI_COMMAND( BuildLODsOnly, "Build LODs", "Only builds LODs (all levels.)", EUserInterfaceActionType::Button, FInputGesture() );
 	UI_COMMAND( LightingQuality_Production, "Production", "Sets precomputed lighting quality to highest possible quality (slowest computation time.)", EUserInterfaceActionType::RadioButton, FInputGesture() );
 	UI_COMMAND( LightingQuality_High, "High", "Sets precomputed lighting quality to high quality", EUserInterfaceActionType::RadioButton, FInputGesture() );
 	UI_COMMAND( LightingQuality_Medium, "Medium", "Sets precomputed lighting quality to medium quality", EUserInterfaceActionType::RadioButton, FInputGesture() );
diff --git a/Engine/Source/Editor/LevelEditor/Private/LevelEditorToolBar.cpp b/Engine/Source/Editor/LevelEditor/Private/LevelEditorToolBar.cpp
--- a/Engine/Source/Editor/LevelEditor/Private/LevelEditorToolBar.cpp
+++ b/Engine/Source/Editor/LevelEditor/Private/LevelEditorToolBar.cpp
@@ -1564,6 +1564,12 @@ TSharedRef< SWidget > FLevelEditorToolBar::GenerateBuildMenuContent( TSharedRef<
 	}
 	MenuBuilder.EndSection();
 
+	MenuBuilder.BeginSection("LevelEditorLOD", LOCTEXT("LODHeading", "Hiearchical LOD"));
+	{
+		MenuBuilder.AddMenuEntry(FLevelEditorCommands::Get().BuildLODsOnly);
+	}
+	MenuBuilder.EndSection();
+
 	MenuBuilder.BeginSection("LevelEditorAutomation", LOCTEXT( "AutomationHeading", "Automation" ) );
 	{
 		MenuBuilder.AddMenuEntry( 
diff --git a/Engine/Source/Editor/LevelEditor/Public/LevelEditorActions.h b/Engine/Source/Editor/LevelEditor/Public/LevelEditorActions.h
--- a/Engine/Source/Editor/LevelEditor/Public/LevelEditorActions.h
+++ b/Engine/Source/Editor/LevelEditor/Public/LevelEditorActions.h
@@ -74,6 +74,7 @@ public:
 	TSharedPtr< FUICommandInfo > BuildGeometryOnly;
 	TSharedPtr< FUICommandInfo > BuildGeometryOnly_OnlyCurrentLevel;
 	TSharedPtr< FUICommandInfo > BuildPathsOnly;
+	TSharedPtr< FUICommandInfo > BuildLODsOnly;
 	TSharedPtr< FUICommandInfo > LightingQuality_Production;
 	TSharedPtr< FUICommandInfo > LightingQuality_High;
 	TSharedPtr< FUICommandInfo > LightingQuality_Medium;
@@ -686,6 +687,7 @@ public:
 	static void BuildGeometryOnly_Execute();
 	static void BuildGeometryOnly_OnlyCurrentLevel_Execute();
 	static void BuildPathsOnly_Execute();
+	static void BuildLODsOnly_Execute();
 	static void SetLightingQuality( ELightingBuildQuality NewQuality );
 	static bool IsLightingQualityChecked( ELightingBuildQuality TestQuality );
 	static void SetLightingToolShowBounds();
diff --git a/Engine/Source/Editor/UnrealEd/Private/EditorBuildUtils.cpp b/Engine/Source/Editor/UnrealEd/Private/EditorBuildUtils.cpp
--- a/Engine/Source/Editor/UnrealEd/Private/EditorBuildUtils.cpp
+++ b/Engine/Source/Editor/UnrealEd/Private/EditorBuildUtils.cpp
@@ -214,6 +214,9 @@ bool FEditorBuildUtils::EditorBuild( UWorld* InWorld, EBuildOptions::Type Id, co
 	case EBuildOptions::BuildSelectedAIPaths:
 		BuildType = SBuildProgressWidget::BUILDTYPE_Paths;
 		break;
+	case EBuildOptions::BuildHierarchicalLOD:
+		BuildType = SBuildProgressWidget::BUILDTYPE_LODs;
+		break;
 	default:
 		BuildType = SBuildProgressWidget::BUILDTYPE_Unknown;	
 		break;
@@ -286,6 +289,23 @@ bool FEditorBuildUtils::EditorBuild( UWorld* InWorld, EBuildOptions::Type Id, co
 			break;
 		}
 
+	case EBuildOptions::BuildHierarchicalLOD:
+		{
+			bDoBuild = GEditor->WarnAboutHiddenLevels( InWorld, false );
+			if ( bDoBuild )
+			{
+				GEditor->ResetTransaction( NSLOCTEXT("UnrealEd", "RebuildLOD", "Rebuilding HierarchicalLOD") );
+
+				// We can't set the busy cursor for all windows, because lighting
+				// needs a cursor for the lighting options dialog.
+				const FScopedBusyCursor BusyCursor;
+
+				TriggerHierarchicalLODBuilder(InWorld, Id);
+			}
+
+			break;
+		}
+
 	case EBuildOptions::BuildAll:
 	case EBuildOptions::BuildAllSubmit:
 		{
@@ -299,6 +319,13 @@ bool FEditorBuildUtils::EditorBuild( UWorld* InWorld, EBuildOptions::Type Id, co
 
 				GUnrealEd->Exec( InWorld, TEXT("MAP REBUILD ALLVISIBLE") );
 
+				// disable this in build all path, it's very slow, it can be painful
+				// they don't have any collision, but make sure it's added first before doing anything
+ 				{
+ 					BuildProgressWidget.Pin()->SetBuildType(SBuildProgressWidget::BUILDTYPE_LODs);
+ 					TriggerHierarchicalLODBuilder(InWorld, Id);
+ 				}
+
 				{
 					BuildProgressWidget.Pin()->SetBuildType(SBuildProgressWidget::BUILDTYPE_Paths);
 					TriggerNavigationBuilder(InWorld, Id);
@@ -764,4 +791,9 @@ void FEditorBuildUtils::TriggerNavigationBuilder(UWorld* InWorld, EBuildOptions:
 	}
 }
 
+void FEditorBuildUtils::TriggerHierarchicalLODBuilder(UWorld* InWorld, EBuildOptions::Type Id)
+{
+	// Invoke navmesh generator
+	InWorld->HierarchicalLODBuilder.Build();
+}
 #undef LOCTEXT_NAMESPACE
\ No newline at end of file
diff --git a/Engine/Source/Editor/UnrealEd/Private/EditorServer.cpp b/Engine/Source/Editor/UnrealEd/Private/EditorServer.cpp
--- a/Engine/Source/Editor/UnrealEd/Private/EditorServer.cpp
+++ b/Engine/Source/Editor/UnrealEd/Private/EditorServer.cpp
@@ -4866,7 +4866,7 @@ void UEditorEngine::AssignReplacementComponentsByActors(TArray<AActor*>& ActorsT
 				FComponentReregisterContext ComponentReattch(PrimitiveComponent);
 
 				// set the replacement
-				PrimitiveComponent->ReplacementPrimitive = ReplacementComponent;
+				PrimitiveComponent->LODParentPrimitive = ReplacementComponent;
 
 				// makr the package as dirty now that we've modified it
 				Actor->MarkPackageDirty();
@@ -5770,7 +5770,7 @@ bool UEditorEngine::HandleSetReplacementCommand( const TCHAR* Str, FOutputDevice
 			// then set it's replacement component
 			if (ClassToReplace == NULL || PrimitiveComponent->IsA(ClassToReplace))
 			{
-				PrimitiveComponent->ReplacementPrimitive = ReplacementComponent;
+				PrimitiveComponent->SetLODParentPrimitive(ReplacementComponent);
 			}
 		}
 	}
diff --git a/Engine/Source/Editor/UnrealEd/Private/HierarchicalLOD.cpp b/Engine/Source/Editor/UnrealEd/Private/HierarchicalLOD.cpp
new file mode 100644
--- /dev/null
+++ b/Engine/Source/Editor/UnrealEd/Private/HierarchicalLOD.cpp
@@ -0,0 +1,637 @@
+// Copyright 1998-2015 Epic Games, Inc. All Rights Reserved.
+
+/*=============================================================================
+	World.cpp: UWorld implementation
+=============================================================================*/
+
+#include "UnrealEd.h"
+#include "HierarchicalLOD.h"
+
+#if WITH_EDITOR
+#include "Editor/UnrealEd/Classes/Editor/EditorEngine.h"
+#include "Editor/UnrealEd/Public/Editor.h"
+#include "Engine/LODActor.h"
+#include "Editor/UnrealEd/Classes/Factories/Factory.h"
+#include "MeshUtilities.h"
+#include "ObjectTools.h"
+
+#endif // WITH_EDITOR
+
+DEFINE_LOG_CATEGORY_STATIC(LogLODGenerator, Log, All);
+
+#define LOCTEXT_NAMESPACE "HierarchicalLOD"
+
+/*-----------------------------------------------------------------------------
+	HierarchicalLOD implementation.
+-----------------------------------------------------------------------------*/
+
+////////////////////////////////////////////////////////////////////////////////////////////
+// Hierarchical LOD System implementation 
+////////////////////////////////////////////////////////////////////////////////////////////
+
+// utility functions
+float CalculateFillingFactor(const FBox& ABox, const float AFillingFactor, const FBox& BBox, const float BFillingFactor)
+{
+	FBox OverlapBox = ABox.Overlap(BBox);
+	FBox UnionBox = ABox + BBox;
+	// it shouldn't be zero or it should be checked outside
+	ensure(UnionBox.GetSize() != FVector::ZeroVector);
+
+	// http://deim.urv.cat/~rivi/pub/3d/icra04b.pdf
+	// cost is calculated based on r^3 / filling factor
+	// since it subtract by AFillingFactor * 1/2 overlap volume + BfillingFactor * 1/2 overlap volume
+	return (AFillingFactor * ABox.GetVolume() + BFillingFactor * BBox.GetVolume() - 0.5f * (AFillingFactor+BFillingFactor) * OverlapBox.GetVolume()) / UnionBox.GetVolume();
+}
+
+float GetDistance(float AMin, float AMax, float BMin, float BMax)
+{
+	if(AMin > BMax)
+	{
+		return (AMin-BMax);
+	}
+
+	if(BMin > AMax)
+	{
+		return (BMin-AMax);
+	}
+
+	return 0.f;
+}
+// utility functions
+float CalculateDistance(const FLODCluster& A, const FLODCluster& B)
+{
+	FVector AMax = A.Bound.Max;
+	FVector AMin = A.Bound.Min;
+	FVector BMax = B.Bound.Max;
+	FVector BMin = B.Bound.Min;
+	FVector DistanceVector = FVector::ZeroVector;
+
+	DistanceVector.X = GetDistance(AMin.X, AMax.X, BMin.X, BMax.X);
+	DistanceVector.Y = GetDistance(AMin.Y, AMax.Y, BMin.Y, BMax.Y);
+	DistanceVector.Z = GetDistance(AMin.Z, AMax.Z, BMin.Z, BMax.Z);
+
+	return DistanceVector.SizeSquared();
+}
+
+////////////////////////////////////////////////////////////////////////////////////////////
+// Hierarchical LOD Builder System implementation 
+////////////////////////////////////////////////////////////////////////////////////////////
+
+FHierarchicalLODBuilder::FHierarchicalLODBuilder(class UWorld* InWorld)
+:	World(InWorld)
+{}
+
+// build hierarchical cluster
+void FHierarchicalLODBuilder::Build()
+{
+	check (World);
+
+	const TArray<class ULevel*>& Levels = World->GetLevels();
+
+	for (const auto& LevelIter : Levels)
+	{
+		BuildClusters(LevelIter);
+	}
+}
+
+
+// for now it only builds per level, it turns to one actor at the end
+void FHierarchicalLODBuilder::BuildClusters(class ULevel* InLevel)
+{
+	// you still have to delete all objects just in case they had it and didn't want it anymore
+	TArray<UObject*> AssetsToDelete;
+	for (int32 ActorId=InLevel->Actors.Num()-1; ActorId >= 0; --ActorId)
+	{
+		ALODActor* LodActor = Cast<ALODActor>(InLevel->Actors[ActorId]);
+		if (LodActor)
+		{
+			for (auto& Asset: LodActor->SubObjects)
+			{
+				AssetsToDelete.Add(Asset);
+			}
+			World->DestroyActor(LodActor);
+		}
+	}
+
+	for (auto& Asset : AssetsToDelete)
+	{
+		ObjectTools::DeleteSingleObject(Asset, false);
+	}
+
+	// only build if it's enabled
+	if(InLevel->GetWorld()->GetWorldSettings()->bEnableHierarchicalLODSystem)
+	{
+		const int32 TotalNumLOD = InLevel->GetWorld()->GetWorldSettings()->HierarchicalLODSetup.Num();
+
+		for(int32 LODId=0; LODId<TotalNumLOD; ++LODId)
+		{
+			// initialize Clusters
+			InitializeClusters(InLevel, LODId);
+
+			// now we have all pair of nodes
+			TArray<FLODCluster> MSTClusters;
+			FindMST(MSTClusters);
+			// replace Clusters to new culled cluster
+			Clusters = MSTClusters;
+
+			// now we have to calculate merge
+			MergeClusters(InLevel, LODId);
+		}
+	}
+}
+
+void FHierarchicalLODBuilder::FindMST(TArray<FLODCluster>& OutClusters) 
+{
+	if (Clusters.Num() > 0)
+	{
+		// now sort edge in the order of weight
+		struct FCompareCluster
+		{
+			FORCEINLINE bool operator()(const FLODCluster& A, const FLODCluster& B) const
+			{
+				return (A.GetCost() < B.GetCost());
+			}
+		};
+
+		Clusters.Sort(FCompareCluster());
+
+		// now start from small to large, find the MST
+		OutClusters.Empty();
+
+		TArray<class AActor*> Processed;
+		for(auto& Cluster: Clusters)
+		{
+			if (Cluster.IsValid())
+			{
+				bool bNeedToAdd=false;
+
+				TArray<class AActor*>& Actors=Cluster.Actors;
+				for(auto& Actor : Actors)
+				{
+					if(!Processed.Contains(Actor))
+					{
+						bNeedToAdd = true;
+						Processed.Add(Actor);
+					}
+				}
+				if(bNeedToAdd)
+				{
+					OutClusters.Add(Cluster);
+				}
+			}
+		}
+	}
+}
+
+void FHierarchicalLODBuilder::MergeClusters(class ULevel* InLevel, const int32 LODIdx)
+{
+	if (Clusters.Num() > 0)
+	{
+		int32 LastIndex = Clusters.Num()-1; //FMath::TruncToInt(Clusters.Num() * 0.70); /*(10% ratio)*/
+		AWorldSettings* WorldSetting = InLevel->GetWorld()->GetWorldSettings();
+		check (WorldSetting->HierarchicalLODSetup.IsValidIndex(LODIdx));
+		float DesiredBoundBoxHalfSize = WorldSetting->HierarchicalLODSetup[LODIdx].DesiredBoundBoxSize * 0.5f;
+		float DesiredFillingRatio = WorldSetting->HierarchicalLODSetup[LODIdx].DesiredFillingPercentage * 0.01f;
+		ensure (DesiredFillingRatio!=0.f);
+		float HighestCost = FMath::Pow(DesiredBoundBoxHalfSize, 3)/(DesiredFillingRatio);
+
+		const int32 TotalCluster = Clusters.Num();
+		// now we have minimum Clusters
+		for(int32 ClusterId=0; ClusterId < TotalCluster; ++ClusterId)
+		{
+			auto& Cluster = Clusters[ClusterId];
+			UE_LOG(LogLODGenerator, Log, TEXT("%d. %0.2f {%s}"), ClusterId+1, Cluster.GetCost(), *Cluster.ToString());
+
+			for (int32 MergedClusterId=0; MergedClusterId < ClusterId; ++MergedClusterId)
+			{ 
+				// compare with previous clusters
+				auto& MergedCluster = Clusters[MergedClusterId];
+				// see if it's valid, if it contains, check the cost
+				if (MergedCluster.IsValid() && Cluster.IsValid())
+				{
+					// if valid, see if it contains any of this actors
+					// merge whole clusters
+					FLODCluster NewCluster = Cluster + MergedCluster;
+					float MergeCost = NewCluster.GetCost();
+
+					// merge two clusters
+					if (MergeCost <= HighestCost)
+					{
+						UE_LOG(LogLODGenerator, Log, TEXT("Merging of Cluster (%d) and (%d) with merge cost (%0.2f) "), ClusterId+1, MergedClusterId+1, MergeCost);
+						
+						MergedCluster = NewCluster;
+						// now this cluster is invalid
+						Cluster.Invalidate();
+						break;
+					}
+					else if (Cluster.Contains(MergedCluster))
+					{
+						UE_LOG(LogLODGenerator, Log, TEXT("Removing actor (%d) from Cluster (%d)"), MergedClusterId+1, ClusterId+1);
+						// if higher than high cost, then we break them apart
+						Cluster -= MergedCluster;
+						break;
+					}
+				}
+			}
+
+			UE_LOG(LogLODGenerator, Log, TEXT("Processed(%s): %0.2f {%s}"), Cluster.IsValid()? TEXT("Valid"):TEXT("Invalid"), Cluster.GetCost(), *Cluster.ToString());
+		}
+
+		static bool bBuildActor=true;
+
+		GWarn->BeginSlowTask(LOCTEXT("HierarchicalLOD_CreatingSimpleMesh", "Creating Simple Mesh"), true);
+		GEditor->BeginTransaction(LOCTEXT("HierarchicalLOD_CreatingSimpleMesh", "Creating Simple Mesh"));
+		// print data
+		int32 Index=0, TotalValidCluster=0;
+		for (auto& Cluster: Clusters)
+		{
+			if (Cluster.IsValid())
+			{
+				++TotalValidCluster;
+			}
+		}
+
+		for(auto& Cluster: Clusters)
+		{
+			if(Cluster.IsValid())
+			{
+				GWarn->StatusUpdate(Index, TotalValidCluster, LOCTEXT("HierarchicalLOD_BuildActor", "Building Actor"));
+				UE_LOG(LogLODGenerator, Log, TEXT("%d. %0.2f (%s)"), ++Index, Cluster.GetCost(), *Cluster.ToString());
+				if (bBuildActor)
+				{
+					Cluster.BuildActor(InLevel->GetWorld(), InLevel, LODIdx);
+				}
+			}
+		}
+
+		GEditor->EndTransaction();
+		GWarn->EndSlowTask();
+	}
+}
+
+bool ShouldGenerateCluster(class AActor* Actor)
+{
+	if (!Actor)
+	{
+		return false;
+	}
+
+	if (!Actor->bEnableAutoLODGeneration)
+	{
+		return false;
+	}
+
+	// make sure we don't generate for LODActor for Level 0
+	if (Actor->IsA(ALODActor::StaticClass()))
+	{
+		return false;
+	}
+
+	FVector Origin, Extent;
+	Actor->GetActorBounds(false, Origin, Extent);
+	if(Extent.SizeSquared() <= 0.1)
+	{
+		return false;
+	}
+
+	// for now only consider staticmesh - I don't think skel mesh would work with simplygon merge right now @fixme
+	TArray<UStaticMeshComponent*> Components;
+	Actor->GetComponents<UStaticMeshComponent>(Components);
+	// TODO: support instanced static meshes
+	Components.RemoveAll([](UStaticMeshComponent* Val){ return Val->IsA(UInstancedStaticMeshComponent::StaticClass()); });
+
+	return (Components.Num() > 0);
+}
+
+void FHierarchicalLODBuilder::InitializeClusters(class ULevel* InLevel, const int32 LODIdx)
+{
+	if (InLevel->Actors.Num() > 0)
+	{
+		if (LODIdx == 0)
+		{
+			Clusters.Empty(InLevel->Actors.Num());
+
+			// first we generate graph with 2 pair nodes
+			// this is very expensive when we have so many actors
+			// so we'll need to optimize later @todo
+			for(int32 ActorId=0; ActorId<InLevel->Actors.Num(); ++ActorId)
+			{
+				AActor* Actor1 = InLevel->Actors[ActorId];
+				// make sure Actor has bound
+				if(ShouldGenerateCluster(Actor1))
+				{
+					for(int32 SubActorId=ActorId+1; SubActorId<InLevel->Actors.Num(); ++SubActorId)
+					{
+						AActor* Actor2 = InLevel->Actors[SubActorId];
+						if(ShouldGenerateCluster(Actor2))
+						{
+							FLODCluster NewClusterCandidate = FLODCluster(Actor1, Actor2);
+							// @todo optimization : filter first - this might create disconnected tree - 
+							// minimal it should be filled up 20% even to consider that
+							// if (NewClusterCandidate.GetFillingFactor() > 0.2f)
+							Clusters.Add(NewClusterCandidate);
+						}
+					}
+				}
+			}
+		}
+		else // at this point we only care for LODActors
+		{
+			Clusters.Empty();
+			TArray<ALODActor*> LODActors;
+			for(int32 ActorId=0; ActorId<InLevel->Actors.Num(); ++ActorId)
+			{
+				ALODActor* Actor = Cast<ALODActor>(InLevel->Actors[ActorId]);
+				if (Actor)
+				{
+					LODActors.Add(Actor);
+				}
+			}
+			// first we generate graph with 2 pair nodes
+			// this is very expensive when we have so many actors
+			// so we'll need to optimize later @todo
+			for(int32 ActorId=0; ActorId<LODActors.Num(); ++ActorId)
+			{
+				ALODActor* Actor1 = (LODActors[ActorId]);
+				if (Actor1->LODLevel == LODIdx)
+				{
+					for(int32 SubActorId=ActorId+1; SubActorId<LODActors.Num(); ++SubActorId)
+					{
+						ALODActor* Actor2 = LODActors[SubActorId];
+
+						if (Actor2->LODLevel == LODIdx)
+						{
+							// create new cluster
+							FLODCluster NewClusterCandidate = FLODCluster(Actor1, Actor2);
+							// @todo optimization : filter first - this might create disconnected tree - 
+							// minimal it should be filled up 20% even to consider that
+							// if (NewClusterCandidate.GetFillingFactor() > 0.2f)
+							Clusters.Add(NewClusterCandidate);
+						}
+					}
+				}
+			}
+		}
+	}
+}
+
+////////////////////////////////////////////////////////////////////////////////////////////
+// LODCluster
+////////////////////////////////////////////////////////////////////////////////////////////
+
+
+FLODCluster::FLODCluster(const FLODCluster& Other)
+: Bound(Other.Bound)
+, Actors(Other.Actors)
+, FillingFactor(Other.FillingFactor)
+, bValid(Other.bValid)
+{
+}
+
+FLODCluster::FLODCluster(class AActor* Actor1, class AActor* Actor2)
+: Bound(ForceInit)
+, bValid(true)
+{
+	FBox Actor1Bound = AddActor(Actor1);
+	FBox Actor2Bound = AddActor(Actor2);
+
+	// calculate new filling factor
+	FillingFactor = CalculateFillingFactor(Actor1Bound, 1.f, Actor2Bound, 1.f);
+}
+
+FBox FLODCluster::AddActor(class AActor* NewActor)
+{
+	ensure (Actors.Contains(NewActor) == false);
+	Actors.Add(NewActor);
+	FVector Origin, Extent;
+	NewActor->GetActorBounds(false, Origin, Extent);
+	FBox NewBound = FBox(Origin-Extent, Origin+Extent);
+	Bound += NewBound;
+
+	return NewBound;
+}
+
+FLODCluster FLODCluster::operator+(const FLODCluster& Other) const
+{
+	FLODCluster UnionCluster(*this);
+	UnionCluster.MergeClusters(Other);
+	return UnionCluster;
+}
+
+FLODCluster FLODCluster::operator+=(const FLODCluster& Other)
+{
+	MergeClusters(Other);
+	return *this;
+}
+
+FLODCluster FLODCluster::operator-(const FLODCluster& Other) const
+{
+	FLODCluster Cluster(*this);
+	Cluster.SubtractCluster(Other);
+	return *this;
+}
+
+FLODCluster FLODCluster::operator-=(const FLODCluster& Other)
+{
+	SubtractCluster(Other);
+	return *this;
+}
+
+FLODCluster& FLODCluster::operator=(const FLODCluster& Other)
+{
+	this->bValid		= Other.bValid;
+	this->Actors		= Other.Actors;
+	this->Bound			= Other.Bound;
+	this->FillingFactor = Other.FillingFactor;
+
+	return *this;
+}
+// Merge two clusters
+void FLODCluster::MergeClusters(const FLODCluster& Other)
+{
+	// please note that when merge, we merge two boxes from each cluster, not exactly all actors' bound
+	// have to recalculate filling factor and bound based on cluster data
+	FillingFactor = CalculateFillingFactor(Bound, FillingFactor, Other.Bound, Other.FillingFactor);
+	Bound += Other.Bound;
+
+	for (auto& Actor: Other.Actors)
+	{
+		Actors.AddUnique(Actor);
+	}
+}
+
+void FLODCluster::SubtractCluster(const FLODCluster& Other)
+{
+	for(int32 ActorId=0; ActorId<Actors.Num(); ++ActorId)
+	{
+		if (Other.Actors.Contains(Actors[ActorId]))
+		{
+			Actors.RemoveAt(ActorId);
+			--ActorId;
+		}
+	}
+
+	TArray<AActor*> NewActors = Actors;
+	Actors.Empty();
+	// need to recalculate parameter
+	if (NewActors.Num() == 0)
+	{
+		Invalidate();
+	}
+	else if (NewActors.Num() == 1)
+	{
+		Bound = FBox(ForceInitToZero);
+		AddActor(NewActors[0]);
+		FillingFactor = 1.f;
+	}
+	else if (NewActors.Num() == 2)
+	{
+		Bound = FBox(ForceInitToZero);
+
+		FBox Actor1Bound = AddActor(NewActors[0]);
+		FBox Actor2Bound = AddActor(NewActors[1]);
+
+		// calculate new filling factor
+		FillingFactor = CalculateFillingFactor(Actor1Bound, 1.f, Actor2Bound, 1.f);
+	}
+	else
+	{
+		// this shouldn't happen
+		// if this has to happen, we'll either have to come up with way to subtrac bound or recalculate from scratch
+		ensure (false);
+	}
+}
+
+// invalidate this cluster - happens when this merged to another cluster
+void FLODCluster::Invalidate()
+{
+	bValid = false;
+}
+
+bool FLODCluster::IsValid() const
+{
+	return bValid;
+}
+
+const float FLODCluster::GetCost() const
+{
+	FVector BoundExtent = Bound.GetExtent();
+	const float Extent = BoundExtent.GetMax();
+	return (FMath::Pow(Extent, 3) / FillingFactor);
+}
+// if criteria matches, build new LODActor and replace current Actors with that. We don't need 
+// this clears previous actors and sets to this new actor
+// this is required when new LOD is created from these actors, this will be replaced
+// to save memory and to reduce memory increase during this process, we discard previous actors and replace with this actor
+void FLODCluster::BuildActor(class UWorld* InWorld, class ULevel* InLevel, const int32 LODIdx)
+{
+	// do big size
+	if (InWorld && InLevel)
+	{
+		////////////////////////////////////////////////////////////////////////////////////
+		// create asset using Actors
+		const FHierarchicalSimplification& LODSetup = InWorld->GetWorldSettings()->HierarchicalLODSetup[LODIdx];
+		const FMeshProxySettings& ProxySetting = LODSetup.Setting;
+		// Where generated assets will be stored
+		UPackage* AssetsOuter = InLevel->GetOutermost(); // this asset is going to save with map, this means, I'll have to delete with it
+		if (AssetsOuter)
+		{
+			TArray<UStaticMeshComponent*> AllComponents;
+
+			for (auto& Actor: Actors)
+			{
+				TArray<UStaticMeshComponent*> Components;
+				Actor->GetComponents<UStaticMeshComponent>(Components);
+				// TODO: support instanced static meshes
+				Components.RemoveAll([](UStaticMeshComponent* Val){ return Val->IsA(UInstancedStaticMeshComponent::StaticClass()); });
+				AllComponents.Append(Components);
+			}
+
+			// it shouldn't even have come here if it didn't have any staticmesh
+			if (ensure(AllComponents.Num() > 0))
+			{
+				// In case we don't have outer generated assets should have same path as LOD level
+				const FString AssetsPath = FPackageName::GetLongPackagePath(AssetsOuter->GetName()) + TEXT("/");
+				class AActor* FirstActor = Actors[0];
+
+				TArray<UObject*> OutAssets;
+				FVector OutProxyLocation = FVector::ZeroVector;
+				// Generate proxy mesh and proxy material assets
+				IMeshUtilities& MeshUtilities = FModuleManager::Get().LoadModuleChecked<IMeshUtilities>("MeshUtilities");
+				// should give unique name, so use level + actor name
+				const FString PackageName = FString::Printf(TEXT("LOD_%s_%s"), *InLevel->GetName(), *FirstActor->GetName());
+				if (LODSetup.bSimplifyMesh)
+				{
+					MeshUtilities.CreateProxyMesh(Actors, ProxySetting, AssetsOuter, PackageName, OutAssets, OutProxyLocation);
+				}
+				else
+				{
+					FMeshMergingSettings MergeSetting;
+					MeshUtilities.MergeActors(Actors, MergeSetting, AssetsPath, OutAssets, OutProxyLocation );
+				}
+
+				UStaticMesh* MainMesh=NULL;
+				// set staticmesh
+				for(auto& Asset: OutAssets)
+				{
+					class UStaticMesh* StaticMesh = Cast<UStaticMesh>(Asset);
+
+					if(StaticMesh)
+					{
+						MainMesh = StaticMesh;
+					}
+				}
+
+				if (MainMesh)
+				{
+					////////////////////////////////////////////////////////////////////////////////////
+					// create LODActors using the current Actors
+					class ALODActor* NewActor = Cast<ALODActor>(InWorld->SpawnActor(ALODActor::StaticClass(), &OutProxyLocation, &FRotator::ZeroRotator));
+
+					NewActor->SubObjects = OutAssets;
+					NewActor->SubActors = Actors;
+					NewActor->LODLevel = LODIdx+1; 
+					float DrawDistance = LODSetup.DrawDistance;
+
+					NewActor->GetStaticMeshComponent()->StaticMesh = MainMesh;
+
+					// now set as parent
+					for(auto& Actor: Actors)
+					{
+						Actor->SetLODParent(NewActor->GetStaticMeshComponent(), DrawDistance);
+					}
+				}
+			}
+		}
+	}
+}
+
+bool FLODCluster::Contains(FLODCluster& Other) const
+{
+	if (IsValid() && Other.IsValid())
+	{
+		for(auto& Actor: Other.Actors)
+		{
+			if(Actors.Contains(Actor))
+			{
+				return true;
+			}
+		}
+	}
+
+	return false;
+}
+
+FString FLODCluster::ToString() const
+{
+	FString ActorList;
+	for (auto& Actor: Actors)
+	{
+		ActorList += Actor->GetActorLabel();
+		ActorList += ", ";
+	}
+
+	return FString::Printf(TEXT("ActorNum(%d), Actor List (%s)"), Actors.Num(), *ActorList);
+}
+
+#undef LOCTEXT_NAMESPACE 
diff --git a/Engine/Source/Editor/UnrealEd/Public/Dialogs/SBuildProgress.h b/Engine/Source/Editor/UnrealEd/Public/Dialogs/SBuildProgress.h
--- a/Engine/Source/Editor/UnrealEd/Public/Dialogs/SBuildProgress.h
+++ b/Engine/Source/Editor/UnrealEd/Public/Dialogs/SBuildProgress.h
@@ -36,6 +36,8 @@ public:
 		BUILDTYPE_Lighting,
 		/** Paths are being built. */
 		BUILDTYPE_Paths,
+		/** LODs are being built */
+		BUILDTYPE_LODs, 
 	};
 
 	/** The various issues that can occur. */
diff --git a/Engine/Source/Editor/UnrealEd/Public/EditorBuildUtils.h b/Engine/Source/Editor/UnrealEd/Public/EditorBuildUtils.h
--- a/Engine/Source/Editor/UnrealEd/Public/EditorBuildUtils.h
+++ b/Engine/Source/Editor/UnrealEd/Public/EditorBuildUtils.h
@@ -27,6 +27,8 @@ namespace EBuildOptions
 		BuildAllSubmit,
 		/** Build everything except for paths only build seleced */
 		BuildAllOnlySelectedPaths,
+		/** Build Hierarchical LOD system - need WorldSetting setup*/
+		BuildHierarchicalLOD, 
 	};
 }
 
@@ -163,6 +165,14 @@ private:
 	 */
 	static void TriggerNavigationBuilder(UWorld* InWorld, EBuildOptions::Type Id);
 
+	/** 
+	 * Trigger LOD builder to (re)generate LODActors
+	 *
+	 * @param	InWorld			WorldContext
+	 * @param	BuildSettings	Build settings that will be used for the editor build
+	 */
+	static void TriggerHierarchicalLODBuilder(UWorld* InWorld, EBuildOptions::Type Id);
+
 	/** Intentionally hide constructors, etc. to prevent instantiation */
 	FEditorBuildUtils();
 	~FEditorBuildUtils();
diff --git a/Engine/Source/Editor/UnrealEd/Public/HierarchicalLOD.h b/Engine/Source/Editor/UnrealEd/Public/HierarchicalLOD.h
new file mode 100644
--- /dev/null
+++ b/Engine/Source/Editor/UnrealEd/Public/HierarchicalLOD.h
@@ -0,0 +1,74 @@
+// Copyright 1998-2015 Epic Games, Inc. All Rights Reserved.
+
+#pragma once
+
+/*=============================================================================
+	HierarchicalLOD.h: Hierarchical LOD definition.
+=============================================================================*/
+
+
+// implementation 
+//  http://deim.urv.cat/~rivi/pub/3d/icra04b.pdf
+
+struct FLODCluster
+{
+	FLODCluster(const FLODCluster& Other);
+	FLODCluster(class AActor* Actor1, class AActor* Actor2);
+
+	FLODCluster operator+( const FLODCluster& Other ) const;
+	FLODCluster operator+=( const FLODCluster& Other );
+	FLODCluster operator-(const FLODCluster& Other) const;
+	FLODCluster operator-=(const FLODCluster& Other);
+	FLODCluster& operator=(const FLODCluster & Other);
+	// invalidate this cluster - happens when this merged to another cluster
+	void Invalidate();
+	bool IsValid() const;
+
+	const float GetCost() const;
+	bool Contains(FLODCluster& Other) const;
+	FString ToString() const;
+	
+	TArray<class AActor*>	Actors;
+	FBox					Bound;
+	float					FillingFactor;
+
+	// we don't delete the invalid entry, but mark as invalid when it's not any more valid
+	// it's initialized to be valid
+	bool bValid;
+
+	// if criteria matches, build new LODActor and replace current Actors with that. We don't need 
+	// this clears previous actors and sets to this new actor
+	// this is required when new LOD is created from these actors, this will be replaced
+	// to save memory and to reduce memory increase during this process, we discard previous actors and replace with this actor
+	void BuildActor(class UWorld* InWorld, class ULevel* InLevel, const int32 LODIdx);
+	// Merge two clusters
+	void MergeClusters(const FLODCluster& Other);
+	void SubtractCluster(const FLODCluster& Other);
+	FBox AddActor(class AActor* NewActor);
+};
+
+struct UNREALED_API FHierarchicalLODBuilder
+{
+	FHierarchicalLODBuilder(class UWorld* InWorld);
+
+	// build hierarchical cluster
+	void Build();
+
+private:
+	// data structure
+	TArray<FLODCluster> Clusters;
+	class UWorld*		World;
+
+	// for now it only builds per level, it turns to one actor at the end
+	void BuildClusters(class ULevel* InLevel);
+
+	// initialize Clusters variable - each Actor becomes Cluster
+	void InitializeClusters(class ULevel* InLevel, const int32 LODIdx);
+
+	// Refresh Cluster Distance for ClusterId
+	//void RefreshClusterDistance(int32 ClusterId, const int32 TotalNumClusters, float** InDistanceCacheMatrix);
+
+	void MergeClusters(class ULevel* InLevel, const int32 LODIdx);
+
+	void FindMST(TArray<FLODCluster>& OutClusters);
+};
diff --git a/Engine/Source/Runtime/Core/Private/Math/Box.cpp b/Engine/Source/Runtime/Core/Private/Math/Box.cpp
--- a/Engine/Source/Runtime/Core/Private/Math/Box.cpp
+++ b/Engine/Source/Runtime/Core/Private/Math/Box.cpp
@@ -178,3 +178,27 @@ FBox FBox::TransformProjectBy(const FMatrix& ProjM) const
 
 	return NewBox;
 }
+
+FBox FBox::Overlap( const FBox& Other ) const
+{
+	if(Intersect(Other) == false)
+	{
+		static FBox EmptyBox(ForceInit);
+		return EmptyBox;
+	}
+
+	// otherwise they overlap
+	// so find overlapping box
+	FVector MinVector, MaxVector;
+
+	MinVector.X = FMath::Max(Min.X, Other.Min.X);
+	MaxVector.X = FMath::Min(Max.X, Other.Max.X);
+
+	MinVector.Y = FMath::Max(Min.Y, Other.Min.Y);
+	MaxVector.Y = FMath::Min(Max.Y, Other.Max.Y);
+
+	MinVector.Z = FMath::Max(Min.Z, Other.Min.Z);
+	MaxVector.Z = FMath::Min(Max.Z, Other.Max.Z);
+
+	return FBox(MinVector, MaxVector);
+}
diff --git a/Engine/Source/Runtime/Core/Public/Math/Box.h b/Engine/Source/Runtime/Core/Public/Math/Box.h
--- a/Engine/Source/Runtime/Core/Public/Math/Box.h
+++ b/Engine/Source/Runtime/Core/Public/Math/Box.h
@@ -308,6 +308,14 @@ public:
 	 */
 	FORCEINLINE bool IntersectXY( const FBox& Other ) const;
 
+	/**
+	 * Returns the overlap FBox of two box
+	 *
+	 * @param Other The bounding box to test overlap
+	 * @return the overlap box. It can be 0 if they don't overlap
+	 */
+	CORE_API FBox Overlap( const FBox& Other ) const;
+
 	/**
 	  * Gets a bounding volume transformed by an inverted FTransform object.
 	  *
diff --git a/Engine/Source/Runtime/Engine/Classes/Components/PrimitiveComponent.h b/Engine/Source/Runtime/Engine/Classes/Components/PrimitiveComponent.h
--- a/Engine/Source/Runtime/Engine/Classes/Components/PrimitiveComponent.h
+++ b/Engine/Source/Runtime/Engine/Classes/Components/PrimitiveComponent.h
@@ -916,8 +916,12 @@ public:
 	// Scene data
 
 	/** Replacement primitive to draw instead of this one (multiple UPrim's will point to the same Replacement) */
-	UPROPERTY()
-	TLazyObjectPtr<class UPrimitiveComponent> ReplacementPrimitive;
+	UPROPERTY(duplicatetransient)
+	TLazyObjectPtr<class UPrimitiveComponent> LODParentPrimitive;
+
+public:
+	UFUNCTION(BlueprintCallable, Category="Rendering|LOD")
+	void SetLODParentPrimitive(UPrimitiveComponent * InLODParentPrimitive);
 
 #if WITH_EDITOR
 	virtual const int32 GetNumUncachedStaticLightingInteractions() const; // recursive function
diff --git a/Engine/Source/Runtime/Engine/Classes/Engine/LODActor.h b/Engine/Source/Runtime/Engine/Classes/Engine/LODActor.h
new file mode 100644
--- /dev/null
+++ b/Engine/Source/Runtime/Engine/Classes/Engine/LODActor.h
@@ -0,0 +1,60 @@
+// Copyright 1998-2015 Epic Games, Inc. All Rights Reserved.
+
+#pragma once
+#include "LODActor.generated.h"
+
+/**
+ * LODActor is an instance of an autogenerated StaticMesh Actors by Hierarchical LOD System
+ * This is essentially just StaticMeshActor that you can't move or edit, but it contains multiple actors reference
+ *
+ * @see https://docs.unrealengine.com/latest/INT/Engine/Actors/LODActor/
+ * @see UStaticMesh
+ */
+UCLASS()
+class ENGINE_API ALODActor : public AActor
+{
+	GENERATED_UCLASS_BODY()
+
+private_subobject:
+	UPROPERTY(Category=LODActor, VisibleAnywhere)
+	class UStaticMeshComponent* StaticMeshComponent;
+
+public:
+	UPROPERTY(Category=LODActor, VisibleAnywhere)
+	TArray<class AActor*> SubActors;
+
+	/** what distance do you want this to show up instead of SubActors */
+	UPROPERTY(Category=LODActor, VisibleAnywhere)
+	float LODDrawDistance;
+
+	/** what distance do you want this to show up instead of SubActors */
+	UPROPERTY(Category=LODActor, VisibleAnywhere)
+	int32 LODLevel;
+
+	/** assets that were created for this, so that we can delete them */
+	UPROPERTY(Category=LODActor, VisibleAnywhere)
+	TArray<class UObject*> SubObjects;
+
+#if WITH_EDITOR
+	// Begin AActor Interface
+	virtual void CheckForErrors() override;
+	virtual bool GetReferencedContentObjects( TArray<UObject*>& Objects ) const override;
+	// End AActor Interface
+#endif // WITH_EDITOR	
+
+protected:
+	// Begin UObject interface.
+	virtual FString GetDetailedInfoInternal() const override;
+#if WITH_EDITOR
+	virtual void PostEditChangeChainProperty(FPropertyChangedChainEvent& PropertyChangedEvent) override;
+#endif // WITH_EDITOR	
+	virtual void PostRegisterAllComponents() override;
+	// End UObject interface.
+
+public:
+	/** Returns StaticMeshComponent subobject **/
+	class UStaticMeshComponent* GetStaticMeshComponent() const { return StaticMeshComponent; }
+};
+
+
+
diff --git a/Engine/Source/Runtime/Engine/Classes/Engine/World.h b/Engine/Source/Runtime/Engine/Classes/Engine/World.h
--- a/Engine/Source/Runtime/Engine/Classes/Engine/World.h
+++ b/Engine/Source/Runtime/Engine/Classes/Engine/World.h
@@ -12,6 +12,11 @@
 #include "Engine/LatentActionManager.h"
 #include "Runtime/RHI/Public/RHIDefinitions.h"
 #include "GameFramework/Actor.h"
+
+#if WITH_EDITOR
+	#include "Editor/UnrealEd/Public/HierarchicalLOD.h"
+#endif // WITH_EDITOR
+
 #include "World.generated.h"
 
 class FPhysScene;
@@ -623,6 +628,11 @@ public:
 	/** Array of actors that are candidates for sending over the network */
 	TArray<class AActor*>						NetworkActors;
 
+#if WITH_EDITOR
+	/** Hierarchical LOD System. Used when WorldSetting.bEnableHierarchicalLODSystem is true */
+	FHierarchicalLODBuilder						HierarchicalLODBuilder;
+#endif // WITH_EDITOR
+
 private:
 
 	/** Pointer to the current level being edited. Level has to be in the Levels array and == PersistentLevel in the game.		*/
diff --git a/Engine/Source/Runtime/Engine/Classes/GameFramework/Actor.h b/Engine/Source/Runtime/Engine/Classes/GameFramework/Actor.h
--- a/Engine/Source/Runtime/Engine/Classes/GameFramework/Actor.h
+++ b/Engine/Source/Runtime/Engine/Classes/GameFramework/Actor.h
@@ -474,6 +474,10 @@ public:
 	UPROPERTY(EditAnywhere, AdvancedDisplay, Category=Actor)
 	uint32 bIgnoresOriginShifting:1;
 	
+	/** If true, and if World setting has bEnableHigerarhicalLOD is true, then it will generate LODActor from groups of clustered Actor */
+	UPROPERTY(EditAnywhere, AdvancedDisplay, Category=Actor)
+	uint32 bEnableAutoLODGeneration:1;
+
 	/** Array of tags that can be used for grouping and categorizing. */
 	UPROPERTY(EditAnywhere, BlueprintReadWrite, Category=Tags)
 	TArray<FName> Tags;
@@ -1326,6 +1330,9 @@ public:
 	/** Called by MirrorActors to perform a mirroring operation on the actor */
 	virtual void EditorApplyMirror(const FVector& MirrorScale, const FVector& PivotLocation);
 
+	/** Set LOD Parent Primitive*/
+	void SetLODParent(class UPrimitiveComponent* InLODParent, float InParentDrawDistance);
+
 	/**
 	 * Simple accessor to check if the actor is hidden upon editor startup
 	 * @return	true if the actor is hidden upon editor startup; false if it is not
diff --git a/Engine/Source/Runtime/Engine/Classes/GameFramework/WorldSettings.h b/Engine/Source/Runtime/Engine/Classes/GameFramework/WorldSettings.h
--- a/Engine/Source/Runtime/Engine/Classes/GameFramework/WorldSettings.h
+++ b/Engine/Source/Runtime/Engine/Classes/GameFramework/WorldSettings.h
@@ -188,6 +188,42 @@ struct ENGINE_API FNetViewer
 	FNetViewer(UNetConnection* InConnection, float DeltaSeconds);
 };
 
+
+USTRUCT()
+struct ENGINE_API FHierarchicalSimplification
+{
+	GENERATED_USTRUCT_BODY()
+
+	/** If this is true, it will simplify mesh. However it is slower. If false, it will just merge actors but not simplify */
+	UPROPERTY(Category=FHierarchicalSimplification, EditAnywhere)
+	bool bSimplifyMesh;
+
+	/** What are the draw distance for this LOD */
+	UPROPERTY(Category=FHierarchicalSimplification, EditAnywhere, AdvancedDisplay, meta=(UIMin=10.f, ClampMin=10.f))
+	float DrawDistance;
+
+	/** Desired Bounding Box Size for clustering */
+	UPROPERTY(EditAnywhere, Category=FHierarchicalSimplification, AdvancedDisplay, meta=(UIMin=10.f, ClampMin=10.f))
+	float DesiredBoundBoxSize;
+
+	/** Desired Filling Percentage for clustering */
+	UPROPERTY(EditAnywhere, Category=FHierarchicalSimplification, AdvancedDisplay, meta=(ClampMin = "0", ClampMax = "100", UIMin = "0", UIMax = "100"))
+	float DesiredFillingPercentage;
+
+	/** Simplification Setting if bSimplifyMesh is true */
+	UPROPERTY(Category=FHierarchicalSimplification, EditAnywhere, AdvancedDisplay, meta=(editcondition = "bSimplifyMesh"))
+	FMeshProxySettings Setting;
+
+	FHierarchicalSimplification()
+		: bSimplifyMesh(true)
+		, DrawDistance(5000)
+		, DesiredBoundBoxSize(10000)
+		, DesiredFillingPercentage(30)
+	{
+		Setting.ScreenSize = 300;
+	}
+};
+
 /**
  * Actor containing all script accessible world properties.
  */
@@ -206,6 +242,11 @@ class ENGINE_API AWorldSettings : public AInfo, public IInterface_AssetUserData
 	UPROPERTY(EditAnywhere, BlueprintReadOnly, config, Category=World, AdvancedDisplay)
 	uint32 bEnableNavigationSystem:1;
 
+#if WITH_EDITORONLY_DATA
+	/** if set to true, hierarchical LODs will be built, which will create hierarchical LODActors*/
+	UPROPERTY(EditAnywhere, config, Category=World, AdvancedDisplay)
+	uint32 bEnableHierarchicalLODSystem:1;
+#endif
 	/** 
 	 * Enables tools for composing a tiled world. 
 	 * Level has to be saved and all sub-levels removed before enabling this option.
@@ -343,6 +384,11 @@ class ENGINE_API AWorldSettings : public AInfo, public IInterface_AssetUserData
 	UPROPERTY(EditAnywhere, Category=Audio)
 	class USoundMix* DefaultBaseSoundMix;
 
+#if WITH_EDITORONLY_DATA
+	/** Hierarchical LOD Setup */
+	UPROPERTY(EditAnywhere, Category=LODSystem, AdvancedDisplay, meta=(editcondition = "bEnableHierarchicalLODSystem"))
+	TArray<struct FHierarchicalSimplification>	HierarchicalLODSetup;
+#endif
 	/************************************/
 	/** DEFAULT SETTINGS **/
 
diff --git a/Engine/Source/Runtime/Engine/Private/Actor.cpp b/Engine/Source/Runtime/Engine/Private/Actor.cpp
--- a/Engine/Source/Runtime/Engine/Private/Actor.cpp
+++ b/Engine/Source/Runtime/Engine/Private/Actor.cpp
@@ -69,6 +69,7 @@ AActor::AActor(const FObjectInitializer& ObjectInitializer)
 	bHiddenEdLevel = false;
 	bActorLabelEditable = true;
 	SpriteScale = 1.0f;
+	bEnableAutoLODGeneration = true;	
 #endif // WITH_EDITORONLY_DATA
 	NetCullDistanceSquared = 225000000.0f;
 	NetDriverName = NAME_GameNetDriver;
diff --git a/Engine/Source/Runtime/Engine/Private/ActorComponent.cpp b/Engine/Source/Runtime/Engine/Private/ActorComponent.cpp
--- a/Engine/Source/Runtime/Engine/Private/ActorComponent.cpp
+++ b/Engine/Source/Runtime/Engine/Private/ActorComponent.cpp
@@ -73,10 +73,10 @@ FGlobalComponentReregisterContext::FGlobalComponentReregisterContext(const TArra
 	{
 		bool bShouldReregister=false;
 		UPrimitiveComponent* PrimitiveComponent = Cast<UPrimitiveComponent>(Component);
-		if (PrimitiveComponent && PrimitiveComponent->ReplacementPrimitive.Get())
+		if (PrimitiveComponent && PrimitiveComponent->LODParentPrimitive.Get())
 		{
-			UPrimitiveComponent* ReplacementPrimitive = PrimitiveComponent->ReplacementPrimitive.Get();
-			AActor* ParentActor = Cast<AActor>(ReplacementPrimitive->GetOuter());
+			UPrimitiveComponent* LODParentPrimitive = PrimitiveComponent->LODParentPrimitive.Get();
+			AActor* ParentActor = Cast<AActor>(LODParentPrimitive->GetOuter());
 			if (ParentActor && InParentActors.Contains(ParentActor))
 			{
 				bShouldReregister = true;
diff --git a/Engine/Source/Runtime/Engine/Private/ActorEditor.cpp b/Engine/Source/Runtime/Engine/Private/ActorEditor.cpp
--- a/Engine/Source/Runtime/Engine/Private/ActorEditor.cpp
+++ b/Engine/Source/Runtime/Engine/Private/ActorEditor.cpp
@@ -698,6 +698,23 @@ bool AActor::GetReferencedContentObjects( TArray<UObject*>& Objects ) const
 	return true;
 }
 
+void AActor::SetLODParent(UPrimitiveComponent* InLODParent, float InParentDrawDistance)
+{
+	if(InLODParent)
+	{
+		InLODParent->MinDrawDistance = InParentDrawDistance;
+		InLODParent->MarkRenderStateDirty();
+	}
+
+	TArray<UPrimitiveComponent*> ComponentsToBeReplaced;
+	GetComponents(ComponentsToBeReplaced);
+
+	for(auto& Component : ComponentsToBeReplaced)
+	{
+		// parent primitive will be null if no LOD parent is selected
+		Component->SetLODParentPrimitive(InLODParent);
+	}
+}
 #undef LOCTEXT_NAMESPACE
 
 #endif // WITH_EDITOR
diff --git a/Engine/Source/Runtime/Engine/Private/LODActor.cpp b/Engine/Source/Runtime/Engine/Private/LODActor.cpp
new file mode 100644
--- /dev/null
+++ b/Engine/Source/Runtime/Engine/Private/LODActor.cpp
@@ -0,0 +1,88 @@
+// Copyright 1998-2015 Epic Games, Inc. All Rights Reserved.
+
+/*=============================================================================
+	LODActorBase.cpp: Static mesh actor base class implementation.
+=============================================================================*/
+
+#include "EnginePrivate.h"
+#include "Engine/LODActor.h"
+#include "MapErrors.h"
+
+#define LOCTEXT_NAMESPACE "LODActor"
+
+ALODActor::ALODActor(const FObjectInitializer& ObjectInitializer)
+	: Super(ObjectInitializer)
+	, LODDrawDistance(5000)
+{
+	bCanBeDamaged = false;
+
+	StaticMeshComponent = CreateDefaultSubobject<UStaticMeshComponent>(TEXT("StaticMeshComponent0"));
+	StaticMeshComponent->SetCollisionProfileName(UCollisionProfile::NoCollision_ProfileName);
+	StaticMeshComponent->Mobility = EComponentMobility::Static;
+	StaticMeshComponent->bGenerateOverlapEvents = false;
+	StaticMeshComponent->bCastDynamicShadow = false;
+	StaticMeshComponent->bCastStaticShadow = false;
+
+	RootComponent = StaticMeshComponent;
+}
+
+FString ALODActor::GetDetailedInfoInternal() const
+{
+	return StaticMeshComponent ? StaticMeshComponent->GetDetailedInfoInternal() : TEXT("No_StaticMeshComponent");
+}
+
+void ALODActor::PostRegisterAllComponents() 
+{
+	Super::PostRegisterAllComponents();
+#if WITH_EDITOR
+	if (StaticMeshComponent->SceneProxy)
+	{
+		ensure (LODLevel >= 1);
+		(StaticMeshComponent->SceneProxy)->SetHierarchicalLOD_GameThread(LODLevel);
+	}
+#endif
+}
+#if WITH_EDITOR
+
+void ALODActor::PostEditChangeChainProperty(FPropertyChangedChainEvent& PropertyChangedEvent)
+{
+	Super::PostEditChangeChainProperty(PropertyChangedEvent);
+}
+
+bool ALODActor::GetReferencedContentObjects( TArray<UObject*>& Objects ) const
+{
+	if (StaticMeshComponent && StaticMeshComponent->StaticMesh)
+	{
+		Objects.Add(StaticMeshComponent->StaticMesh);
+	}
+	return true;
+}
+
+void ALODActor::CheckForErrors()
+{
+	Super::CheckForErrors();
+
+	FMessageLog MapCheck("MapCheck");
+
+	if( !StaticMeshComponent )
+	{
+		MapCheck.Warning()
+			->AddToken(FUObjectToken::Create(this))
+			->AddToken(FTextToken::Create(LOCTEXT( "MapCheck_Message_StaticMeshComponent", "Static mesh actor has NULL StaticMeshComponent property - please delete" ) ))
+			->AddToken(FMapErrorToken::Create(FMapErrors::StaticMeshComponent));
+	}
+	else if( StaticMeshComponent->StaticMesh == NULL )
+	{
+		MapCheck.Warning()
+			->AddToken(FUObjectToken::Create(this))
+			->AddToken(FTextToken::Create(LOCTEXT( "MapCheck_Message_StaticMeshNull", "Static mesh actor has NULL StaticMesh property" ) ))
+			->AddToken(FMapErrorToken::Create(FMapErrors::StaticMeshNull));
+	}
+
+	// @todo error message when missing Actors
+}
+
+#endif // WITH_EDITOR
+
+#undef LOCTEXT_NAMESPACE
+
diff --git a/Engine/Source/Runtime/Engine/Private/PrimitiveComponent.cpp b/Engine/Source/Runtime/Engine/Private/PrimitiveComponent.cpp
--- a/Engine/Source/Runtime/Engine/Private/PrimitiveComponent.cpp
+++ b/Engine/Source/Runtime/Engine/Private/PrimitiveComponent.cpp
@@ -2588,6 +2588,12 @@ void UPrimitiveComponent::SetRenderInMainPass(bool bValue)
 	}
 }
 
+void UPrimitiveComponent::SetLODParentPrimitive(UPrimitiveComponent * InLODParentPrimitive)
+{
+	// @todo, what do we do with old parent. We can't just reset undo parent because the parent migh be used by other primitive
+	LODParentPrimitive = InLODParentPrimitive;
+	MarkRenderStateDirty();
+}
 #if WITH_EDITOR
 const int32 UPrimitiveComponent::GetNumUncachedStaticLightingInteractions() const
 {
diff --git a/Engine/Source/Runtime/Engine/Private/PrimitiveSceneProxy.cpp b/Engine/Source/Runtime/Engine/Private/PrimitiveSceneProxy.cpp
--- a/Engine/Source/Runtime/Engine/Private/PrimitiveSceneProxy.cpp
+++ b/Engine/Source/Runtime/Engine/Private/PrimitiveSceneProxy.cpp
@@ -77,6 +77,7 @@ FPrimitiveSceneProxy::FPrimitiveSceneProxy(const UPrimitiveComponent* InComponen
 ,	ComponentForDebuggingOnly(InComponent)
 #if WITH_EDITOR
 ,	NumUncachedStaticLightingInteractions(0)
+,	HierarchicalLODOverride(0)
 #endif
 {
 	check(Scene);
@@ -334,6 +335,27 @@ void FPrimitiveSceneProxy::SetCollisionEnabled_RenderThread(const bool bNewEnabl
 	bCollisionEnabled = bNewEnabled;
 }
 
+#if WITH_EDITOR
+void FPrimitiveSceneProxy::SetHierarchicalLOD_GameThread(const int32 InLODLevel)
+{
+	check(IsInGameThread());
+
+	// Enqueue a message to the rendering thread to change draw state
+	ENQUEUE_UNIQUE_RENDER_COMMAND_TWOPARAMETER(
+		etHierarchicalLOD,
+		FPrimitiveSceneProxy*, PrimSceneProxy, this,
+		const int32, InLODLevel, InLODLevel,
+		{
+		PrimSceneProxy->SetHierarchicalLOD_RenderThread(InLODLevel);
+	});
+}
+
+void FPrimitiveSceneProxy::SetHierarchicalLOD_RenderThread(const int32 InLODLevel)
+{
+	check(IsInRenderingThread());
+	HierarchicalLODOverride = InLODLevel;
+}
+#endif 
 /** @return True if the primitive is visible in the given View. */
 bool FPrimitiveSceneProxy::IsShown(const FSceneView* View) const
 {
@@ -357,6 +379,11 @@ bool FPrimitiveSceneProxy::IsShown(const FSceneView* View) const
 		{
 			return false;
 		}
+
+		if (View->Family->HierarchicalLODOverride >= 0 && View->Family->HierarchicalLODOverride != HierarchicalLODOverride)
+		{
+			return false;
+		}
 	}
 	else
 #endif
diff --git a/Engine/Source/Runtime/Engine/Private/SceneView.cpp b/Engine/Source/Runtime/Engine/Private/SceneView.cpp
--- a/Engine/Source/Runtime/Engine/Private/SceneView.cpp
+++ b/Engine/Source/Runtime/Engine/Private/SceneView.cpp
@@ -1449,6 +1449,7 @@ FSceneViewFamily::FSceneViewFamily( const ConstructionValues& CVS )
 	}
 
 	LandscapeLODOverride = -1;
+	HierarchicalLODOverride = -1;
 	bDrawBaseInfo = true;
 #endif
 
diff --git a/Engine/Source/Runtime/Engine/Private/World.cpp b/Engine/Source/Runtime/Engine/Private/World.cpp
--- a/Engine/Source/Runtime/Engine/Private/World.cpp
+++ b/Engine/Source/Runtime/Engine/Private/World.cpp
@@ -42,6 +42,7 @@
 #include "GameFramework/PlayerState.h"
 
 #include "Materials/MaterialParameterCollectionInstance.h"
+
 #if WITH_EDITOR
 	#include "Editor/UnrealEd/Public/Kismet2/KismetEditorUtilities.h"
 	#include "Editor/UnrealEd/Public/Kismet2/BlueprintEditorUtils.h"
@@ -108,6 +109,9 @@ UWorld::UWorld( const FObjectInitializer& ObjectInitializer )
 ,   bIsBuilt(false)
 ,	FlushLevelStreamingType(EFlushLevelStreamingType::None)
 ,	NextTravelType(TRAVEL_Relative)
+#if WITH_EDITOR
+,	HierarchicalLODBuilder(this)
+#endif
 {
 	TimerManager = new FTimerManager();
 #if WITH_EDITOR
@@ -127,6 +131,9 @@ UWorld::UWorld( const FObjectInitializer& ObjectInitializer,const FURL& InURL )
 ,   bIsBuilt(false)
 ,	FlushLevelStreamingType(EFlushLevelStreamingType::None)
 ,	NextTravelType(TRAVEL_Relative)
+#if WITH_EDITOR
+,	HierarchicalLODBuilder(this)
+#endif
 {
 	SetFlags( RF_Transactional );
 	TimerManager = new FTimerManager();
diff --git a/Engine/Source/Runtime/Engine/Private/WorldSettings.cpp b/Engine/Source/Runtime/Engine/Private/WorldSettings.cpp
--- a/Engine/Source/Runtime/Engine/Private/WorldSettings.cpp
+++ b/Engine/Source/Runtime/Engine/Private/WorldSettings.cpp
@@ -33,6 +33,12 @@ AWorldSettings::AWorldSettings(const FObjectInitializer& ObjectInitializer)
 	bEnableNavigationSystem = true;
 	bEnableWorldComposition = false;
 	bEnableWorldOriginRebasing = false;
+#if WITH_EDITORONLY_DATA	
+	bEnableHierarchicalLODSystem = false;
+
+ 	FHierarchicalSimplification LODBaseSetup;
+	HierarchicalLODSetup.Add(LODBaseSetup);
+#endif
 
 	KillZ = -HALF_WORLD_MAX1;
 	KillZDamageType = ConstructorStatics.DmgType_Environmental_Object.Object;
diff --git a/Engine/Source/Runtime/Engine/Public/PrimitiveSceneProxy.h b/Engine/Source/Runtime/Engine/Public/PrimitiveSceneProxy.h
--- a/Engine/Source/Runtime/Engine/Public/PrimitiveSceneProxy.h
+++ b/Engine/Source/Runtime/Engine/Public/PrimitiveSceneProxy.h
@@ -367,6 +367,9 @@ public:
 
 #if WITH_EDITOR
 	inline int32 GetNumUncachedStaticLightingInteractions() { return NumUncachedStaticLightingInteractions; }
+
+	void SetHierarchicalLOD_GameThread(const int32 InLODLevel);
+	void SetHierarchicalLOD_RenderThread(const int32 InLODLevel);
 #endif
 
 	inline FLinearColor GetWireframeColor() const { return WireframeColor; }
@@ -653,6 +656,9 @@ private:
 	*/
 	int32 NumUncachedStaticLightingInteractions;
 	friend class FLightPrimitiveInteraction;
+
+	/** this is used if world setting has EnableHierarchical LOD true */
+	int32 HierarchicalLODOverride;
 #endif
 
 	/** Updates the proxy's actor position, called from the game thread. */
diff --git a/Engine/Source/Runtime/Engine/Public/SceneView.h b/Engine/Source/Runtime/Engine/Public/SceneView.h
--- a/Engine/Source/Runtime/Engine/Public/SceneView.h
+++ b/Engine/Source/Runtime/Engine/Public/SceneView.h
@@ -774,6 +774,9 @@ public:
 	// Override the LOD of landscape in this viewport
 	int8 LandscapeLODOverride;
 
+	// Override the LOD of landscape in this viewport
+	int8 HierarchicalLODOverride;
+
 	/** Indicates whether, of not, the base attachment volume should be drawn. */
 	bool bDrawBaseInfo;
 #endif
diff --git a/Engine/Source/Runtime/Renderer/Private/PrimitiveSceneInfo.cpp b/Engine/Source/Runtime/Renderer/Private/PrimitiveSceneInfo.cpp
--- a/Engine/Source/Runtime/Renderer/Private/PrimitiveSceneInfo.cpp
+++ b/Engine/Source/Runtime/Renderer/Private/PrimitiveSceneInfo.cpp
@@ -108,7 +108,7 @@ FPrimitiveSceneInfo::FPrimitiveSceneInfo(UPrimitiveComponent* InComponent,FScene
 	{
 		LightingAttachmentRoot = SearchParentComponent->ComponentId;
 	}
-		
+
 	// Only create hit proxies in the Editor as that's where they are used.
 	if (GIsEditor)
 	{
@@ -119,6 +119,14 @@ FPrimitiveSceneInfo::FPrimitiveSceneInfo(UPrimitiveComponent* InComponent,FScene
 			DefaultDynamicHitProxyId = DefaultDynamicHitProxy->Id;
 		}
 	}
+
+	// set LOD parent info if exists
+	if (InComponent->LODParentPrimitive.IsValid())
+	{
+		// I think this should be checked in editor feature
+		check (InComponent->LODParentPrimitive.Get() != InComponent);
+		LODParentComponentId = InComponent->LODParentPrimitive.Get()->ComponentId;
+	}
 }
 
 FPrimitiveSceneInfo::~FPrimitiveSceneInfo()
@@ -312,6 +320,24 @@ void FPrimitiveSceneInfo::BeginDeferredUpdateStaticMeshes()
 	bNeedsStaticMeshUpdate = true;
 }
 
+void FPrimitiveSceneInfo::LinkLODParentComponent()
+{
+	if (LODParentComponentId.IsValid())
+	{
+		Scene->SceneLODHierarchy.AddChildNode(LODParentComponentId, this);
+	}
+}
+
+void FPrimitiveSceneInfo::UnlinkLODParentComponent()
+{
+	if(LODParentComponentId.IsValid())
+	{
+		Scene->SceneLODHierarchy.RemoveChildNode(LODParentComponentId, this);
+		// I don't think this will be reused but just in case
+		LODParentComponentId = FPrimitiveComponentId();
+	}
+}
+
 void FPrimitiveSceneInfo::LinkAttachmentGroup()
 {
 	// Add the primitive to its attachment group.
diff --git a/Engine/Source/Runtime/Renderer/Private/PrimitiveSceneInfo.h b/Engine/Source/Runtime/Renderer/Private/PrimitiveSceneInfo.h
--- a/Engine/Source/Runtime/Renderer/Private/PrimitiveSceneInfo.h
+++ b/Engine/Source/Runtime/Renderer/Private/PrimitiveSceneInfo.h
@@ -178,6 +178,11 @@ public:
 	 */
 	FPrimitiveComponentId LightingAttachmentRoot;
 
+	/** 
+	 * The component id of the LOD parent if valid.
+	 */
+	FPrimitiveComponentId LODParentComponentId;
+
 	/** The primitive's static meshes. */
 	TIndirectArray<class FStaticMesh> StaticMeshes;
 
@@ -264,6 +269,12 @@ public:
 	/** Removes the primitive's static meshes from the scene. */
 	void RemoveStaticMeshes();
 
+	/** Set LOD Parent primitive information to the scene. */
+	void LinkLODParentComponent();
+
+	/** clear LOD parent primitive information from the scene. */
+	void UnlinkLODParentComponent();
+
 	/** Adds the primitive to the scene's attachment groups. */
 	void LinkAttachmentGroup();
 
diff --git a/Engine/Source/Runtime/Renderer/Private/RendererScene.cpp b/Engine/Source/Runtime/Renderer/Private/RendererScene.cpp
--- a/Engine/Source/Runtime/Renderer/Private/RendererScene.cpp
+++ b/Engine/Source/Runtime/Renderer/Private/RendererScene.cpp
@@ -276,10 +276,20 @@ void FScene::AddPrimitiveSceneInfo_RenderThread(FRHICommandListImmediate& RHICmd
 	// Note: must happen before AddToScene because AddToScene depends on LightingAttachmentRoot
 	PrimitiveSceneInfo->LinkAttachmentGroup();
 
+	// Set lod Parent information if valid
+	PrimitiveSceneInfo->LinkLODParentComponent();
+
 	// Add the primitive to the scene.
 	PrimitiveSceneInfo->AddToScene(RHICmdList, true);
 
 	DistanceFieldSceneData.AddPrimitive(PrimitiveSceneInfo);
+
+	// LOD Parent, if this is LOD parent, we should update Proxy Scene Info
+	// LOD parent gets removed WHEN no children is accessing
+	// LOD parent can be recreated as scene updates
+	// I update if the parent component ID is still valid
+	// @Todo : really remove it if you know this is being destroyed - should happen from game thread as streaming in/out
+	SceneLODHierarchy.UpdateNodeSceneInfo(PrimitiveSceneInfo->PrimitiveComponentId, PrimitiveSceneInfo);
 }
 
 /**
@@ -325,6 +335,7 @@ FScene::FScene(UWorld* InWorld, bool bInRequiresHitProxies, bool bInIsEditorScen
 ,	LowerDynamicSkylightColor(FLinearColor::Black)
 ,	NumVisibleLights(0)
 ,	bHasSkyLight(false)
+,	SceneLODHierarchy(this)
 {
 	check(World);
 	World->Scene = this;
@@ -664,6 +675,9 @@ void FScene::RemovePrimitiveSceneInfo_RenderThread(FPrimitiveSceneInfo* Primitiv
 {
 	SCOPE_CYCLE_COUNTER(STAT_RemoveScenePrimitiveTime);
 
+	// clear it up, parent is getting removed
+	SceneLODHierarchy.UpdateNodeSceneInfo(PrimitiveSceneInfo->PrimitiveComponentId, nullptr);
+
 	CheckPrimitiveArrays();
 
 	int32 PrimitiveIndex = PrimitiveSceneInfo->PackedIndex;
@@ -687,6 +701,9 @@ void FScene::RemovePrimitiveSceneInfo_RenderThread(FPrimitiveSceneInfo* Primitiv
 	// Unlink the primitive from its shadow parent.
 	PrimitiveSceneInfo->UnlinkAttachmentGroup();
 
+	// Unlink the LOD parent info if valid
+	PrimitiveSceneInfo->UnlinkLODParentComponent();
+
 	// Remove the primitive from the scene.
 	PrimitiveSceneInfo->RemoveFromScene(true);
 
diff --git a/Engine/Source/Runtime/Renderer/Private/ScenePrivate.h b/Engine/Source/Runtime/Renderer/Private/ScenePrivate.h
--- a/Engine/Source/Runtime/Renderer/Private/ScenePrivate.h
+++ b/Engine/Source/Runtime/Renderer/Private/ScenePrivate.h
@@ -1161,10 +1161,74 @@ public:
 	TArray<FPrimitiveSceneInfo*> Primitives;
 
 	FAttachmentGroupSceneInfo() :
-		ParentSceneInfo(NULL)
+		ParentSceneInfo(nullptr)
 	{}
 };
 
+class FLODSceneTree
+{
+public:
+	FLODSceneTree(FScene* InScene)
+		: Scene(InScene)
+		, UpdateCount(0)
+	{}
+
+	/** Information about the primitives that are attached together. */
+	struct FLODSceneNode
+	{
+		/** The primitive. */
+		FPrimitiveSceneInfo* SceneInfo;
+
+		/** Children scene infos. */
+		TArray<FPrimitiveSceneInfo*> ChildrenSceneInfos;
+
+		/** Last updated FrameCount */
+		int32 LatestUpdateCount;
+
+		FLODSceneNode() :
+			SceneInfo(nullptr), 
+			LatestUpdateCount(INDEX_NONE)
+		{}
+
+		void AddChild(FPrimitiveSceneInfo * NewChild)
+		{
+			if(NewChild && !ChildrenSceneInfos.Contains(NewChild))
+			{
+				ChildrenSceneInfos.Add(NewChild);
+			}
+		}
+
+		void RemoveChild(FPrimitiveSceneInfo * ChildToDelete)
+		{
+			if(ChildToDelete && ChildrenSceneInfos.Contains(ChildToDelete))
+			{
+				ChildrenSceneInfos.Remove(ChildToDelete);
+			}
+		}
+	};
+
+	void AddChildNode(FPrimitiveComponentId NodeId, FPrimitiveSceneInfo* ChildSceneInfo);
+	void RemoveChildNode(FPrimitiveComponentId NodeId, FPrimitiveSceneInfo* ChildSceneInfo);
+
+	void UpdateNodeSceneInfo(FPrimitiveComponentId NodeId, FPrimitiveSceneInfo* SceneInfo);
+	void PopulateHiddenFlags(FViewInfo& View, FSceneBitArray& HiddenFlags);
+
+	bool IsActive() { return SceneNodes.Num() > 0; }
+
+private:
+	/** Scene this Tree belong to */
+	FScene* Scene;
+
+	/** The LOd groups in the scene.  The map key is the current primitive who has children. */
+	TMap<FPrimitiveComponentId, FLODSceneNode> SceneNodes;
+
+	/**  Update Count. This is used to skip Child node that has been updated */
+	int32 UpdateCount;
+
+	/** Populate Hidden Flags to the children **/
+	void PopulateHiddenFlagsToChildren(FSceneBitArray& HiddenFlags, FLODSceneNode& Node);
+};
+
 typedef TMap<FMaterial*, FMaterialShaderMap*> FMaterialsToUpdateMap;
 
 class FScene : public FSceneInterface
@@ -1362,6 +1426,9 @@ public:
 	/** Uniform buffers for parameter collections with the corresponding Ids. */
 	TMap<FGuid, FUniformBufferRHIRef> ParameterCollections;
 
+	/** LOD Tree Holder for massive LOD system */
+	FLODSceneTree SceneLODHierarchy;
+
 	/** Initialization constructor. */
 	FScene(UWorld* InWorld, bool bInRequiresHitProxies,bool bInIsEditorScene, bool bCreateFXSystem, ERHIFeatureLevel::Type InFeatureLevel);
 
diff --git a/Engine/Source/Runtime/Renderer/Private/SceneVisibility.cpp b/Engine/Source/Runtime/Renderer/Private/SceneVisibility.cpp
--- a/Engine/Source/Runtime/Renderer/Private/SceneVisibility.cpp
+++ b/Engine/Source/Runtime/Renderer/Private/SceneVisibility.cpp
@@ -669,8 +669,9 @@ static void ComputeRelevanceForView(
 
 	for (FSceneSetBitIterator BitIt(View.PrimitiveVisibilityMap); BitIt; ++BitIt)
 	{
-		FPrimitiveSceneInfo* PrimitiveSceneInfo = Scene->Primitives[BitIt.GetIndex()];
-		FPrimitiveViewRelevance& ViewRelevance = View.PrimitiveViewRelevanceMap[BitIt.GetIndex()];
+		const int32 BitIndex = BitIt.GetIndex();
+		FPrimitiveSceneInfo* PrimitiveSceneInfo = Scene->Primitives[BitIndex];
+		FPrimitiveViewRelevance& ViewRelevance = View.PrimitiveViewRelevanceMap[BitIndex];
 		ViewRelevance = PrimitiveSceneInfo->Proxy->GetViewRelevance(&View);
 		ViewRelevance.bInitializedThisFrame = true;
 
@@ -1975,6 +1976,24 @@ void FSceneRenderer::ComputeViewVisibility(FRHICommandListImmediate& RHICmdList)
 			STAT(NumOccludedPrimitives += NumOccludedPrimitivesInView);
 		}
 
+		// visibility test is done, so now build the hidden flags based on visibility set up
+		bool bLODSceneTreeActive = Scene->SceneLODHierarchy.IsActive();
+		FSceneBitArray PrimitiveHiddenByLODMap;
+		if (bLODSceneTreeActive)
+		{
+			PrimitiveHiddenByLODMap.Init(false, View.PrimitiveVisibilityMap.Num());
+			Scene->SceneLODHierarchy.PopulateHiddenFlags(View, PrimitiveHiddenByLODMap);
+
+			// now iterate through turn off visibility if hidden by LOD
+			for(FSceneSetBitIterator BitIt(PrimitiveHiddenByLODMap); BitIt; ++BitIt)
+			{
+				if(PrimitiveHiddenByLODMap.AccessCorrespondingBit(BitIt))
+				{
+					View.PrimitiveVisibilityMap.AccessCorrespondingBit(BitIt) = false;
+				}
+			}
+		}
+
 		MarkAllPrimitivesForReflectionProxyUpdate(Scene);
 		Scene->ConditionalMarkStaticMeshElementsForUpdate();
 
@@ -1992,6 +2011,14 @@ void FSceneRenderer::ComputeViewVisibility(FRHICommandListImmediate& RHICmdList)
 				ComputeRelevanceForView(RHICmdList, Scene, View, ViewBit, RelevantStaticPrimitives, HasDynamicMeshElementsMasks, HasDynamicEditorMeshElementsMasks);
 				MarkRelevantStaticMeshesForView(Scene, View, RelevantStaticPrimitives);
 			}
+
+			if (bLODSceneTreeActive)
+			{
+				for (FSceneBitArray::FIterator BitIt(PrimitiveHiddenByLODMap); BitIt; ++BitIt)
+				{
+					View.PrimitiveViewRelevanceMap[BitIt.GetIndex()].bInitializedThisFrame = true;
+				}
+			}
 		}
 
 #if !(UE_BUILD_SHIPPING || UE_BUILD_TEST)
@@ -2196,3 +2223,128 @@ void FDeferredShadingSceneRenderer::InitViews(FRHICommandListImmediate& RHICmdLi
 	OnStartFrame();
 }
 
+/*------------------------------------------------------------------------------
+	FLODSceneTree Implementation
+------------------------------------------------------------------------------*/
+void FLODSceneTree::AddChildNode(const FPrimitiveComponentId NodeId, FPrimitiveSceneInfo* ChildSceneInfo)
+{
+	if (NodeId.IsValid() && ChildSceneInfo)
+	{
+		FLODSceneNode* Node = SceneNodes.Find(NodeId);
+
+		if(!Node)
+		{
+			Node = &SceneNodes.Add(NodeId, FLODSceneNode());
+
+			// scene info can be added later depending on order of adding to the scene
+			// but at least add componentId, that way when parent is added, it will add its info properly
+			int32 ParentIndex = Scene->PrimitiveComponentIds.Find(NodeId);
+			if(Scene->Primitives.IsValidIndex(ParentIndex))
+			{
+				Node->SceneInfo = Scene->Primitives[ParentIndex];
+			}
+		}
+
+		Node->AddChild(ChildSceneInfo);
+	}
+}
+
+void FLODSceneTree::RemoveChildNode(const FPrimitiveComponentId NodeId, FPrimitiveSceneInfo* ChildSceneInfo)
+{
+	if(NodeId.IsValid() && ChildSceneInfo)
+	{
+		FLODSceneNode* Node = SceneNodes.Find(NodeId);
+		if (Node)
+		{
+			Node->RemoveChild(ChildSceneInfo);
+
+			// delete from scene if no children remains
+			if(Node->ChildrenSceneInfos.Num() == 0)
+			{
+				SceneNodes.Remove(NodeId);
+			}
+		}
+	}
+}
+
+void FLODSceneTree::UpdateNodeSceneInfo(FPrimitiveComponentId NodeId, FPrimitiveSceneInfo* SceneInfo)
+{
+	FLODSceneNode* Node = SceneNodes.Find(NodeId);
+	if(Node)
+	{
+		Node->SceneInfo = SceneInfo;
+	}
+}
+
+void FLODSceneTree::PopulateHiddenFlagsToChildren(FSceneBitArray& HiddenFlags, FLODSceneNode& Node)
+{
+	// if already updated, no reason to do this
+	if(Node.LatestUpdateCount != UpdateCount)
+	{
+		Node.LatestUpdateCount = UpdateCount;
+		// if node doesn't have scene info, that means it doesn't to populate, children is disconnected, so don't bother
+		// in this case we still update children when this node is missing scene info
+		// because parent is showing, so you don't have to show anyway anybody below
+		// although this node might be MIA at this moment
+		for(const auto& Child : Node.ChildrenSceneInfos)
+		{
+			const int32 ChildIndex = Child->GetIndex();
+
+			// first update the flags
+			FRelativeBitReference BitRef(ChildIndex);
+			HiddenFlags.AccessCorrespondingBit(BitRef) = true;
+			// find the node for it
+			FLODSceneNode* ChildNode = SceneNodes.Find(Child->PrimitiveComponentId);
+			// if you have child, populate it again, 
+			if (ChildNode)
+			{
+				PopulateHiddenFlagsToChildren(HiddenFlags, *ChildNode);
+			}
+		}
+	}
+}
+
+void FLODSceneTree::PopulateHiddenFlags(FViewInfo& View, FSceneBitArray& HiddenFlags)
+{
+	++UpdateCount;
+
+	// @todo this is experimental code - hide the children if parent is showing
+	for(auto Iter = SceneNodes.CreateIterator(); Iter; ++Iter)
+	{
+		FLODSceneNode& Node = Iter.Value();
+		// if already updated, no reason to do this
+		if(Node.LatestUpdateCount != UpdateCount)
+		{
+			Node.LatestUpdateCount = UpdateCount;
+			// if node doesn't have scene info, that means it doesn't have any
+			if(Node.SceneInfo)
+			{
+				int32 NodeIndex = Node.SceneInfo->GetIndex();
+				// if this node is visible, children shouldn't show up
+				if(View.PrimitiveVisibilityMap[NodeIndex])
+				{
+					for(const auto& Child : Node.ChildrenSceneInfos)
+					{
+						const int32 ChildIndex = Child->GetIndex();
+
+						// first update the flags
+						FRelativeBitReference BitRef(ChildIndex);
+						HiddenFlags.AccessCorrespondingBit(BitRef) = true;
+
+						// find the node for it
+						FLODSceneNode* ChildNode = SceneNodes.Find(Child->PrimitiveComponentId);
+						// if you have child, populate it again, 
+						if(ChildNode)
+						{
+							PopulateHiddenFlagsToChildren(HiddenFlags, *ChildNode);
+						}
+					}
+				}
+				else
+				{
+					HiddenFlags.AccessCorrespondingBit(FRelativeBitReference(NodeIndex)) = true;
+				}
+			}
+		}
+	}
+}
\ No newline at end of file
