Fix incorrect initialization order.

diff --git a/Engine/Source/Editor/UnrealEd/Private/HierarchicalLOD.cpp b/Engine/Source/Editor/UnrealEd/Private/HierarchicalLOD.cpp
--- a/Engine/Source/Editor/UnrealEd/Private/HierarchicalLOD.cpp
+++ b/Engine/Source/Editor/UnrealEd/Private/HierarchicalLOD.cpp
@@ -381,8 +381,8 @@ void FHierarchicalLODBuilder::InitializeClusters(class ULevel* InLevel, const in
 
 
 FLODCluster::FLODCluster(const FLODCluster& Other)
-: Bound(Other.Bound)
-, Actors(Other.Actors)
+: Actors(Other.Actors)
+, Bound(Other.Bound)
 , FillingFactor(Other.FillingFactor)
 , bValid(Other.bValid)
 {
