Fixed bug in pivot render, and set pivot not to be rendered by default.

diff --git a/Engine/Source/Editor/UnrealEd/Private/EditorComponents.cpp b/Engine/Source/Editor/UnrealEd/Private/EditorComponents.cpp
--- a/Engine/Source/Editor/UnrealEd/Private/EditorComponents.cpp
+++ b/Engine/Source/Editor/UnrealEd/Private/EditorComponents.cpp
@@ -310,7 +310,7 @@ FEditorCommonDrawHelper.
 ------------------------------------------------------------------------------*/
 FEditorCommonDrawHelper::FEditorCommonDrawHelper()
 	: bDrawGrid(true)
-	, bDrawPivot(true)
+	, bDrawPivot(false)
 	, bDrawBaseInfo(true)
 	, bDrawWorldBox(false)
 	, bDrawKillZ(false)
@@ -606,7 +606,7 @@ void FEditorCommonDrawHelper::DrawPivot(const FSceneView* View,FPrimitiveDrawInt
 	const FVector PivLoc = GLevelEditorModeTools().SnappedLocation;
 
 	const float ZoomFactor = FMath::Min<float>(View->ViewMatrices.ProjMatrix.M[0][0], View->ViewMatrices.ProjMatrix.M[1][1]);
-	const float WidgetRadius = View->ViewMatrices.ProjMatrix.TransformPosition(PivLoc).W * (PivotSize / ZoomFactor);
+	const float WidgetRadius = View->ViewMatrices.GetViewProjMatrix().TransformPosition(PivLoc).W * (PivotSize / ZoomFactor);
 
 	const FVector CamX = CameraToWorld.TransformVector( FVector(1,0,0) );
 	const FVector CamY = CameraToWorld.TransformVector( FVector(0,1,0) );
