When selecting a deleted level from favorites list, changed the error popup from a message box to a notification.

diff --git a/Engine/Source/Editor/UnrealEd/Private/MRUFavoritesList.cpp b/Engine/Source/Editor/UnrealEd/Private/MRUFavoritesList.cpp
--- a/Engine/Source/Editor/UnrealEd/Private/MRUFavoritesList.cpp
+++ b/Engine/Source/Editor/UnrealEd/Private/MRUFavoritesList.cpp
@@ -4,6 +4,8 @@
 #include "UnrealEd.h"
 #include "MRUFavoritesList.h"
 #include "ConfigCacheIni.h"
+#include "SNotificationList.h"
+#include "NotificationManager.h"
 
 const FString FMainMRUFavoritesList::FAVORITES_INI_SECTION = TEXT("FavoriteFiles");
 
@@ -126,7 +128,10 @@ bool FMainMRUFavoritesList::VerifyFavoritesFile( int32 ItemIndex )
 	// If the file doesn't exist any more, remove it from the favorites list and alert the user
 	if ( !bFileExists )
 	{
-		FMessageDialog::Open( EAppMsgType::Ok, FText::Format( NSLOCTEXT("UnrealEd", "Error_FavoritesFileDoesNotExist", "File does not exist: '{0}'. It will be removed from the Favorites list."), FText::FromString(CurFileName) ) );
+		FNotificationInfo Info(FText::Format(NSLOCTEXT("UnrealEd", "Error_FavoritesFileDoesNotExist", "Map '{0}' does not exist - it will be removed from the Favorites list."), FText::FromString(FPaths::GetCleanFilename(CurFileName))));
+		Info.bUseThrobber = false;
+		Info.ExpireDuration = 8.0f;
+		FSlateNotificationManager::Get().AddNotification(Info)->SetCompletionState(SNotificationItem::CS_Fail);
 		RemoveFavoritesItem( CurFileName );
 	}
 
