Bookkeeping merge from UE4-Releases/4.7 [REDACTED_HANDLE]. Accepting target except for CookOnTheFlyServer.cpp.

diff --git a/Engine/Source/Editor/UnrealEd/Private/CookOnTheFlyServer.cpp b/Engine/Source/Editor/UnrealEd/Private/CookOnTheFlyServer.cpp
--- a/Engine/Source/Editor/UnrealEd/Private/CookOnTheFlyServer.cpp
+++ b/Engine/Source/Editor/UnrealEd/Private/CookOnTheFlyServer.cpp
@@ -2695,6 +2695,8 @@ void UCookOnTheFlyServer::CancelCookByTheBook()
 		// save the cook requests 
 		CookRequests.DequeueAllRequests(CookByTheBookOptions->PreviousCookRequests);
 		CookByTheBookOptions->bRunning = false;
+
+		SandboxFile = NULL;
 	}	
 }
 
