[INTEGRATE] Change 2412028 by [REDACTED_PERSON][REDACTED_HANDLE].7 on 2015/01/20 07:02:45

	Added automatic content monitoring for all mounted content directories

	Modifications to source content files that reside inside these folders will now automatically be reimported, if they relate to an asset.
	New files will also be automatically imported, and deleted files will prompt the user to delete associated assets.
	The feature is currently disabled by default, and enabled in Editor Preferences / Loading and Saving.

	Notable features still todo/refine:
	 * Source files rename currently process as a delete/add
	 * Slow task dialogs need work. We can hopefully get it running in the background / asynchronously over multiple frames
	 * Batch operations need work
	 * Profiling/optimisation
	 * Various message boxes/dialogs need some refinement
	 * Asset import factories that are dynamically loaded/unloaded will not be picked up by the import manager. Need to figure out a good way of fixing this.

diff --git a/Engine/Plugins/2D/Paper2D/Source/Paper2D/Classes/PaperTileMap.h b/Engine/Plugins/2D/Paper2D/Source/Paper2D/Classes/PaperTileMap.h
--- a/Engine/Plugins/2D/Paper2D/Source/Paper2D/Classes/PaperTileMap.h
+++ b/Engine/Plugins/2D/Paper2D/Source/Paper2D/Classes/PaperTileMap.h
@@ -92,7 +92,7 @@ class PAPER2D_API UPaperTileMap : public UObject
 public:
 #if WITH_EDITORONLY_DATA
 	/** Importing data and options used for this tile map */
-	UPROPERTY(Category=ImportSettings, VisibleAnywhere, Instanced)
+	UPROPERTY(Category=ImportSettings, VisibleAnywhere, Instanced, AssetRegistrySearchable)
 	class UAssetImportData* AssetImportData;
 
 	/** The currently selected layer index */
diff --git a/Engine/Plugins/2D/Paper2D/Source/PaperJsonImporter/Private/PaperSpriteSheet.h b/Engine/Plugins/2D/Paper2D/Source/PaperJsonImporter/Private/PaperSpriteSheet.h
--- a/Engine/Plugins/2D/Paper2D/Source/PaperJsonImporter/Private/PaperSpriteSheet.h
+++ b/Engine/Plugins/2D/Paper2D/Source/PaperJsonImporter/Private/PaperSpriteSheet.h
@@ -26,7 +26,7 @@ public:
 
 #if WITH_EDITORONLY_DATA
 	// Import data for this 
-	UPROPERTY(VisibleAnywhere, Instanced, Category=ImportSettings)
+	UPROPERTY(VisibleAnywhere, Instanced, Category=ImportSettings, AssetRegistrySearchable)
 	class UAssetImportData* AssetImportData;
 #endif
 };
diff --git a/Engine/Plugins/ScriptPlugin/Source/ScriptPlugin/Classes/ScriptBlueprint.h b/Engine/Plugins/ScriptPlugin/Source/ScriptPlugin/Classes/ScriptBlueprint.h
--- a/Engine/Plugins/ScriptPlugin/Source/ScriptPlugin/Classes/ScriptBlueprint.h
+++ b/Engine/Plugins/ScriptPlugin/Source/ScriptPlugin/Classes/ScriptBlueprint.h
@@ -17,7 +17,7 @@ class SCRIPTPLUGIN_API UScriptBlueprint : public UBlueprint
 public:
 
 	/** Path to the script text used to construct this asset. Relative to the object's package, BaseDir() or absolute */
-	UPROPERTY(Category = Script, VisibleAnywhere, BlueprintReadWrite)
+	UPROPERTY(Category = Script, VisibleAnywhere, BlueprintReadWrite, AssetRegistrySearchable, meta=(ReimportPath))
 	FString SourceFilePath;
 
 	/** Date/Time-stamp of the file from the last import. */
diff --git a/Engine/Source/Editor/UnrealEd/Classes/AutoReimport/AutoReimportManager.h b/Engine/Source/Editor/UnrealEd/Classes/AutoReimport/AutoReimportManager.h
deleted file mode 100644
diff --git a/Engine/Source/Editor/UnrealEd/Classes/Settings/EditorLoadingSavingSettings.h b/Engine/Source/Editor/UnrealEd/Classes/Settings/EditorLoadingSavingSettings.h
--- a/Engine/Source/Editor/UnrealEd/Classes/Settings/EditorLoadingSavingSettings.h
+++ b/Engine/Source/Editor/UnrealEd/Classes/Settings/EditorLoadingSavingSettings.h
@@ -27,17 +27,21 @@ public:
 
 public:
 
+	/**Automatically reimports textures when a change to source content is detected */
+	UPROPERTY(EditAnywhere, config, Category=AutoReimport, meta=(DisplayName="Monitor Content Directories", ToolTip="When enabled, changes to made to source content files inside the content directories will automatically be reflected in the content browser."))
+	bool bMonitorContentDirectories;
+
 	/** Directories being monitored for Auto Reimport */
-	UPROPERTY(EditAnywhere, config, Category=AutoReimport, meta=(DisplayName="Auto Reimport Directories"))
+	UPROPERTY(EditAnywhere, config, Category=AutoReimport, meta=(DisplayName="Additional Directories to Monitor"))
 	TArray<FString> AutoReimportDirectories;
 
-	/**Automatically reimports textures when a change to source content is detected */
-	UPROPERTY(EditAnywhere, config, Category=AutoReimport, meta=(DisplayName="Auto Reimport Textures"))
-	bool bAutoReimportTextures;
+private:
 
-	/**Automatically reimports CSV files when a change to source content is detected */
-	UPROPERTY(EditAnywhere, config, Category=AutoReimport, meta=(DisplayName="Auto Reimport CSVData"))
-	bool bAutoReimportCSV;
+	/** Deprecated properties - we automatically monitor all source file types that are supported by in memory factories now */
+	UPROPERTY(config)
+	bool bAutoReimportTextures_DEPRECATED;
+	UPROPERTY(config)
+	bool bAutoReimportCSV_DEPRECATED;
 
 public:
 
diff --git a/Engine/Source/Editor/UnrealEd/Private/AutoReimportManager.cpp b/Engine/Source/Editor/UnrealEd/Private/AutoReimportManager.cpp
--- a/Engine/Source/Editor/UnrealEd/Private/AutoReimportManager.cpp
+++ b/Engine/Source/Editor/UnrealEd/Private/AutoReimportManager.cpp
@@ -6,394 +6,522 @@
 #include "Developer/DirectoryWatcher/Public/DirectoryWatcherModule.h"
 #include "Engine/CurveTable.h"
 #include "Engine/DataTable.h"
-
+#include "AutoReimport/AutoReimportManager.h"
+#include "PackageTools.h"
+#include "ObjectTools.h"
+#include "ARFilter.h"
+#include "ContentBrowserModule.h"
+#include "AssetRegistryModule.h"
 
 #define LOCTEXT_NAMESPACE "AutoReimportManager"
 
 DEFINE_LOG_CATEGORY_STATIC(LogAutoReimportManager, Log, All);
 
-
-UAutoReimportManager::UAutoReimportManager(const FObjectInitializer& ObjectInitializer)
-	: Super(ObjectInitializer)
+namespace
 {
-	ReimportTime = 0.0f;
-}
+	static const FName ReimportPathName("ReimportPath");
 
-FDelegateHandle UAutoReimportManager::AddReimportDelegate( const FString& FileType, const FAutoReimportDelegate& Delegate )
-{
-	if(ReimportDelegates* Del = ExtensionToReimporter.Find(FileType))
+	/** Template helper function to count the number of elements in the specified array that pass a predicate */
+	template<typename T, typename P>
+	int32 CountIf(const TArray<T>& InArray, P Predicate)
 	{
-		Del->Add(Delegate);
-		return Del->Last().GetHandle();
+		int32 Count = 0;
+		for (const T& X : InArray)
+		{
+			if (Predicate(X))
+			{
+				++Count;
+			}
+		}
+		return Count;
 	}
 
-	ReimportDelegates Delegates;
-	Delegates.Add(Delegate);
-	return ExtensionToReimporter.Add(FileType, Delegates).Last().GetHandle();
-}
-
-void UAutoReimportManager::RemoveReimportDelegate( const FString& FileType, const FAutoReimportDelegate& InDelegate )
-{
-	if(ReimportDelegates* Del = ExtensionToReimporter.Find(FileType))
+	/** Reduce the array to the specified accumulator */
+	template<typename T, typename P, typename A>
+	A&& Reduce(const TArray<T>& InArray, P Predicate, A&& Accumulator)
 	{
-		Del->RemoveAll([&](const FAutoReimportDelegate& Delegate){ return Delegate.DEPRECATED_Compare(InDelegate); });
+		for (const T& X : InArray)
+		{
+			Predicate(X, Accumulator);
+		}
+		return MoveTemp(Accumulator);
 	}
-}
 
-void UAutoReimportManager::RemoveReimportDelegate( const FString& FileType, FDelegateHandle Handle )
-{
-	if(ReimportDelegates* Del = ExtensionToReimporter.Find(FileType))
+	/** Retrieve a semi-colon separated string of file extensions supported by all available editor import factories */
+	FString GetAllFactoryExtensions()
 	{
-		Del->RemoveAll([=](const FAutoReimportDelegate& Delegate){ return Delegate.GetHandle() == Handle; });
+		FString AllExtensions;
+
+		// Use a scratch buffer to avoid unnecessary re-allocation
+		FString Scratch;
+		Scratch.Reserve(32);
+
+		for (TObjectIterator<UClass> ClassIt; ClassIt; ++ClassIt)
+		{
+			UClass* Class = *ClassIt;
+			if (Class->IsChildOf(UFactory::StaticClass()) && !Class->HasAnyClassFlags(CLASS_Abstract))
+			{
+				UFactory* Factory = Cast<UFactory>(Class->GetDefaultObject());
+				if (Factory->bEditorImport)
+				{
+					for (const FString& Format : Factory->Formats)
+					{
+						int32 Index = INDEX_NONE;
+						if (Format.FindChar(';', Index) && Index > 0)
+						{
+							Scratch.GetCharArray().Reset();
+							// Include the ;
+							Scratch.AppendChars(*Format, Index + 1);
+
+							if (AllExtensions.Find(Scratch) == INDEX_NONE)
+							{
+								AllExtensions += Scratch;
+							}
+						}
+					}
+				}
+			}
+		}
+
+		return AllExtensions;
 	}
-}
 
-void UAutoReimportManager::Initialize()
-{
-	//Add all auto reimport delegates
-
-	//csv types
-	AddReimportDelegate("csv", FAutoReimportDelegate::CreateStatic(&UAutoReimportManager::AutoReimportCurve));
-	AddReimportDelegate("csv", FAutoReimportDelegate::CreateStatic(&UAutoReimportManager::AutoReimportCurveTable));
-	AddReimportDelegate("csv", FAutoReimportDelegate::CreateStatic(&UAutoReimportManager::AutoReimportDataTable));
-	
-	//textures
-	AddReimportDelegate( "tga", FAutoReimportDelegate::CreateStatic( &UAutoReimportManager::AutoReimportTexture ) );
-	AddReimportDelegate( "dds", FAutoReimportDelegate::CreateStatic( &UAutoReimportManager::AutoReimportTexture ) );
-	AddReimportDelegate( "png", FAutoReimportDelegate::CreateStatic( &UAutoReimportManager::AutoReimportTexture ) );
-	AddReimportDelegate( "jpg", FAutoReimportDelegate::CreateStatic( &UAutoReimportManager::AutoReimportTexture ) );
-	AddReimportDelegate( "jpeg", FAutoReimportDelegate::CreateStatic( &UAutoReimportManager::AutoReimportTexture ) );
-	AddReimportDelegate( "psd", FAutoReimportDelegate::CreateStatic( &UAutoReimportManager::AutoReimportTexture ) );
-	AddReimportDelegate( "bmp", FAutoReimportDelegate::CreateStatic( &UAutoReimportManager::AutoReimportTexture ) );
-
-	//Want notification when files change in specific directories
-	GetMutableDefault<UEditorLoadingSavingSettings>()->OnSettingChanged().AddUObject(this, &UAutoReimportManager::HandleLoadingSavingSettingChanged);
-
-	AutoReimportDirectories = GetDefault<UEditorLoadingSavingSettings>()->AutoReimportDirectories;
-	for(auto It(AutoReimportDirectories.CreateIterator()); It; ++It)
+	/** Find a list of assets that were last imported from the specified filename. */
+	TArray<FAssetData> FindAssetsPertainingToFile(const IAssetRegistry& Registry, const FString& AbsoluteFilename)
 	{
-		//Fix up existing user paths
-		*It = FixUserPath(*It);
+		TArray<FAssetData> Assets, Result;
+		const FString LeafName = FPaths::GetCleanFilename(AbsoluteFilename);
+
+		FARFilter Filter;
+		Filter.SourceFilenames.Add(*LeafName);
+		Filter.bIncludeOnlyOnDiskAssets = true;
+		Registry.GetAssets(Filter, Assets);
 
-		//watch it
-		WatchedDirectories.Add(this, *It);
+		for (const auto& Asset : Assets)
+		{
+			for (const auto& Pair : Asset.TagsAndValues)
+			{
+				// We don't compare numbers on FNames for this check because the tag "ReimportPath" may exist multiple times
+				const bool bCompareNumber = false;
+				if (Pair.Key.IsEqual(ReimportPathName, ENameCase::IgnoreCase, bCompareNumber))
+				{
+					// Either relative to the asset itself, or relative to the base path, or absolute.
+					// Would ideally use FReimportManager::ResolveImportFilename here, but we don't want to ask the file system whether the file exists.
+					if (AbsoluteFilename == FPaths::ConvertRelativePathToFull(FPackageName::LongPackageNameToFilename(Asset.PackagePath.ToString()) / Pair.Value) ||
+						AbsoluteFilename == FPaths::ConvertRelativePathToFull(Pair.Value))
+					{
+						Result.Add(Asset);
+						break;
+					}
+				}
+			}
+		}
+
+		return Result;
 	}
-	
-}
 
-FString UAutoReimportManager::FixUserPath(  FString File )
-{
-	FPaths::NormalizeFilename(File);
-	if(File.EndsWith(TEXT("/")) == false)
+	/** Generate a config from the specified options, to pass to FFileCache on construction */
+	FFileCacheConfig GenerateFileCacheConfig(const FString& InPath, const FString& InSupportedExtensions, const FString& InMountedContentPath)
 	{
-		File += TEXT("/");
+		FFileCacheConfig Config;
+		Config.Directory = FPaths::ConvertRelativePathToFull(InPath);
+		Config.IncludeExtensions = InSupportedExtensions;
+
+		const FString& HashString = InMountedContentPath.IsEmpty() ? Config.Directory : InMountedContentPath;
+		const uint32 CRC = FCrc::MemCrc32(*HashString, HashString.Len()*sizeof(TCHAR));	
+
+		Config.CacheFile = FPaths::ConvertRelativePathToFull(FPaths::GameIntermediateDir()) / TEXT("ReimportCache") / FString::Printf(TEXT("%u.bin"), CRC);
+		return Config;
 	}
-	return File;
-}
 
-void UAutoReimportManager::BeginDestroy( void ) 
-{
-	for(auto It(AutoReimportDirectories.CreateIterator());It;++It)
+	/** Choose a location to import the specified asset into */
+	TOptional<FString> ChooseAssetImportLocation(const FString& Filename)
 	{
-		FString Directory = *It;
-		WatchedDirectories.Remove(this, Directory);
-	}
-	Super::BeginDestroy();
-}
+		FString DefaultPath = TEXT("/Game/");
 
-void UAutoReimportManager::Tick( float DeltaTime ) 
-{
-	ReimportTime += DeltaTime;
+		FString NameSuggestion;
+		{
+			FString PackageName = DefaultPath / FPaths::GetBaseFilename(Filename);
+			FString Name;
+			FAssetToolsModule& AssetToolsModule = FModuleManager::LoadModuleChecked<FAssetToolsModule>("AssetTools");
+			AssetToolsModule.Get().CreateUniqueAssetName(PackageName, TEXT(""), PackageName, Name);
 
-	//we wait this amount of time after a modification to re-import, this is because some tools save the file multiple times
-	const double DeltaSinceModified = 1.0;
+			NameSuggestion = FPaths::GetCleanFilename(PackageName);
+		}
 
-	//last modification must have been prior to this time
-	const double CanImportTime = ReimportTime - DeltaSinceModified;
+		FSaveAssetDialogConfig Config;
+		Config.DialogTitleOverride = LOCTEXT("ImportAssetDialogTitle", "Import asset to:");
+		Config.DefaultPath = DefaultPath;
+		Config.DefaultAssetName = NameSuggestion;
+		Config.ExistingAssetPolicy = ESaveAssetDialogExistingAssetPolicy::Disallow;
 
-	//Reimport files that have been modified
-	for (auto It(ModifiedFiles.CreateIterator());It;++It)
-	{
-		if(It.Value() < CanImportTime)
+		FContentBrowserModule& ContentBrowserModule = FModuleManager::LoadModuleChecked<FContentBrowserModule>("ContentBrowser");
+		FString SaveObjectPath = ContentBrowserModule.Get().CreateModalSaveAssetDialog(Config);
+		if ( !SaveObjectPath.IsEmpty() )
 		{
-			if(ReimportDelegates* Reimporters = ExtensionToReimporter.Find(FPaths::GetExtension(It.Key())))
-			{
-				OnFileModified(*Reimporters, It.Key());
-			}
-			It.RemoveCurrent();
+			return SaveObjectPath;
 		}
+
+		return TOptional<FString>();
 	}
 }
 
-TStatId UAutoReimportManager::GetStatId() const
+FContentDirectoryMonitor::FContentDirectoryMonitor(const FString& InDirectory, const FString& InSupportedExtensions, const FString& InMountedContentPath)
+	: Cache(GenerateFileCacheConfig(InDirectory, InSupportedExtensions, InMountedContentPath))
+	, MountedContentPath(InMountedContentPath)
+	, LastProcessTime(0)
+	, LastSaveTime(0)
 {
-	RETURN_QUICK_DECLARE_CYCLE_STAT(UAutoReimportManager, STATGROUP_Tickables);
+	Cache.Scan();
 }
 
-bool UAutoReimportManager::ReimportFile(UObject* Obj, const FString File)
+void FContentDirectoryMonitor::Tick(const IAssetRegistry& Registry)
 {
-	UPackage* Package = Obj->GetOutermost();
-	check(Package);
+	const double Now = FPlatformTime::Seconds();
 
-	//found a match, reimport!
-	if(FReimportManager::Instance()->Reimport( Obj ))
+	// Keep accumulating periodically 'til we have no more changes
+	auto NewChanges = Cache.GetOutstandingChanges();
+	if (NewChanges.Num() != 0)
 	{
-		return true;
+		// Check again for any more in 1 second
+		OutstandingChanges.Append(MoveTemp(NewChanges));
+		LastProcessTime = Now;
+	}
+	else if (Now - LastProcessTime > ReactionTimeS && OutstandingChanges.Num())
+	{
+		LastProcessTime = Now;
+		ProcessOutstandingChanges(Registry);
 	}
 
-	return false;
+	if (Now - LastSaveTime > ResaveIntervalS)
+	{
+		LastSaveTime = Now;
+		Cache.WriteCache();
+	}
 }
 
-void UAutoReimportManager::HandleLoadingSavingSettingChanged( FName PropertyName )
+void FContentDirectoryMonitor::ProcessOutstandingChanges(const IAssetRegistry& Registry)
 {
-	if (PropertyName.ToString() == "AutoReimportDirectories")
+	// Count the number of additions so we can apprximate the work load
+	const int32 NumAdditions = CountIf(OutstandingChanges, [](const FUpdateCacheTransaction& Change){ return Change.Action == FFileChangeData::FCA_Added; });
+
+	auto* ReimportManager = FReimportManager::Instance();
+
+	TArray<UObject*> AssetsToDelete;
 	{
-		for (auto It(AutoReimportDirectories.CreateIterator()); It; ++It)
-		{
-			WatchedDirectories.Remove(this, *It);
-		}
+		bool bCreatedDialog = false;
+		FScopedSlowTask SlowTask(OutstandingChanges.Num() + NumAdditions, LOCTEXT("ProcessingChanges", "Processing source content changes"));
+		SlowTask.MakeDialog();
 
-		//Fix up slash issues with whatever user entered
-		for (auto It(GetMutableDefault<UEditorLoadingSavingSettings>()->AutoReimportDirectories.CreateIterator()); It; ++It)
+		TArray<UPackage*> PackagesToSave;
+
+		for (auto& Change : OutstandingChanges)
 		{
-			*It = FixUserPath(*It);
+			const FString FullFilename = Cache.GetDirectory() + Change.Filename.Get();
+
+			if (!bCreatedDialog && FPlatformTime::Seconds() - SlowTask.StartTime > 2)
+			{
+				//SlowTask.MakeDialog();
+				bCreatedDialog = true;
+			}
+			
+			SlowTask.EnterProgressFrame();
+			switch(Change.Action)
+			{
+				case FFileChangeData::FCA_Modified:
+					// @todo: arodham: enable/disable specific extensions/asset types?
+					{
+						for (const auto& AssetData : FindAssetsPertainingToFile(Registry, FullFilename))
+						{
+							UObject* Asset = AssetData.GetAsset();
+							ReimportManager->Reimport(Asset);
+						}
+					}
+					break;
+
+				case FFileChangeData::FCA_Added:
+					{
+						FAssetToolsModule& AssetToolsModule = FModuleManager::Get().LoadModuleChecked<FAssetToolsModule>("AssetTools");
+
+						TArray<FString> File;
+						File.Add(FullFilename);
+
+						FString ParentPackage;
+						if (MountedContentPath.IsEmpty())
+						{
+							// @todo: arodham: this is sketchy and needs revisiting. PromptUserForDirectory?
+							auto UserChosenPath = ChooseAssetImportLocation(FullFilename);
+							if (UserChosenPath.IsSet())
+							{
+								ParentPackage = UserChosenPath.GetValue();
+							}
+						}
+						else
+						{
+							ParentPackage = MountedContentPath / PackageTools::SanitizePackageName(FPaths::GetPath(Change.Filename.Get()));
+						}
+
+						if (!ParentPackage.IsEmpty())
+						{
+							auto NewAssets = AssetToolsModule.Get().ImportAssets(File, ParentPackage);
+
+							// @todo: arodham: Error reporting
+
+							for (auto* Asset : NewAssets)
+							{
+								UPackage* Package = Asset->GetOutermost();
+								if (Package)
+								{
+									PackagesToSave.Add(Package);
+								}
+							}
+						}
+					}
+					break;
+
+				case FFileChangeData::FCA_Removed:
+					for (const auto& AssetData : FindAssetsPertainingToFile(Registry, FullFilename))
+					{
+						AssetsToDelete.Add(AssetData.GetAsset());
+					}
+					break;
+
+				default:
+					break;
+			}
+
+			// Let the cache know that we've dealt with this change
+			Cache.CompleteTransaction(MoveTemp(Change));
 		}
 
-		AutoReimportDirectories = GetDefault<UEditorLoadingSavingSettings>()->AutoReimportDirectories;
-		for (auto It(AutoReimportDirectories.CreateIterator()); It; ++It)
+		SlowTask.EnterProgressFrame(NumAdditions);
+		if (PackagesToSave.Num() > 0)
 		{
-			WatchedDirectories.Add(this, *It);
+			const bool bAlreadyCheckedOut = false;
+			const bool bCheckDirty = false;
+			const bool bPromptToSave = false;
+			// @todo: arodham: Error reporting?
+			FEditorFileUtils::PromptForCheckoutAndSave(PackagesToSave, bCheckDirty, bPromptToSave, nullptr, bAlreadyCheckedOut);
 		}
 	}
-}
 
-void UAutoReimportManager::OnDirectoryChanged( const TArray<struct FFileChangeData>& FileChanges )
-{
-	for (int32 FileIdx = 0; FileIdx < FileChanges.Num(); ++FileIdx)
-	{
-	    FString File = FileChanges[FileIdx].Filename;
-		FPaths::NormalizeFilename(File);
-		const FString FileExtension = FPaths::GetExtension(File);
+	// Open a dialog to ask about deleted files
+	// @todo: this needs improvement
 
-		if ((FileChanges[FileIdx].Action != FFileChangeData::FCA_Removed) && ExtensionToReimporter.Contains(FileExtension))
-		{
-			ModifiedFiles.Add(File,ReimportTime);//will overwrite with most recent modified time
-		} 	
-	}
-}
-
-bool UAutoReimportManager::OnFileModified(ReimportDelegates& Delegates,  const FString& File )
-{
-	for( TObjectIterator<UObject> ObjIt; ObjIt; ++ObjIt )
+	if (AssetsToDelete.Num() > 0)
 	{
-		UObject* Object  = *ObjIt;
-		for( auto It = Delegates.CreateIterator();It;++It)
-		{
-			FAutoReimportDelegate& Delegate = *It;
-			if(Delegate.Execute(Object, File))
-			{
-				return true;
-			}
-		}
+		ObjectTools::DeleteObjects(AssetsToDelete);
 	}
-	return false;
+
+	OutstandingChanges.Empty();
 }
 
-bool UAutoReimportManager::JustFileNameMatches(const FString& ObjectPath, const FString& FileName)
+UAutoReimportManager::UAutoReimportManager(const FObjectInitializer& ObjectInitializer)
+	: Super(ObjectInitializer)
+	, bIsProcessingChanges(false)
 {
-	return FPaths::GetCleanFilename(ObjectPath) == FPaths::GetCleanFilename(FileName);
+
 }
 
-bool UAutoReimportManager::IsSamePath(const FString& ObjectPath, const FString& FileName)
+void UAutoReimportManager::Initialize()
 {
-	if(ObjectPath.Len() == 0)
-	{
-		return false;
-	}
+	// @todo: arodham: update this when modules are reloaded or new factory types are available?
+	SupportedExtensions = GetAllFactoryExtensions();
 
-	if(FPaths::GetBaseFilename(ObjectPath) != FPaths::GetBaseFilename(FileName))
-	{
-		return false;
-	}
+	auto* Settings = GetMutableDefault<UEditorLoadingSavingSettings>();
+	Settings->OnSettingChanged().AddUObject(this, &UAutoReimportManager::HandleLoadingSavingSettingChanged);
 
-	FString NormalizedName = ObjectPath;
-	FPaths::NormalizeFilename(NormalizedName);
+	FPackageName::OnContentPathMounted().AddUObject(this, &UAutoReimportManager::OnContentPathMounted);
+	FPackageName::OnContentPathDismounted().AddUObject(this, &UAutoReimportManager::OnContentPathDismounted);
 
-	if(ObjectPath == FileName) 
-	{
-		return true;
-	}	
+	IAssetRegistry& AssetRegistry = FModuleManager::LoadModuleChecked<FAssetRegistryModule>("AssetRegistry").Get();
+	AssetRegistry.OnInMemoryAssetDeleted().AddUObject(this, &UAutoReimportManager::OnAssetDeleted);
 
-	const bool ObjIsRelative = FPaths::IsRelative(NormalizedName);
-	const bool FileIsRelative = FPaths::IsRelative(FileName);
 
-	if(ObjIsRelative && !FileIsRelative)
-	{
-		return FPaths::ConvertRelativePathToFull(NormalizedName) == FileName;
-	}
-	else if(!ObjIsRelative && FileIsRelative)
+	// Only set this up for content directories if the user has this enabled
+	if (Settings->bMonitorContentDirectories)
 	{
-		return NormalizedName == FPaths::ConvertRelativePathToFull(FileName);
+		SetUpMountedContentDirectories();
 	}
 
-	return false;
+	// We always monitor user directories if they are set
+	SetUpUserDirectories();
+}
+
+void UAutoReimportManager::BeginDestroy() 
+{
+	// Force a save of all the caches
+	MountedDirectoryMonitors.Empty();
+	UserDirectoryMonitors.Empty();
 
+	Super::BeginDestroy();
 }
 
-void UAutoReimportManager::FDirectories::Add( UAutoReimportManager* AutoReimportManager, const FString&  FileName )
+void UAutoReimportManager::Tick(float DeltaTime)
 {
-	bool bFound = false;
-	for(auto It(Paths.CreateIterator());It;++It)
+	const IAssetRegistry& Registry = FModuleManager::LoadModuleChecked<FAssetRegistryModule>("AssetRegistry").Get();
+
+	// We can't do this if the registry is loading assets as we won't necesssarily be able to action changes effectively
+	if (!Registry.IsLoadingAssets())
 	{
-		FDirectory& Path = *It;	
-		if(Path.Path == FileName)
+		TGuardValue<bool> Guard(bIsProcessingChanges, true);
+
+		for (auto& Monitor : MountedDirectoryMonitors)
 		{
-			Path.Count++;
-			bFound = true;
-			break;
+			Monitor.Tick(Registry);
+		}
+		for (auto& Monitor : UserDirectoryMonitors)
+		{
+			Monitor.Tick(Registry);
 		}
 	}
-	if(!bFound)
-	{
-		Paths.Add(FDirectory(FileName));
-		AutoReimportManager->ActivateAutoReimport(FileName);
-	}
 }
 
-void UAutoReimportManager::FDirectories::Remove( UAutoReimportManager* AutoReimportManager,  const FString& FileName )
+TStatId UAutoReimportManager::GetStatId() const
+{
+	RETURN_QUICK_DECLARE_CYCLE_STAT(UAutoReimportManager, STATGROUP_Tickables);
+}
+
+void UAutoReimportManager::HandleLoadingSavingSettingChanged(FName PropertyName)
 {
-	bool bFound = false;
-	for(auto It(Paths.CreateIterator());It;++It)
+	if (PropertyName == GET_MEMBER_NAME_CHECKED(UEditorLoadingSavingSettings, AutoReimportDirectories))
 	{
-		FDirectory& Path = *It;	
-		if(Path.Path == FileName)
+		SetUpUserDirectories();
+	}
+	else if (PropertyName == GET_MEMBER_NAME_CHECKED(UEditorLoadingSavingSettings, bMonitorContentDirectories))
+	{
+		const auto* Settings = GetDefault<UEditorLoadingSavingSettings>();
+		if (Settings->bMonitorContentDirectories)
 		{
-			Path.Count--;
-			if(0 == Path.Count)
+			// Should already be empty, but just to be sure
+			MountedDirectoryMonitors.Empty();
+			SetUpMountedContentDirectories();
+		}
+		else
+		{
+			// Destroy all the existing montirs, including their file caches
+			for (auto& Monitor : MountedDirectoryMonitors)
 			{
-				Paths.RemoveAt(It.GetIndex());
-				AutoReimportManager->DeactivateAutoReimport(FileName);
+				Monitor.Destroy();
 			}
-			bFound = true;
-			break;
+			MountedDirectoryMonitors.Empty();
 		}
 	}
+}
 
-	check(bFound);
+void UAutoReimportManager::OnContentPathMounted(const FString& InAssetPath, const FString& FileSystemPath)
+{
+	MountedDirectoryMonitors.Emplace(FPaths::ConvertRelativePathToFull(FileSystemPath), SupportedExtensions, InAssetPath);
 }
 
-void UAutoReimportManager::ActivateAutoReimport( const FString& Directory)
+void UAutoReimportManager::OnContentPathDismounted(const FString& InAssetPath, const FString& FileSystemPath)
 {
-	// Register for notifications when files change in the content folders
-	FDirectoryWatcherModule& DirectoryWatcherModule = FModuleManager::Get().LoadModuleChecked<FDirectoryWatcherModule>(TEXT("DirectoryWatcher"));
-	IDirectoryWatcher* DirectoryWatcher = DirectoryWatcherModule.Get();
-	if (DirectoryWatcher)
-	{
-		if(Directory.IsEmpty() == false)
-		{
-			DirectoryWatcher->RegisterDirectoryChangedCallback_Handle(Directory, IDirectoryWatcher::FDirectoryChanged::CreateUObject(this, &UAutoReimportManager::OnDirectoryChanged), OnDirectoryChangedDelegateHandle);
-		}
-	}
+	MountedDirectoryMonitors.RemoveAll([&](const FContentDirectoryMonitor& Monitor){
+		return Monitor.GetMountedContentPath() == InAssetPath;
+	});
 }
 
-void UAutoReimportManager::DeactivateAutoReimport( const FString& Directory )
+void UAutoReimportManager::OnAssetDeleted(UObject* Object)
 {
-	FDirectoryWatcherModule& DirectoryWatcherModule = FModuleManager::Get().LoadModuleChecked<FDirectoryWatcherModule>(TEXT("DirectoryWatcher"));
-	IDirectoryWatcher* DirectoryWatcher = DirectoryWatcherModule.Get();
-	if (DirectoryWatcher)
+	if (bIsProcessingChanges)
 	{
-		if(Directory.IsEmpty() == false)
-		{
-			DirectoryWatcher->UnregisterDirectoryChangedCallback_Handle(Directory, OnDirectoryChangedDelegateHandle);
-		}
+		return;
 	}
-}
 
-//////////////////////////////////////////////////////////////////////////
-// Auto Reimport handlers
+	const IAssetRegistry& Registry = FModuleManager::LoadModuleChecked<FAssetRegistryModule>("AssetRegistry").Get();
 
-bool UAutoReimportManager::AutoReimportCurve( UObject* Object, const FString& FileName )
-{
-	if (GetDefault<UEditorLoadingSavingSettings>()->bAutoReimportCSV)
+	bool bCanDelete = true;
+
+	TArray<FAssetRegistryTag> Tags;
+	Object->GetAssetRegistryTags(Tags);
+
+	TArray<FString> AbsoluteSourceFileNames;
+
+	for (const auto& Tag : Tags)
 	{
-		if (UCurveBase* Curve = Cast<UCurveBase>(Object))
+		// Don't compare the FName number suffix
+		if (Tag.Name.IsEqual(ReimportPathName, ENameCase::IgnoreCase, false))
 		{
-			if (JustFileNameMatches(Curve->ImportPath, FileName))
+			const FString SourceFile = FReimportManager::ResolveImportFilename(Tag.Value, Object);
+			// Check there aren't any other assets referencing this thing
+			if (FindAssetsPertainingToFile(Registry, SourceFile).Num() != 0)
 			{
-				Curve->ImportPath = FReimportManager::SanitizeImportFilename(FileName, Curve);
-				ReimportFile(Curve,FileName);
-
-				return true;
+				bCanDelete = false;
+				break;
+			}
+			else
+			{
+				AbsoluteSourceFileNames.Add(MoveTemp(SourceFile));
 			}
 		}
 	}
 
-	return false;
-}
-
-bool UAutoReimportManager::AutoReimportCurveTable( UObject* Object, const FString& FileName )
-{
-	if (GetDefault<UEditorLoadingSavingSettings>()->bAutoReimportCSV)
+	if (bCanDelete)
 	{
-		if (UCurveTable* CurveTable = Cast<UCurveTable>(Object))
+		// Offer to delete the source file(s) as well
+		const FText MessageText = FText::Format(LOCTEXT("DeleteRelatedAssets", "The asset '{0}' has been deleted. Would you like to delete the source files it was imported from as well?"), FText::FromString(Object->GetName()) );
+		if (FMessageDialog::Open(EAppMsgType::YesNo, MessageText) == EAppReturnType::Yes)
 		{
-			if (JustFileNameMatches(CurveTable->ImportPath, FileName))
-			{
-				CurveTable->ImportPath = FReimportManager::SanitizeImportFilename(FileName, CurveTable);
-				ReimportFile(CurveTable, FileName);
+			auto& FileManager = IFileManager::Get();
 
-				return true;
+			for (auto& Filename : AbsoluteSourceFileNames)
+			{
+				FileManager.Delete(*Filename, false /* RequireExists */, true /* Even if read only */, true /* Quiet */);
 			}
 		}
 	}
-
-	return false;
 }
 
-bool UAutoReimportManager::AutoReimportDataTable( UObject* Object, const FString& FileName )
+void UAutoReimportManager::SetUpMountedContentDirectories()
 {
-	if (GetDefault<UEditorLoadingSavingSettings>()->bAutoReimportCSV)
+	TArray<FString> ContentPaths;
+	FPackageName::QueryRootContentPaths(ContentPaths);
+	for (const auto& ContentPath : ContentPaths)
 	{
-		if (UDataTable* DataTable = Cast<UDataTable>(Object))
-		{
-			if (JustFileNameMatches(DataTable->ImportPath, FileName))
-			{
-				DataTable->ImportPath = FReimportManager::SanitizeImportFilename(FileName, DataTable);
-				ReimportFile(DataTable, FileName);
-
-				return true;
-			}
-		}
+		const FString Directory = FPaths::ConvertRelativePathToFull(FPackageName::LongPackageNameToFilename(ContentPath));
+		MountedDirectoryMonitors.Emplace(Directory, SupportedExtensions, ContentPath);
 	}
-
-	return false;
 }
 
-bool UAutoReimportManager::AutoReimportTexture( UObject* Object, const FString& FileName )
+void UAutoReimportManager::SetUpUserDirectories()
 {
-	if (GetDefault<UEditorLoadingSavingSettings>()->bAutoReimportTextures)
+	TArray<bool> PersistedDirectories;
+	PersistedDirectories.AddZeroed(UserDirectoryMonitors.Num());
+
+	// Also set up user-specified directories
+	TArray<FString> NewDirectories = GetDefault<UEditorLoadingSavingSettings>()->AutoReimportDirectories;
+
+	for (int32 Index = 0; Index < NewDirectories.Num(); ++Index)
 	{
-		if (UTexture* Texture = Cast<UTexture>(Object))
+		NewDirectories[Index] = FPaths::ConvertRelativePathToFull(NewDirectories[Index]);
+		const int32 Existing = UserDirectoryMonitors.IndexOfByPredicate([&](const FContentDirectoryMonitor& InMonitor){ return InMonitor.GetDirectory() == NewDirectories[Index]; });
+		if (Existing != INDEX_NONE)
 		{
-			if (IsSamePath(FReimportManager::ResolveImportFilename(Texture->SourceFilePath, Texture), FileName))
-			{
-				// Prevent the texture from being compressed immediately, so the user can see the results
-				auto SavedCompressionSetting = Texture->DeferCompression;
-				Texture->DeferCompression = true;
-
-				if (ReimportFile(Texture, FileName))
-				{
-					return true;
-				}
+			PersistedDirectories[Existing] = true;
+		}
+		else
+		{
+			// Add a new one
+			UserDirectoryMonitors.Emplace(NewDirectories[Index], SupportedExtensions);
+			PersistedDirectories.Add(true);
+		}
+	}
 
-				Texture->DeferCompression  = SavedCompressionSetting;
-			}
+	// Delete anything that is no longer needed (we delete the cache file altogether)
+	for (int32 Index = PersistedDirectories.Num() - 1; Index >= 0; --Index)
+	{
+		if (!PersistedDirectories[Index])
+		{
+			UserDirectoryMonitors[Index].Destroy();
+			UserDirectoryMonitors.RemoveAt(Index, 1, false);
 		}
 	}
 
-	return false;
+	if (UserDirectoryMonitors.GetSlack() > UserDirectoryMonitors.Max() / 2)
+	{
+		UserDirectoryMonitors.Shrink();
+	}
 }
 
-
 #undef LOCTEXT_NAMESPACE
 
diff --git a/Engine/Source/Editor/UnrealEd/Private/Editor.cpp b/Engine/Source/Editor/UnrealEd/Private/Editor.cpp
--- a/Engine/Source/Editor/UnrealEd/Private/Editor.cpp
+++ b/Engine/Source/Editor/UnrealEd/Private/Editor.cpp
@@ -2557,6 +2557,7 @@ FReimportManager* FReimportManager::Instance()
 void FReimportManager::RegisterHandler( FReimportHandler& InHandler )
 {
 	Handlers.AddUnique( &InHandler );
+	bHandlersNeedSorting = true;
 }
 
 /**
@@ -2594,16 +2595,17 @@ bool FReimportManager::Reimport( UObject* Obj, bool bAskForNewFileIfMissing )
 	bool bSuccess = false;
 	if ( Obj )
 	{
+		if (bHandlersNeedSorting)
+		{
+			// Use > operator because we want higher priorities earlier in the list
+			Handlers.Sort([](const FReimportHandler& A, const FReimportHandler& B) { return A.GetPriority() > B.GetPriority(); });
+			bHandlersNeedSorting = false;
+		}
+		
 		bool bShowNotification = true;
 		bool bValidSourceFilename = false;
 		TArray<FString> SourceFilenames;
-		
-		struct FCompareReimportHandlerPriority
-		{
-			// Use > operator because we want higher priorities earlier in the list
-			FORCEINLINE bool operator()(const FReimportHandler& A, const FReimportHandler& B) const { return A.GetPriority() > B.GetPriority(); }
-		};
-		Handlers.Sort(FCompareReimportHandlerPriority());
+
 		for( int32 HandlerIndex = 0; HandlerIndex < Handlers.Num(); ++HandlerIndex )
 		{
 			SourceFilenames.Empty();
diff --git a/Engine/Source/Editor/UnrealEd/Private/FileCache.cpp b/Engine/Source/Editor/UnrealEd/Private/FileCache.cpp
new file mode 100644
--- /dev/null
+++ b/Engine/Source/Editor/UnrealEd/Private/FileCache.cpp
@@ -0,0 +1,383 @@
+// Copyright 1998-2015 Epic Games, Inc. All Rights Reserved.
+
+
+#include "UnrealEd.h"
+#include "FileCache.h"
+#include "Developer/DirectoryWatcher/Public/DirectoryWatcherModule.h"
+
+const FGuid FFileCacheCustomVersion::Key(0x8E7DDCB3, 0x80DA47BB, 0x9FD346A2, 0x93984DF6);
+FCustomVersionRegistration GRegisterFileCacheVersion(FFileCacheCustomVersion::Key, FFileCacheCustomVersion::Latest, TEXT("FileCacheVersion"));
+
+namespace
+{
+	/** Helper template function to remove duplicates from an array, given some predicate */
+	template<typename T, typename P>
+	void RemoveDuplicates(TArray<T>& Array, P Predicate)
+	{
+		if (Array.Num() == 0)
+			return;
+
+		const int32 NumElements = Array.Num();
+		int32 Write = 0;
+		for (int32 Read = 0; Read < NumElements; ++Read)
+		{
+			for (int32 DuplRead = Read + 1; DuplRead < NumElements; ++DuplRead)
+			{
+				if (Predicate(Array[Read], Array[DuplRead]))
+				{
+					++Read;
+					goto next;
+				}
+			}
+
+			if (Write != Read)
+			{
+				Swap(Array[Write++], Array[Read]);
+			}
+	next:
+			;
+		}
+		const int32 NumDuplicates = NumElements - (Write + 1);
+		if (NumDuplicates != 0)
+		{
+			Array.RemoveAt(Write + 1, NumDuplicates, false);
+		}
+	}
+}
+
+FFileCache::FFileCache(const FFileCacheConfig& InConfig)
+	: Config(InConfig)
+	, bSavedCacheDirty(false)
+{
+	// Ensure that the extension strings are of the form ;ext1;ext2;ext3;
+	if (!Config.IncludeExtensions.IsEmpty() && Config.IncludeExtensions[0] != ';')
+	{
+		Config.IncludeExtensions.InsertAt(0, ';');
+	}
+	if (!Config.ExcludeExtensions.IsEmpty() && Config.ExcludeExtensions[0] != ';')
+	{
+		Config.ExcludeExtensions.InsertAt(0, ';');
+	}
+
+	// Attempt to load an existing cache file
+	auto ExistingCache = ReadCache();
+	if (ExistingCache.IsSet())
+	{
+		CachedDirectoryState = MoveTemp(ExistingCache.GetValue());
+	}
+
+	FDirectoryWatcherModule& Module = FModuleManager::LoadModuleChecked<FDirectoryWatcherModule>(TEXT("DirectoryWatcher"));
+	if (IDirectoryWatcher* DirectoryWatcher = Module.Get())
+	{
+		auto Callback = IDirectoryWatcher::FDirectoryChanged::CreateRaw(this, &FFileCache::OnDirectoryChanged);
+		DirectoryWatcher->RegisterDirectoryChangedCallback_Handle(Config.Directory, Callback, WatcherDelegate);
+	}
+}
+
+FFileCache::~FFileCache()
+{
+	UnbindWatcher();
+	WriteCache();
+}
+
+void FFileCache::Destroy()
+{
+	// Delete the cache file, and clear out everything
+	bSavedCacheDirty = false;
+	IFileManager::Get().Delete(*Config.CacheFile, false, true, true);
+
+	OutstandingChanges.Empty();
+	CachedDirectoryState = FDirectoryState();
+
+	UnbindWatcher();
+}
+
+void FFileCache::UnbindWatcher()
+{
+	if (WatcherDelegate == FDelegateHandle())
+	{
+		return;
+	}
+
+	if (FDirectoryWatcherModule* Module = FModuleManager::GetModulePtr<FDirectoryWatcherModule>(TEXT("DirectoryWatcher")))
+	{
+		if (IDirectoryWatcher* DirectoryWatcher = Module->Get())
+		{
+			DirectoryWatcher->UnregisterDirectoryChangedCallback_Handle(Config.Directory, WatcherDelegate);
+		}
+	}
+
+	WatcherDelegate = FDelegateHandle();
+}
+
+FDirectoryState FFileCache::WalkDirectory() const
+{
+	struct FDirectoryVisitor : public IPlatformFile::FDirectoryVisitor
+	{
+		FDirectoryState State;
+		int32 LeadingDirectoryLength;
+		IFileManager* FileManager;
+
+		virtual bool Visit(const TCHAR* FilenameOrDirectory, bool bIsDirectory)
+		{
+			if (!bIsDirectory)
+			{
+				State.Files.Emplace(FString(FilenameOrDirectory + LeadingDirectoryLength), FileManager->GetTimeStamp(FilenameOrDirectory));
+			}
+			return true;
+		}
+	};
+
+
+	auto& FileManager = IFileManager::Get();
+
+	// Iterate the directory
+	FDirectoryVisitor Visitor;
+	Visitor.LeadingDirectoryLength = Config.Directory.Len();
+	Visitor.FileManager = &FileManager;
+	FileManager.IterateDirectoryRecursively(*Config.Directory, Visitor);
+
+	return MoveTemp(Visitor.State);
+}
+
+TOptional<FDirectoryState> FFileCache::ReadCache() const
+{
+	TOptional<FDirectoryState> Optional;
+
+	FArchive* Ar = IFileManager::Get().CreateFileReader(*Config.CacheFile);
+	if (Ar)
+	{
+		FDirectoryState Result;
+
+		*Ar << Result;
+		Ar->Close();
+		delete Ar;
+
+		Optional.Emplace(MoveTemp(Result));
+	}
+
+	return Optional;
+}
+
+void FFileCache::WriteCache()
+{
+	if (bSavedCacheDirty)
+	{
+		FArchive* Ar = IFileManager::Get().CreateFileWriter(*Config.CacheFile);
+		check(Ar);
+
+		*Ar << CachedDirectoryState;
+
+		Ar->Close();
+		delete Ar;
+
+		// @todo: arodham: can't do this with maps
+		// if (CachedDirectoryState.Files.Num() < CachedDirectoryState.Files.Max() * 0.5f)
+		// {
+		// 	CachedDirectoryState.Files.Shrink();
+		// }
+
+		bSavedCacheDirty = false;
+	}
+}
+
+TArray<FUpdateCacheTransaction> FFileCache::GetOutstandingChanges()
+{
+	TArray<FUpdateCacheTransaction> Moved;
+	Swap(OutstandingChanges, Moved);
+	return Moved;
+}
+
+void FFileCache::CompleteTransaction(FUpdateCacheTransaction&& Transaction)
+{
+	auto* CachedData = CachedDirectoryState.Files.Find(Transaction.Filename);
+
+	switch (Transaction.Action)
+	{
+	case FFileChangeData::FCA_Modified:
+		if (CachedData && CachedData->Timestamp < Transaction.Timestamp)
+		{
+			// Update the timestamp
+			CachedData->Timestamp = Transaction.Timestamp;
+
+			bSavedCacheDirty = true;
+		}
+		break;
+	case FFileChangeData::FCA_Added:
+		if (!CachedData)
+		{
+			// Add the file information to the cache
+			CachedDirectoryState.Files.Emplace(Transaction.Filename, Transaction.Timestamp);
+
+			bSavedCacheDirty = true;
+		}
+		break;
+	case FFileChangeData::FCA_Removed:
+		if (CachedData)
+		{
+			// Remove the file information to the cache
+			CachedDirectoryState.Files.Remove(Transaction.Filename);
+
+			bSavedCacheDirty = true;
+		}
+		break;
+
+	default:
+		checkf(false, TEXT("Invalid file cached transaction"));
+		break;
+	}
+}
+
+void FFileCache::Scan()
+{
+	auto CurrentFileData = WalkDirectory();
+
+	if (CachedDirectoryState.Files.Num() == 0)
+	{
+		// If we don't have any cached data yet, just use the file data we just harvested
+		CachedDirectoryState = MoveTemp(CurrentFileData);
+		bSavedCacheDirty = true;
+		return;
+	}
+
+	// We already have cached data so we need to compare it with the harvested data
+	// to detect additions, modifications, and removals
+	for (const auto& FilenameAndData : CurrentFileData.Files)
+	{
+		const FString& Filename = FilenameAndData.Key.Get();
+		if (!IsFileApplicable(Filename))
+		{
+			continue;
+		}
+
+		const auto* CachedData = CachedDirectoryState.Files.Find(Filename);
+		// Have we encountered this before?
+		if (CachedData)
+		{
+			if (CachedData->Timestamp != FilenameAndData.Value.Timestamp)
+			{
+				OutstandingChanges.Add(FUpdateCacheTransaction(Filename, FFileChangeData::FCA_Modified, FilenameAndData.Value.Timestamp));
+			}
+		}
+		else
+		{
+			// No file already - create one. We don't do this if the cache didn't exist before as we have no idea what is new and what isn't
+			OutstandingChanges.Add(FUpdateCacheTransaction(Filename, FFileChangeData::FCA_Added, FilenameAndData.Value.Timestamp));
+		}
+	}
+
+	// Delete anything that doesn't exist on disk anymore
+	for (auto It = CachedDirectoryState.Files.CreateIterator(); It; ++It)
+	{
+		const FImmutableString& Filename = It.Key();
+		if (!CurrentFileData.Files.Contains(Filename) && IsFileApplicable(Filename.Get()))
+		{
+			OutstandingChanges.Add(FUpdateCacheTransaction(Filename, FFileChangeData::FCA_Removed));
+		}
+	}
+}
+
+void FFileCache::OnDirectoryChanged(const TArray<FFileChangeData>& FileChanges)
+{
+	auto& FileManager = IFileManager::Get();
+
+	OutstandingChanges.Reserve(OutstandingChanges.Num() + FileChanges.Num());
+
+	for (const auto& ThisEntry : FileChanges)
+	{
+		// If it's a directory or is not applicable, ignore it
+		if (FileManager.DirectoryExists(*ThisEntry.Filename) || !IsFileApplicable(ThisEntry.Filename))
+		{
+			continue;
+		}
+
+		FImmutableString RelativeFilename;
+		{
+			FString Temp = FPaths::ConvertRelativePathToFull(ThisEntry.Filename);
+			Temp.GetCharArray().RemoveAt(0, Config.Directory.Len());
+			RelativeFilename = Temp;
+		}
+		
+		switch (ThisEntry.Action)
+		{
+		case FFileChangeData::FCA_Added:
+			{
+				const int32 NumRemoved = OutstandingChanges.RemoveAll([&](const FUpdateCacheTransaction& X){
+					return X.Action == FFileChangeData::FCA_Removed && X.Filename == RelativeFilename;
+				});
+
+				const auto Action = NumRemoved == 0 ? FFileChangeData::FCA_Added : FFileChangeData::FCA_Modified;
+				OutstandingChanges.Add(FUpdateCacheTransaction(RelativeFilename, Action, FileManager.GetTimeStamp(*ThisEntry.Filename)));
+			}
+			break;
+
+		case FFileChangeData::FCA_Removed:
+			{
+				bool bPreviouslyAdded = false;
+				const int32 NumRemoved = OutstandingChanges.RemoveAll([&](const FUpdateCacheTransaction& X){
+					if (X.Filename == RelativeFilename)
+					{
+						bPreviouslyAdded = X.Action == FFileChangeData::FCA_Added || bPreviouslyAdded;
+						return true;
+					}
+					return false;
+				});
+
+				if (!bPreviouslyAdded)
+				{
+					OutstandingChanges.Add(FUpdateCacheTransaction(RelativeFilename, FFileChangeData::FCA_Removed));
+				}
+			}
+			
+			break;
+
+		case FFileChangeData::FCA_Modified:
+			{
+				const bool bPreviouslyAdded = OutstandingChanges.ContainsByPredicate([&](const FUpdateCacheTransaction& X){
+					return X.Filename == RelativeFilename && X.Action == FFileChangeData::FCA_Added;
+				});
+
+				if (!bPreviouslyAdded)
+				{
+					OutstandingChanges.Add(FUpdateCacheTransaction(RelativeFilename, FFileChangeData::FCA_Modified, FileManager.GetTimeStamp(*ThisEntry.Filename)));
+				}
+			}
+			break;
+
+		default:
+			break;
+		}
+	}
+
+	RemoveDuplicates(OutstandingChanges, [](const FUpdateCacheTransaction& A, const FUpdateCacheTransaction& B){
+		return A.Action == B.Action && A.Filename == B.Filename;
+	});
+}
+
+bool FFileCache::IsFileApplicable(const FString& Filename)
+{
+	return (Config.ExcludeExtensions.IsEmpty() || !MatchExtensionString(Config.ExcludeExtensions, Filename)) &&
+		   (Config.IncludeExtensions.IsEmpty() ||  MatchExtensionString(Config.IncludeExtensions, Filename));
+}
+
+bool FFileCache::MatchExtensionString(const FString& InSource, const FString& InPath)
+{
+	int32 SlashPosition = InPath.Find(TEXT("/"), ESearchCase::CaseSensitive, ESearchDir::FromEnd);
+	// in case we are using backslashes on a platform that doesn't use backslashes
+	SlashPosition = FMath::Max(SlashPosition, InPath.Find(TEXT("\\"), ESearchCase::CaseSensitive, ESearchDir::FromEnd));
+
+	const int32 DotPosition = InPath.Find(TEXT("."), ESearchCase::CaseSensitive, ESearchDir::FromEnd);
+
+	if (DotPosition > INDEX_NONE && DotPosition > SlashPosition && DotPosition < InPath.Len() - 1)
+	{
+		FString Search;
+		Search.Reserve(InPath.Len() - DotPosition + 2);;
+		Search.InsertAt(0, ';');
+		Search += &InPath[DotPosition + 1];
+		Search.AppendChar(';');
+
+		return InSource.Find(Search, ESearchCase::IgnoreCase) != INDEX_NONE;
+	}
+
+	return false;
+}
\ No newline at end of file
diff --git a/Engine/Source/Editor/UnrealEd/Private/FileHelpers.cpp b/Engine/Source/Editor/UnrealEd/Private/FileHelpers.cpp
--- a/Engine/Source/Editor/UnrealEd/Private/FileHelpers.cpp
+++ b/Engine/Source/Editor/UnrealEd/Private/FileHelpers.cpp
@@ -2856,51 +2856,56 @@ FEditorFileUtils::EPromptReturnCode FEditorFileUtils::PromptForCheckoutAndSave(
 
 			const FScopedBusyCursor BusyCursor;
 			FSaveErrorOutputDevice SaveErrors;
-			GWarn->BeginSlowTask( NSLOCTEXT("UnrealEd", "SavingPackagesE", "Saving packages..."), true );
-			for( TArray<UPackage*>::TConstIterator PackageIter( FinalSaveList ); PackageIter; ++PackageIter )
+
 			{
-				UPackage* Package = *PackageIter;
-				
-				if( !Package->IsFullyLoaded() )
-				{
-					// Packages must be fully loaded to save.
-					Package->FullyLoad();
-				}
+				FScopedSlowTask SlowTask(FinalSaveList.Num()*2, NSLOCTEXT("UnrealEd", "SavingPackagesE", "Saving packages..."));
+				SlowTask.MakeDialog();
 
-				const UWorld* const AssociatedWorld = UWorld::FindWorldInPackage(Package);
-				const bool bIsMapPackage = AssociatedWorld != nullptr;
+				for (auto* Package : FinalSaveList)
+				{
+					SlowTask.EnterProgressFrame(1);
 
-				const FText SavingPackageText = (bIsMapPackage) 
-					? FText::Format(NSLOCTEXT("UnrealEd", "SavingMapf", "Saving map {0}"), FText::FromString(Package->GetName()))
-					: FText::Format(NSLOCTEXT("UnrealEd", "SavingAssetf", "Saving asset {0}"), FText::FromString(Package->GetName()));
+					if( !Package->IsFullyLoaded() )
+					{
+						// Packages must be fully loaded to save.
+						Package->FullyLoad();
+					}
 
-				GWarn->StatusForceUpdate( PackageIter.GetIndex(), FinalSaveList.Num(), SavingPackageText );
-				
-				// Save the package
-				bool bPackageLocallyWritable;
-				const int32 SaveStatus = InternalSavePackage( Package, bPackageLocallyWritable, SaveErrors );
-				
-				// If InternalSavePackage reported that the provided package was locally writable, add it to the list of writable files
-				// to warn the user about
-				if ( bPackageLocallyWritable )
-				{
-					WritablePackageFiles.Add( Package );
-				}
+					const UWorld* const AssociatedWorld = UWorld::FindWorldInPackage(Package);
+					const bool bIsMapPackage = AssociatedWorld != nullptr;
+
+					const FText SavingPackageText = (bIsMapPackage) 
+						? FText::Format(NSLOCTEXT("UnrealEd", "SavingMapf", "Saving map {0}"), FText::FromString(Package->GetName()))
+						: FText::Format(NSLOCTEXT("UnrealEd", "SavingAssetf", "Saving asset {0}"), FText::FromString(Package->GetName()));
+
+					SlowTask.EnterProgressFrame(1, SavingPackageText);
+					
+					// Save the package
+					bool bPackageLocallyWritable;
+					const int32 SaveStatus = InternalSavePackage( Package, bPackageLocallyWritable, SaveErrors );
+					
+					// If InternalSavePackage reported that the provided package was locally writable, add it to the list of writable files
+					// to warn the user about
+					if ( bPackageLocallyWritable )
+					{
+						WritablePackageFiles.Add( Package );
+					}
 
-				if( SaveStatus == EAppReturnType::No )
-				{
-					// The package could not be saved so add it to the failed array and change the return response to indicate failure
-					FailedPackages.Add( Package );
-					ReturnResponse = PR_Failure;
-				}
-				else if( SaveStatus == EAppReturnType::Cancel )
-				{
-					// No need to save anything else, the user wants to cancel everything
-					ReturnResponse = PR_Cancelled;
-					break;
+					if( SaveStatus == EAppReturnType::No )
+					{
+						// The package could not be saved so add it to the failed array and change the return response to indicate failure
+						FailedPackages.Add( Package );
+						ReturnResponse = PR_Failure;
+					}
+					else if( SaveStatus == EAppReturnType::Cancel )
+					{
+						// No need to save anything else, the user wants to cancel everything
+						ReturnResponse = PR_Cancelled;
+						break;
+					}
 				}
 			}
-			GWarn->EndSlowTask();
+
 			SaveErrors.Flush();
 
 			if( UserResponse == false && PackagesNotNeedingCheckout.Num() > 0 )
diff --git a/Engine/Source/Editor/UnrealEd/Private/Settings/SettingsClasses.cpp b/Engine/Source/Editor/UnrealEd/Private/Settings/SettingsClasses.cpp
--- a/Engine/Source/Editor/UnrealEd/Private/Settings/SettingsClasses.cpp
+++ b/Engine/Source/Editor/UnrealEd/Private/Settings/SettingsClasses.cpp
@@ -90,6 +90,7 @@ void UEditorExperimentalSettings::PostEditChangeProperty( struct FPropertyChange
 
 UEditorLoadingSavingSettings::UEditorLoadingSavingSettings( const FObjectInitializer& ObjectInitializer )
 	: Super(ObjectInitializer)
+	, bMonitorContentDirectories(false)
 {
 	TextDiffToolPath.FilePath = TEXT("P4Merge.exe");
 }
diff --git a/Engine/Source/Editor/UnrealEd/Private/UnrealEdEngine.cpp b/Engine/Source/Editor/UnrealEd/Private/UnrealEdEngine.cpp
--- a/Engine/Source/Editor/UnrealEd/Private/UnrealEdEngine.cpp
+++ b/Engine/Source/Editor/UnrealEd/Private/UnrealEdEngine.cpp
@@ -28,6 +28,7 @@
 #include "ComponentVisualizer.h"
 #include "Editor/EditorLiveStreaming/Public/IEditorLiveStreaming.h"
 #include "SourceCodeNavigation.h"
+#include "AutoReimport/AutoReimportManager.h"
 #include "NotificationManager.h"
 #include "SNotificationList.h"
 
@@ -481,16 +482,16 @@ void UUnrealEdEngine::OnSourceControlStateUpdated(const FSourceControlOperationR
 					// To get here, either "prompt for checkout on asset modification" is set, or "automatically checkout on asset modification"
 					// is set, but it failed.
 
-					// Allow packages that are not checked out to pass through.
-					// Allow packages that are not current or checked out by others pass through.  
-					// The user wont be able to checkout these packages but the checkout dialog will show up with a special icon 
-					// to let the user know they wont be able to checkout the package they are modifying.
+			// Allow packages that are not checked out to pass through.
+			// Allow packages that are not current or checked out by others pass through.  
+			// The user wont be able to checkout these packages but the checkout dialog will show up with a special icon 
+			// to let the user know they wont be able to checkout the package they are modifying.
 
 					PackageToNotifyState.Add(Package, SourceControlState->CanCheckout() ? NS_PendingPrompt : NS_PendingWarning);
-					// We need to prompt since a new package was added
-					bNeedToPromptForCheckout = true;
-				}
-			}
+			// We need to prompt since a new package was added
+			bNeedToPromptForCheckout = true;
+		}
+	}
 		}
 	}
 }
@@ -510,7 +511,7 @@ void UUnrealEdEngine::OnPackageCheckedOut(const FSourceControlOperationRef& Sour
 		if (ResultType == ECommandResult::Succeeded)
 		{
 			if (SourceControlState.IsValid() && SourceControlState->IsCheckedOut())
-			{
+	{
 				FNotificationInfo Notification(FText::Format(NSLOCTEXT("SourceControl", "AutoCheckOutNotification", "Package '{Package}' automatically checked out."), Arguments));
 				Notification.bFireAndForget = true;
 				Notification.ExpireDuration = 4.0f;
diff --git a/Engine/Source/Editor/UnrealEd/Public/AutoReimport/AutoReimportManager.h b/Engine/Source/Editor/UnrealEd/Public/AutoReimport/AutoReimportManager.h
new file mode 100644
--- /dev/null
+++ b/Engine/Source/Editor/UnrealEd/Public/AutoReimport/AutoReimportManager.h
@@ -0,0 +1,124 @@
+// Copyright 1998-2015 Epic Games, Inc. All Rights Reserved.
+
+#pragma once
+
+#include "FileCache.h"
+
+#include "AutoReimportManager.generated.h"
+
+class IAssetRegistry;
+
+/** Class responsible for watching a specific content directory for changes */
+class FContentDirectoryMonitor
+{
+public:
+
+	/**
+	 * Constructor.
+	 * @param InDirectory			Content directory path to monitor. Assumed to be absolute.
+	 * @param InSupportedExtensions A string containing semi-colon separated extensions to monitor.
+	 * @param InMountedContentPath	(optional) Mounted content path (eg /Engine/, /Game/) to which InDirectory maps.
+	 */
+	FContentDirectoryMonitor(const FString& InDirectory, const FString& InSupportedExtensions, const FString& InMountedContentPath = FString());
+
+	/** Tick this directory monitor. Will potentially process changed files */
+	void Tick(const IAssetRegistry& Registry);
+
+	/** Destroy this monitor including its cache */
+	void Destroy() { return Cache.Destroy(); }
+
+	/** Get the directory that this monitor applies to */
+	const FString& GetDirectory() const { return Cache.GetDirectory(); }
+
+	/** Get the mounted content path that this monitor applies to */
+	const FString& GetMountedContentPath() const { return MountedContentPath; }
+
+private:
+
+	/** Process the outstanding changes that we have cached */
+	void ProcessOutstandingChanges(const IAssetRegistry& Registry);
+
+private:
+
+	/** The file cache that monitors and reflects the content directory. */
+	FFileCache Cache;
+
+	/** The mounted content path for this monitor (eg /Engine/) */
+	FString MountedContentPath;
+
+	/** A list of file system changes that are due to be processed */
+	TArray<FUpdateCacheTransaction> OutstandingChanges;
+
+	/** The time we last processed outstanding changes */
+	double LastProcessTime;
+
+	/** The last time we attempted to save the cache file */
+	double LastSaveTime;
+
+	/** The amount of time (in seconds) we must wait before reacting to changed files */
+	static const int32 ReactionTimeS = 1;
+
+	/** The interval between potential re-saves of the cache file */
+	static const int32 ResaveIntervalS = 60;
+};
+
+
+/* Deals with auto reimporting of objects when the objects file on disk is modified*/
+UCLASS(config=Editor, transient)
+class UNREALED_API UAutoReimportManager : public UObject, public FTickableEditorObject
+{
+	GENERATED_UCLASS_BODY()
+
+public:
+
+	/** Initialize during engine startup */
+	void Initialize();
+
+private:
+
+	/** FTickableEditorObject interface*/
+	virtual void Tick(float DeltaTime) override;
+	virtual bool IsTickable() const override { return true; }
+	virtual TStatId GetStatId() const override;
+
+private:
+
+	/** UObject Interface */
+	virtual void BeginDestroy() override;
+
+private:
+
+	/** Called when a new asset path has been mounted */
+	void OnContentPathMounted(const FString& InAssetPath, const FString& FileSystemPath);
+
+	/** Called when a new asset path has been dismounted */
+	void OnContentPathDismounted(const FString& InAssetPath, const FString& FileSystemPath);
+
+	/** Called when an asset has been deleted */
+	void OnAssetDeleted(UObject* Object);
+
+private:
+
+	/** Callback for when an editor user setting has changed */
+	void HandleLoadingSavingSettingChanged(FName PropertyName);
+
+	/** Set up monitors to all the mounted content directories */
+	void SetUpMountedContentDirectories();
+
+	/** Set up the user directories from the editor Loading/Saving settings */
+	void SetUpUserDirectories();
+
+private:
+
+	/** Cached string of extensions that are supported by all available factories */
+	FString SupportedExtensions;
+
+	/** Array of directory monitors for automatically managed mounted content paths */
+	TArray<FContentDirectoryMonitor> MountedDirectoryMonitors;
+
+	/** Additional external directories to monitor for changes specified by the user */
+	TArray<FContentDirectoryMonitor> UserDirectoryMonitors;
+
+	/** Set when we are processing changes to avoid reentrant logic for asset deletions etc */
+	bool bIsProcessingChanges;
+};
diff --git a/Engine/Source/Editor/UnrealEd/Public/EditorReimportHandler.h b/Engine/Source/Editor/UnrealEd/Public/EditorReimportHandler.h
--- a/Engine/Source/Editor/UnrealEd/Public/EditorReimportHandler.h
+++ b/Engine/Source/Editor/UnrealEd/Public/EditorReimportHandler.h
@@ -95,6 +95,9 @@ private:
 	/** Reimport handlers registered with this manager */
 	TArray<FReimportHandler*> Handlers;
 
+	/** True when the Handlers array has been modified such that it needs sorting */
+	bool bHandlersNeedSorting;
+
 	/** Delegate to call before the asset is reimported */
 	FPreReimportNotification PreReimport;
 
diff --git a/Engine/Source/Editor/UnrealEd/Public/FileCache.h b/Engine/Source/Editor/UnrealEd/Public/FileCache.h
new file mode 100644
--- /dev/null
+++ b/Engine/Source/Editor/UnrealEd/Public/FileCache.h
@@ -0,0 +1,249 @@
+// Copyright 1998-2015 Epic Games, Inc. All Rights Reserved.
+
+#pragma once
+
+#include "Serialization/CustomVersion.h"
+#include "IDirectoryWatcher.h"
+
+/** An immutable string with a cached CRC for efficient comparison with other strings */
+struct FImmutableString
+{
+	/** Constructible from a raw string */
+	FImmutableString(FString InString = FString()) : String(MoveTemp(InString)), CachedHash(0) {}
+	FImmutableString& operator=(FString InString) { String = MoveTemp(InString); CachedHash = 0; return *this; }
+
+	FImmutableString(const FImmutableString&) = default;
+	FImmutableString& operator=(const FImmutableString&) = default;
+
+	/** Move construction/assignment */
+	FImmutableString(FImmutableString&& In) : String(MoveTemp(In.String)), CachedHash(MoveTemp(In.CachedHash)) {}
+	FImmutableString& operator=(FImmutableString&& In) { Swap(String, In.String); Swap(CachedHash, In.CachedHash); return *this; }
+
+	/** Check for equality with another immutable string (case-sensitive) */
+	friend bool operator==(const FImmutableString& A, const FImmutableString& B)
+	{
+		return GetTypeHash(A) == GetTypeHash(B) && A.String == B.String;
+	}
+
+	/** Calculate the type hash for this string */
+	friend uint32 GetTypeHash(const FImmutableString& In)
+	{
+		if (In.CachedHash == 0 && In.String.Len() != 0)
+		{
+			In.CachedHash = GetTypeHash(In.String);
+		}
+		return In.CachedHash;
+	}
+
+	/** Get the underlying string */
+	const FString& Get() const { return String; }
+
+	/** Serialise this string */
+	friend FArchive& operator<<(FArchive& Ar, FImmutableString& String)
+	{
+		Ar << String.String;
+		if (Ar.IsSaving())
+		{
+			GetTypeHash(String);	
+		}
+		
+		Ar << String.CachedHash;
+
+		return Ar;
+	}
+
+private:
+	/** The string itself. */
+	FString String;
+
+	/** The cached hash of our string. 0 until compared with another string. */
+	mutable uint32 CachedHash;
+};
+
+/** Custom serialization version for FFileCache */
+struct UNREALED_API FFileCacheCustomVersion
+{
+	static const FGuid Key;
+	enum Type {	Initial, Latest = Initial };
+};
+
+/** Structure representing specific information about a particular file */
+struct FFileData
+{
+	/** Constructiors */
+	FFileData() {}
+	FFileData(const FDateTime& InTimestamp) : Timestamp(InTimestamp) {}
+
+	/** Serializer for this type */
+	friend FArchive& operator<<(FArchive& Ar, FFileData& Data)
+	{
+		if (Ar.CustomVer(FFileCacheCustomVersion::Key) >= FFileCacheCustomVersion::Initial)
+		{
+			Ar << Data.Timestamp;
+		}
+
+		return Ar;
+	}
+
+	FDateTime Timestamp;
+};
+
+/** Structure representing the file data for a number of files in a directory */
+struct FDirectoryState
+{
+	/** Default construction */
+	FDirectoryState() {}
+
+	/** Move construction */
+	FDirectoryState(FDirectoryState&& In) : Files(MoveTemp(In.Files)) {}
+	FDirectoryState& operator=(FDirectoryState&& In) { Swap(Files, In.Files); return *this; }
+
+	/** Filename -> data map */
+	TMap<FImmutableString, FFileData> Files;
+
+	/** Serializer for this type */
+	friend FArchive& operator<<(FArchive& Ar, FDirectoryState& State)
+	{
+		Ar.UsingCustomVersion(FFileCacheCustomVersion::Key);
+
+		if (Ar.CustomVer(FFileCacheCustomVersion::Key) >= FFileCacheCustomVersion::Initial)
+		{
+			// Number of files
+			int32 Num = State.Files.Num();
+			Ar << Num;
+			if (Ar.IsLoading())
+			{
+				State.Files.Reserve(Num);
+			}
+
+			Ar << State.Files;
+		}
+
+		return Ar;
+	}
+};
+
+/** A transaction issued by FFileCache to describe a change to the cache. The change is only committed once the transaction is returned to the cache (see FFileCache::CompleteTransaction). */
+struct FUpdateCacheTransaction
+{
+	/** The path of the file to which this transaction relates */
+	FImmutableString Filename;
+
+	/** The timestamp of the file at the time of this change */
+	FDateTime Timestamp;
+
+	/** The type of action that prompted this transaction */
+	FFileChangeData::EFileChangeAction Action;
+
+	/** Publically moveable */
+	FUpdateCacheTransaction(FUpdateCacheTransaction&& In) : Filename(MoveTemp(In.Filename)), Timestamp(MoveTemp(In.Timestamp)), Action(MoveTemp(In.Action)) {}
+	FUpdateCacheTransaction& operator=(FUpdateCacheTransaction&& In) { Swap(Filename, In.Filename); Swap(Timestamp, In.Timestamp); Swap(Action, In.Action); return *this; }
+
+private:
+	friend class FFileCache;
+	
+	/** Construction responsibility is held by FFileCache */
+	FUpdateCacheTransaction(FImmutableString InFilename, FFileChangeData::EFileChangeAction InAction, const FDateTime& InTimestamp = 0)
+		: Filename(MoveTemp(InFilename)), Timestamp(InTimestamp), Action(InAction)
+	{}
+
+	/** Copyable only by FFileCache*/
+	FUpdateCacheTransaction(const FUpdateCacheTransaction& In) = default;
+	FUpdateCacheTransaction& operator=(const FUpdateCacheTransaction& In) = default;
+};
+
+/** Configuration structure required to construct a FFileCache */
+struct FFileCacheConfig
+{
+	/** String specifying the directory on disk that the cache should reflect */
+	FString Directory;
+
+	/** String specifying the file that the cache should be saved to */
+	FString CacheFile;
+
+	/**
+	 * Strings specifying extensions to include/exclude in the cache (or empty to include everything).
+	 * Specified as a list of semi-colon separated extensions (eg "png;jpg;gif;").
+	 */
+	FString IncludeExtensions, ExcludeExtensions;
+};
+
+/**
+ * A class responsible for scanning a directory, and maintaining a cache of its state (files and timestamps).
+ * Changes in the cache can be retrieved through GetOutstandingChanges(). Changes will be reported for any
+ * change in the cached state even between runs of the process.
+ */
+class UNREALED_API FFileCache
+{
+public:
+
+	/** Construction from a config */
+	FFileCache(const FFileCacheConfig& InConfig);
+
+	FFileCache(const FFileCache&) = delete;
+	FFileCache& operator=(const FFileCache&) = delete;
+
+	/** Destructor */
+	~FFileCache();
+
+	/** Destroy this cache. Cleans out the in-memory state and deletes the cache file, if present. */
+	void Destroy();
+
+	/** Get the absolute path of the directory this cache reflects */
+	const FString& GetDirectory() const { return Config.Directory; }
+
+	/** Force a scan the directory to detect any changes in the cache. */
+	void Scan();
+
+	/** Write out the cached file, if we have any changes to write */
+	void WriteCache();
+
+	/**
+	 * Return a transaction to the cache for completion. Will update the cached state with the change
+	 * described in the transaction, and mark the cache as needing to be saved.
+	 */
+	void CompleteTransaction(FUpdateCacheTransaction&& Transaction);
+
+	/** Get all pending changes to the cache. Transactions must be returned to CompleteTransaction to update the cache. */
+	TArray<FUpdateCacheTransaction> GetOutstandingChanges();
+
+private:
+
+	/** Called when the directory we are monitoring has been changed in some way */
+	void OnDirectoryChanged(const TArray<FFileChangeData>& FileChanges);
+
+	/** Check whether a file is applicable to our cache */
+	bool IsFileApplicable(const FString& InPath);
+
+	/**
+	 * Check whether a filename's extension matches the specified extension list.
+	 * InExtensions must be of the form ";ext1;ext2;ext3;". Note the leading and trailing ;
+	 */
+	static bool MatchExtensionString(const FString& InExtensions, const FString& InFilename);
+
+	/** Unbind the watcher from the directory. Called on destruction. */
+	void UnbindWatcher();
+	
+	/** Walk the directory and return all the file information we find pertaining to all applicable files */
+	FDirectoryState WalkDirectory() const;
+
+	/** Read the cache file data and return the contents */
+	TOptional<FDirectoryState> ReadCache() const;
+
+private:
+
+	/** Temporary array of transactions to be returned to the client through GetOutstandingChanges() */
+	TArray<FUpdateCacheTransaction> OutstandingChanges;
+
+	/** Our in-memory view of the cached directory state. */
+	FDirectoryState CachedDirectoryState;
+
+	/** Configuration settings applied on construction */
+	FFileCacheConfig Config;
+
+	/** Handle to the directory watcher delegate so we can delete it properly */
+	FDelegateHandle WatcherDelegate;
+
+	/** True when the cached state we have in memory is more up to date than the serialized file. Enables WriteCache() when true. */
+	bool bSavedCacheDirty;
+};
diff --git a/Engine/Source/Runtime/AssetRegistry/Private/AssetRegistry.cpp b/Engine/Source/Runtime/AssetRegistry/Private/AssetRegistry.cpp
--- a/Engine/Source/Runtime/AssetRegistry/Private/AssetRegistry.cpp
+++ b/Engine/Source/Runtime/AssetRegistry/Private/AssetRegistry.cpp
@@ -11,6 +11,9 @@
 
 DEFINE_LOG_CATEGORY(LogAssetRegistry);
 
+
+static const FName ReimportPathName("ReimportPath");
+
 /** Returns the appropriate ChunkProgressReportingType for the given Asset enum */
 EChunkProgressReportingType::Type GetChunkAvailabilityProgressType(EAssetAvailabilityProgressReportingType::Type ReportType)
 {
@@ -169,6 +172,9 @@ FAssetRegistry::~FAssetRegistry()
 	CachedAssetsByClass.Empty();
 	CachedAssetsByTag.Empty();
 	CachedDependsNodes.Empty();
+#if WITH_EDITORONLY_DATA
+	CachedAssetsBySourceFileName.Empty();
+#endif
 
 	// Stop listening for content mount point events
 	FPackageName::OnContentPathMounted().RemoveAll( this );
@@ -549,6 +555,21 @@ bool FAssetRegistry::GetAssets(const FARFilter& Filter, TArray<FAssetData>& OutA
 		}
 	}
 
+#if WITH_EDITORONLY_DATA
+	// Asset source file filter
+	if ( GIsEditor && Filter.SourceFilenames.Num() )
+	{
+		for (const auto& Filename : Filter.SourceFilenames)
+		{
+			const auto* Assets = CachedAssetsBySourceFileName.Find(Filename);
+			if (Assets)
+			{
+				DiskFilterSets.Add(*Assets);
+			}
+		}
+	}
+#endif
+
 	// If we have any filter sets, add the assets which are contained in the sets to OutAssetData
 	if ( DiskFilterSets.Num() > 0 )
 	{
@@ -1626,6 +1647,15 @@ void FAssetRegistry::AddAssetData(FAssetData* AssetData)
 
 		auto& TagAssets = CachedAssetsByTag.FindOrAdd(Key);
 		TagAssets.Add(AssetData);
+
+#if WITH_EDITORONLY_DATA
+		// Don't check FName number so we still accumulate them even if there are multiple
+		if (GIsEditor && TagIt.Key().IsEqual(ReimportPathName, ENameCase::IgnoreCase, false))
+		{
+			auto& SourceFiles = CachedAssetsBySourceFileName.FindOrAdd(*FPaths::GetCleanFilename(TagIt.Value()));
+			SourceFiles.Add(AssetData);
+		}
+#endif
 	}
 
 	// Notify subscribers
@@ -1709,6 +1739,15 @@ void FAssetRegistry::UpdateAssetData(FAssetData* AssetData, const FAssetData& Ne
 			auto OldTagAssets = CachedAssetsByTag.Find(FNameKey);
 
 			OldTagAssets->Remove(AssetData);
+
+#if WITH_EDITORONLY_DATA
+			// Don't check FName number so we still remove them even if there are multiple
+			if (GIsEditor && TagIt.Key().IsEqual(ReimportPathName, ENameCase::IgnoreCase, false))
+			{
+				auto* SourceFiles = CachedAssetsBySourceFileName.Find(*FPaths::GetCleanFilename(TagIt.Value()));
+				SourceFiles->Remove(AssetData);
+			}
+#endif
 		}
 
 		for (TMap<FName, FString>::TConstIterator TagIt(NewAssetData.TagsAndValues); TagIt; ++TagIt)
@@ -1717,6 +1756,15 @@ void FAssetRegistry::UpdateAssetData(FAssetData* AssetData, const FAssetData& Ne
 			auto& NewTagAssets = CachedAssetsByTag.FindOrAdd(FNameKey);
 
 			NewTagAssets.Add(AssetData);
+
+#if WITH_EDITORONLY_DATA
+			// Don't check FName number so we still add them even if there are multiple
+			if (GIsEditor && TagIt.Key().IsEqual(ReimportPathName, ENameCase::IgnoreCase, false))
+			{
+				auto& SourceFiles = CachedAssetsBySourceFileName.FindOrAdd(*FPaths::GetCleanFilename(TagIt.Value()));
+				SourceFiles.Add(AssetData);
+			}
+#endif
 		}
 	}
 
@@ -1779,6 +1827,16 @@ bool FAssetRegistry::RemoveAssetData(FAssetData* AssetData)
 		{
 			auto OldTagAssets = CachedAssetsByTag.Find(TagIt.Key());
 			OldTagAssets->Remove(AssetData);
+
+
+#if WITH_EDITORONLY_DATA
+			// Don't check FName number so we still remove them even if there are multiple
+			if (GIsEditor && TagIt.Key().IsEqual(ReimportPathName, ENameCase::IgnoreCase, false))
+			{
+				auto* SourceFiles = CachedAssetsBySourceFileName.Find(*FPaths::GetCleanFilename(TagIt.Value()));
+				SourceFiles->Remove(AssetData);
+			}
+#endif
 		}
 
 		// We need to update the cached dependencies references cache so that they know we no
diff --git a/Engine/Source/Runtime/AssetRegistry/Private/AssetRegistry.h b/Engine/Source/Runtime/AssetRegistry/Private/AssetRegistry.h
--- a/Engine/Source/Runtime/AssetRegistry/Private/AssetRegistry.h
+++ b/Engine/Source/Runtime/AssetRegistry/Private/AssetRegistry.h
@@ -181,6 +181,11 @@ private:
 	/** The map of asset tag to asset data for assets saved to disk */
 	TMap<FName, TArray<FAssetData*> > CachedAssetsByTag;
 
+#if WITH_EDITORONLY_DATA
+	/** A map of clean source file name (asset.ext) to assets saved to disk */
+	TMap<FName, TArray<FAssetData*>> CachedAssetsBySourceFileName;
+#endif
+
 	/** A map of object names to dependency data */
 	TMap<FName, FDependsNode*> CachedDependsNodes;
 
diff --git a/Engine/Source/Runtime/AssetRegistry/Private/BackgroundAssetData.cpp b/Engine/Source/Runtime/AssetRegistry/Private/BackgroundAssetData.cpp
--- a/Engine/Source/Runtime/AssetRegistry/Private/BackgroundAssetData.cpp
+++ b/Engine/Source/Runtime/AssetRegistry/Private/BackgroundAssetData.cpp
@@ -2,7 +2,7 @@
 
 #include "AssetRegistryPCH.h"
 
-FBackgroundAssetData::FBackgroundAssetData(const FString& InPackageName, const FString& InPackagePath, const FString& InGroupNames, const FString& InAssetName, const FString& InAssetClass, const TMap<FString, FString>& InTags, const TArray<int32>& InChunkIDs)
+FBackgroundAssetData::FBackgroundAssetData(const FString& InPackageName, const FString& InPackagePath, const FString& InGroupNames, const FString& InAssetName, const FString& InAssetClass, const TMultiMap<FString, FString>& InTags, const TArray<int32>& InChunkIDs)
 {
 	PackageName = InPackageName;
 	PackagePath = InPackagePath;
@@ -44,9 +44,23 @@ FBackgroundAssetData::FBackgroundAssetData(const FAssetData& InAssetData)
 FAssetData FBackgroundAssetData::ToAssetData() const
 {
 	TMap<FName, FString> CopiedTagsAndValues;
-	for (TMap<FString,FString>::TConstIterator TagIt(TagsAndValues); TagIt; ++TagIt)
+
+	// Copy over tags and values, retaining unique entries in the map for duplicate names
+	TSet<FString> VisitedKeys;
+	for (const auto& Pair : TagsAndValues)
 	{
-		CopiedTagsAndValues.Add(FName(*TagIt.Key()), TagIt.Value());
+		if (VisitedKeys.Find(Pair.Key))
+		{
+			continue;
+		}
+
+		VisitedKeys.Add(Pair.Key);
+
+		int32 NameIndex = 0;
+		for (auto It = TagsAndValues.CreateConstKeyIterator(Pair.Key); It; ++It, ++NameIndex)
+		{
+			CopiedTagsAndValues.Add(FName(FName(*It.Key()), NameIndex), *It.Value());
+		}
 	}
 
 	return FAssetData(
diff --git a/Engine/Source/Runtime/AssetRegistry/Private/BackgroundAssetData.h b/Engine/Source/Runtime/AssetRegistry/Private/BackgroundAssetData.h
--- a/Engine/Source/Runtime/AssetRegistry/Private/BackgroundAssetData.h
+++ b/Engine/Source/Runtime/AssetRegistry/Private/BackgroundAssetData.h
@@ -7,7 +7,7 @@ class FBackgroundAssetData
 {
 public:
 	/** Constructor */
-	FBackgroundAssetData(const FString& InPackageName, const FString& InPackagePath, const FString& InGroupNames, const FString& InAssetName, const FString& InAssetClass, const TMap<FString, FString>& InTags, const TArray<int32>& InChunkIDs);
+	FBackgroundAssetData(const FString& InPackageName, const FString& InPackagePath, const FString& InGroupNames, const FString& InAssetName, const FString& InAssetClass, const TMultiMap<FString, FString>& InTags, const TArray<int32>& InChunkIDs);
 	FBackgroundAssetData(const FAssetData& InAssetData);
 
 	/** Creates an AssetData object based on this object */
@@ -26,7 +26,7 @@ public:
 	/** The name of the asset's class */
 	FString AssetClass;
 	/** The map of values for properties that were marked AssetRegistrySearchable */
-	TMap<FString, FString> TagsAndValues;
+	TMultiMap<FString, FString> TagsAndValues;
 	/** The IDs of the chunks this asset is located in for streaming install.  Empty if not assigned to a chunk */
 	TArray<int32> ChunkIDs;
 };
diff --git a/Engine/Source/Runtime/AssetRegistry/Private/PackageReader.cpp b/Engine/Source/Runtime/AssetRegistry/Private/PackageReader.cpp
--- a/Engine/Source/Runtime/AssetRegistry/Private/PackageReader.cpp
+++ b/Engine/Source/Runtime/AssetRegistry/Private/PackageReader.cpp
@@ -139,7 +139,7 @@ bool FPackageReader::ReadAssetRegistryData (TArray<FBackgroundAssetData*>& Asset
 		if (bLegacyPackage || bNoMapAsset)
 		{
 			const FString AssetName = FPackageName::GetLongPackageAssetName(PackageName);
-			TMap<FString, FString> TagsAndValues;
+			TMultiMap<FString, FString> TagsAndValues;
 			AssetDataList.Add(new FBackgroundAssetData(PackageName, PackagePath, TEXT(""), AssetName, TEXT("World"), TagsAndValues, PackageFileSummary.ChunkIDs));
 		}
 	}
@@ -154,7 +154,7 @@ bool FPackageReader::ReadAssetRegistryData (TArray<FBackgroundAssetData*>& Asset
 		*this << ObjectClassName;
 		*this << TagCount;
 
-		TMap<FString, FString> TagsAndValues;
+		TMultiMap<FString, FString> TagsAndValues;
 
 		for(int32 TagIdx = 0; TagIdx < TagCount; ++TagIdx)
 		{
@@ -244,7 +244,7 @@ bool FPackageReader::ReadAssetDataFromThumbnailCache(TArray<FBackgroundAssetData
 		}
 
 		// Create a new FBackgroundAssetData for this asset and update it with the gathered data
-		TMap<FString, FString> TagsAndValues;
+		TMultiMap<FString, FString> TagsAndValues;
 		AssetDataList.Add(new FBackgroundAssetData(PackageName, PackagePath, GroupNames, AssetName, AssetClassName, TagsAndValues, PackageFileSummary.ChunkIDs));
 	}
 
diff --git a/Engine/Source/Runtime/AssetRegistry/Public/ARFilter.h b/Engine/Source/Runtime/AssetRegistry/Public/ARFilter.h
--- a/Engine/Source/Runtime/AssetRegistry/Public/ARFilter.h
+++ b/Engine/Source/Runtime/AssetRegistry/Public/ARFilter.h
@@ -15,6 +15,11 @@ struct FARFilter
 	TArray<FName> ClassNames;
 	/** The filter component for properties marked with the AssetRegistrySearchable flag */
 	TMultiMap<FName, FString> TagsAndValues;
+#if WITH_EDITORONLY_DATA
+	/** The filter component for asset source file names */
+	TArray<FName> SourceFilenames;
+#endif
+
 	/** If bRecursiveClasses is true, this results will exclude classes (including subclasses) in this list */
 	TSet<FName> RecursiveClassesExclusionSet;
 	/** If true, PackagePath components will be recursive */
@@ -38,7 +43,9 @@ struct FARFilter
 		PackagePaths.Append(Other.PackagePaths);
 		ObjectPaths.Append(Other.ObjectPaths);
 		ClassNames.Append(Other.ClassNames);
-
+#if WITH_EDITORONLY_DATA
+		SourceFilenames.Append(Other.SourceFilenames);
+#endif
 		for (auto TagIt = Other.TagsAndValues.CreateConstIterator(); TagIt; ++TagIt)
 		{
 			TagsAndValues.Add(TagIt.Key(), TagIt.Value());
@@ -54,7 +61,11 @@ struct FARFilter
 	/** Returns true if this filter has no entries */
 	bool IsEmpty() const
 	{
-		return PackageNames.Num() + PackagePaths.Num() + ObjectPaths.Num() + ClassNames.Num() + TagsAndValues.Num() == 0;
+		return PackageNames.Num() + PackagePaths.Num() + ObjectPaths.Num() + ClassNames.Num() + TagsAndValues.Num()
+	#if WITH_EDITORONLY_DATA
+			+ SourceFilenames.Num()
+	#endif
+			== 0;
 	}
 
 	/** Clears this filter of all entries */
@@ -65,6 +76,9 @@ struct FARFilter
 		ObjectPaths.Empty();
 		ClassNames.Empty();
 		TagsAndValues.Empty();
+#if WITH_EDITORONLY_DATA
+		SourceFilenames.Empty();
+#endif
 		RecursiveClassesExclusionSet.Empty();
 
 		bRecursivePaths = false;
diff --git a/Engine/Source/Runtime/CoreUObject/Private/UObject/Obj.cpp b/Engine/Source/Runtime/CoreUObject/Private/UObject/Obj.cpp
--- a/Engine/Source/Runtime/CoreUObject/Private/UObject/Obj.cpp
+++ b/Engine/Source/Runtime/CoreUObject/Private/UObject/Obj.cpp
@@ -1123,13 +1123,42 @@ static FString GetConfigFilename( UObject* SourceObject )
 
 void UObject::FAssetRegistryTag::GetAssetRegistryTagsFromSearchableProperties(const UObject* Object, TArray<FAssetRegistryTag>& OutTags)
 {
+	static FName ReimportPathName(TEXT("ReimportPath"));
+
 	check(NULL != Object);
 	for( TFieldIterator<UProperty> FieldIt( Object->GetClass() ); FieldIt; ++FieldIt )
 	{
 		if ( FieldIt->HasAnyPropertyFlags(CPF_AssetRegistrySearchable) )
 		{
-			FString PropertyStr;
 			const uint8* PropertyAddr = FieldIt->ContainerPtrToValuePtr<uint8>(Object);
+
+			if (const auto* Property = Cast<UStructProperty>(*FieldIt))
+			{
+				if (Property->Struct)
+				{
+					GetAssetRegistryTagsFromSearchableProperties(Property->Struct, OutTags);
+				}
+				continue;
+			}
+			else if (const auto* Property = Cast<UObjectProperty>(*FieldIt))
+			{
+				if (const UObject* Obj = Property->GetObjectPropertyValue(PropertyAddr))
+				{
+					GetAssetRegistryTagsFromSearchableProperties(Obj, OutTags);
+				}
+				continue;
+			}
+
+
+#if WITH_EDITOR
+			if (!GIsEditor && FieldIt->HasMetaData(ReimportPathName))
+			{
+				// ignore reimport paths in non-editor
+				continue;
+			}
+#endif
+
+			FString PropertyStr;
 			FieldIt->ExportTextItem( PropertyStr, PropertyAddr, PropertyAddr, NULL, PPF_None );
 
 			FAssetRegistryTag::ETagType TagType;
@@ -1165,7 +1194,15 @@ void UObject::FAssetRegistryTag::GetAssetRegistryTagsFromSearchableProperties(co
 				TagType = FAssetRegistryTag::TT_Alphabetical;
 			}
 
-			OutTags.Add( FAssetRegistryTag(FieldIt->GetFName(), PropertyStr, TagType) );
+			FName Key = FieldIt->GetFName();
+
+#if WITH_EDITOR
+			if (GIsEditor && FieldIt->HasMetaData(ReimportPathName))
+			{
+				Key = ReimportPathName;
+			}
+#endif
+			OutTags.Add( FAssetRegistryTag(Key, PropertyStr, TagType) );
 		}
 	}
 }
diff --git a/Engine/Source/Runtime/Engine/Classes/Animation/AnimSequence.h b/Engine/Source/Runtime/Engine/Classes/Animation/AnimSequence.h
--- a/Engine/Source/Runtime/Engine/Classes/Animation/AnimSequence.h
+++ b/Engine/Source/Runtime/Engine/Classes/Animation/AnimSequence.h
@@ -525,7 +525,7 @@ class UAnimSequence : public UAnimSequenceBase
 	uint32 bWasCompressedWithoutTranslations:1;
 
 	/** Importing data and options used for this mesh */
-	UPROPERTY(EditAnywhere, Instanced, Category=Reimport)
+	UPROPERTY(EditAnywhere, Instanced, Category=Reimport, AssetRegistrySearchable)
 	class UAssetImportData* AssetImportData;
 
 	/***  for Reimport **/
diff --git a/Engine/Source/Runtime/Engine/Classes/Curves/CurveBase.h b/Engine/Source/Runtime/Engine/Classes/Curves/CurveBase.h
--- a/Engine/Source/Runtime/Engine/Classes/Curves/CurveBase.h
+++ b/Engine/Source/Runtime/Engine/Classes/Curves/CurveBase.h
@@ -446,7 +446,7 @@ class ENGINE_API UCurveBase : public UObject, public FCurveOwnerInterface
 	GENERATED_UCLASS_BODY()
 
 	/** The filename imported to create this object. Relative to this object's package, BaseDir() or absolute */
-	UPROPERTY()
+	UPROPERTY(AssetRegistrySearchable, meta=(ReimportPath))
 	FString ImportPath;
 
 	/** Get the time range across all curves */
diff --git a/Engine/Source/Runtime/Engine/Classes/EditorFramework/AssetImportData.h b/Engine/Source/Runtime/Engine/Classes/EditorFramework/AssetImportData.h
--- a/Engine/Source/Runtime/Engine/Classes/EditorFramework/AssetImportData.h
+++ b/Engine/Source/Runtime/Engine/Classes/EditorFramework/AssetImportData.h
@@ -12,7 +12,7 @@ class UAssetImportData : public UObject
 	GENERATED_UCLASS_BODY()
 
 	/** Path to the resource used to construct this static mesh. Relative to the object's package, BaseDir() or absolute */
-	UPROPERTY(EditAnywhere, Category=ImportSettings)
+	UPROPERTY(EditAnywhere, Category=ImportSettings, AssetRegistrySearchable, meta=(ReimportPath))
 	FString SourceFilePath;
 
 	/** Date/Time-stamp of the file from the last import */
diff --git a/Engine/Source/Runtime/Engine/Classes/Engine/CurveTable.h b/Engine/Source/Runtime/Engine/Classes/Engine/CurveTable.h
--- a/Engine/Source/Runtime/Engine/Classes/Engine/CurveTable.h
+++ b/Engine/Source/Runtime/Engine/Classes/Engine/CurveTable.h
@@ -25,7 +25,7 @@ class UCurveTable
 	GENERATED_UCLASS_BODY()
 
 	/** The filename that was used to create this object. Relative to the object's package, BaseDir() or absolute */
-	UPROPERTY()
+	UPROPERTY(AssetRegistrySearchable, meta=(ReimportPath))
 	FString ImportPath;
 
 
diff --git a/Engine/Source/Runtime/Engine/Classes/Engine/DataTable.h b/Engine/Source/Runtime/Engine/Classes/Engine/DataTable.h
--- a/Engine/Source/Runtime/Engine/Classes/Engine/DataTable.h
+++ b/Engine/Source/Runtime/Engine/Classes/Engine/DataTable.h
@@ -41,7 +41,7 @@ class UDataTable
 	UScriptStruct*			RowStruct;
 
 	/** The filename that was used to create this object. Relative to the object's package, BaseDir() or absolute */
-	UPROPERTY()
+	UPROPERTY(AssetRegistrySearchable, meta=(ReimportPath))
 	FString ImportPath;
 
 	/** Map of name of row to row data structure. */
diff --git a/Engine/Source/Runtime/Engine/Classes/Engine/Font.h b/Engine/Source/Runtime/Engine/Classes/Engine/Font.h
--- a/Engine/Source/Runtime/Engine/Classes/Engine/Font.h
+++ b/Engine/Source/Runtime/Engine/Classes/Engine/Font.h
@@ -168,7 +168,7 @@ class UFont : public UObject, public IFontProviderInterface
 	FName LegacyFontName;
 
 	/** Embedded composite font data */
-	UPROPERTY()
+	UPROPERTY(AssetRegistrySearchable)
 	FCompositeFont CompositeFont;
 
 public:
diff --git a/Engine/Source/Runtime/Engine/Classes/Engine/SkeletalMesh.h b/Engine/Source/Runtime/Engine/Classes/Engine/SkeletalMesh.h
--- a/Engine/Source/Runtime/Engine/Classes/Engine/SkeletalMesh.h
+++ b/Engine/Source/Runtime/Engine/Classes/Engine/SkeletalMesh.h
@@ -600,7 +600,7 @@ public:
 #if WITH_EDITORONLY_DATA
 
 	/** Importing data and options used for this mesh */
-	UPROPERTY(EditAnywhere, Instanced, Category = Reimport)
+	UPROPERTY(EditAnywhere, Instanced, Category = Reimport, AssetRegistrySearchable)
 	class UAssetImportData* AssetImportData;
 
 	/** Path to the resource used to construct this skeletal mesh */
diff --git a/Engine/Source/Runtime/Engine/Classes/Engine/StaticMesh.h b/Engine/Source/Runtime/Engine/Classes/Engine/StaticMesh.h
--- a/Engine/Source/Runtime/Engine/Classes/Engine/StaticMesh.h
+++ b/Engine/Source/Runtime/Engine/Classes/Engine/StaticMesh.h
@@ -373,7 +373,7 @@ class UStaticMesh : public UObject, public IInterface_CollisionDataProvider, pub
 
 #if WITH_EDITORONLY_DATA
 	/** Importing data and options used for this mesh */
-	UPROPERTY(VisibleAnywhere, Instanced, Category=ImportSettings)
+	UPROPERTY(VisibleAnywhere, Instanced, Category=ImportSettings, AssetRegistrySearchable)
 	class UAssetImportData* AssetImportData;
 
 	/** Path to the resource used to construct this static mesh */
diff --git a/Engine/Source/Runtime/Engine/Classes/Engine/Texture.h b/Engine/Source/Runtime/Engine/Classes/Engine/Texture.h
--- a/Engine/Source/Runtime/Engine/Classes/Engine/Texture.h
+++ b/Engine/Source/Runtime/Engine/Classes/Engine/Texture.h
@@ -343,7 +343,7 @@ private:
 
 public:
 	/** Path to the resource used to construct this texture. Relative to the object's package, BaseDir() or absolute */
-	UPROPERTY(Category=Texture, VisibleAnywhere, BlueprintReadWrite)
+	UPROPERTY(Category=Texture, VisibleAnywhere, BlueprintReadWrite, AssetRegistrySearchable, meta=(ReimportPath))
 	FString SourceFilePath;
 
 	/** Date/Time-stamp of the file from the last import. */
diff --git a/Engine/Source/Runtime/Engine/Classes/Sound/SoundWave.h b/Engine/Source/Runtime/Engine/Classes/Sound/SoundWave.h
--- a/Engine/Source/Runtime/Engine/Classes/Sound/SoundWave.h
+++ b/Engine/Source/Runtime/Engine/Classes/Sound/SoundWave.h
@@ -211,7 +211,7 @@ class ENGINE_API USoundWave : public USoundBase
 
 #if WITH_EDITORONLY_DATA
 	/** Path to the resource used to construct this sound node wave. Relative to the object's package, BaseDir() or absolute. */
-	UPROPERTY(Category=Info, VisibleAnywhere)
+	UPROPERTY(Category=Info, VisibleAnywhere, AssetRegistrySearchable, meta=(ReimportPath))
 	FString SourceFilePath;
 
 	/** Date/Time-stamp of the file from the last import */
diff --git a/Engine/Source/Runtime/Engine/Classes/VectorField/VectorFieldStatic.h b/Engine/Source/Runtime/Engine/Classes/VectorField/VectorFieldStatic.h
--- a/Engine/Source/Runtime/Engine/Classes/VectorField/VectorFieldStatic.h
+++ b/Engine/Source/Runtime/Engine/Classes/VectorField/VectorFieldStatic.h
@@ -35,7 +35,7 @@ public:
 
 #if WITH_EDITORONLY_DATA
 	/** Path to the resource used to construct this vector field. Relative to the object's package, BaseDir() or absolute. */
-	UPROPERTY(Category=SourceAsset, VisibleAnywhere)
+	UPROPERTY(Category=SourceAsset, VisibleAnywhere, AssetRegistrySearchable, meta=(ReimportPath))
 	FString SourceFilePath;
 
 	/** Date/Time-stamp of the file from the last import */
diff --git a/Engine/Source/Runtime/SlateCore/Public/Fonts/CompositeFont.h b/Engine/Source/Runtime/SlateCore/Public/Fonts/CompositeFont.h
--- a/Engine/Source/Runtime/SlateCore/Public/Fonts/CompositeFont.h
+++ b/Engine/Source/Runtime/SlateCore/Public/Fonts/CompositeFont.h
@@ -42,7 +42,7 @@ struct SLATECORE_API FFontData
 	bool operator!=(const FFontData& Other) const;
 
 	/** The filename of the font to use - this may not always exist on disk, as we may have previously loaded and cached the font data inside an asset */
-	UPROPERTY(EditAnywhere, BlueprintReadWrite, Category=Font)
+	UPROPERTY(EditAnywhere, BlueprintReadWrite, Category=Font, AssetRegistrySearchable, meta=(ReimportPath))
 	FString FontFilename;
 
 	/** The data associated with the font */
