[INTEGRATE] Change 2412353 by [REDACTED_PERSON][REDACTED_HANDLE] on 2015/01/20 11:50:42

	Fixed bad padding and labels in "Convert to Blueprint" dialog

diff --git a/Engine/Source/Editor/KismetWidgets/Private/CreateBlueprintFromActorDialog.cpp b/Engine/Source/Editor/KismetWidgets/Private/CreateBlueprintFromActorDialog.cpp
--- a/Engine/Source/Editor/KismetWidgets/Private/CreateBlueprintFromActorDialog.cpp
+++ b/Engine/Source/Editor/KismetWidgets/Private/CreateBlueprintFromActorDialog.cpp
@@ -24,7 +24,7 @@ void FCreateBlueprintFromActorDialog::OpenDialog(bool bInHarvest)
 		(
 		SAssignNew(CreateBlueprintFromActorDialog, SCreateAssetFromActor, PickBlueprintPathWidget)
 		.AssetFilenameSuffix(TEXT("Blueprint"))
-		.HeadingText(LOCTEXT("CreateBlueprintFromActor_Heading", "Blueprint Name:"))
+		.HeadingText(LOCTEXT("CreateBlueprintFromActor_Heading", "Blueprint Name"))
 		.CreateButtonText(LOCTEXT("CreateBlueprintFromActor_ButtonLabel", "Create Blueprint"))
 		.OnCreateAssetAction(FOnPathChosen::CreateStatic(FCreateBlueprintFromActorDialog::OnCreateBlueprint, bInHarvest))
 		);
diff --git a/Engine/Source/Editor/UnrealEd/Private/SCreateAssetFromActor.cpp b/Engine/Source/Editor/UnrealEd/Private/SCreateAssetFromActor.cpp
--- a/Engine/Source/Editor/UnrealEd/Private/SCreateAssetFromActor.cpp
+++ b/Engine/Source/Editor/UnrealEd/Private/SCreateAssetFromActor.cpp
@@ -81,6 +81,7 @@ void SCreateAssetFromActor::Construct(const FArguments& InArgs, TSharedPtr<SWind
 		]
 		+ SVerticalBox::Slot()
 		.HAlign(HAlign_Right)
+		.Padding(0, 20, 0, 0)
 		.AutoHeight()
 		[
 			SNew(SHorizontalBox)
@@ -89,8 +90,8 @@ void SCreateAssetFromActor::Construct(const FArguments& InArgs, TSharedPtr<SWind
 			.AutoWidth()
 			[
 				SNew(SButton)
-				.VAlign(VAlign_Center)
-				.HAlign(HAlign_Center)
+				.VAlign(VAlign_Bottom)
+				.ContentPadding(FMargin(8, 2, 8, 2))
 				.OnClicked(this, &SCreateAssetFromActor::OnCreateAssetFromActorClicked)
 				.IsEnabled(this, &SCreateAssetFromActor::IsCreateAssetFromActorEnabled)
 				[
@@ -103,8 +104,7 @@ void SCreateAssetFromActor::Construct(const FArguments& InArgs, TSharedPtr<SWind
 			.AutoWidth()
 			[
 				SNew(SButton)
-				.VAlign(VAlign_Center)
-				.HAlign(HAlign_Center)
+				.VAlign(VAlign_Bottom)
 				.ContentPadding(FMargin(8, 2, 8, 2))
 				.OnClicked(this, &SCreateAssetFromActor::OnCancelCreateAssetFromActor)
 				[
