[INTEGRATE] Change 2412538 by [REDACTED_PERSON][REDACTED_HANDLE] on 2015/01/20 13:45:02

	Focus selected on components now works properly

diff --git a/Engine/Source/Editor/UnrealEd/Classes/Editor/EditorEngine.h b/Engine/Source/Editor/UnrealEd/Classes/Editor/EditorEngine.h
--- a/Engine/Source/Editor/UnrealEd/Classes/Editor/EditorEngine.h
+++ b/Engine/Source/Editor/UnrealEd/Classes/Editor/EditorEngine.h
@@ -836,24 +836,6 @@ public:
 	 */
 	AActor* SelectNamedActor(const TCHAR *TargetActorName);
 
-	/**
-	 * Moves all viewport cameras to the target actor.
-	 * @param	Actor					Target actor.
-	 * @param	bActiveViewportOnly		If true, move/reorient only the active viewport.
-	 */
-	void MoveViewportCamerasToActor(AActor& Actor,  bool bActiveViewportOnly);
-
-	/** 
-	 * Snaps an actor in a direction.  Optionally will align with the trace normal.
-	 * @param InActor			Actor to move to the floor.
-	 * @param InAlign			sWhether or not to rotate the actor to align with the trace normal.
-	 * @param InUseLineTrace	Whether or not to only trace with a line through the world.
-	 * @param InUseBounds		Whether or not to base the line trace off of the bounds.
-	 * @param InUsePivot		Whether or not to use the pivot position.
-	 * @param InDestination		The destination actor we want to move this actor to, NULL assumes we just want to go towards the floor
-	 * @return					Whether or not the actor was moved.
-	 */
-	bool SnapActorTo( AActor* InActor, const bool InAlign, const bool InUseLineTrace, const bool InUseBounds, const bool InUsePivot, const AActor* InDestination = NULL );
 
 	/**
 	 * Moves an actor in front of a camera specified by the camera's origin and direction.
@@ -866,13 +848,41 @@ public:
 	 */
 	void MoveActorInFrontOfCamera( AActor& InActor, const FVector& InCameraOrigin, const FVector& InCameraDirection );
 
+	/**
+	 * Moves all viewport cameras to the target actor.
+	 * @param	Actor					Target actor.
+	 * @param	bActiveViewportOnly		If true, move/reorient only the active viewport.
+	 */
+	void MoveViewportCamerasToActor(AActor& Actor,  bool bActiveViewportOnly);
+
 	/**
 	* Moves all viewport cameras to focus on the provided array of actors.
 	* @param	Actors					Target actors.
+
 	* @param	bActiveViewportOnly		If true, move/reorient only the active viewport.
 	*/
 	void MoveViewportCamerasToActor(const TArray<AActor*> &Actors, bool bActiveViewportOnly);
 
+	/**
+	* Moves all viewport cameras to focus on the provided array of actors.
+	* @param	Actors					Target actors.
+	* @param	Components				Target components (used of actors array is empty)
+	* @param	bActiveViewportOnly		If true, move/reorient only the active viewport.
+	*/
+	void MoveViewportCamerasToActor(const TArray<AActor*> &Actors, const TArray<UPrimitiveComponent*>& Components, bool bActiveViewportOnly);
+
+	/** 
+	 * Snaps an actor in a direction.  Optionally will align with the trace normal.
+	 * @param InActor			Actor to move to the floor.
+	 * @param InAlign			sWhether or not to rotate the actor to align with the trace normal.
+	 * @param InUseLineTrace	Whether or not to only trace with a line through the world.
+	 * @param InUseBounds		Whether or not to base the line trace off of the bounds.
+	 * @param InUsePivot		Whether or not to use the pivot position.
+	 * @param InDestination		The destination actor we want to move this actor to, NULL assumes we just want to go towards the floor
+	 * @return					Whether or not the actor was moved.
+	 */
+	bool SnapActorTo( AActor* InActor, const bool InAlign, const bool InUseLineTrace, const bool InUseBounds, const bool InUsePivot, const AActor* InDestination = NULL );
+
 	/**
 	 * Snaps the view of the camera to that of the provided actor.
 	 *
diff --git a/Engine/Source/Editor/UnrealEd/Private/EditorServer.cpp b/Engine/Source/Editor/UnrealEd/Private/EditorServer.cpp
--- a/Engine/Source/Editor/UnrealEd/Private/EditorServer.cpp
+++ b/Engine/Source/Editor/UnrealEd/Private/EditorServer.cpp
@@ -4236,13 +4236,17 @@ void UEditorEngine::MoveViewportCamerasToActor(AActor& Actor,  bool bActiveViewp
 
 	Actors.Add( &Actor );
 
-	MoveViewportCamerasToActor( Actors, bActiveViewportOnly );
+	MoveViewportCamerasToActor( Actors, TArray<UPrimitiveComponent*>(), bActiveViewportOnly );
 }
 
-
 void UEditorEngine::MoveViewportCamerasToActor(const TArray<AActor*> &Actors, bool bActiveViewportOnly)
 {
-	if( Actors.Num() == 0 )
+	MoveViewportCamerasToActor( Actors, TArray<UPrimitiveComponent*>(), bActiveViewportOnly );
+}
+
+void UEditorEngine::MoveViewportCamerasToActor(const TArray<AActor*> &Actors, const TArray<UPrimitiveComponent*>& Components, bool bActiveViewportOnly)
+{
+	if( Actors.Num() == 0 && Components.Num() == 0 )
 	{
 		return;
 	}
@@ -4257,116 +4261,140 @@ void UEditorEngine::MoveViewportCamerasToActor(const TArray<AActor*> &Actors, bo
 		}
 	}
 
+	struct ComponentTypeMatcher
+	{
+		ComponentTypeMatcher(UPrimitiveComponent* InComponentToMatch)
+			: ComponentToMatch(InComponentToMatch)
+		{}
+
+		bool operator()(const UClass* ComponentClass) const
+		{
+			return ComponentToMatch->IsA(ComponentClass);
+		}
+
+		UPrimitiveComponent* ComponentToMatch;
+	};
+
 	TArray<AActor*> InvisLevelActors;
 
 	TArray<UClass*> PrimitiveComponentTypesToIgnore;
 	PrimitiveComponentTypesToIgnore.Add( UShapeComponent::StaticClass() );
 	PrimitiveComponentTypesToIgnore.Add( UNavLinkRenderingComponent::StaticClass() );
 	PrimitiveComponentTypesToIgnore.Add( UDrawFrustumComponent::StaticClass() );
-
 	// Create a bounding volume of all of the selected actors.
 	FBox BoundingBox( 0 );
-	int32 NumActiveActors = 0;
-	for( int32 ActorIdx = 0; ActorIdx < Actors.Num(); ActorIdx++ )
-	{
-		AActor* Actor = Actors[ActorIdx];
 
-		if( Actor )
+	if( Components.Num() > 0 )
+	{
+		// First look at components
+		for(UPrimitiveComponent* PrimitiveComponent : Components)
 		{
-
-			// Don't allow moving the viewport cameras to actors in invisible levels
-			if ( !FLevelUtils::IsLevelVisible( Actor->GetLevel() ) )
+			if(PrimitiveComponent)
 			{
-				InvisLevelActors.Add( Actor );
-				continue;
-			}
+				if(!FLevelUtils::IsLevelVisible(PrimitiveComponent->GetComponentLevel()))
+				{
+					continue;
+				}
+
+				// Some components can have huge bounds but are not visible.  Ignore these components unless it is the only component on the actor 
+				const bool bIgnore = Components.Num() > 1 && PrimitiveComponentTypesToIgnore.IndexOfByPredicate(ComponentTypeMatcher(PrimitiveComponent)) != INDEX_NONE;
+
+				if(!bIgnore && PrimitiveComponent->IsRegistered())
+				{
+					BoundingBox += PrimitiveComponent->Bounds.GetBox();
+				}
 
-			const bool bActorIsEmitter = (Cast<AEmitter>(Actor) != NULL);
-	
-			if (bActorIsEmitter && bCustomCameraAlignEmitter)
-			{
-				const FVector DefaultExtent(CustomCameraAlignEmitterDistance,CustomCameraAlignEmitterDistance,CustomCameraAlignEmitterDistance);
-				const FBox DefaultSizeBox(Actor->GetActorLocation() - DefaultExtent, Actor->GetActorLocation() + DefaultExtent);
-				BoundingBox += DefaultSizeBox;
 			}
-			else
+		}
+	}
+	else 
+	{
+		for(int32 ActorIdx = 0; ActorIdx < Actors.Num(); ActorIdx++)
+		{
+			AActor* Actor = Actors[ActorIdx];
+
+			if(Actor)
 			{
-				TInlineComponentArray<UPrimitiveComponent*> Components;
-				Actor->GetComponents(Components);
 
-				for(int32 ComponentIndex = 0; ComponentIndex < Components.Num(); ++ComponentIndex)
+				// Don't allow moving the viewport cameras to actors in invisible levels
+				if(!FLevelUtils::IsLevelVisible(Actor->GetLevel()))
+				{
+					InvisLevelActors.Add(Actor);
+					continue;
+				}
+
+				const bool bActorIsEmitter = (Cast<AEmitter>(Actor) != NULL);
+
+				if(bActorIsEmitter && bCustomCameraAlignEmitter)
 				{
-					UPrimitiveComponent* PrimitiveComponent = Components[ComponentIndex];
+					const FVector DefaultExtent(CustomCameraAlignEmitterDistance, CustomCameraAlignEmitterDistance, CustomCameraAlignEmitterDistance);
+					const FBox DefaultSizeBox(Actor->GetActorLocation() - DefaultExtent, Actor->GetActorLocation() + DefaultExtent);
+					BoundingBox += DefaultSizeBox;
+				}
+				else
+				{
+					TInlineComponentArray<UPrimitiveComponent*> Components;
+					Actor->GetComponents(Components);
 
-					if( PrimitiveComponent->IsRegistered() )
+					for(int32 ComponentIndex = 0; ComponentIndex < Components.Num(); ++ComponentIndex)
 					{
-						struct ComponentTypeMatcher
+						UPrimitiveComponent* PrimitiveComponent = Components[ComponentIndex];
+
+						if(PrimitiveComponent->IsRegistered())
 						{
-							ComponentTypeMatcher( UPrimitiveComponent* InComponentToMatch )
-								: ComponentToMatch( InComponentToMatch )
-							{}
 
-							bool operator()(const UClass* ComponentClass) const
+							// Some components can have huge bounds but are not visible.  Ignore these components unless it is the only component on the actor 
+							const bool bIgnore = Components.Num() > 1 && PrimitiveComponentTypesToIgnore.IndexOfByPredicate(ComponentTypeMatcher(PrimitiveComponent)) != INDEX_NONE;
+
+							if(!bIgnore)
 							{
-								return ComponentToMatch->IsA(ComponentClass);
+								BoundingBox += PrimitiveComponent->Bounds.GetBox();
 							}
-
-							UPrimitiveComponent* ComponentToMatch;
-						};
-
-						// Some components can have huge bounds but are not visible.  Ignore these components unless it is the only component on the actor 
-						const bool bIgnore = Components.Num() > 1 && PrimitiveComponentTypesToIgnore.IndexOfByPredicate(ComponentTypeMatcher(PrimitiveComponent)) != INDEX_NONE;
-
-						if( !bIgnore )
-						{
-							BoundingBox += PrimitiveComponent->Bounds.GetBox();
 						}
 					}
-				}
-
-				if (Actor->IsA<ABrush>() && GLevelEditorModeTools().IsModeActive(FBuiltinEditorModes::EM_Geometry))
-				{
-					FEdModeGeometry* GeometryMode = GLevelEditorModeTools().GetActiveModeTyped<FEdModeGeometry>(FBuiltinEditorModes::EM_Geometry);
 
-					TArray<FGeomVertex*> SelectedVertices;
-					TArray<FGeomPoly*> SelectedPolys;
-					TArray<FGeomEdge*> SelectedEdges;
+					if(Actor->IsA<ABrush>() && GLevelEditorModeTools().IsModeActive(FBuiltinEditorModes::EM_Geometry))
+					{
+						FEdModeGeometry* GeometryMode = GLevelEditorModeTools().GetActiveModeTyped<FEdModeGeometry>(FBuiltinEditorModes::EM_Geometry);
 
-					GeometryMode->GetSelectedVertices(SelectedVertices);
-					GeometryMode->GetSelectedPolygons(SelectedPolys);
-					GeometryMode->GetSelectedEdges(SelectedEdges);
+						TArray<FGeomVertex*> SelectedVertices;
+						TArray<FGeomPoly*> SelectedPolys;
+						TArray<FGeomEdge*> SelectedEdges;
 
-					if (SelectedVertices.Num() + SelectedPolys.Num() + SelectedEdges.Num() > 0)
-					{
-						BoundingBox.Init();
+						GeometryMode->GetSelectedVertices(SelectedVertices);
+						GeometryMode->GetSelectedPolygons(SelectedPolys);
+						GeometryMode->GetSelectedEdges(SelectedEdges);
 
-						for (FGeomVertex* Vertex : SelectedVertices)
+						if(SelectedVertices.Num() + SelectedPolys.Num() + SelectedEdges.Num() > 0)
 						{
-							BoundingBox += Vertex->GetWidgetLocation();
-						}
+							BoundingBox.Init();
 
-						for (FGeomPoly* Poly : SelectedPolys)
-						{
-							BoundingBox += Poly->GetWidgetLocation();
-						}
+							for(FGeomVertex* Vertex : SelectedVertices)
+							{
+								BoundingBox += Vertex->GetWidgetLocation();
+							}
 
-						for (FGeomEdge* Edge : SelectedEdges)
-						{
-							BoundingBox += Edge->GetWidgetLocation();
-						}
+							for(FGeomPoly* Poly : SelectedPolys)
+							{
+								BoundingBox += Poly->GetWidgetLocation();
+							}
 
-						// Zoom out a little bit so you can see the selection
-						BoundingBox = BoundingBox.ExpandBy(25);
+							for(FGeomEdge* Edge : SelectedEdges)
+							{
+								BoundingBox += Edge->GetWidgetLocation();
+							}
+
+							// Zoom out a little bit so you can see the selection
+							BoundingBox = BoundingBox.ExpandBy(25);
+						}
 					}
 				}
 			}
-
-			NumActiveActors++;
 		}
 	}
 
 	// Make sure we had atleast one non-null actor in the array passed in.
-	if( NumActiveActors > 0 )
+	if( BoundingBox.GetSize() != FVector::ZeroVector )
 	{
 		if ( bActiveViewportOnly )
 		{
@@ -4686,15 +4714,24 @@ bool UEditorEngine::Exec_Camera( const TCHAR* Str, FOutputDevice& Ar )
 				Actors.Add( Actor );
 			}
 
-			if( Actors.Num() )
+			TArray<UPrimitiveComponent*> SelectedComponents;
+			for( FSelectionIterator It( GetSelectedComponentIterator() ); It; ++It )
+			{
+				UPrimitiveComponent* PrimitiveComp = Cast<UPrimitiveComponent>( *It );
+				if( PrimitiveComp )
+				{
+					SelectedComponents.Add( PrimitiveComp );
+				}
+			}
+
+			if( Actors.Num() || SelectedComponents.Num() )
 			{
-				MoveViewportCamerasToActor( Actors, bActiveViewportOnly );
-				Ar.Log( TEXT("Aligned camera to fit all selected actors.") );
+				MoveViewportCamerasToActor( Actors, SelectedComponents, bActiveViewportOnly );
 				return true;
 			}
 			else
 			{
-				Ar.Log( TEXT("Can't find target actor.") );
+				Ar.Log( TEXT("Can't find target actor or component.") );
 				return false;
 			}
 		}
