[INTEGRATE] Change 2412743 by [REDACTED_PERSON][REDACTED_HANDLE] on 2015/01/20 15:31:09

	Make Create Blueprint from Actor harvest InstanceComponents and be allowed to create sub-classes of Blueprints

diff --git a/Engine/Source/Editor/DetailCustomizations/Private/ActorDetails.cpp b/Engine/Source/Editor/DetailCustomizations/Private/ActorDetails.cpp
--- a/Engine/Source/Editor/DetailCustomizations/Private/ActorDetails.cpp
+++ b/Engine/Source/Editor/DetailCustomizations/Private/ActorDetails.cpp
@@ -816,22 +816,6 @@ void FActorDetails::AddSingleBlueprintRow( IDetailCategoryBuilder& BlueprintCate
 		Description = FText::FromString( InBlueprint->GetDesc() );
 	}
 
-	int NumBlueprintableActors = 0;
-	bool IsBlueprintBased = InBlueprint ? true : false;
-
-	if(!InBlueprint)
-	{
-		for ( TWeakObjectPtr<AActor>& Actor : SelectedActors )
-		{
-			if( FKismetEditorUtilities::CanCreateBlueprintOfClass(Actor->GetClass()))
-			{
-				NumBlueprintableActors++;
-			}
-		}
-	}
-
-	const bool bCanCreateActorBlueprint = ( !IsBlueprintBased && ( NumBlueprintableActors == 1 ) );
-
 	// Set up PathPickerConfig.
 	PathForActorBlueprint = FString("/Game/");
 	FPathPickerConfig PathPickerConfig;
@@ -839,21 +823,25 @@ void FActorDetails::AddSingleBlueprintRow( IDetailCategoryBuilder& BlueprintCate
 	PathPickerConfig.OnPathSelected = FOnPathSelected::CreateRaw(this, &FActorDetails::OnSelectBlueprintPath);
 	FContentBrowserModule& ContentBrowserModule = FModuleManager::LoadModuleChecked<FContentBrowserModule>("ContentBrowser");
 
+	bool bCanCreateActorBlueprint = false;
+	FString SelectedActorClassName;
 	FText CreateBlueprintToolTip = LOCTEXT("CreateBlueprintSuccess_ToolTip", "Create a Blueprint based on this Actor");
-	if(IsBlueprintBased)
+	if (SelectedActors.Num() == 1)
 	{
-		CreateBlueprintToolTip = LOCTEXT("CreateBlueprintFail_ToolTip", "Cannot create Blueprint using this Actor class");
+		if (FKismetEditorUtilities::CanCreateBlueprintOfClass(SelectedActors[0]->GetClass()))
+		{
+			SelectedActorClassName = SelectedActors[0]->GetClass()->GetName();
+			bCanCreateActorBlueprint = true;
+		}
+		else
+		{
+			CreateBlueprintToolTip = LOCTEXT("CreateBlueprintFail_ToolTip", "Cannot create Blueprint using this Actor class");
+		}
 	}
-	else if(!IsBlueprintBased && NumBlueprintableActors > 1)
+	else
 	{
 		CreateBlueprintToolTip = LOCTEXT("CreateBlueprintFailTooMany_ToolTip", "Cannot create Blueprint, please select only one Actor");
 	}
-
-	FString SelectedActorClassName;
-	if(NumBlueprintableActors == 1)
-	{
-		SelectedActorClassName = SelectedActors[0]->GetClass()->GetName();
-	}
 	
 	FText CurrentBlueprintLong = FText::Format( FText::FromString( TEXT( "{0} {1}" ) ), Name, Description );
 
@@ -991,8 +979,6 @@ void FActorDetails::AddSingleBlueprintRow( IDetailCategoryBuilder& BlueprintCate
 		]
 	];
 
-	if(!InBlueprint)
-	{
 		BlueprintCategory.AddCustomRow( LOCTEXT("CreateBlueprintFilterString", "Create Blueprint"), true )
 		.WholeRowContent()
 		.MinDesiredWidth(200)
@@ -1036,7 +1022,6 @@ void FActorDetails::AddSingleBlueprintRow( IDetailCategoryBuilder& BlueprintCate
 			]
 		];
 	}
-}
 END_SLATE_FUNCTION_BUILD_OPTIMIZATION
 
 FReply FActorDetails::OnPickBlueprintPathClicked(bool bHarvest )
diff --git a/Engine/Source/Editor/UnrealEd/Private/Kismet2/Kismet2.cpp b/Engine/Source/Editor/UnrealEd/Private/Kismet2/Kismet2.cpp
--- a/Engine/Source/Editor/UnrealEd/Private/Kismet2/Kismet2.cpp
+++ b/Engine/Source/Editor/UnrealEd/Private/Kismet2/Kismet2.cpp
@@ -916,20 +916,8 @@ UBlueprint* FKismetEditorUtilities::CreateBlueprintFromActor(const FName Bluepri
 	{
 		if (Outer)
 		{
-			UActorFactory* FactoryToUse = GEditor->FindActorFactoryForActorClass( Actor->GetClass() );
-
-			if( FactoryToUse != NULL )
-			{
-				// Create the blueprint
-				UObject* Asset = FactoryToUse->GetAssetFromActorInstance(Actor);
-				// For Actors that don't have an asset associated with them, Asset will be null
-				NewBlueprint = FactoryToUse->CreateBlueprint( Asset, Outer, BlueprintName, FName("CreateFromActor") );
-			}
-			else 
-			{
-				// We don't have a factory, but we can still try to create a blueprint for this actor class
-				NewBlueprint = FKismetEditorUtilities::CreateBlueprint( Actor->GetClass(), Outer, BlueprintName, EBlueprintType::BPTYPE_Normal, UBlueprint::StaticClass(), UBlueprintGeneratedClass::StaticClass(), FName("CreateFromActor") );
-			}
+			// We don't have a factory, but we can still try to create a blueprint for this actor class
+			NewBlueprint = FKismetEditorUtilities::CreateBlueprint( Actor->GetClass(), Outer, BlueprintName, EBlueprintType::BPTYPE_Normal, UBlueprint::StaticClass(), UBlueprintGeneratedClass::StaticClass(), FName("CreateFromActor") );
 		}
 
 		if(NewBlueprint)
@@ -940,6 +928,84 @@ UBlueprint* FKismetEditorUtilities::CreateBlueprintFromActor(const FName Bluepri
 			// Mark the package dirty
 			Outer->MarkPackageDirty();
 
+			// If the source Actor has Instance Components we need to translate these in to SCS Nodes
+			if (Actor->InstanceComponents.Num() > 0)
+			{
+				USimpleConstructionScript* SCS = NewBlueprint->SimpleConstructionScript;
+
+				TArray<UBlueprint*> ParentBPStack;
+				UBlueprint::GetBlueprintHierarchyFromClass(NewBlueprint->GeneratedClass, ParentBPStack);
+
+				TMap<USceneComponent*, USCS_Node*> SceneComponentsToAdd;
+				TMap<USceneComponent*, USCS_Node*> InstanceComponentToNodeMap;
+
+				for (UActorComponent* ActorComponent : Actor->InstanceComponents)
+				{
+					USCS_Node* SCSNode = SCS->CreateNode(ActorComponent->GetClass());
+					UEditorEngine::CopyPropertiesForUnrelatedObjects(ActorComponent,SCSNode->ComponentTemplate);
+
+					// The easy part is non-scene component or the Root simply add it
+					if (!ActorComponent->IsA<USceneComponent>() || ActorComponent == Actor->GetRootComponent())
+					{
+						SCS->AddNode(SCSNode);
+					}
+					else
+					{
+						USceneComponent* SceneComponent = CastChecked<USceneComponent>(ActorComponent);
+						check(SceneComponent->AttachParent);
+
+						InstanceComponentToNodeMap.Add(SceneComponent,SCSNode);
+
+						// If we're attached to a blueprint component look it up as the variable name is the component name
+						if (SceneComponent->AttachParent->bCreatedByConstructionScript)
+						{
+							USCS_Node* ParentSCSNode = nullptr;
+							for (UBlueprint* Blueprint : ParentBPStack)
+							{
+								ParentSCSNode = Blueprint->SimpleConstructionScript->FindSCSNode(SceneComponent->AttachParent->GetFName());
+								if (ParentSCSNode)
+								{
+									break;
+								}
+							}
+							check(ParentSCSNode);
+
+							if (ParentSCSNode->GetSCS() != SCS)
+							{
+								SCS->AddNode(SCSNode);
+							}
+
+							SCSNode->SetParent(ParentSCSNode);
+						}
+						// If we're attached to a native component
+						else if (!Actor->InstanceComponents.Contains(SceneComponent->AttachParent))
+						{
+							SCS->AddNode(SCSNode);
+							SCSNode->SetParent(SceneComponent->AttachParent);
+						}
+						else
+						{
+							// Otherwise check if we've already created the parents' new SCS node and attach to that or cache it off to do next pass
+							USCS_Node** ParentSCSNode = InstanceComponentToNodeMap.Find(SceneComponent->AttachParent);
+							if (ParentSCSNode)
+							{
+								(*ParentSCSNode)->AddChildNode(SCSNode);
+							}
+							else
+							{
+								SceneComponentsToAdd.Add(SceneComponent, SCSNode);
+							}
+						}
+					}
+				}
+
+				// Hook up the remaining components nodes that the parent's node was missing when it was processed
+				for (auto ComponentIt = SceneComponentsToAdd.CreateConstIterator(); ComponentIt; ++ComponentIt)
+				{
+					InstanceComponentToNodeMap.FindChecked(ComponentIt.Key()->AttachParent)->AddChildNode(ComponentIt.Value());
+				}
+			}
+
 			if(NewBlueprint->GeneratedClass)
 			{
 				UObject* CDO = NewBlueprint->GeneratedClass->GetDefaultObject();
@@ -957,6 +1023,9 @@ UBlueprint* FKismetEditorUtilities::CreateBlueprintFromActor(const FName Bluepri
 						// Ensure the light mass information is cleaned up
 						Scene->InvalidateLightingCache();
 					}
+
+					// Clear the instance properties array as we created SCS nodes for them
+					CDOAsActor->InstanceComponents.Empty();
 				}
 			}
 
