[INTEGRATE] Change 2412835 by [REDACTED_PERSON][REDACTED_HANDLE] on 2015/01/20 16:36:19

	The generated name when adding a component in-world matched the pattern in the BP SCS editor

diff --git a/Engine/Source/Editor/Kismet/Private/SSCSEditor.cpp b/Engine/Source/Editor/Kismet/Private/SSCSEditor.cpp
--- a/Engine/Source/Editor/Kismet/Private/SSCSEditor.cpp
+++ b/Engine/Source/Editor/Kismet/Private/SSCSEditor.cpp
@@ -240,7 +240,9 @@ FSCSEditorTreeNode::FSCSEditorTreeNode(UActorComponent* InComponentTemplate)
 		if (!FComponentEditorUtils::IsValidVariableNameString(InComponentTemplate, InComponentTemplate->GetName()))
 		{
 			ERenameFlags RenameFlags = REN_DontCreateRedirectors | REN_NonTransactional;
-			InComponentTemplate->Rename(*InComponentTemplate->GetName().Replace(TEXT("Component_"), TEXT("")), nullptr, RenameFlags);
+			FString NewComponentName = FComponentEditorUtils::GenerateValidVariableName(InComponentTemplate);
+
+			InComponentTemplate->Rename(*NewComponentName, nullptr, RenameFlags);
 		}
 		
 		InstancedComponentName = InComponentTemplate->GetFName();
diff --git a/Engine/Source/Editor/UnrealEd/Private/Kismet2/ComponentEditorUtils.cpp b/Engine/Source/Editor/UnrealEd/Private/Kismet2/ComponentEditorUtils.cpp
--- a/Engine/Source/Editor/UnrealEd/Private/Kismet2/ComponentEditorUtils.cpp
+++ b/Engine/Source/Editor/UnrealEd/Private/Kismet2/ComponentEditorUtils.cpp
@@ -107,6 +107,38 @@ bool FComponentEditorUtils::IsValidVariableNameString(const UActorComponent* InC
 	return bIsValid;
 }
 
+FString FComponentEditorUtils::GenerateValidVariableName(const UActorComponent* InComponent)
+{
+	int32 Counter = 1;
+	FString ComponentTypeName = *InComponent->GetClass()->GetName().Replace(TEXT("Component"), TEXT(""));
+	FString ComponentInstanceName;
+
+	// Make sure that none of the current components have the same name
+	TInlineComponentArray<UActorComponent*> Components;
+	InComponent->GetOwner()->GetComponents(Components);
+
+	bool bNameIsValid = false;
+	while (!bNameIsValid)
+	{
+		bNameIsValid = true;
+
+		// Try the name with the lowest possible number
+		ComponentInstanceName = FString::Printf(TEXT("%s%d"), *ComponentTypeName, Counter++);
+
+		// The name is valid if no other components in the owning actor have the same name
+		for (auto Component : Components)
+		{
+			if (Component->GetName() == ComponentInstanceName)
+			{
+				bNameIsValid = false;
+				break;
+			}
+		}
+	}
+
+	return ComponentInstanceName;
+}
+
 void FComponentEditorUtils::PropagateTransformPropertyChange(
 	class USceneComponent* InSceneComponentTemplate,
 	const FTransformData& OldDefaultTransform,
diff --git a/Engine/Source/Editor/UnrealEd/Public/Kismet2/ComponentEditorUtils.h b/Engine/Source/Editor/UnrealEd/Public/Kismet2/ComponentEditorUtils.h
--- a/Engine/Source/Editor/UnrealEd/Public/Kismet2/ComponentEditorUtils.h
+++ b/Engine/Source/Editor/UnrealEd/Public/Kismet2/ComponentEditorUtils.h
@@ -16,6 +16,9 @@ public:
 	/** Test whether or not the given string is a valid variable name string for the given component instance */
 	static bool IsValidVariableNameString(const UActorComponent* InComponent, const FString& InString);
 
+	/** Generate a valid variable name string for the given component instance */
+	static FString GenerateValidVariableName(const UActorComponent* InComponent);
+
 	struct FTransformData
 	{
 		FTransformData(const class USceneComponent& Component)
