[INTEGRATE] Change 2413021 by [REDACTED_PERSON][REDACTED_HANDLE] on 2015/01/20 18:02:50

	Make Apply changes for an Actor to Blueprint add Instance Components in to the Blueprint

diff --git a/Engine/Source/Editor/DetailCustomizations/Private/ActorDetails.cpp b/Engine/Source/Editor/DetailCustomizations/Private/ActorDetails.cpp
--- a/Engine/Source/Editor/DetailCustomizations/Private/ActorDetails.cpp
+++ b/Engine/Source/Editor/DetailCustomizations/Private/ActorDetails.cpp
@@ -1220,11 +1220,18 @@ bool FActorDetails::PushToBlueprintDefaults_IsEnabled( TWeakObjectPtr<UBlueprint
 			&& Blueprint != NULL
 			&& Actor->GetClass()->ClassGeneratedBy == Blueprint)
 		{
-			AActor* BlueprintCDO = Actor->GetClass()->GetDefaultObject<AActor>();
-			if(BlueprintCDO != NULL)
+			if (Actor->InstanceComponents.Num() > 0)
 			{
-				const auto CopyOptions = (EditorUtilities::ECopyOptions::Type)(EditorUtilities::ECopyOptions::PreviewOnly|EditorUtilities::ECopyOptions::OnlyCopyEditOrInterpProperties);
-				bIsEnabled = EditorUtilities::CopyActorProperties(Actor, BlueprintCDO, CopyOptions) > 0;
+				bIsEnabled = true;
+			}
+			else
+			{
+				AActor* BlueprintCDO = Actor->GetClass()->GetDefaultObject<AActor>();
+				if(BlueprintCDO != NULL)
+				{
+					const auto CopyOptions = (EditorUtilities::ECopyOptions::Type)(EditorUtilities::ECopyOptions::PreviewOnly|EditorUtilities::ECopyOptions::OnlyCopyEditOrInterpProperties);
+					bIsEnabled = EditorUtilities::CopyActorProperties(Actor, BlueprintCDO, CopyOptions) > 0;
+				}
 			}
 		}
 	}
@@ -1250,6 +1257,7 @@ FText FActorDetails::PushToBlueprintDefaults_ToolTipText( TWeakObjectPtr<UBluepr
 				const auto CopyOptions = (EditorUtilities::ECopyOptions::Type)(EditorUtilities::ECopyOptions::PreviewOnly|EditorUtilities::ECopyOptions::OnlyCopyEditOrInterpProperties);
 				NumChangedProperties += EditorUtilities::CopyActorProperties(Actor, BlueprintCDO, CopyOptions);
 			}
+			NumChangedProperties += Actor->InstanceComponents.Num();
 		}
 	}
 
@@ -1287,9 +1295,16 @@ FReply FActorDetails::PushToBlueprintDefaults_OnClicked( TWeakObjectPtr<UBluepri
 			{
 				const auto CopyOptions = (EditorUtilities::ECopyOptions::Type)(EditorUtilities::ECopyOptions::OnlyCopyEditOrInterpProperties|EditorUtilities::ECopyOptions::PropagateChangesToArcheypeInstances);
 				NumChangedProperties = EditorUtilities::CopyActorProperties(Actor, BlueprintCDO, CopyOptions);
+				if (Actor->InstanceComponents.Num() > 0)
+				{
+					FKismetEditorUtilities::AddComponentsToBlueprint(Blueprint, Actor->InstanceComponents);
+					NumChangedProperties += Actor->InstanceComponents.Num();
+					Actor->InstanceComponents.Empty();
+				}
 				if(NumChangedProperties > 0)
 				{
 					FBlueprintEditorUtils::MarkBlueprintAsModified(Blueprint);
+					FKismetEditorUtilities::CompileBlueprint(Blueprint);
 				}
 			}
 		}
diff --git a/Engine/Source/Editor/UnrealEd/Private/Kismet2/Kismet2.cpp b/Engine/Source/Editor/UnrealEd/Private/Kismet2/Kismet2.cpp
--- a/Engine/Source/Editor/UnrealEd/Private/Kismet2/Kismet2.cpp
+++ b/Engine/Source/Editor/UnrealEd/Private/Kismet2/Kismet2.cpp
@@ -908,6 +908,95 @@ UBlueprint* FKismetEditorUtilities::CreateBlueprintFromActor(const FString& Path
 	return NewBlueprint;
 }
 
+void FKismetEditorUtilities::AddComponentsToBlueprint(UBlueprint* Blueprint, const TArray<UActorComponent*>& Components)
+{
+	USimpleConstructionScript* SCS = Blueprint->SimpleConstructionScript;
+
+	TArray<UBlueprint*> ParentBPStack;
+	UBlueprint::GetBlueprintHierarchyFromClass(Blueprint->GeneratedClass, ParentBPStack);
+
+	TMap<USceneComponent*, USCS_Node*> SceneComponentsToAdd;
+	TMap<USceneComponent*, USCS_Node*> InstanceComponentToNodeMap;
+
+	AActor* Actor = nullptr;
+
+	for (UActorComponent* ActorComponent : Components)
+	{
+		if (Actor)
+		{
+			check(Actor == ActorComponent->GetOwner());
+		}
+		else
+		{
+			Actor = ActorComponent->GetOwner();
+			check(Actor);
+		}
+
+		USCS_Node* SCSNode = SCS->CreateNode(ActorComponent->GetClass());
+		UEditorEngine::CopyPropertiesForUnrelatedObjects(ActorComponent,SCSNode->ComponentTemplate);
+
+		// The easy part is non-scene component or the Root simply add it
+		if (!ActorComponent->IsA<USceneComponent>() || ActorComponent == Actor->GetRootComponent())
+		{
+			SCS->AddNode(SCSNode);
+		}
+		else
+		{
+			USceneComponent* SceneComponent = CastChecked<USceneComponent>(ActorComponent);
+			check(SceneComponent->AttachParent);
+
+			InstanceComponentToNodeMap.Add(SceneComponent,SCSNode);
+
+			// If we're attached to a blueprint component look it up as the variable name is the component name
+			if (SceneComponent->AttachParent->bCreatedByConstructionScript)
+			{
+				USCS_Node* ParentSCSNode = nullptr;
+				for (UBlueprint* Blueprint : ParentBPStack)
+				{
+					ParentSCSNode = Blueprint->SimpleConstructionScript->FindSCSNode(SceneComponent->AttachParent->GetFName());
+					if (ParentSCSNode)
+					{
+						break;
+					}
+				}
+				check(ParentSCSNode);
+
+				if (ParentSCSNode->GetSCS() != SCS)
+				{
+					SCS->AddNode(SCSNode);
+				}
+
+				SCSNode->SetParent(ParentSCSNode);
+			}
+			// If we're attached to a native component
+			else if (!Components.Contains(SceneComponent->AttachParent))
+			{
+				SCS->AddNode(SCSNode);
+				SCSNode->SetParent(SceneComponent->AttachParent);
+			}
+			else
+			{
+				// Otherwise check if we've already created the parents' new SCS node and attach to that or cache it off to do next pass
+				USCS_Node** ParentSCSNode = InstanceComponentToNodeMap.Find(SceneComponent->AttachParent);
+				if (ParentSCSNode)
+				{
+					(*ParentSCSNode)->AddChildNode(SCSNode);
+				}
+				else
+				{
+					SceneComponentsToAdd.Add(SceneComponent, SCSNode);
+				}
+			}
+		}
+	}
+
+	// Hook up the remaining components nodes that the parent's node was missing when it was processed
+	for (auto ComponentIt = SceneComponentsToAdd.CreateConstIterator(); ComponentIt; ++ComponentIt)
+	{
+		InstanceComponentToNodeMap.FindChecked(ComponentIt.Key()->AttachParent)->AddChildNode(ComponentIt.Value());
+	}
+}
+
 UBlueprint* FKismetEditorUtilities::CreateBlueprintFromActor(const FName BlueprintName, UObject* Outer, AActor* Actor, const bool bReplaceActor )
 {
 	UBlueprint* NewBlueprint = nullptr;
@@ -931,79 +1020,7 @@ UBlueprint* FKismetEditorUtilities::CreateBlueprintFromActor(const FName Bluepri
 			// If the source Actor has Instance Components we need to translate these in to SCS Nodes
 			if (Actor->InstanceComponents.Num() > 0)
 			{
-				USimpleConstructionScript* SCS = NewBlueprint->SimpleConstructionScript;
-
-				TArray<UBlueprint*> ParentBPStack;
-				UBlueprint::GetBlueprintHierarchyFromClass(NewBlueprint->GeneratedClass, ParentBPStack);
-
-				TMap<USceneComponent*, USCS_Node*> SceneComponentsToAdd;
-				TMap<USceneComponent*, USCS_Node*> InstanceComponentToNodeMap;
-
-				for (UActorComponent* ActorComponent : Actor->InstanceComponents)
-				{
-					USCS_Node* SCSNode = SCS->CreateNode(ActorComponent->GetClass());
-					UEditorEngine::CopyPropertiesForUnrelatedObjects(ActorComponent,SCSNode->ComponentTemplate);
-
-					// The easy part is non-scene component or the Root simply add it
-					if (!ActorComponent->IsA<USceneComponent>() || ActorComponent == Actor->GetRootComponent())
-					{
-						SCS->AddNode(SCSNode);
-					}
-					else
-					{
-						USceneComponent* SceneComponent = CastChecked<USceneComponent>(ActorComponent);
-						check(SceneComponent->AttachParent);
-
-						InstanceComponentToNodeMap.Add(SceneComponent,SCSNode);
-
-						// If we're attached to a blueprint component look it up as the variable name is the component name
-						if (SceneComponent->AttachParent->bCreatedByConstructionScript)
-						{
-							USCS_Node* ParentSCSNode = nullptr;
-							for (UBlueprint* Blueprint : ParentBPStack)
-							{
-								ParentSCSNode = Blueprint->SimpleConstructionScript->FindSCSNode(SceneComponent->AttachParent->GetFName());
-								if (ParentSCSNode)
-								{
-									break;
-								}
-							}
-							check(ParentSCSNode);
-
-							if (ParentSCSNode->GetSCS() != SCS)
-							{
-								SCS->AddNode(SCSNode);
-							}
-
-							SCSNode->SetParent(ParentSCSNode);
-						}
-						// If we're attached to a native component
-						else if (!Actor->InstanceComponents.Contains(SceneComponent->AttachParent))
-						{
-							SCS->AddNode(SCSNode);
-							SCSNode->SetParent(SceneComponent->AttachParent);
-						}
-						else
-						{
-							// Otherwise check if we've already created the parents' new SCS node and attach to that or cache it off to do next pass
-							USCS_Node** ParentSCSNode = InstanceComponentToNodeMap.Find(SceneComponent->AttachParent);
-							if (ParentSCSNode)
-							{
-								(*ParentSCSNode)->AddChildNode(SCSNode);
-							}
-							else
-							{
-								SceneComponentsToAdd.Add(SceneComponent, SCSNode);
-							}
-						}
-					}
-				}
-
-				// Hook up the remaining components nodes that the parent's node was missing when it was processed
-				for (auto ComponentIt = SceneComponentsToAdd.CreateConstIterator(); ComponentIt; ++ComponentIt)
-				{
-					InstanceComponentToNodeMap.FindChecked(ComponentIt.Key()->AttachParent)->AddChildNode(ComponentIt.Value());
-				}
+				AddComponentsToBlueprint(NewBlueprint, Actor->InstanceComponents);
 			}
 
 			if(NewBlueprint->GeneratedClass)
diff --git a/Engine/Source/Editor/UnrealEd/Public/Kismet2/KismetEditorUtilities.h b/Engine/Source/Editor/UnrealEd/Public/Kismet2/KismetEditorUtilities.h
--- a/Engine/Source/Editor/UnrealEd/Public/Kismet2/KismetEditorUtilities.h
+++ b/Engine/Source/Editor/UnrealEd/Public/Kismet2/KismetEditorUtilities.h
@@ -86,6 +86,9 @@ public:
 	/** @return true is it's possible to create a blueprint from the specified class */
 	static bool CanCreateBlueprintOfClass(const UClass* Class);
 
+	/** Take a list of components that belong to a single Actor and add them to a blueprint as SCSNodes */
+	static void AddComponentsToBlueprint(UBlueprint* Blueprint, const TArray<UActorComponent*>& Components);
+
 	/** 
 	 * Take an Actor and generate a blueprint based on it. Uses the Actors type as the parent class. 
 	 * @param Path					The path to use when creating the package for the new blueprint
