[INTEGRATE] Change 2413355 by [REDACTED_PERSON][REDACTED_HANDLE] on 2015/01/20 22:22:34

	Another non-unity build fix.

diff --git a/Engine/Source/Editor/UnrealEd/Private/EditorHook.cpp b/Engine/Source/Editor/UnrealEd/Private/EditorHook.cpp
--- a/Engine/Source/Editor/UnrealEd/Private/EditorHook.cpp
+++ b/Engine/Source/Editor/UnrealEd/Private/EditorHook.cpp
@@ -8,6 +8,7 @@
 #include "Engine/Selection.h"
 
 #include "LandscapeProxy.h"
+#include "LandscapeComponent.h"
 
 uint32			EngineThreadId;
 const TCHAR*	GItem;
