[INTEGRATE] Change 2413898 by [REDACTED_PERSON][REDACTED_HANDLE].7 on 2015/01/21 11:48:00

	File caches used by auto reimport manager now scan for changes asynchronously over multiple frames

	Currently this is rate-limited to half a 60fps frame. We may move the work onto a thread at a later date.

diff --git a/Engine/Source/Editor/UnrealEd/Private/AutoReimportManager.cpp b/Engine/Source/Editor/UnrealEd/Private/AutoReimportManager.cpp
--- a/Engine/Source/Editor/UnrealEd/Private/AutoReimportManager.cpp
+++ b/Engine/Source/Editor/UnrealEd/Private/AutoReimportManager.cpp
@@ -19,7 +19,6 @@ DEFINE_LOG_CATEGORY_STATIC(LogAutoReimportManager, Log, All);
 
 namespace
 {
-
 	/** Template helper function to count the number of elements in the specified array that pass a predicate */
 	template<typename T, typename P>
 	int32 CountIf(const TArray<T>& InArray, P Predicate)
@@ -35,17 +34,6 @@ namespace
 		return Count;
 	}
 
-	/** Reduce the array to the specified accumulator */
-	template<typename T, typename P, typename A>
-	A&& Reduce(const TArray<T>& InArray, P Predicate, A&& Accumulator)
-	{
-		for (const T& X : InArray)
-		{
-			Predicate(X, Accumulator);
-		}
-		return MoveTemp(Accumulator);
-	}
-
 	/** Retrieve a semi-colon separated string of file extensions supported by all available editor import factories */
 	FString GetAllFactoryExtensions()
 	{
@@ -123,14 +111,16 @@ namespace
 	/** Generate a config from the specified options, to pass to FFileCache on construction */
 	FFileCacheConfig GenerateFileCacheConfig(const FString& InPath, const FString& InSupportedExtensions, const FString& InMountedContentPath)
 	{
-		FFileCacheConfig Config;
-		Config.Directory = FPaths::ConvertRelativePathToFull(InPath);
-		Config.IncludeExtensions = InSupportedExtensions;
+		FString Directory = FPaths::ConvertRelativePathToFull(InPath);
 
-		const FString& HashString = InMountedContentPath.IsEmpty() ? Config.Directory : InMountedContentPath;
+		const FString& HashString = InMountedContentPath.IsEmpty() ? Directory : InMountedContentPath;
 		const uint32 CRC = FCrc::MemCrc32(*HashString, HashString.Len()*sizeof(TCHAR));	
+		FString CacheFilename = FPaths::ConvertRelativePathToFull(FPaths::GameIntermediateDir()) / TEXT("ReimportCache") / FString::Printf(TEXT("%u.bin"), CRC);
+
+		FFileCacheConfig Config(MoveTemp(Directory), MoveTemp(CacheFilename));
+		Config.BatchDelayS = 1;
+		Config.IncludeExtensions = InSupportedExtensions;
 
-		Config.CacheFile = FPaths::ConvertRelativePathToFull(FPaths::GameIntermediateDir()) / TEXT("ReimportCache") / FString::Printf(TEXT("%u.bin"), CRC);
 		return Config;
 	}
 
@@ -169,38 +159,51 @@ namespace
 FContentDirectoryMonitor::FContentDirectoryMonitor(const FString& InDirectory, const FString& InSupportedExtensions, const FString& InMountedContentPath)
 	: Cache(GenerateFileCacheConfig(InDirectory, InSupportedExtensions, InMountedContentPath))
 	, MountedContentPath(InMountedContentPath)
-	, LastProcessTime(0)
 	, LastSaveTime(0)
 {
-	Cache.Scan();
+	
 }
 
-void FContentDirectoryMonitor::Tick(const IAssetRegistry& Registry)
+void FContentDirectoryMonitor::Destroy()
 {
-	const double Now = FPlatformTime::Seconds();
+	Cache.Destroy();
+}
 
-	// Keep accumulating periodically 'til we have no more changes
-	auto NewChanges = Cache.GetOutstandingChanges();
-	if (NewChanges.Num() != 0)
+void FContentDirectoryMonitor::Tick(const IAssetRegistry& Registry, const FWorkLimiter& Limiter)
+{
+	Cache.Tick(Limiter);
+
+	if (Limiter.ShouldLimit())
+	{
+		return;
+	}
+	// We can't do this if the registry is loading assets as we won't necesssarily be able to react to changes effectively
+	else if (!Registry.IsLoadingAssets())
 	{
-		// Check again for any more in 1 second
-		OutstandingChanges.Append(MoveTemp(NewChanges));
-		LastProcessTime = Now;
+		OutstandingChanges = Cache.GetOutstandingChanges();
+		if (OutstandingChanges.Num() != 0)
+		{
+			ProcessOutstandingChanges(Registry, Limiter);
+		}
 	}
-	else if (Now - LastProcessTime > ReactionTimeS && OutstandingChanges.Num())
+	
+	if (Limiter.ShouldLimit())
 	{
-		LastProcessTime = Now;
-		ProcessOutstandingChanges(Registry);
+		return;
 	}
-
-	if (Now - LastSaveTime > ResaveIntervalS)
+	// Save out the cache file if we've got time to do so, and have not saved for a while
+	else
 	{
-		LastSaveTime = Now;
-		Cache.WriteCache();
+		const double Now = FPlatformTime::Seconds();
+		if (Now - LastSaveTime > ResaveIntervalS)
+		{
+			LastSaveTime = Now;
+			Cache.WriteCache();
+		}
 	}
 }
 
-void FContentDirectoryMonitor::ProcessOutstandingChanges(const IAssetRegistry& Registry)
+void FContentDirectoryMonitor::ProcessOutstandingChanges(const IAssetRegistry& Registry, const FWorkLimiter& Limiter)
 {
 	// Count the number of additions so we can apprximate the work load
 	const int32 NumAdditions = CountIf(OutstandingChanges, [](const FUpdateCacheTransaction& Change){ return Change.Action == FFileChangeData::FCA_Added; });
@@ -359,20 +362,28 @@ void UAutoReimportManager::BeginDestroy()
 
 void UAutoReimportManager::Tick(float DeltaTime)
 {
+	// Only spend half a frame on this, at most
+	const FWorkLimiter Limiter(1.0 / 120);
+
 	const IAssetRegistry& Registry = FModuleManager::LoadModuleChecked<FAssetRegistryModule>("AssetRegistry").Get();
 
-	// We can't do this if the registry is loading assets as we won't necesssarily be able to action changes effectively
-	if (!Registry.IsLoadingAssets())
-	{
-		TGuardValue<bool> Guard(bIsProcessingChanges, true);
+	TGuardValue<bool> Guard(bIsProcessingChanges, true);
 
-		for (auto& Monitor : MountedDirectoryMonitors)
+	for (auto& Monitor : MountedDirectoryMonitors)
+	{
+		Monitor.Tick(Registry, Limiter);
+		if (Limiter.ShouldLimit())
 		{
-			Monitor.Tick(Registry);
+			return;
 		}
-		for (auto& Monitor : UserDirectoryMonitors)
+	}
+
+	for (auto& Monitor : UserDirectoryMonitors)
+	{
+		Monitor.Tick(Registry, Limiter);
+		if (Limiter.ShouldLimit())
 		{
-			Monitor.Tick(Registry);
+			return;
 		}
 	}
 }
diff --git a/Engine/Source/Editor/UnrealEd/Private/FileCache.cpp b/Engine/Source/Editor/UnrealEd/Private/FileCache.cpp
--- a/Engine/Source/Editor/UnrealEd/Private/FileCache.cpp
+++ b/Engine/Source/Editor/UnrealEd/Private/FileCache.cpp
@@ -45,9 +45,90 @@ namespace
 	}
 }
 
+FAsyncDirectoryReader::FAsyncDirectoryReader(const FString& InDirectory)
+{
+	PendingDirectories.Add(InDirectory);
+	State.Emplace();
+}
+
+TOptional<FDirectoryState> FAsyncDirectoryReader::GetFinalState()
+{
+	TOptional<FDirectoryState> OldState;
+	Swap(OldState, State);
+
+	return OldState;
+}
+
+bool FAsyncDirectoryReader::IsComplete() const
+{
+	return PendingFiles.Num() == 0 && PendingDirectories.Num() == 0;
+}
+
+EProgressResult FAsyncDirectoryReader::Tick(const FWorkLimiter& Limiter)
+{
+	auto& FileManager = IFileManager::Get();
+
+	for (int32 Index = 0; Index < PendingDirectories.Num(); ++Index)
+	{
+		ScanDirectory(PendingDirectories[Index]);
+		if (Limiter.ShouldLimit())
+		{
+			// We've spent too long, bail
+			PendingDirectories.RemoveAt(0, Index, false);
+			return EProgressResult::Pending;
+		}
+	}
+	PendingDirectories.Empty();
+
+	for (int32 Index = 0; Index < PendingFiles.Num(); ++Index)
+	{
+		const auto& File = PendingFiles[Index];
+		State->Files.Emplace(File, FileManager.GetTimeStamp(*File));
+
+		if (Index % 100 == 0 && Limiter.ShouldLimit())
+		{
+			// We've spent too long, bail
+			PendingFiles.RemoveAt(0, Index, false);
+			return EProgressResult::Pending;
+		}
+	}
+	PendingFiles.Empty();
+
+	return EProgressResult::Finished;
+}
+
+void FAsyncDirectoryReader::ScanDirectory(const FString& InDirectory)
+{
+	struct FVisitor : public IPlatformFile::FDirectoryVisitor
+	{
+		TArray<FString>* PendingFiles;
+		TArray<FString>* PendingDirectories;
+
+		virtual bool Visit(const TCHAR* FilenameOrDirectory, bool bIsDirectory)
+		{
+			if (bIsDirectory)
+			{
+				PendingDirectories->Add(FString(FilenameOrDirectory));
+			}
+			else
+			{
+				PendingFiles->Add(FString(FilenameOrDirectory));
+			}
+			return true;
+		}
+	};
+
+	FVisitor Visitor;
+	Visitor.PendingFiles = &PendingFiles;
+	Visitor.PendingDirectories = &PendingDirectories;
+	IFileManager::Get().IterateDirectory(*InDirectory, Visitor);
+}
+
 FFileCache::FFileCache(const FFileCacheConfig& InConfig)
 	: Config(InConfig)
+	, DirectoryReader(InConfig.Directory)
 	, bSavedCacheDirty(false)
+	, LastChangeTimeS(0)
 {
 	// Ensure that the extension strings are of the form ;ext1;ext2;ext3;
 	if (!Config.IncludeExtensions.IsEmpty() && Config.IncludeExtensions[0] != ';')
@@ -86,6 +167,7 @@ void FFileCache::Destroy()
 	bSavedCacheDirty = false;
 	IFileManager::Get().Delete(*Config.CacheFile, false, true, true);
 
+	DirectoryReader = FAsyncDirectoryReader(Config.Directory);
 	OutstandingChanges.Empty();
 	CachedDirectoryState = FDirectoryState();
 
@@ -110,36 +192,6 @@ void FFileCache::UnbindWatcher()
 	WatcherDelegate = FDelegateHandle();
 }
 
-FDirectoryState FFileCache::WalkDirectory() const
-{
-	struct FDirectoryVisitor : public IPlatformFile::FDirectoryVisitor
-	{
-		FDirectoryState State;
-		int32 LeadingDirectoryLength;
-		IFileManager* FileManager;
-
-		virtual bool Visit(const TCHAR* FilenameOrDirectory, bool bIsDirectory)
-		{
-			if (!bIsDirectory)
-			{
-				State.Files.Emplace(FString(FilenameOrDirectory + LeadingDirectoryLength), FileManager->GetTimeStamp(FilenameOrDirectory));
-			}
-			return true;
-		}
-	};
-
-
-	auto& FileManager = IFileManager::Get();
-
-	// Iterate the directory
-	FDirectoryVisitor Visitor;
-	Visitor.LeadingDirectoryLength = Config.Directory.Len();
-	Visitor.FileManager = &FileManager;
-	FileManager.IterateDirectoryRecursively(*Config.Directory, Visitor);
-
-	return MoveTemp(Visitor.State);
-}
-
 TOptional<FDirectoryState> FFileCache::ReadCache() const
 {
 	TOptional<FDirectoryState> Optional;
@@ -171,11 +223,7 @@ void FFileCache::WriteCache()
 		Ar->Close();
 		delete Ar;
 
-		// @todo: arodham: can't do this with maps
-		// if (CachedDirectoryState.Files.Num() < CachedDirectoryState.Files.Max() * 0.5f)
-		// {
-		// 	CachedDirectoryState.Files.Shrink();
-		// }
+		CachedDirectoryState.Files.Shrink();
 
 		bSavedCacheDirty = false;
 	}
@@ -183,9 +231,16 @@ void FFileCache::WriteCache()
 
 TArray<FUpdateCacheTransaction> FFileCache::GetOutstandingChanges()
 {
-	TArray<FUpdateCacheTransaction> Moved;
-	Swap(OutstandingChanges, Moved);
-	return Moved;
+	if (FPlatformTime::Seconds() - LastChangeTimeS >= Config.BatchDelayS)
+	{
+		TArray<FUpdateCacheTransaction> Moved;
+		Swap(OutstandingChanges, Moved);
+		return Moved;
+	}
+	else
+	{
+		return TArray<FUpdateCacheTransaction>();
+	}
 }
 
 void FFileCache::CompleteTransaction(FUpdateCacheTransaction&& Transaction)
@@ -228,21 +283,28 @@ void FFileCache::CompleteTransaction(FUpdateCacheTransaction&& Transaction)
 	}
 }
 
-void FFileCache::Scan()
+void FFileCache::Tick(const FWorkLimiter& Limiter)
 {
-	auto CurrentFileData = WalkDirectory();
+	if (DirectoryReader.IsComplete() || DirectoryReader.Tick(Limiter) == EProgressResult::Pending)
+	{
+		return;
+	}
 
+	// We should only ever get here once. The directory reader has finished scanning, and we can now diff the results with what we had saved in the cache file.
+	check(DirectoryReader.IsComplete());
+	
+	TOptional<FDirectoryState> CurrentFileData = DirectoryReader.GetFinalState();
 	if (CachedDirectoryState.Files.Num() == 0)
 	{
 		// If we don't have any cached data yet, just use the file data we just harvested
-		CachedDirectoryState = MoveTemp(CurrentFileData);
+		CachedDirectoryState = MoveTemp(CurrentFileData.GetValue());
 		bSavedCacheDirty = true;
 		return;
 	}
 
 	// We already have cached data so we need to compare it with the harvested data
 	// to detect additions, modifications, and removals
-	for (const auto& FilenameAndData : CurrentFileData.Files)
+	for (const auto& FilenameAndData : CurrentFileData->Files)
 	{
 		const FString& Filename = FilenameAndData.Key.Get();
 		if (!IsFileApplicable(Filename))
@@ -270,7 +332,7 @@ void FFileCache::Scan()
 	for (auto It = CachedDirectoryState.Files.CreateIterator(); It; ++It)
 	{
 		const FImmutableString& Filename = It.Key();
-		if (!CurrentFileData.Files.Contains(Filename) && IsFileApplicable(Filename.Get()))
+		if (!CurrentFileData->Files.Contains(Filename) && IsFileApplicable(Filename.Get()))
 		{
 			OutstandingChanges.Add(FUpdateCacheTransaction(Filename, FFileChangeData::FCA_Removed));
 		}
@@ -352,6 +414,8 @@ void FFileCache::OnDirectoryChanged(const TArray<FFileChangeData>& FileChanges)
 	RemoveDuplicates(OutstandingChanges, [](const FUpdateCacheTransaction& A, const FUpdateCacheTransaction& B){
 		return A.Action == B.Action && A.Filename == B.Filename;
 	});
+
+	LastChangeTimeS = FPlatformTime::Seconds();
 }
 
 bool FFileCache::IsFileApplicable(const FString& Filename)
diff --git a/Engine/Source/Editor/UnrealEd/Public/AutoReimport/AutoReimportManager.h b/Engine/Source/Editor/UnrealEd/Public/AutoReimport/AutoReimportManager.h
--- a/Engine/Source/Editor/UnrealEd/Public/AutoReimport/AutoReimportManager.h
+++ b/Engine/Source/Editor/UnrealEd/Public/AutoReimport/AutoReimportManager.h
@@ -22,10 +22,10 @@ public:
 	FContentDirectoryMonitor(const FString& InDirectory, const FString& InSupportedExtensions, const FString& InMountedContentPath = FString());
 
 	/** Tick this directory monitor. Will potentially process changed files */
-	void Tick(const IAssetRegistry& Registry);
+	void Tick(const IAssetRegistry& Registry, const FWorkLimiter& Limiter);
 
 	/** Destroy this monitor including its cache */
-	void Destroy() { return Cache.Destroy(); }
+	void Destroy();
 
 	/** Get the directory that this monitor applies to */
 	const FString& GetDirectory() const { return Cache.GetDirectory(); }
@@ -36,7 +36,7 @@ public:
 private:
 
 	/** Process the outstanding changes that we have cached */
-	void ProcessOutstandingChanges(const IAssetRegistry& Registry);
+	void ProcessOutstandingChanges(const IAssetRegistry& Registry, const FWorkLimiter& Limiter);
 
 private:
 
@@ -49,15 +49,9 @@ private:
 	/** A list of file system changes that are due to be processed */
 	TArray<FUpdateCacheTransaction> OutstandingChanges;
 
-	/** The time we last processed outstanding changes */
-	double LastProcessTime;
-
 	/** The last time we attempted to save the cache file */
 	double LastSaveTime;
 
-	/** The amount of time (in seconds) we must wait before reacting to changed files */
-	static const int32 ReactionTimeS = 1;
-
 	/** The interval between potential re-saves of the cache file */
 	static const int32 ResaveIntervalS = 60;
 };
diff --git a/Engine/Source/Editor/UnrealEd/Public/FileCache.h b/Engine/Source/Editor/UnrealEd/Public/FileCache.h
--- a/Engine/Source/Editor/UnrealEd/Public/FileCache.h
+++ b/Engine/Source/Editor/UnrealEd/Public/FileCache.h
@@ -152,15 +152,88 @@ private:
 	FUpdateCacheTransaction& operator=(const FUpdateCacheTransaction& In) = default;
 };
 
+/** Small utility for specifying when a call to do some work should bail out, if at all. */
+struct FWorkLimiter
+{
+	/** Constructor specifying that we should never bail out early */
+	FWorkLimiter() : StopAtTime(-1) {}
+
+	/** Constructor specifying not to run over the specified number of seconds */
+	FWorkLimiter(float NumSeconds) : StopAtTime(FPlatformTime::Seconds() + NumSeconds) {}
+
+	/** Check whether we have exceeded the stop time */
+	bool ShouldLimit() const { return StopAtTime != -1 && FPlatformTime::Seconds() >= StopAtTime; }
+
+private:
+
+	/** User specified time to stop processing */
+	double StopAtTime;
+};
+
+enum class EProgressResult
+{
+	Finished, Pending
+};
+
+/**
+ * Class responsible for 'asyncronously' scanning a folder for files and timestamps.
+ * Example usage:
+ *		FAsyncDirectoryReader Reader(TEXT("C:\\Path"));
+ *
+ *		while(!Reader.IsComplete())
+ *		{
+ *			FPlatformProcess::Sleep(1);
+ *			Reader.Tick(FWorkLimiter(1));	// Do 1 second of work
+ *		}
+ *		TOptional<FDirectoryState> State = Reader.GetFinalState();
+ */
+struct FAsyncDirectoryReader
+{
+	/** Constructor that sets up the directory reader to the specified directory */
+	FAsyncDirectoryReader(const FString& InDirectory);
+
+	/**
+	 * Get the state of the directory once finished. Relinquishes the currently stored directory state to the client.
+	 * Returns nothing if incomplete, or if GetFinalState() has already been called.
+	 */
+	TOptional<FDirectoryState> GetFinalState();
+
+	/** Returns true when this directory reader has finished scanning the directory */
+	bool IsComplete() const;
+
+	/** Tick this reader (discover new directories / files). Returns progress state. */
+	EProgressResult Tick(const FWorkLimiter& Limiter);
+
+private:
+	/** Non-recursively scan a single directory for its contents. Adds results to Pending arrays. */
+	void ScanDirectory(const FString& InDirectory);
+
+	/** The currently discovered state of the directory */
+	TOptional<FDirectoryState> State;
+
+	/** A list of directories we have recursively found on our travels */
+	TArray<FString> PendingDirectories;
+
+	/** A list of files we have recursively found on our travels */
+	TArray<FString> PendingFiles;
+};
+
 /** Configuration structure required to construct a FFileCache */
 struct FFileCacheConfig
 {
+	FFileCacheConfig(FString InDirectory, FString InCacheFile)
+		: Directory(InDirectory), CacheFile(InCacheFile), BatchDelayS(0.f)
+	{}
+
 	/** String specifying the directory on disk that the cache should reflect */
 	FString Directory;
 
 	/** String specifying the file that the cache should be saved to */
 	FString CacheFile;
 
+	/** Optionally specified delay that should be waited before reporting changes to the watched directory. Useful for batching together large or rapid changes. */
+	float BatchDelayS;
+
 	/**
 	 * Strings specifying extensions to include/exclude in the cache (or empty to include everything).
 	 * Specified as a list of semi-colon separated extensions (eg "png;jpg;gif;").
@@ -192,8 +265,8 @@ public:
 	/** Get the absolute path of the directory this cache reflects */
 	const FString& GetDirectory() const { return Config.Directory; }
 
-	/** Force a scan the directory to detect any changes in the cache. */
-	void Scan();
+	/** Tick this FileCache, optionally specifying a limiter to limit the amount of time spent in here */
+	void Tick(const FWorkLimiter& Limiter);
 
 	/** Write out the cached file, if we have any changes to write */
 	void WriteCache();
@@ -223,27 +296,30 @@ private:
 
 	/** Unbind the watcher from the directory. Called on destruction. */
 	void UnbindWatcher();
-	
-	/** Walk the directory and return all the file information we find pertaining to all applicable files */
-	FDirectoryState WalkDirectory() const;
 
 	/** Read the cache file data and return the contents */
 	TOptional<FDirectoryState> ReadCache() const;
 
 private:
 
+	/** Configuration settings applied on construction */
+	FFileCacheConfig Config;
+
+	/** 'Asynchronous' directory reader responsible for gathering all file/timestamp information recursively from our cache directory */
+	FAsyncDirectoryReader DirectoryReader;
+
 	/** Temporary array of transactions to be returned to the client through GetOutstandingChanges() */
 	TArray<FUpdateCacheTransaction> OutstandingChanges;
 
 	/** Our in-memory view of the cached directory state. */
 	FDirectoryState CachedDirectoryState;
 
-	/** Configuration settings applied on construction */
-	FFileCacheConfig Config;
-
 	/** Handle to the directory watcher delegate so we can delete it properly */
 	FDelegateHandle WatcherDelegate;
 
 	/** True when the cached state we have in memory is more up to date than the serialized file. Enables WriteCache() when true. */
 	bool bSavedCacheDirty;
+
+	/** The time that a change to the file-system was last logged. */
+	double LastChangeTimeS;
 };
