[INTEGRATE] Change 2413978 by [REDACTED_PERSON][REDACTED_HANDLE] on 2015/01/21 12:31:08

	BP - Unified BP editor now on by default.

diff --git a/Engine/Source/Editor/UnrealEd/Private/Settings/SettingsClasses.cpp b/Engine/Source/Editor/UnrealEd/Private/Settings/SettingsClasses.cpp
--- a/Engine/Source/Editor/UnrealEd/Private/Settings/SettingsClasses.cpp
+++ b/Engine/Source/Editor/UnrealEd/Private/Settings/SettingsClasses.cpp
@@ -53,7 +53,8 @@ UDestructableMeshEditorSettings::UDestructableMeshEditorSettings( const FObjectI
 
 UEditorExperimentalSettings::UEditorExperimentalSettings( const FObjectInitializer& ObjectInitializer )
 	: Super(ObjectInitializer)
-{ 
+	, bUnifiedBlueprintEditor(true)
+{
 }
 
 void UEditorExperimentalSettings::PostEditChangeProperty( struct FPropertyChangedEvent& PropertyChangedEvent )
