[INTEGRATE] Change 2414079 by [REDACTED_PERSON][REDACTED_HANDLE] on 2015/01/21 13:31:58

	[UE-7838] Deltas are now applied properly to components edited in the world
	- Additionally, enables selection of a component via editor-only component (completes UE-7801)

diff --git a/Engine/Source/Editor/UnrealEd/Private/LevelEditorViewport.cpp b/Engine/Source/Editor/UnrealEd/Private/LevelEditorViewport.cpp
--- a/Engine/Source/Editor/UnrealEd/Private/LevelEditorViewport.cpp
+++ b/Engine/Source/Editor/UnrealEd/Private/LevelEditorViewport.cpp
@@ -1789,18 +1789,13 @@ void FLevelEditorViewportClient::ProcessClick(FSceneView& View, HHitProxy* HitPr
 			// 3. The LMB was pressed or we already have a component selected (so that right-clicking a selected actor won't select a component)
 			// 4. The click was not a double click (unless a component has already been selected)
 			// 5. The level viewport didn't just receive focus this frame (again unless a component has already been selected)
-			// 6. The component selected is not a hidden editor-only component
 			const bool bActorAlreadySelectedExclusively = GEditor->GetSelectedActors()->IsSelected(ActorHitProxy->Actor) && ( GEditor->GetSelectedActorCount() == 1 );
 			const bool bComponentAlreadySelected = GEditor->GetSelectedComponentCount() > 0;
-	
 			const bool bCanBeginSelectingComponents =  (Click.GetKey() == EKeys::LeftMouseButton) 
 													&& (Click.GetEvent() != IE_DoubleClick) 
 													&& !bReceivedFocusRecently;
 
-			auto HitComponent = ActorHitProxy->PrimComponent;
-			const bool bHitSelectableComponent = HitComponent != nullptr && HitComponent->AlwaysLoadOnClient && HitComponent->AlwaysLoadOnServer;
-
-			const bool bSelectComponent = bActorAlreadySelectedExclusively && bHitSelectableComponent && ( bComponentAlreadySelected || bCanBeginSelectingComponents );
+			const bool bSelectComponent = bActorAlreadySelectedExclusively && ( bComponentAlreadySelected || bCanBeginSelectingComponents );
 
 			if (bSelectComponent && GetDefault<UEditorExperimentalSettings>()->bInWorldBPEditing)
 			{
@@ -3056,7 +3051,7 @@ void FLevelEditorViewportClient::ApplyDeltaToComponent(USceneComponent* InCompon
 		&InDeltaDrag,
 		&InDeltaRot,
 		&ModifiedDeltaScale,
-		GEditor->GetPivotLocation());
+		InComponent->RelativeLocation);
 }
 
 /** Helper function for ModifyScale - Convert the active Dragging Axis to per-axis flags */
@@ -3401,7 +3396,6 @@ void FLevelEditorViewportClient::CheckHoveredHitProxy( HHitProxy* HoveredHitProx
 		}
 	}
 
-
 	// Check to see if there are any hovered objects that need to be updated
 	{
 		bool bAnyHoverChanges = false;
@@ -4027,7 +4021,7 @@ void FLevelEditorViewportClient::AddHoverEffect( FViewportHoverTarget& InHoverTa
 	AActor* ActorUnderCursor = InHoverTarget.HoveredActor;
 	UModel* ModelUnderCursor = InHoverTarget.HoveredModel;
 
-	if( ActorUnderCursor != NULL )
+	if( ActorUnderCursor != nullptr )
 	{
 		TInlineComponentArray<UPrimitiveComponent*> Components;
 		ActorUnderCursor->GetComponents(Components);
@@ -4041,7 +4035,7 @@ void FLevelEditorViewportClient::AddHoverEffect( FViewportHoverTarget& InHoverTa
 			}
 		}
 	}
-	else if( ModelUnderCursor != NULL )
+	else if (ModelUnderCursor != nullptr)
 	{
 		check( InHoverTarget.ModelSurfaceIndex != INDEX_NONE );
 		check( InHoverTarget.ModelSurfaceIndex < (uint32)ModelUnderCursor->Surfs.Num() );
@@ -4052,14 +4046,14 @@ void FLevelEditorViewportClient::AddHoverEffect( FViewportHoverTarget& InHoverTa
 
 
 /**
- * Static: Removes a hover effect to the specified object
+ * Static: Removes a hover effect from the specified object
  *
  * @param	InHoverTarget	The hoverable object to remove the effect from
  */
 void FLevelEditorViewportClient::RemoveHoverEffect( FViewportHoverTarget& InHoverTarget )
 {
 	AActor* CurHoveredActor = InHoverTarget.HoveredActor;
-	if( CurHoveredActor != NULL )
+	if( CurHoveredActor != nullptr )
 	{
 		TInlineComponentArray<UPrimitiveComponent*> Components;
 		CurHoveredActor->GetComponents(Components);
@@ -4076,7 +4070,7 @@ void FLevelEditorViewportClient::RemoveHoverEffect( FViewportHoverTarget& InHove
 	}
 
 	UModel* CurHoveredModel = InHoverTarget.HoveredModel;
-	if( CurHoveredModel != NULL )
+	if( CurHoveredModel != nullptr )
 	{
 		if( InHoverTarget.ModelSurfaceIndex != INDEX_NONE &&
 			(uint32)CurHoveredModel->Surfs.Num() >= InHoverTarget.ModelSurfaceIndex )
