[INTEGRATE] Change 2415127 by [REDACTED_PERSON][REDACTED_HANDLE].7 on 2015/01/22 04:49:05

	Expanded FFileCache to allow storing both relative and absolute filenames, and ensured auto reimport caches are stored relative

diff --git a/Engine/Source/Editor/UnrealEd/Private/AutoReimportManager.cpp b/Engine/Source/Editor/UnrealEd/Private/AutoReimportManager.cpp
--- a/Engine/Source/Editor/UnrealEd/Private/AutoReimportManager.cpp
+++ b/Engine/Source/Editor/UnrealEd/Private/AutoReimportManager.cpp
@@ -120,6 +120,8 @@ namespace
 		FFileCacheConfig Config(MoveTemp(Directory), MoveTemp(CacheFilename));
 		Config.BatchDelayS = 1;
 		Config.IncludeExtensions = InSupportedExtensions;
+		// We always store paths inside content folders relative to the folder
+		Config.PathType = EPathType::Relative;
 
 		return Config;
 	}
diff --git a/Engine/Source/Editor/UnrealEd/Private/FileCache.cpp b/Engine/Source/Editor/UnrealEd/Private/FileCache.cpp
--- a/Engine/Source/Editor/UnrealEd/Private/FileCache.cpp
+++ b/Engine/Source/Editor/UnrealEd/Private/FileCache.cpp
@@ -45,7 +45,8 @@ namespace
 	}
 }
 
-FAsyncDirectoryReader::FAsyncDirectoryReader(const FString& InDirectory)
+FAsyncDirectoryReader::FAsyncDirectoryReader(const FString& InDirectory, EPathType InPathType)
+	: RootPath(InDirectory), PathType(InPathType)
 {
 	PendingDirectories.Add(InDirectory);
 	State.Emplace();
@@ -64,9 +65,10 @@ bool FAsyncDirectoryReader::IsComplete() const
 	return PendingFiles.Num() == 0 && PendingDirectories.Num() == 0;
 }
 
-EProgressResult FAsyncDirectoryReader::Tick(const FWorkLimiter& Limiter)
+FAsyncDirectoryReader::EProgressResult FAsyncDirectoryReader::Tick(const FWorkLimiter& Limiter)
 {
 	auto& FileManager = IFileManager::Get();
+	const int32 RootPathLen = RootPath.Len();
 
 	for (int32 Index = 0; Index < PendingDirectories.Num(); ++Index)
 	{
@@ -83,7 +85,10 @@ EProgressResult FAsyncDirectoryReader::Tick(const FWorkLimiter& Limiter)
 	for (int32 Index = 0; Index < PendingFiles.Num(); ++Index)
 	{
 		const auto& File = PendingFiles[Index];
-		State->Files.Emplace(File, FileManager.GetTimeStamp(*File));
+
+		// Store the file relative or absolute
+		FString Filename = (PathType == EPathType::Relative ? *File + RootPathLen : *File);
+		State->Files.Emplace(MoveTemp(Filename), FileManager.GetTimeStamp(*File));
 
 		if (Index % 100 == 0 && Limiter.ShouldLimit())
 		{
@@ -126,7 +131,7 @@ void FAsyncDirectoryReader::ScanDirectory(const FString& InDirectory)
 
 FFileCache::FFileCache(const FFileCacheConfig& InConfig)
 	: Config(InConfig)
-	, DirectoryReader(InConfig.Directory)
+	, DirectoryReader(InConfig.Directory, Config.PathType)
 	, bSavedCacheDirty(false)
 	, LastChangeTimeS(0)
 {
@@ -167,7 +172,7 @@ void FFileCache::Destroy()
 	bSavedCacheDirty = false;
 	IFileManager::Get().Delete(*Config.CacheFile, false, true, true);
 
-	DirectoryReader = FAsyncDirectoryReader(Config.Directory);
+	DirectoryReader = FAsyncDirectoryReader(Config.Directory, Config.PathType);
 	OutstandingChanges.Empty();
 	CachedDirectoryState = FDirectoryState();
 
@@ -285,7 +290,7 @@ void FFileCache::CompleteTransaction(FUpdateCacheTransaction&& Transaction)
 
 void FFileCache::Tick(const FWorkLimiter& Limiter)
 {
-	if (DirectoryReader.IsComplete() || DirectoryReader.Tick(Limiter) == EProgressResult::Pending)
+	if (DirectoryReader.IsComplete() || DirectoryReader.Tick(Limiter) == FAsyncDirectoryReader::EProgressResult::Pending)
 	{
 		return;
 	}
@@ -345,6 +350,7 @@ void FFileCache::OnDirectoryChanged(const TArray<FFileChangeData>& FileChanges)
 
 	OutstandingChanges.Reserve(OutstandingChanges.Num() + FileChanges.Num());
 
+	const int32 RootPathLen = Config.Directory.Len();
 	for (const auto& ThisEntry : FileChanges)
 	{
 		// If it's a directory or is not applicable, ignore it
@@ -353,11 +359,17 @@ void FFileCache::OnDirectoryChanged(const TArray<FFileChangeData>& FileChanges)
 			continue;
 		}
 
-		FImmutableString RelativeFilename;
+		FImmutableString TransactionFilename;
 		{
 			FString Temp = FPaths::ConvertRelativePathToFull(ThisEntry.Filename);
-			Temp.GetCharArray().RemoveAt(0, Config.Directory.Len());
-			RelativeFilename = Temp;
+			if (Config.PathType == EPathType::Relative)
+			{
+				TransactionFilename = FString(*Temp + RootPathLen);
+			}
+			else
+			{
+				TransactionFilename = MoveTemp(Temp);
+			}
 		}
 		
 		switch (ThisEntry.Action)
@@ -365,11 +377,11 @@ void FFileCache::OnDirectoryChanged(const TArray<FFileChangeData>& FileChanges)
 		case FFileChangeData::FCA_Added:
 			{
 				const int32 NumRemoved = OutstandingChanges.RemoveAll([&](const FUpdateCacheTransaction& X){
-					return X.Action == FFileChangeData::FCA_Removed && X.Filename == RelativeFilename;
+					return X.Action == FFileChangeData::FCA_Removed && X.Filename == TransactionFilename;
 				});
 
 				const auto Action = NumRemoved == 0 ? FFileChangeData::FCA_Added : FFileChangeData::FCA_Modified;
-				OutstandingChanges.Add(FUpdateCacheTransaction(RelativeFilename, Action, FileManager.GetTimeStamp(*ThisEntry.Filename)));
+				OutstandingChanges.Add(FUpdateCacheTransaction(TransactionFilename, Action, FileManager.GetTimeStamp(*ThisEntry.Filename)));
 			}
 			break;
 
@@ -377,7 +389,7 @@ void FFileCache::OnDirectoryChanged(const TArray<FFileChangeData>& FileChanges)
 			{
 				bool bPreviouslyAdded = false;
 				const int32 NumRemoved = OutstandingChanges.RemoveAll([&](const FUpdateCacheTransaction& X){
-					if (X.Filename == RelativeFilename)
+					if (X.Filename == TransactionFilename)
 					{
 						bPreviouslyAdded = X.Action == FFileChangeData::FCA_Added || bPreviouslyAdded;
 						return true;
@@ -387,7 +399,7 @@ void FFileCache::OnDirectoryChanged(const TArray<FFileChangeData>& FileChanges)
 
 				if (!bPreviouslyAdded)
 				{
-					OutstandingChanges.Add(FUpdateCacheTransaction(RelativeFilename, FFileChangeData::FCA_Removed));
+					OutstandingChanges.Add(FUpdateCacheTransaction(TransactionFilename, FFileChangeData::FCA_Removed));
 				}
 			}
 			
@@ -396,12 +408,12 @@ void FFileCache::OnDirectoryChanged(const TArray<FFileChangeData>& FileChanges)
 		case FFileChangeData::FCA_Modified:
 			{
 				const bool bPreviouslyAdded = OutstandingChanges.ContainsByPredicate([&](const FUpdateCacheTransaction& X){
-					return X.Filename == RelativeFilename && X.Action == FFileChangeData::FCA_Added;
+					return X.Filename == TransactionFilename && X.Action == FFileChangeData::FCA_Added;
 				});
 
 				if (!bPreviouslyAdded)
 				{
-					OutstandingChanges.Add(FUpdateCacheTransaction(RelativeFilename, FFileChangeData::FCA_Modified, FileManager.GetTimeStamp(*ThisEntry.Filename)));
+					OutstandingChanges.Add(FUpdateCacheTransaction(TransactionFilename, FFileChangeData::FCA_Modified, FileManager.GetTimeStamp(*ThisEntry.Filename)));
 				}
 			}
 			break;
diff --git a/Engine/Source/Editor/UnrealEd/Public/FileCache.h b/Engine/Source/Editor/UnrealEd/Public/FileCache.h
--- a/Engine/Source/Editor/UnrealEd/Public/FileCache.h
+++ b/Engine/Source/Editor/UnrealEd/Public/FileCache.h
@@ -170,15 +170,20 @@ private:
 	double StopAtTime;
 };
 
-enum class EProgressResult
+/** Enum specifying whether a path should be relative or absolute */
+enum EPathType
 {
-	Finished, Pending
+	/** Paths should be cached relative to the root cache directory */
+	Relative,
+
+	/** Paths should be cached as absolute file system paths */
+	Absolute
 };
 
 /**
  * Class responsible for 'asyncronously' scanning a folder for files and timestamps.
  * Example usage:
- *		FAsyncDirectoryReader Reader(TEXT("C:\\Path"));
+ *		FAsyncDirectoryReader Reader(TEXT("C:\\Path"), EPathType::Relative);
  *
  *		while(!Reader.IsComplete())
  *		{
@@ -189,8 +194,13 @@ enum class EProgressResult
  */
 struct FAsyncDirectoryReader
 {
+	enum class EProgressResult
+	{
+		Finished, Pending
+	};
+
 	/** Constructor that sets up the directory reader to the specified directory */
-	FAsyncDirectoryReader(const FString& InDirectory);
+	FAsyncDirectoryReader(const FString& InDirectory, EPathType InPathType);
 
 	/**
 	 * Get the state of the directory once finished. Relinquishes the currently stored directory state to the client.
@@ -208,6 +218,12 @@ private:
 	/** Non-recursively scan a single directory for its contents. Adds results to Pending arrays. */
 	void ScanDirectory(const FString& InDirectory);
 
+	/** Path to the root directory we want to scan */
+	FString RootPath;
+
+	/** Whether we should return relative or absolute paths */
+	EPathType PathType;
+
 	/** The currently discovered state of the directory */
 	TOptional<FDirectoryState> State;
 
@@ -222,7 +238,7 @@ private:
 struct FFileCacheConfig
 {
 	FFileCacheConfig(FString InDirectory, FString InCacheFile)
-		: Directory(InDirectory), CacheFile(InCacheFile), BatchDelayS(0.f)
+		: Directory(InDirectory), CacheFile(InCacheFile), BatchDelayS(0.f), PathType(EPathType::Relative)
 	{}
 
 	/** String specifying the directory on disk that the cache should reflect */
@@ -231,14 +247,17 @@ struct FFileCacheConfig
 	/** String specifying the file that the cache should be saved to */
 	FString CacheFile;
 
-	/** Optionally specified delay that should be waited before reporting changes to the watched directory. Useful for batching together large or rapid changes. */
-	float BatchDelayS;
-
 	/**
 	 * Strings specifying extensions to include/exclude in the cache (or empty to include everything).
 	 * Specified as a list of semi-colon separated extensions (eg "png;jpg;gif;").
 	 */
 	FString IncludeExtensions, ExcludeExtensions;
+
+	/** Optionally specified delay that should be waited before reporting changes to the watched directory. Useful for batching together large or rapid changes. */
+	float BatchDelayS;
+
+	/** Path type to return, relative to the directory or absolute. */
+	EPathType PathType;
 };
 
 /**
