[INTEGRATE] Change 2415144 by [REDACTED_PERSON][REDACTED_HANDLE].7 on 2015/01/22 05:18:27

	Disabled automatic asset imports for engine content folder

diff --git a/Engine/Source/Editor/UnrealEd/Private/AutoReimportManager.cpp b/Engine/Source/Editor/UnrealEd/Private/AutoReimportManager.cpp
--- a/Engine/Source/Editor/UnrealEd/Private/AutoReimportManager.cpp
+++ b/Engine/Source/Editor/UnrealEd/Private/AutoReimportManager.cpp
@@ -494,7 +494,12 @@ void UAutoReimportManager::SetUpMountedContentDirectories()
 	for (const auto& ContentPath : ContentPaths)
 	{
 		const FString Directory = FPaths::ConvertRelativePathToFull(FPackageName::LongPackageNameToFilename(ContentPath));
-		MountedDirectoryMonitors.Emplace(Directory, SupportedExtensions, ContentPath);
+
+		// We ignore the engine content directory for now
+		if (Directory != FPaths::ConvertRelativePathToFull(FPaths::EngineContentDir()))
+		{
+			MountedDirectoryMonitors.Emplace(Directory, SupportedExtensions, ContentPath);
+		}
 	}
 }
 
