[INTEGRATE] Change 2415268 by [REDACTED_PERSON][REDACTED_HANDLE] on 2015/01/22 09:10:34

	Improvements to the way thumbnails are shown in the Content Browser

	UE-7184 - Show game c++ classes and engine c++ classes in the content browser

	C++ classes and Blueprints now both show you either a live-thumbnail (for asset types), or a class thumbnail for their class type. If showing a live-thumbnail, they also show a mini class icon at the bottom right of the thumbnail.

	This iconography change has been applied to all asset types, and they now show you their asset type icon on the colored overlay, rather than show you the asset type text. They will still fallback to using the asset type text if there is no icon available for the asset type.

diff --git a/Engine/Config/BaseEditor.ini b/Engine/Config/BaseEditor.ini
--- a/Engine/Config/BaseEditor.ini
+++ b/Engine/Config/BaseEditor.ini
@@ -15,6 +15,7 @@ CheckSingleInfluenceLOD=True
 MaxNumberOfLogLines=2000
 
 [/Script/UnrealEd.ThumbnailManager]
++RenderableThumbnailTypes=(ClassNeedingThumbnailName="/Script/CoreUObject.Class",RendererClassName="/Script/UnrealEd.ClassThumbnailRenderer")
 +RenderableThumbnailTypes=(ClassNeedingThumbnailName="/Script/Engine.Blueprint",RendererClassName="/Script/UnrealEd.BlueprintThumbnailRenderer")
 +RenderableThumbnailTypes=(ClassNeedingThumbnailName="/Script/Engine.Font",RendererClassName="/Script/UnrealEd.FontThumbnailRenderer")
 +RenderableThumbnailTypes=(ClassNeedingThumbnailName="/Script/Engine.Texture2D",RendererClassName="/Script/UnrealEd.TextureThumbnailRenderer")
diff --git a/Engine/Source/Developer/AssetTools/Private/AssetTypeActions/AssetTypeActions_Class.cpp b/Engine/Source/Developer/AssetTools/Private/AssetTypeActions/AssetTypeActions_Class.cpp
--- a/Engine/Source/Developer/AssetTools/Private/AssetTypeActions/AssetTypeActions_Class.cpp
+++ b/Engine/Source/Developer/AssetTools/Private/AssetTypeActions/AssetTypeActions_Class.cpp
@@ -43,7 +43,7 @@ void FAssetTypeActions_Class::GetActions(const TArray<UObject*>& InObjects, FMen
 	if(InObjects.Num() == 1)
 	{
 		const FText BaseClassName = FText::FromName(BaseClass->GetFName());
-		NewDerivedClassClassLabel = FText::Format(LOCTEXT("Class_NewDerivedClassLabel_CreateFrom", "New Class from {0}"), BaseClassName);
+		NewDerivedClassClassLabel = FText::Format(LOCTEXT("Class_NewDerivedClassLabel_CreateFrom", "Create class derived from {0}"), BaseClassName);
 		if(bIsValidBaseClass)
 		{
 			NewDerivedClassClassToolTip = FText::Format(LOCTEXT("Class_NewDerivedClassTooltip_CreateFrom", "Create a new class deriving from {0}."), BaseClassName);
@@ -55,7 +55,7 @@ void FAssetTypeActions_Class::GetActions(const TArray<UObject*>& InObjects, FMen
 	}
 	else
 	{
-		NewDerivedClassClassLabel = LOCTEXT("Class_NewDerivedClassLabel_InvalidNumberOfBases", "New Class from...");
+		NewDerivedClassClassLabel = LOCTEXT("Class_NewDerivedClassLabel_InvalidNumberOfBases", "New class derived from...");
 		NewDerivedClassClassToolTip = LOCTEXT("Class_NewDerivedClassTooltip_InvalidNumberOfBases", "Can only create a derived class when there is a single base class selected.");
 	}
 
@@ -70,6 +70,14 @@ void FAssetTypeActions_Class::GetActions(const TArray<UObject*>& InObjects, FMen
 		);
 }
 
+UThumbnailInfo* FAssetTypeActions_Class::GetThumbnailInfo(UObject* Asset) const
+{
+	// todo: jdale - CLASS - We need to generate and store proper thumbnail info for classes so that we can store their custom render transforms
+	// This can't be stored in the UClass instance (like we do for Blueprints), so we'll need another place to store it
+	// This will need to be accessible to FClassThumbnailScene::GetSceneThumbnailInfo
+	return nullptr;
+}
+
 void FAssetTypeActions_Class::OpenAssetEditor( const TArray<UObject*>& InObjects, TSharedPtr<class IToolkitHost> EditWithinLevelEditor )
 {
 	for(UObject* Object : InObjects)
diff --git a/Engine/Source/Developer/AssetTools/Private/AssetTypeActions/AssetTypeActions_Class.h b/Engine/Source/Developer/AssetTools/Private/AssetTypeActions/AssetTypeActions_Class.h
--- a/Engine/Source/Developer/AssetTools/Private/AssetTypeActions/AssetTypeActions_Class.h
+++ b/Engine/Source/Developer/AssetTools/Private/AssetTypeActions/AssetTypeActions_Class.h
@@ -8,13 +8,15 @@ class FAssetTypeActions_Class : public FAssetTypeActions_Base
 {
 public:
 	// IAssetTypeActions Implementation
-	virtual FText GetName() const override { return NSLOCTEXT("AssetTypeActions", "AssetTypeActions_Class", "Class"); }
-	virtual FColor GetTypeColor() const override { return FColor::White; }
+	virtual FText GetName() const override { return NSLOCTEXT("AssetTypeActions", "AssetTypeActions_Class", "C++ Class"); }
+	virtual FColor GetTypeColor() const override { return FColor(72, 155, 129); }
 	virtual UClass* GetSupportedClass() const override { return UClass::StaticClass(); }
 	virtual uint32 GetCategories() override { return EAssetTypeCategories::None; }
 	
 	virtual bool HasActions(const TArray<UObject*>& InObjects) const override;
 	virtual void GetActions(const TArray<UObject*>& InObjects, FMenuBuilder& MenuBuilder) override;
 
+	virtual class UThumbnailInfo* GetThumbnailInfo(UObject* Asset) const override;
+
 	virtual void OpenAssetEditor( const TArray<UObject*>& InObjects, TSharedPtr<class IToolkitHost> EditWithinLevelEditor = TSharedPtr<IToolkitHost>() ) override;
 };
diff --git a/Engine/Source/Editor/PlacementMode/Private/SPlacementModeTools.cpp b/Engine/Source/Editor/PlacementMode/Private/SPlacementModeTools.cpp
--- a/Engine/Source/Editor/PlacementMode/Private/SPlacementModeTools.cpp
+++ b/Engine/Source/Editor/PlacementMode/Private/SPlacementModeTools.cpp
@@ -65,7 +65,6 @@ public:
 		const TAttribute< FLinearColor >& HintColorAndOpacity = FLinearColor( 0.0f, 0.0f, 0.0f, 0.0f );
 		bool AllowHintText = true;
 		FName ClassThumbnailBrushOverride = NAME_None;
-		bool ShowBackground = false;
 
 		ChildSlot
 		[
@@ -76,8 +75,8 @@ public:
 			HighlightedText,
 			HintColorAndOpacity,
 			AllowHintText,
-			ClassThumbnailBrushOverride,
-			ShowBackground )
+			ClassThumbnailBrushOverride
+			)
 		];
 	}
 
diff --git a/Engine/Source/Editor/UnrealEd/Classes/ThumbnailRendering/ClassThumbnailRenderer.h b/Engine/Source/Editor/UnrealEd/Classes/ThumbnailRendering/ClassThumbnailRenderer.h
new file mode 100644
--- /dev/null
+++ b/Engine/Source/Editor/UnrealEd/Classes/ThumbnailRendering/ClassThumbnailRenderer.h
@@ -0,0 +1,23 @@
+// Copyright 1998-2015 Epic Games, Inc. All Rights Reserved.
+
+#pragma once
+#include "ClassThumbnailRenderer.generated.h"
+
+UCLASS(config=Editor,MinimalAPI)
+class UClassThumbnailRenderer : public UDefaultSizedThumbnailRenderer
+{
+	GENERATED_UCLASS_BODY()
+
+
+	// Begin UThumbnailRenderer Object
+	virtual bool CanVisualizeAsset(UObject* Object) override;
+	UNREALED_API virtual void Draw(UObject* Object, int32 X, int32 Y, uint32 Width, uint32 Height, FRenderTarget*, FCanvas* Canvas) override;
+	// End UThumbnailRenderer Object
+
+	// UObject implementation
+	UNREALED_API virtual void BeginDestroy() override;
+	// End UObject implementation
+
+private:
+	class FClassThumbnailScene* ThumbnailScene;
+};
diff --git a/Engine/Source/Editor/UnrealEd/Private/AssetThumbnail.cpp b/Engine/Source/Editor/UnrealEd/Private/AssetThumbnail.cpp
--- a/Engine/Source/Editor/UnrealEd/Private/AssetThumbnail.cpp
+++ b/Engine/Source/Editor/UnrealEd/Private/AssetThumbnail.cpp
@@ -22,7 +22,6 @@ public:
 		, _HighlightedText(FText::GetEmpty())
 		, _HintColorAndOpacity(FLinearColor(0.0f, 0.0f, 0.0f, 0.0f))
 		, _ClassThumbnailBrushOverride(NAME_None)
-		, _ShowClassBackground( true )
 		{}
 
 		SLATE_ARGUMENT( FName, Style )
@@ -35,7 +34,6 @@ public:
 		SLATE_ATTRIBUTE( FText, HighlightedText )
 		SLATE_ATTRIBUTE( FLinearColor, HintColorAndOpacity )
 		SLATE_ARGUMENT( FName, ClassThumbnailBrushOverride )
-		SLATE_ARGUMENT( bool, ShowClassBackground )
 
 	SLATE_END_ARGS()
 
@@ -47,7 +45,6 @@ public:
 		Label = InArgs._Label;
 		HintColorAndOpacity = InArgs._HintColorAndOpacity;
 		bAllowHintText = InArgs._AllowHintText;
-		bShowClassBackground = InArgs._ShowClassBackground;
 
 		AssetThumbnail = InArgs._AssetThumbnail;
 		bHasRenderedThumbnail = false;
@@ -74,54 +71,44 @@ public:
 
 		TSharedRef<SOverlay> OverlayWidget = SNew(SOverlay);
 
-		if( Class == UClass::StaticClass() )
-		{
-			ClassAssetClass = FindObject<UClass>(ANY_PACKAGE, *AssetData.AssetName.ToString());
-		}
+		UpdateClassAssetClass(AssetData.AssetName, Class);
 
 		ClassThumbnailBrushOverride = InArgs._ClassThumbnailBrushOverride;
 
-		const FSlateBrush* BackgroundBrush = GetBackgroundBrush();
-
-		if( ClassAssetClass != nullptr )
-		{
-			OverlayWidget->AddSlot()
-			[
-				SAssignNew(ClassThumbnailWidget, SBorder)
-				.BorderImage(BackgroundBrush)
-				.BorderBackgroundColor(AssetColor)
-				.Padding(GenericThumbnailBorderPadding)
-				.VAlign(VAlign_Center)
-				.HAlign(HAlign_Center)
-				.Visibility(this, &SAssetThumbnail::GetClassThumbnailVisibility)
-				[
-					SNew(SImage)
-					.Image(this, &SAssetThumbnail::GetClassThumbnailBrush)
-				]
-			];
-		}
-
 		// The generic representation of the thumbnail, for use before the rendered version, if it exists
 		OverlayWidget->AddSlot()
 		[
 			SNew(SBorder)
-			.BorderImage( BackgroundBrush )
-			.BorderBackgroundColor( AssetColor )
-			.Padding( GenericThumbnailBorderPadding )
-			.VAlign(VAlign_Center) .HAlign(HAlign_Center)
-			.Visibility( this, &SAssetThumbnail::GetGenericThumbnailVisibility )
+			.BorderImage(GetAssetBackgroundBrush())
+			.BorderBackgroundColor(AssetColor.CopyWithNewOpacity(0.3f))
+			.Padding(GenericThumbnailBorderPadding)
+			.VAlign(VAlign_Center)
+			.HAlign(HAlign_Center)
+			.Visibility(this, &SAssetThumbnail::GetGenericThumbnailVisibility)
 			[
-				SAssignNew(LabelTextBlock, STextBlock)
-				.Text( GetLabelText() )
-				.Font( GetTextFont() )
-				.ColorAndOpacity( FEditorStyle::GetColor(Style, ".ColorAndOpacity") )
-				.ShadowOffset( FEditorStyle::GetVector(Style, ".ShadowOffset") )
-				.ShadowColorAndOpacity( FEditorStyle::GetColor(Style, ".ShadowColorAndOpacity") )
-				.HighlightText( HighlightedText )
+				SNew(SOverlay)
+
+				+SOverlay::Slot()
+				[
+					SAssignNew(GenericLabelTextBlock, STextBlock)
+					.Text(GetLabelText())
+					.Font(GetTextFont())
+					.Justification(ETextJustify::Center)
+					.ColorAndOpacity(FEditorStyle::GetColor(Style, ".ColorAndOpacity"))
+					.ShadowOffset(FEditorStyle::GetVector(Style, ".ShadowOffset"))
+					.ShadowColorAndOpacity( FEditorStyle::GetColor(Style, ".ShadowColorAndOpacity"))
+					.HighlightText(HighlightedText)
+				]
+
+				+SOverlay::Slot()
+				[
+					SAssignNew(GenericThumbnailImage, SImage)
+					.Image(this, &SAssetThumbnail::GetClassThumbnailBrush)
+				]
 			]
 		];
 
-		if ( InArgs._ThumbnailPool.IsValid() && !InArgs._ForceGenericThumbnail && Class != UClass::StaticClass() )
+		if ( InArgs._ThumbnailPool.IsValid() && !InArgs._ForceGenericThumbnail )
 		{
 			ViewportFadeAnimation = FCurveSequence();
 			ViewportFadeCurve = ViewportFadeAnimation.AddCurve(0.f, 0.25f, ECurveEaseFunction::QuadOut);
@@ -155,6 +142,22 @@ public:
 			];
 		}
 
+		if( ClassAssetClass != nullptr && bIsClassType )
+		{
+			OverlayWidget->AddSlot()
+			.VAlign(VAlign_Bottom)
+			.HAlign(HAlign_Right)
+			.Padding(TAttribute<FMargin>(this, &SAssetThumbnail::GetClassIconPadding))
+			[
+				SAssignNew(ClassIconWidget, SBorder)
+				.BorderImage(FEditorStyle::GetNoBrush())
+				[
+					SNew(SImage)
+					.Image(this, &SAssetThumbnail::GetClassIconBrush)
+				]
+			];
+		}
+
 		if( bAllowHintText )
 		{
 			OverlayWidget->AddSlot()
@@ -180,7 +183,7 @@ public:
 				];
 		}
 
-		// The asset color strip, only visible when the rendered viewport is
+		// The asset color strip
 		OverlayWidget->AddSlot()
 		.HAlign(HAlign_Fill)
 		.VAlign(VAlign_Bottom)
@@ -200,6 +203,40 @@ public:
 
 	}
 
+	void UpdateClassAssetClass(const FName& InAssetName, UClass* InAssetClass)
+	{
+		ClassAssetClass = InAssetClass;
+		bIsClassType = false;
+
+		if( InAssetClass == UClass::StaticClass() )
+		{
+			ClassAssetClass = FindObject<UClass>(ANY_PACKAGE, *InAssetName.ToString());
+			bIsClassType = true;
+		}
+		else if( InAssetClass->IsChildOf<UBlueprint>() )
+		{
+			static const FName NativeParentClassTag = "NativeParentClass";
+			static const FName ParentClassTag = "ParentClass";
+
+			// We need to use the asset data to get the parent class as the blueprint may not be loaded
+			const FAssetData& AssetData = AssetThumbnail->GetAssetData();
+			const FString* ParentClassNamePtr = AssetData.TagsAndValues.Find(NativeParentClassTag);
+			if(!ParentClassNamePtr)
+			{
+				ParentClassNamePtr = AssetData.TagsAndValues.Find(ParentClassTag);
+			}
+			if(ParentClassNamePtr && !ParentClassNamePtr->IsEmpty())
+			{
+				UObject* Outer = nullptr;
+				FString ParentClassName = *ParentClassNamePtr;
+				ResolveName(Outer, ParentClassName, false, false);
+				ClassAssetClass = FindObject<UClass>(ANY_PACKAGE, *ParentClassName);
+			}
+
+			bIsClassType = true;
+		}
+	}
+
 	FSlateColor GetHintBackgroundColor() const
 	{
 		const FLinearColor Color = HintColorAndOpacity.Get();
@@ -225,10 +262,10 @@ public:
 			WidthLastFrame = AllottedGeometry.Size.X;
 
 			// The width changed, update the font
-			if ( LabelTextBlock.IsValid() )
+			if ( GenericLabelTextBlock.IsValid() )
 			{
-				LabelTextBlock->SetFont( GetTextFont() );
-				LabelTextBlock->SetWrapTextAt( GetTextWrapWidth() );
+				GenericLabelTextBlock->SetFont( GetTextFont() );
+				GenericLabelTextBlock->SetWrapTextAt( GetTextWrapWidth() );
 			}
 
 			if ( HintTextBlock.IsValid() )
@@ -242,9 +279,9 @@ public:
 private:
 	void OnAssetDataChanged()
 	{
-		if ( LabelTextBlock.IsValid() )
+		if ( GenericLabelTextBlock.IsValid() )
 		{
-			LabelTextBlock->SetText( GetLabelText() );
+			GenericLabelTextBlock->SetText( GetLabelText() );
 		}
 
 		if ( HintTextBlock.IsValid() )
@@ -278,10 +315,7 @@ private:
 			AssetTypeActions = AssetToolsModule.Get().GetAssetTypeActionsForClass(Class);
 		}
 
-		if (Class == UClass::StaticClass())
-		{
-			ClassAssetClass = FindObject<UClass>(ANY_PACKAGE, *AssetData.AssetName.ToString());
-		}
+		UpdateClassAssetClass(AssetData.AssetName, Class);
 
 		AssetColor = FLinearColor(1.f, 1.f, 1.f, 1.f);
 		if ( AssetTypeActions.IsValid() )
@@ -307,20 +341,15 @@ private:
 		return WidthLastFrame - GenericThumbnailBorderPadding * 2.f;
 	}
 
-	const FSlateBrush* GetBackgroundBrush() const
+	const FSlateBrush* GetAssetBackgroundBrush() const
 	{
-		const FAssetData& AssetData = AssetThumbnail->GetAssetData();
-		FString Substyle;
-		if ( AssetData.AssetClass == UClass::StaticClass()->GetFName() )
-		{
-			Substyle = TEXT(".ClassBackground");
-		}
-		else
-		{
-			Substyle = TEXT(".AssetBackground");
-		}
+		const FName BackgroundBrushName( *(Style.ToString() + TEXT(".AssetBackground")) );
+		return FEditorStyle::GetBrush(BackgroundBrushName);
+	}
 
-		const FName BackgroundBrushName( *(Style.ToString() + Substyle) );
+	const FSlateBrush* GetClassBackgroundBrush() const
+	{
+		const FName BackgroundBrushName( *(Style.ToString() + TEXT(".ClassBackground")) );
 		return FEditorStyle::GetBrush(BackgroundBrushName);
 	}
 
@@ -339,10 +368,15 @@ private:
 		return bHasRenderedThumbnail ? EVisibility::Visible : EVisibility::Collapsed;
 	}
 
-	FMargin GetAssetColorStripPadding() const
+	float GetAssetColorStripHeight() const
 	{
 		// The strip is 2.5% the height of the thumbnail, but at least 3 units tall
-		const float Height = FMath::Max(FMath::CeilToFloat(WidthLastFrame*0.025f), 3.0f);
+		return FMath::Max(FMath::CeilToFloat(WidthLastFrame*0.025f), 3.0f);
+	}
+
+	FMargin GetAssetColorStripPadding() const
+	{
+		const float Height = GetAssetColorStripHeight();
 		return FMargin(0,Height,0,0);
 	}
 
@@ -350,7 +384,10 @@ private:
 	{
 		if (ClassThumbnailBrushOverride.IsNone())
 		{
-			return FClassIconFinder::FindThumbnailForClass(ClassAssetClass.Get());
+			// For non-class types, use the default based upon the actual asset class
+			// This has the side effect of not showing a class icon for assets that don't have a proper thumbnail image available
+			const FName DefaultThumbnail = (bIsClassType) ? NAME_None : FName(*FString::Printf(TEXT("ClassThumbnail.%s"), *AssetThumbnail->GetAssetData().AssetClass.ToString()));
+			return FClassIconFinder::FindThumbnailForClass(ClassAssetClass.Get(), DefaultThumbnail);
 		}
 		else
 		{
@@ -366,15 +403,9 @@ private:
 		if(!bHasRenderedThumbnail)
 		{
 			const FSlateBrush* ClassThumbnailBrush = GetClassThumbnailBrush();
-			if( ClassThumbnailBrush )
+			if( ClassThumbnailBrush && ClassAssetClass.Get() )
 			{
-				const FAssetData& AssetData = AssetThumbnail->GetAssetData();
-				FString AssetClass = AssetData.AssetClass.ToString();
-				UClass* Class = FindObjectSafe<UClass>(ANY_PACKAGE, *AssetClass);
-				if(Class == UClass::StaticClass())
-				{
-					return EVisibility::Visible;
-				}
+				return EVisibility::Visible;
 			}
 		}
 
@@ -383,12 +414,23 @@ private:
 
 	EVisibility GetGenericThumbnailVisibility() const
 	{
-		return ( CachedClassThumbnailVisibility == EVisibility::Visible ) || (bHasRenderedThumbnail && ViewportFadeAnimation.IsAtEnd()) ? EVisibility::Collapsed : EVisibility::Visible;
+		return (bHasRenderedThumbnail && ViewportFadeAnimation.IsAtEnd()) ? EVisibility::Collapsed : EVisibility::Visible;
+	}
+
+	const FSlateBrush* GetClassIconBrush() const
+	{
+		return FClassIconFinder::FindIconForClass(ClassAssetClass.Get());
+	}
+
+	FMargin GetClassIconPadding() const
+	{
+		const float Height = GetAssetColorStripHeight();
+		return FMargin(0,0,0,Height);
 	}
 
 	EVisibility GetHintTextVisibility() const
 	{
-		if ( bAllowHintText && ( bHasRenderedThumbnail || !LabelTextBlock.IsValid() ) && HintColorAndOpacity.Get().A > 0 )
+		if ( bAllowHintText && ( bHasRenderedThumbnail || !GenericLabelTextBlock.IsValid() ) && HintColorAndOpacity.Get().A > 0 )
 		{
 			return EVisibility::Visible;
 		}
@@ -551,30 +593,33 @@ private:
 
 	void UpdateThumbnailVisibilities()
 	{
-	
-		CachedClassThumbnailVisibility = GetClassThumbnailVisibility();
-	
-		if( ClassThumbnailWidget.IsValid() )
+		// Either the generic label or thumbnail should be shown, but not both at once
+		const EVisibility ClassThumbnailVisibility = GetClassThumbnailVisibility();
+		if( GenericThumbnailImage.IsValid() )
 		{
-			ClassThumbnailWidget->SetVisibility( CachedClassThumbnailVisibility );
+			GenericThumbnailImage->SetVisibility( ClassThumbnailVisibility );
 		}
-
-
-		if( RenderedThumbnailWidget.IsValid() )
+		if( GenericLabelTextBlock.IsValid() )
 		{
-			RenderedThumbnailWidget->SetVisibility( GetViewportVisibility() );
+			GenericLabelTextBlock->SetVisibility( (ClassThumbnailVisibility == EVisibility::Visible) ? EVisibility::Collapsed : EVisibility::Visible );
 		}
 
-		if( AssetColorStripWidget.IsValid() )
+		const EVisibility ViewportVisibility = GetViewportVisibility();
+		if( RenderedThumbnailWidget.IsValid() )
 		{
-			AssetColorStripWidget->SetVisibility( GetViewportVisibility() );
+			RenderedThumbnailWidget->SetVisibility( ViewportVisibility );
+			if( ClassIconWidget.IsValid() )
+			{
+				ClassIconWidget->SetVisibility( ViewportVisibility );
+			}
 		}
 	}
 
 private:
-	TSharedPtr<STextBlock> LabelTextBlock;
+	TSharedPtr<STextBlock> GenericLabelTextBlock;
 	TSharedPtr<STextBlock> HintTextBlock;
-	TSharedPtr<SBorder> ClassThumbnailWidget;
+	TSharedPtr<SImage> GenericThumbnailImage;
+	TSharedPtr<SBorder> ClassIconWidget;
 	TSharedPtr<SBorder> RenderedThumbnailWidget;
 	TSharedPtr<SBorder> AssetColorStripWidget;
 	TSharedPtr<FAssetThumbnail> AssetThumbnail;
@@ -589,18 +634,15 @@ private:
 	TAttribute< FText > HighlightedText;
 	EThumbnailLabel::Type Label;
 
-	/** Cached class thumbnail visibility. This only ever changes if the asset changes */ 
-	EVisibility CachedClassThumbnailVisibility;
-
 	TAttribute< FLinearColor > HintColorAndOpacity;
 	bool bAllowHintText;
 
-	bool bShowClassBackground;
-
 	/** The name of the thumbnail which should be used instead of the class thumbnail. */
 	FName ClassThumbnailBrushOverride;
 	/** The class instance for assets which represent a class. */
 	TWeakObjectPtr<UClass> ClassAssetClass;
+	/** Are we showing a class type? (UClass, UBlueprint) */
+	bool bIsClassType;
 };
 
 
@@ -710,8 +752,8 @@ TSharedRef<SWidget> FAssetThumbnail::MakeThumbnailWidget(
 	const TAttribute< FText >& HighlightedText,
 	const TAttribute< FLinearColor >& HintColorAndOpacity,
 	bool AllowHintText,
-	FName ClassThumbnailBrushOverride,
-	bool ShowClassBackground )
+	FName ClassThumbnailBrushOverride
+	)
 {
 	return
 		SNew(SAssetThumbnail)
@@ -723,8 +765,7 @@ TSharedRef<SWidget> FAssetThumbnail::MakeThumbnailWidget(
 		.HighlightedText( HighlightedText )
 		.HintColorAndOpacity( HintColorAndOpacity )
 		.AllowHintText( AllowHintText )
-		.ClassThumbnailBrushOverride( ClassThumbnailBrushOverride )
-		.ShowClassBackground( ShowClassBackground );
+		.ClassThumbnailBrushOverride( ClassThumbnailBrushOverride );
 }
 
 void FAssetThumbnail::RefreshThumbnail()
diff --git a/Engine/Source/Editor/UnrealEd/Private/ClassIconFinder.cpp b/Engine/Source/Editor/UnrealEd/Private/ClassIconFinder.cpp
--- a/Engine/Source/Editor/UnrealEd/Private/ClassIconFinder.cpp
+++ b/Engine/Source/Editor/UnrealEd/Private/ClassIconFinder.cpp
@@ -109,7 +109,7 @@ FName FClassIconFinder::FindIconNameForActor( const TWeakObjectPtr<AActor>& InAc
 
 const FSlateBrush* FClassIconFinder::FindIconForClass(const UClass* InClass, const FName& InDefaultName )
 {
-	return FClassIconFinder::LookupBrush( FindIconNameImpl( InClass, InDefaultName ) );
+	return FClassIconFinder::LookupBrush( FindIconNameForClass( InClass, InDefaultName ) );
 }
 
 FName FClassIconFinder::FindIconNameForClass(const UClass* InClass, const FName& InDefaultName )
@@ -119,7 +119,7 @@ FName FClassIconFinder::FindIconNameForClass(const UClass* InClass, const FName&
 
 const FSlateBrush* FClassIconFinder::FindThumbnailForClass(const UClass* InClass, const FName& InDefaultName )
 {
-	return FClassIconFinder::LookupBrush( FindIconNameImpl( InClass, InDefaultName, TEXT("ClassThumbnail") ) );
+	return FClassIconFinder::LookupBrush( FindThumbnailNameForClass( InClass, InDefaultName ) );
 }
 
 FName FClassIconFinder::FindThumbnailNameForClass(const UClass* InClass, const FName& InDefaultName )
diff --git a/Engine/Source/Editor/UnrealEd/Private/ThumbnailHelpers.cpp b/Engine/Source/Editor/UnrealEd/Private/ThumbnailHelpers.cpp
--- a/Engine/Source/Editor/UnrealEd/Private/ThumbnailHelpers.cpp
+++ b/Engine/Source/Editor/UnrealEd/Private/ThumbnailHelpers.cpp
@@ -864,32 +864,28 @@
 
 /*
 ***************************************************************
-  FBlueprintThumbnailScene
+  FClassActorThumbnailScene
 ***************************************************************
 */
 
-FBlueprintThumbnailScene::FBlueprintThumbnailScene()
+FClassActorThumbnailScene::FClassActorThumbnailScene()
 	: FThumbnailPreviewScene()
 {
-	CurrentBlueprint = NULL;
-
-	FCoreUObjectDelegates::PreGarbageCollect.AddRaw(this, &FBlueprintThumbnailScene::OnPreGarbageCollect);
+	FCoreUObjectDelegates::PreGarbageCollect.AddRaw(this, &FClassActorThumbnailScene::OnPreGarbageCollect);
 }
 
-FBlueprintThumbnailScene::~FBlueprintThumbnailScene()
+FClassActorThumbnailScene::~FClassActorThumbnailScene()
 {
 	FCoreUObjectDelegates::PreGarbageCollect.RemoveAll(this);
 }
 
-void FBlueprintThumbnailScene::SetBlueprint(UBlueprint* Blueprint)
+void FClassActorThumbnailScene::SetObject(UObject* Obj)
 {
-	CurrentBlueprint = Blueprint;
-
-	if ( Blueprint )
+	if ( Obj )
 	{
-		VisualizableBlueprintComponents = GetPooledVisualizableComponents(Blueprint);
+		VisualizableComponents = GetPooledVisualizableComponents(Obj);
 
-		for ( auto CompIt = VisualizableBlueprintComponents.CreateConstIterator(); CompIt; ++CompIt )
+		for ( auto CompIt = VisualizableComponents.CreateConstIterator(); CompIt; ++CompIt )
 		{
 			UPrimitiveComponent* PrimComp = *CompIt;
 			PrimComp->bVisible = true;
@@ -898,40 +894,48 @@
 	}
 	else
 	{
-		for ( auto CompIt = VisualizableBlueprintComponents.CreateConstIterator(); CompIt; ++CompIt )
+		for ( auto CompIt = VisualizableComponents.CreateConstIterator(); CompIt; ++CompIt )
 		{
 			UPrimitiveComponent* PrimComp = *CompIt;
 			PrimComp->bVisible = false;
 			PrimComp->MarkRenderStateDirty();
 		}
 
-		VisualizableBlueprintComponents.Empty();
+		VisualizableComponents.Empty();
 	}
 
 	GetWorld()->SendAllEndOfFrameUpdates();
 }
 
-bool FBlueprintThumbnailScene::IsValidComponentForVisualization(UActorComponent* Component) const
+bool FClassActorThumbnailScene::IsValidComponentForVisualization(UActorComponent* Component) const
 {
 	UPrimitiveComponent* PrimComp = Cast<UPrimitiveComponent>(Component);
 	if ( PrimComp && PrimComp->IsVisible() && !PrimComp->bHiddenInGame )
 	{
+		UStaticMeshComponent* StaticMeshComp = Cast<UStaticMeshComponent>(Component);
+		if ( StaticMeshComp && StaticMeshComp->StaticMesh )
+		{
+			return true;
+		}
+
+		USkeletalMeshComponent* SkelMeshComp = Cast<USkeletalMeshComponent>(Component);
+		if ( SkelMeshComp && SkelMeshComp->SkeletalMesh )
+		{
+			return true;
+		}
+
+		/*
 		UParticleSystemComponent* ParicleSystemComp = Cast<UParticleSystemComponent>(PrimComp);
 
 		// @TODO Support particle systems in thumbnails
 		return !ParicleSystemComp;
+		*/
 	}
 
 	return false;
 }
 
-void FBlueprintThumbnailScene::BlueprintChanged(class UBlueprint* Blueprint)
+void FClassActorThumbnailScene::AddReferencedObjects( FReferenceCollector& Collector )
 {
-	// We could clear only the components for the specified blueprint, but we clear all components anyway because it is quick to regenerate them.
-	ClearComponentsPool();
-}
-
-void FBlueprintThumbnailScene::AddReferencedObjects( FReferenceCollector& Collector )
-{
 	// Clear all components so they are never considered "Referenced"
 	ClearComponentsPool();
@@ -939,14 +943,12 @@
 	FThumbnailPreviewScene::AddReferencedObjects(Collector);
 }
 
-void FBlueprintThumbnailScene::GetViewMatrixParameters(const float InFOVDegrees, FVector& OutOrigin, float& OutOrbitPitch, float& OutOrbitYaw, float& OutOrbitZoom) const
+void FClassActorThumbnailScene::GetViewMatrixParameters(const float InFOVDegrees, FVector& OutOrigin, float& OutOrbitPitch, float& OutOrbitYaw, float& OutOrbitZoom) const
 {
-	check(CurrentBlueprint);
-
 	const float HalfFOVRadians = FMath::DegreesToRadians<float>(InFOVDegrees) * 0.5f;
 	// Add extra size to view slightly outside of the sphere to compensate for perspective
 	FBoxSphereBounds Bounds(ForceInitToZero);
-	for ( auto CompIt = VisualizableBlueprintComponents.CreateConstIterator(); CompIt; ++CompIt )
+	for ( auto CompIt = VisualizableComponents.CreateConstIterator(); CompIt; ++CompIt )
 	{
 		Bounds = Bounds + (*CompIt)->Bounds;
 	}
@@ -955,18 +957,8 @@
 	const float BoundsZOffset = GetBoundsZOffset(Bounds);
 	const float TargetDistance = HalfMeshSize / FMath::Tan(HalfFOVRadians);
 
-	USceneThumbnailInfo* ThumbnailInfo = Cast<USceneThumbnailInfo>(CurrentBlueprint->ThumbnailInfo);
-	if ( ThumbnailInfo )
-	{
-		if ( TargetDistance + ThumbnailInfo->OrbitZoom < 0 )
-		{
-			ThumbnailInfo->OrbitZoom = -TargetDistance;
-		}
-	}
-	else
-	{
-		ThumbnailInfo = USceneThumbnailInfo::StaticClass()->GetDefaultObject<USceneThumbnailInfo>();
-	}
+	USceneThumbnailInfo* ThumbnailInfo = GetSceneThumbnailInfo(TargetDistance);
+	check(ThumbnailInfo);
 
 	OutOrigin = FVector(0, 0, -BoundsZOffset);
 	OutOrbitPitch = ThumbnailInfo->OrbitPitch;
@@ -974,7 +966,7 @@
 	OutOrbitZoom = TargetDistance + ThumbnailInfo->OrbitZoom;
 }
 
-UActorComponent* FBlueprintThumbnailScene::CreateComponentInstanceFromTemplate(UActorComponent* ComponentTemplate) const
+UActorComponent* FClassActorThumbnailScene::CreateComponentInstanceFromTemplate(UActorComponent* ComponentTemplate) const
 {
 	check(ComponentTemplate);
 
@@ -1015,6 +1007,59 @@
 	return NewComponent;
 }
 
+void FClassActorThumbnailScene::OnPreGarbageCollect()
+{
+	// This is a good time to clear the component pool to deal with invalid or stale entries. It will be re-populated as needed.
+	ClearComponentsPool();
+}
+
+void FClassActorThumbnailScene::ClearComponentsPool()
+{
+	for ( auto PoolIt = AllComponentsPool.CreateConstIterator(); PoolIt; ++PoolIt )
+	{
+		TArray< TWeakObjectPtr<UActorComponent> > WeakComponents = PoolIt.Value();
+		for ( auto CompIt = WeakComponents.CreateConstIterator(); CompIt; ++CompIt )
+		{
+			UActorComponent* ActorComp = (*CompIt).GetEvenIfUnreachable();
+			if ( ActorComp )
+			{
+				FPreviewScene::RemoveComponent(ActorComp);
+			}
+		}
+	}
+
+	AllComponentsPool.Empty();
+	VisualizableComponentsPool.Empty();
+}
+
+/*
+***************************************************************
+  FBlueprintThumbnailScene
+***************************************************************
+*/
+
+FBlueprintThumbnailScene::FBlueprintThumbnailScene()
+	: FClassActorThumbnailScene()
+{
+	CurrentBlueprint = nullptr;
+}
+
+FBlueprintThumbnailScene::~FBlueprintThumbnailScene()
+{
+}
+
+void FBlueprintThumbnailScene::SetBlueprint(UBlueprint* Blueprint)
+{
+	CurrentBlueprint = Blueprint;
+	SetObject(CurrentBlueprint);
+}
+
+void FBlueprintThumbnailScene::BlueprintChanged(class UBlueprint* Blueprint)
+{
+	// We could clear only the components for the specified blueprint, but we clear all components anyway because it is quick to regenerate them.
+	ClearComponentsPool();
+}
+
 void FBlueprintThumbnailScene::InstanceComponents(USCS_Node* CurrentNode, USceneComponent* ParentComponent, const TMap<UActorComponent*, UActorComponent*>& NativeInstanceMap, TArray<UActorComponent*>& OutComponents)
 {
 	check(CurrentNode != NULL);
@@ -1075,8 +1120,9 @@
 	}
 }
 
-TArray<UPrimitiveComponent*> FBlueprintThumbnailScene::GetPooledVisualizableComponents(UBlueprint* Blueprint)
+TArray<UPrimitiveComponent*> FBlueprintThumbnailScene::GetPooledVisualizableComponents(UObject* Obj)
 {
+	UBlueprint* Blueprint = Cast<UBlueprint>(Obj);
 	check(Blueprint);
 
 	TArray<UPrimitiveComponent*> VisualizableComponentsListForBlueprint;
@@ -1269,27 +1315,190 @@
 	return VisualizableComponentsListForBlueprint;
 }
 
-void FBlueprintThumbnailScene::OnPreGarbageCollect()
+USceneThumbnailInfo* FBlueprintThumbnailScene::GetSceneThumbnailInfo(const float TargetDistance) const
 {
-	// This is a good time to clear the component pool to deal with invalid or stale entries. It will be re-populated as needed.
-	ClearComponentsPool();
+	check(CurrentBlueprint);
+
+	USceneThumbnailInfo* ThumbnailInfo = Cast<USceneThumbnailInfo>(CurrentBlueprint->ThumbnailInfo);
+	if ( ThumbnailInfo )
+	{
+		if ( TargetDistance + ThumbnailInfo->OrbitZoom < 0 )
+		{
+			ThumbnailInfo->OrbitZoom = -TargetDistance;
+		}
+	}
+	else
+	{
+		ThumbnailInfo = USceneThumbnailInfo::StaticClass()->GetDefaultObject<USceneThumbnailInfo>();
+	}
+
+	return ThumbnailInfo;
 }
 
-void FBlueprintThumbnailScene::ClearComponentsPool()
+/*
+***************************************************************
+  FClassThumbnailScene
+***************************************************************
+*/
+
+FClassThumbnailScene::FClassThumbnailScene()
+	: FClassActorThumbnailScene()
 {
-	for ( auto PoolIt = AllComponentsPool.CreateConstIterator(); PoolIt; ++PoolIt )
+	CurrentClass = nullptr;
+}
+
+FClassThumbnailScene::~FClassThumbnailScene()
+{
+}
+
+void FClassThumbnailScene::SetClass(UClass* Class)
+{
+	CurrentClass = Class;
+	SetObject(CurrentClass);
+}
+
+TArray<UPrimitiveComponent*> FClassThumbnailScene::GetPooledVisualizableComponents(UObject* Obj)
+{
+	UClass* Class = Cast<UClass>(Obj);
+	check(Class);
+
+	TArray<UPrimitiveComponent*> VisualizableComponentsListForBlueprint;
+
+	TArray< TWeakObjectPtr<UPrimitiveComponent> >* PooledComponentsPtr = VisualizableComponentsPool.Find(Class);
+	if ( PooledComponentsPtr )
 	{
-		TArray< TWeakObjectPtr<UActorComponent> > WeakComponents = PoolIt.Value();
-		for ( auto CompIt = WeakComponents.CreateConstIterator(); CompIt; ++CompIt )
+		TArray< TWeakObjectPtr<UPrimitiveComponent> >& PooledComponents = *PooledComponentsPtr;
+		for ( auto CompIt = PooledComponents.CreateConstIterator(); CompIt; ++CompIt )
 		{
-			UActorComponent* ActorComp = (*CompIt).GetEvenIfUnreachable();
-			if ( ActorComp )
+			UPrimitiveComponent* Component = (*CompIt).GetEvenIfUnreachable();
+			if ( Component )
 			{
-				FPreviewScene::RemoveComponent(ActorComp);
+				VisualizableComponentsListForBlueprint.Add(Component);
 			}
 		}
 	}
+	else
+	{
+		// We need to find the RootComponent in order to property transform the components
+		USceneComponent* RootComponent = nullptr;
+		TArray<UActorComponent*> AllCreatedActorComponents;
+		TMap<UActorComponent*, UActorComponent*> NativeInstanceMap;
 
-	AllComponentsPool.Empty();
-	VisualizableComponentsPool.Empty();
+		if ( Class->IsChildOf(AActor::StaticClass()) )
+		{
+			// Instance native components based on the CDO template
+			AActor* CDO = Class->GetDefaultObject<AActor>();
+
+			TInlineComponentArray<UActorComponent*> Components;
+			CDO->GetComponents(Components);
+
+			for ( auto CompIt = Components.CreateConstIterator(); CompIt; ++CompIt )
+			{
+				UActorComponent* Comp = *CompIt;
+				if ( Comp != NULL )
+				{
+					UActorComponent* InstancedComponent = CreateComponentInstanceFromTemplate(Comp);
+					if ( InstancedComponent != NULL )
+					{
+						NativeInstanceMap.Add(Comp, InstancedComponent);
+					}
+				}
+			}
+
+			// Fix up parent and child attachments to point at the new instances
+			for ( auto CompIt = NativeInstanceMap.CreateConstIterator(); CompIt; ++CompIt )
+			{
+				UActorComponent* ActorComp = CompIt.Value();
+				if(ActorComp != NULL)
+				{
+					AllCreatedActorComponents.Add(ActorComp);
+
+					USceneComponent* SceneComp = Cast<USceneComponent>(ActorComp);
+					if(SceneComp != NULL)
+					{
+						if(SceneComp->AttachParent != NULL)
+						{
+							SceneComp->AttachParent = Cast<USceneComponent>(NativeInstanceMap.FindRef(SceneComp->AttachParent));
+						}
+						else if(CompIt.Key() == CDO->GetRootComponent())
+						{
+							RootComponent = SceneComp;
+
+							// Make sure the root component transform is Identity. Actors ignore the transform of the root component.
+							RootComponent->SetRelativeTransform( FTransform::Identity );
+						}
+
+						for(auto ChildCompIt = SceneComp->AttachChildren.CreateIterator(); ChildCompIt; ++ChildCompIt)
+						{
+							*ChildCompIt = Cast<USceneComponent>(NativeInstanceMap.FindRef(*ChildCompIt));
+						}
+					}
+				}
+			}
+		}
+
+		// Calculate the bounds. Include all visualizable components.
+		// Update the transform for all components since they will be used to transform the visualizable ones too.
+		FBoxSphereBounds Bounds(ForceInitToZero);
+		for ( auto CompIt = AllCreatedActorComponents.CreateConstIterator(); CompIt; ++CompIt )
+		{
+			USceneComponent* SceneComp = Cast<USceneComponent>(*CompIt);
+			if ( SceneComp )
+			{
+				SceneComp->UpdateComponentToWorld();
+
+				if ( IsValidComponentForVisualization(*CompIt) )
+				{
+					UPrimitiveComponent* PrimComp = CastChecked<UPrimitiveComponent>(*CompIt);
+					Bounds = Bounds + PrimComp->Bounds;
+				}
+			}
+		}
+
+		// Center the mesh at the world origin then offset to put it on top of the plane
+		const float BoundsZOffset = GetBoundsZOffset(Bounds);
+		FTransform CompTransform(-Bounds.Origin + FVector(0, 0, BoundsZOffset));
+
+		// Add all instanced scene components to the scene.
+		// Hide all non-visualizable ones.
+		TArray< TWeakObjectPtr<UActorComponent> > WeakAllComponentsList;
+		TArray< TWeakObjectPtr<UPrimitiveComponent> > WeakVisualizableComponentsList;
+		for ( auto CompIt = AllCreatedActorComponents.CreateConstIterator(); CompIt; ++CompIt )
+		{
+			WeakAllComponentsList.Add(*CompIt);
+
+			USceneComponent* SceneComp = Cast<USceneComponent>(*CompIt);
+			if ( SceneComp != NULL )
+			{
+				if ( IsValidComponentForVisualization(SceneComp) )
+				{
+					UPrimitiveComponent* PrimComp = CastChecked<UPrimitiveComponent>(SceneComp);
+					WeakVisualizableComponentsList.Add(PrimComp);
+					VisualizableComponentsListForBlueprint.Add(PrimComp);
+				}
+				else
+				{
+					// If this was a non-visualizable scene component, mark it invisible.
+					SceneComp->bVisible = false;
+				}
+
+				// Add the component to the scene.
+				FPreviewScene::AddComponent(*CompIt, CompTransform);
+			}
+		}
+
+		// Keep track of all components to reuse them in future render calls.
+		// These lists are transient and are rebuilt after garbage collection
+		AllComponentsPool.Add( Class, WeakAllComponentsList );
+		VisualizableComponentsPool.Add( Class, WeakVisualizableComponentsList );
+	}
+	
+	return VisualizableComponentsListForBlueprint;
 }
+
+USceneThumbnailInfo* FClassThumbnailScene::GetSceneThumbnailInfo(const float TargetDistance) const
+{
+	// todo: jdale - CLASS - Needs proper thumbnail info for class (see FAssetTypeActions_Class::GetThumbnailInfo)
+	USceneThumbnailInfo* ThumbnailInfo = USceneThumbnailInfo::StaticClass()->GetDefaultObject<USceneThumbnailInfo>();
+	return ThumbnailInfo;
+}
diff --git a/Engine/Source/Editor/UnrealEd/Private/ThumbnailRendering/ClassThumbnailRenderer.cpp b/Engine/Source/Editor/UnrealEd/Private/ThumbnailRendering/ClassThumbnailRenderer.cpp
new file mode 100644
--- /dev/null
+++ b/Engine/Source/Editor/UnrealEd/Private/ThumbnailRendering/ClassThumbnailRenderer.cpp
@@ -0,0 +1,78 @@
+// Copyright 1998-2015 Epic Games, Inc. All Rights Reserved.
+
+#include "UnrealEd.h"
+
+// FPreviewScene derived helpers for rendering
+#include "ThumbnailHelpers.h"
+#include "EngineModule.h"
+#include "RendererInterface.h"
+
+UClassThumbnailRenderer::UClassThumbnailRenderer(const FObjectInitializer& ObjectInitializer)
+	: Super(ObjectInitializer)
+{
+	ThumbnailScene = nullptr;
+}
+
+bool UClassThumbnailRenderer::CanVisualizeAsset(UObject* Object)
+{
+	if (ThumbnailScene == nullptr)
+	{
+		ThumbnailScene = new FClassThumbnailScene();
+	}
+
+	UClass* Class = Cast<UClass>(Object);
+
+	// Only visualize actor based classes
+	if (Class && Class->IsChildOf(AActor::StaticClass()))
+	{
+		// Try to find any visible primitive components in the class' CDO
+		AActor* CDO = Class->GetDefaultObject<AActor>();
+
+		TInlineComponentArray<UActorComponent*> Components;
+		CDO->GetComponents(Components);
+
+		for (auto CompIt = Components.CreateConstIterator(); CompIt; ++CompIt)
+		{
+			if (ThumbnailScene->IsValidComponentForVisualization(*CompIt))
+			{
+				return true;
+			}
+		}
+	}
+
+	return false;
+}
+
+void UClassThumbnailRenderer::Draw(UObject* Object, int32 X, int32 Y, uint32 Width, uint32 Height, FRenderTarget* RenderTarget, FCanvas* Canvas)
+{
+	UClass* Class = Cast<UClass>(Object);
+	if (Class != nullptr)
+	{
+		if ( ThumbnailScene == nullptr )
+		{
+			ThumbnailScene = new FClassThumbnailScene();
+		}
+
+		ThumbnailScene->SetClass(Class);
+		FSceneViewFamilyContext ViewFamily( FSceneViewFamily::ConstructionValues( RenderTarget, ThumbnailScene->GetScene(), FEngineShowFlags(ESFIM_Game) )
+			.SetWorldTimes(FApp::GetCurrentTime() - GStartTime, FApp::GetDeltaTime(), FApp::GetCurrentTime() - GStartTime));
+
+		ViewFamily.EngineShowFlags.DisableAdvancedFeatures();
+		ViewFamily.EngineShowFlags.MotionBlur = 0;
+
+		ThumbnailScene->GetView(&ViewFamily, X, Y, Width, Height);
+		GetRendererModule().BeginRenderingViewFamily(Canvas,&ViewFamily);
+		ThumbnailScene->SetClass(nullptr);
+	}
+}
+
+void UClassThumbnailRenderer::BeginDestroy()
+{
+	if ( ThumbnailScene != nullptr )
+	{
+		delete ThumbnailScene;
+		ThumbnailScene = nullptr;
+	}
+
+	Super::BeginDestroy();
+}
diff --git a/Engine/Source/Editor/UnrealEd/Public/AssetThumbnail.h b/Engine/Source/Editor/UnrealEd/Public/AssetThumbnail.h
--- a/Engine/Source/Editor/UnrealEd/Public/AssetThumbnail.h
+++ b/Engine/Source/Editor/UnrealEd/Public/AssetThumbnail.h
@@ -82,8 +82,8 @@ public:
 		const TAttribute< FText >& HighlightedText = TAttribute< FText >( FText::GetEmpty() ),
 		const TAttribute< FLinearColor >& HintColorAndOpacity = FLinearColor( 0.0f, 0.0f, 0.0f, 0.0f ),
 		bool AllowHintText = true,
-		FName ClassThumbnailBrushOverride = NAME_None,
-		bool ShowBackground = true);
+		FName ClassThumbnailBrushOverride = NAME_None
+		);
 
 	/** Re-renders this thumbnail */
 	UNREALED_API void RefreshThumbnail();
diff --git a/Engine/Source/Editor/UnrealEd/Public/ClassIconFinder.h b/Engine/Source/Editor/UnrealEd/Public/ClassIconFinder.h
--- a/Engine/Source/Editor/UnrealEd/Public/ClassIconFinder.h
+++ b/Engine/Source/Editor/UnrealEd/Public/ClassIconFinder.h
@@ -30,7 +30,7 @@ public:
 	/** Find the small icon name to use for the supplied class */
 	UNREALED_API static FName FindIconNameForClass(const UClass* InClass, const FName& InDefaultName = FName() );
 
-	/** Find the small icon to use for the supplied class */
+	/** Find the large thumbnail name to use for the supplied class */
 	UNREALED_API static const FSlateBrush* FindThumbnailForClass(const UClass* InClass, const FName& InDefaultName = FName() );
 
 	/** Find the large thumbnail name to use for the supplied class */
diff --git a/Engine/Source/Editor/UnrealEd/Public/ThumbnailHelpers.h b/Engine/Source/Editor/UnrealEd/Public/ThumbnailHelpers.h
--- a/Engine/Source/Editor/UnrealEd/Public/ThumbnailHelpers.h
+++ b/Engine/Source/Editor/UnrealEd/Public/ThumbnailHelpers.h
@@ -182,22 +182,16 @@ private:
 
 class UActorComponent;
 
-class UNREALED_API FBlueprintThumbnailScene : public FThumbnailPreviewScene
+class UNREALED_API FClassActorThumbnailScene : public FThumbnailPreviewScene
 {
 public:
 	/** Constructor/Destructor */
-	FBlueprintThumbnailScene();
-	~FBlueprintThumbnailScene();
-
-	/** Sets the static mesh to use in the next GetView() */
-	void SetBlueprint(class UBlueprint* Blueprint);
+	FClassActorThumbnailScene();
+	~FClassActorThumbnailScene();
 
 	/** Returns true if this component can be visualized */
 	bool IsValidComponentForVisualization(UActorComponent* Component) const;
 
-	/** Refreshes components for the specified blueprint */
-	void BlueprintChanged(class UBlueprint* Blueprint);
-
 	void AddReferencedObjects( FReferenceCollector& Collector ) override;
 
 protected:
@@ -207,29 +201,78 @@ protected:
 	/** Returns a duplicate of the specified component whose outer is the transient package. If the component can not be created, a placeholder component is made in its place. */
 	UActorComponent* CreateComponentInstanceFromTemplate(UActorComponent* ComponentTemplate) const;
 
-	/** Creates instances of template components found in a blueprint's simple construction script */
-	void InstanceComponents(USCS_Node* CurrentNode, USceneComponent* ParentComponent, const TMap<UActorComponent*, UActorComponent*>& NativeInstanceMap, TArray<UActorComponent*>& OutComponents);
-
-	/** Get/Release for the component pool */
-	TArray<UPrimitiveComponent*> GetPooledVisualizableComponents(UBlueprint* Blueprint);
-
 	/** Handler for when garbage collection occurs. Used to clear the ComponentsPool */
 	void OnPreGarbageCollect();
 
 	/** Removes all references to components in the components pool and empties all lists. This prepares all loaded components for GC. */
 	void ClearComponentsPool();
 
-private:
-	/** The blueprint that is currently being rendered. NULL when not rendering. */
-	UBlueprint* CurrentBlueprint;
+	/** Sets the object (class or blueprint) used in the next GetView() */
+	virtual void SetObject(class UObject* Obj);
+
+	/** Get/Release for the component pool */
+	virtual TArray<UPrimitiveComponent*> GetPooledVisualizableComponents(UObject* Obj) = 0;
+
+	/** Get the scene thumbnail info to use for the object currently being rendered */
+	virtual USceneThumbnailInfo* GetSceneThumbnailInfo(const float TargetDistance) const = 0;
 
 	/** Instances of the visualizable components found in the blueprint. This array only has elements while the blueprint is being rendered.  */
-	TArray<UPrimitiveComponent*> VisualizableBlueprintComponents;
+	TArray<UPrimitiveComponent*> VisualizableComponents;
 
-	/** A map of Blueprint to "pooled component list".
+	/** A map of objects to "pooled component list".
 	  * This will be populated with components that were created by this preview scene.
 	  * Objects in this map are not persistent and will be garbage collected at GC time.
 	  */
-	TMap< TWeakObjectPtr<UBlueprint>, TArray< TWeakObjectPtr<UActorComponent> > > AllComponentsPool;
-	TMap< TWeakObjectPtr<UBlueprint>, TArray< TWeakObjectPtr<UPrimitiveComponent> > > VisualizableComponentsPool;
-};
\ No newline at end of file
+	TMap< TWeakObjectPtr<UObject>, TArray< TWeakObjectPtr<UActorComponent> > > AllComponentsPool;
+	TMap< TWeakObjectPtr<UObject>, TArray< TWeakObjectPtr<UPrimitiveComponent> > > VisualizableComponentsPool;
+};
+
+class UNREALED_API FBlueprintThumbnailScene : public FClassActorThumbnailScene
+{
+public:
+	/** Constructor/Destructor */
+	FBlueprintThumbnailScene();
+	~FBlueprintThumbnailScene();
+
+	/** Sets the static mesh to use in the next GetView() */
+	void SetBlueprint(class UBlueprint* Blueprint);
+
+	/** Refreshes components for the specified blueprint */
+	void BlueprintChanged(class UBlueprint* Blueprint);
+
+protected:
+	/** Creates instances of template components found in a blueprint's simple construction script */
+	void InstanceComponents(USCS_Node* CurrentNode, USceneComponent* ParentComponent, const TMap<UActorComponent*, UActorComponent*>& NativeInstanceMap, TArray<UActorComponent*>& OutComponents);
+
+	/** Get/Release for the component pool */
+	virtual TArray<UPrimitiveComponent*> GetPooledVisualizableComponents(UObject* Obj) override;
+
+	/** Get the scene thumbnail info to use for the object currently being rendered */
+	virtual USceneThumbnailInfo* GetSceneThumbnailInfo(const float TargetDistance) const override;
+
+private:
+	/** The blueprint that is currently being rendered. NULL when not rendering. */
+	UBlueprint* CurrentBlueprint;
+};
+
+class UNREALED_API FClassThumbnailScene : public FClassActorThumbnailScene
+{
+public:
+	/** Constructor/Destructor */
+	FClassThumbnailScene();
+	~FClassThumbnailScene();
+
+	/** Sets the class use in the next GetView() */
+	void SetClass(class UClass* Class);
+
+protected:
+	/** Get/Release for the component pool */
+	virtual TArray<UPrimitiveComponent*> GetPooledVisualizableComponents(UObject* Obj) override;
+
+	/** Get the scene thumbnail info to use for the object currently being rendered */
+	virtual USceneThumbnailInfo* GetSceneThumbnailInfo(const float TargetDistance) const override;
+
+private:
+	/** The class that is currently being rendered. NULL when not rendering. */
+	UClass* CurrentClass;
+};
