[INTEGRATE] Change 2415366 by [REDACTED_PERSON][REDACTED_HANDLE] on 2015/01/22 11:18:18

	Fix cases where entries in InstanceComponents can end up null

diff --git a/Engine/Source/Editor/DetailCustomizations/Private/ActorDetails.cpp b/Engine/Source/Editor/DetailCustomizations/Private/ActorDetails.cpp
--- a/Engine/Source/Editor/DetailCustomizations/Private/ActorDetails.cpp
+++ b/Engine/Source/Editor/DetailCustomizations/Private/ActorDetails.cpp
@@ -1441,7 +1441,10 @@ FReply FActorDetails::ResetToBlueprintDefaults_OnClicked( TWeakObjectPtr<UBluepr
 				NumChangedProperties += InstanceComponents.Num();
 				for (UActorComponent* ActorComponent : InstanceComponents)
 				{
-					ActorComponent->DestroyComponent();
+					if (ActorComponent)
+					{
+						ActorComponent->DestroyComponent();
+					}
 				}
 			}
 		}
diff --git a/Engine/Source/Editor/UnrealEd/Private/Kismet2/Kismet2.cpp b/Engine/Source/Editor/UnrealEd/Private/Kismet2/Kismet2.cpp
--- a/Engine/Source/Editor/UnrealEd/Private/Kismet2/Kismet2.cpp
+++ b/Engine/Source/Editor/UnrealEd/Private/Kismet2/Kismet2.cpp
@@ -922,69 +922,72 @@ void FKismetEditorUtilities::AddComponentsToBlueprint(UBlueprint* Blueprint, con
 
 	for (UActorComponent* ActorComponent : Components)
 	{
-		if (Actor)
+		if (ActorComponent)
 		{
-			check(Actor == ActorComponent->GetOwner());
-		}
-		else
-		{
-			Actor = ActorComponent->GetOwner();
-			check(Actor);
-		}
+			if (Actor)
+			{
+				check(Actor == ActorComponent->GetOwner());
+			}
+			else
+			{
+				Actor = ActorComponent->GetOwner();
+				check(Actor);
+			}
 
-		USCS_Node* SCSNode = SCS->CreateNode(ActorComponent->GetClass());
-		UEditorEngine::CopyPropertiesForUnrelatedObjects(ActorComponent,SCSNode->ComponentTemplate);
+			USCS_Node* SCSNode = SCS->CreateNode(ActorComponent->GetClass());
+			UEditorEngine::CopyPropertiesForUnrelatedObjects(ActorComponent,SCSNode->ComponentTemplate);
 
-		// The easy part is non-scene component or the Root simply add it
-		if (!ActorComponent->IsA<USceneComponent>() || ActorComponent == Actor->GetRootComponent())
-		{
-			SCS->AddNode(SCSNode);
-		}
-		else
-		{
-			USceneComponent* SceneComponent = CastChecked<USceneComponent>(ActorComponent);
-			check(SceneComponent->AttachParent);
+			// The easy part is non-scene component or the Root simply add it
+			if (!ActorComponent->IsA<USceneComponent>() || ActorComponent == Actor->GetRootComponent())
+			{
+				SCS->AddNode(SCSNode);
+			}
+			else
+			{
+				USceneComponent* SceneComponent = CastChecked<USceneComponent>(ActorComponent);
+				check(SceneComponent->AttachParent);
 
-			InstanceComponentToNodeMap.Add(SceneComponent,SCSNode);
+				InstanceComponentToNodeMap.Add(SceneComponent,SCSNode);
 
-			// If we're attached to a blueprint component look it up as the variable name is the component name
-			if (SceneComponent->AttachParent->bCreatedByConstructionScript)
-			{
-				USCS_Node* ParentSCSNode = nullptr;
-				for (UBlueprint* Blueprint : ParentBPStack)
+				// If we're attached to a blueprint component look it up as the variable name is the component name
+				if (SceneComponent->AttachParent->bCreatedByConstructionScript)
 				{
-					ParentSCSNode = Blueprint->SimpleConstructionScript->FindSCSNode(SceneComponent->AttachParent->GetFName());
-					if (ParentSCSNode)
+					USCS_Node* ParentSCSNode = nullptr;
+					for (UBlueprint* Blueprint : ParentBPStack)
 					{
-						break;
+						ParentSCSNode = Blueprint->SimpleConstructionScript->FindSCSNode(SceneComponent->AttachParent->GetFName());
+						if (ParentSCSNode)
+						{
+							break;
+						}
 					}
-				}
-				check(ParentSCSNode);
+					check(ParentSCSNode);
 
-				if (ParentSCSNode->GetSCS() != SCS)
-				{
-					SCS->AddNode(SCSNode);
-				}
+					if (ParentSCSNode->GetSCS() != SCS)
+					{
+						SCS->AddNode(SCSNode);
+					}
 
-				SCSNode->SetParent(ParentSCSNode);
-			}
-			// If we're attached to a native component
-			else if (!Components.Contains(SceneComponent->AttachParent))
-			{
-				SCS->AddNode(SCSNode);
-				SCSNode->SetParent(SceneComponent->AttachParent);
-			}
-			else
-			{
-				// Otherwise check if we've already created the parents' new SCS node and attach to that or cache it off to do next pass
-				USCS_Node** ParentSCSNode = InstanceComponentToNodeMap.Find(SceneComponent->AttachParent);
-				if (ParentSCSNode)
+					SCSNode->SetParent(ParentSCSNode);
+				}
+				// If we're attached to a native component
+				else if (!Components.Contains(SceneComponent->AttachParent))
 				{
-					(*ParentSCSNode)->AddChildNode(SCSNode);
+					SCS->AddNode(SCSNode);
+					SCSNode->SetParent(SceneComponent->AttachParent);
 				}
 				else
 				{
-					SceneComponentsToAdd.Add(SceneComponent, SCSNode);
+					// Otherwise check if we've already created the parents' new SCS node and attach to that or cache it off to do next pass
+					USCS_Node** ParentSCSNode = InstanceComponentToNodeMap.Find(SceneComponent->AttachParent);
+					if (ParentSCSNode)
+					{
+						(*ParentSCSNode)->AddChildNode(SCSNode);
+					}
+					else
+					{
+						SceneComponentsToAdd.Add(SceneComponent, SCSNode);
+					}
 				}
 			}
 		}
