[INTEGRATE] Change 2415497 by [REDACTED_PERSON][REDACTED_HANDLE] on 2015/01/22 13:33:05

	Stopped derived types from UBlueprint being processed as if they were a BP class

	UE-7184 - Show game c++ classes and engine c++ classes in the content browser

	Also renamed "ClassAssetClass" to "ThumbnailClass" to better reflect its purpose.

diff --git a/Engine/Source/Editor/UnrealEd/Private/AssetThumbnail.cpp b/Engine/Source/Editor/UnrealEd/Private/AssetThumbnail.cpp
--- a/Engine/Source/Editor/UnrealEd/Private/AssetThumbnail.cpp
+++ b/Engine/Source/Editor/UnrealEd/Private/AssetThumbnail.cpp
@@ -71,7 +71,7 @@ public:
 
 		TSharedRef<SOverlay> OverlayWidget = SNew(SOverlay);
 
-		UpdateClassAssetClass(AssetData.AssetName, Class);
+		UpdateThumbnailClass(AssetData.AssetName, Class);
 
 		ClassThumbnailBrushOverride = InArgs._ClassThumbnailBrushOverride;
 
@@ -142,7 +142,7 @@ public:
 			];
 		}
 
-		if( ClassAssetClass != nullptr && bIsClassType )
+		if( ThumbnailClass.Get() && bIsClassType )
 		{
 			OverlayWidget->AddSlot()
 			.VAlign(VAlign_Bottom)
@@ -203,17 +203,17 @@ public:
 
 	}
 
-	void UpdateClassAssetClass(const FName& InAssetName, UClass* InAssetClass)
+	void UpdateThumbnailClass(const FName& InAssetName, UClass* InAssetClass)
 	{
-		ClassAssetClass = InAssetClass;
+		ThumbnailClass = InAssetClass;
 		bIsClassType = false;
 
 		if( InAssetClass == UClass::StaticClass() )
 		{
-			ClassAssetClass = FindObject<UClass>(ANY_PACKAGE, *InAssetName.ToString());
+			ThumbnailClass = FindObject<UClass>(ANY_PACKAGE, *InAssetName.ToString());
 			bIsClassType = true;
 		}
-		else if( InAssetClass->IsChildOf<UBlueprint>() )
+		else if( InAssetClass == UBlueprint::StaticClass() )
 		{
 			static const FName NativeParentClassTag = "NativeParentClass";
 			static const FName ParentClassTag = "ParentClass";
@@ -230,7 +230,7 @@ public:
 				UObject* Outer = nullptr;
 				FString ParentClassName = *ParentClassNamePtr;
 				ResolveName(Outer, ParentClassName, false, false);
-				ClassAssetClass = FindObject<UClass>(ANY_PACKAGE, *ParentClassName);
+				ThumbnailClass = FindObject<UClass>(ANY_PACKAGE, *ParentClassName);
 			}
 
 			bIsClassType = true;
@@ -315,7 +315,7 @@ private:
 			AssetTypeActions = AssetToolsModule.Get().GetAssetTypeActionsForClass(Class);
 		}
 
-		UpdateClassAssetClass(AssetData.AssetName, Class);
+		UpdateThumbnailClass(AssetData.AssetName, Class);
 
 		AssetColor = FLinearColor(1.f, 1.f, 1.f, 1.f);
 		if ( AssetTypeActions.IsValid() )
@@ -387,7 +387,7 @@ private:
 			// For non-class types, use the default based upon the actual asset class
 			// This has the side effect of not showing a class icon for assets that don't have a proper thumbnail image available
 			const FName DefaultThumbnail = (bIsClassType) ? NAME_None : FName(*FString::Printf(TEXT("ClassThumbnail.%s"), *AssetThumbnail->GetAssetData().AssetClass.ToString()));
-			return FClassIconFinder::FindThumbnailForClass(ClassAssetClass.Get(), DefaultThumbnail);
+			return FClassIconFinder::FindThumbnailForClass(ThumbnailClass.Get(), DefaultThumbnail);
 		}
 		else
 		{
@@ -403,7 +403,7 @@ private:
 		if(!bHasRenderedThumbnail)
 		{
 			const FSlateBrush* ClassThumbnailBrush = GetClassThumbnailBrush();
-			if( ClassThumbnailBrush && ClassAssetClass.Get() )
+			if( ClassThumbnailBrush && ThumbnailClass.Get() )
 			{
 				return EVisibility::Visible;
 			}
@@ -419,7 +419,7 @@ private:
 
 	const FSlateBrush* GetClassIconBrush() const
 	{
-		return FClassIconFinder::FindIconForClass(ClassAssetClass.Get());
+		return FClassIconFinder::FindIconForClass(ThumbnailClass.Get());
 	}
 
 	FMargin GetClassIconPadding() const
@@ -639,8 +639,8 @@ private:
 
 	/** The name of the thumbnail which should be used instead of the class thumbnail. */
 	FName ClassThumbnailBrushOverride;
-	/** The class instance for assets which represent a class. */
-	TWeakObjectPtr<UClass> ClassAssetClass;
+	/** The class to use when finding the thumbnail. */
+	TWeakObjectPtr<UClass> ThumbnailClass;
 	/** Are we showing a class type? (UClass, UBlueprint) */
 	bool bIsClassType;
 };
