[INTEGRATE] Change 2415616 by [REDACTED_PERSON][REDACTED_HANDLE] on 2015/01/22 14:24:26

	Fix crash when converting Actor with instanced root component with and instanced child of that root in to a Blueprint.

diff --git a/Engine/Source/Editor/UnrealEd/Private/Kismet2/Kismet2.cpp b/Engine/Source/Editor/UnrealEd/Private/Kismet2/Kismet2.cpp
--- a/Engine/Source/Editor/UnrealEd/Private/Kismet2/Kismet2.cpp
+++ b/Engine/Source/Editor/UnrealEd/Private/Kismet2/Kismet2.cpp
@@ -937,20 +937,23 @@ void FKismetEditorUtilities::AddComponentsToBlueprint(UBlueprint* Blueprint, con
 			USCS_Node* SCSNode = SCS->CreateNode(ActorComponent->GetClass());
 			UEditorEngine::CopyPropertiesForUnrelatedObjects(ActorComponent,SCSNode->ComponentTemplate);
 
+			USceneComponent* SceneComponent = Cast<USceneComponent>(ActorComponent);
+
 			// The easy part is non-scene component or the Root simply add it
-			if (!ActorComponent->IsA<USceneComponent>() || ActorComponent == Actor->GetRootComponent())
+			if (SceneComponent == nullptr)
 			{
 				SCS->AddNode(SCSNode);
 			}
 			else
 			{
-				USceneComponent* SceneComponent = CastChecked<USceneComponent>(ActorComponent);
-				check(SceneComponent->AttachParent);
-
 				InstanceComponentToNodeMap.Add(SceneComponent,SCSNode);
 
+				if (ActorComponent == Actor->GetRootComponent())
+				{
+					SCS->AddNode(SCSNode);
+				}
 				// If we're attached to a blueprint component look it up as the variable name is the component name
-				if (SceneComponent->AttachParent->bCreatedByConstructionScript)
+				else if (SceneComponent->AttachParent->bCreatedByConstructionScript)
 				{
 					USCS_Node* ParentSCSNode = nullptr;
 					for (UBlueprint* Blueprint : ParentBPStack)
