[INTEGRATE] Change 2415866 by [REDACTED_PERSON][REDACTED_HANDLE] on 2015/01/22 16:51:13

	[UE-7827] [UE-7995]

	SCS checks to make sure the entered name is not already taken by another component

diff --git a/Engine/Source/Editor/Kismet/Private/SSCSEditor.cpp b/Engine/Source/Editor/Kismet/Private/SSCSEditor.cpp
--- a/Engine/Source/Editor/Kismet/Private/SSCSEditor.cpp
+++ b/Engine/Source/Editor/Kismet/Private/SSCSEditor.cpp
@@ -2309,10 +2309,18 @@ UBlueprint* SSCS_RowWidget::GetBlueprint() const
 
 bool SSCS_RowWidget::OnNameTextVerifyChanged(const FText& InNewText, FText& OutErrorMessage)
 {
-	if(!InNewText.IsEmpty() && !FComponentEditorUtils::IsValidVariableNameString(NodePtr->GetComponentTemplate(), InNewText.ToString()))
+	if(!InNewText.IsEmpty())
 	{
-		OutErrorMessage = LOCTEXT("RenameFailed_NotValid", "This name is reserved for engine use.");
-		return false;
+		if (!FComponentEditorUtils::IsValidVariableNameString(NodePtr->GetComponentTemplate(), InNewText.ToString()))
+		{
+			OutErrorMessage = LOCTEXT("RenameFailed_EngineReservedName", "This name is reserved for engine use.");
+			return false;
+		}
+		else if (!FComponentEditorUtils::IsComponentNameAvailable(InNewText.ToString(), NodePtr->GetComponentTemplate()->GetOwner(), NodePtr->GetComponentTemplate()))
+		{
+			OutErrorMessage = LOCTEXT("RenameFailed_ExistingName", "Another component already has the same name.");
+			return false;
+		}
 	}
 
 	UBlueprint* Blueprint = GetBlueprint();
diff --git a/Engine/Source/Editor/UnrealEd/Private/Kismet2/ComponentEditorUtils.cpp b/Engine/Source/Editor/UnrealEd/Private/Kismet2/ComponentEditorUtils.cpp
--- a/Engine/Source/Editor/UnrealEd/Private/Kismet2/ComponentEditorUtils.cpp
+++ b/Engine/Source/Editor/UnrealEd/Private/Kismet2/ComponentEditorUtils.cpp
@@ -107,6 +107,27 @@ bool FComponentEditorUtils::IsValidVariableNameString(const UActorComponent* InC
 	return bIsValid;
 }
 
+bool FComponentEditorUtils::IsComponentNameAvailable(const FString& InString, const AActor* ComponentOwner, const UActorComponent* ComponentToIgnore)
+{
+	check(ComponentOwner);
+
+	bool bNameIsAvailable = true;
+
+	TInlineComponentArray<UActorComponent*> Components;
+	ComponentOwner->GetComponents(Components);
+
+	for (auto Component : Components)
+	{
+		if (Component != ComponentToIgnore && Component->GetName() == InString)
+		{
+			bNameIsAvailable = false;
+			break;
+		}
+	}
+
+	return bNameIsAvailable;
+}
+
 FString FComponentEditorUtils::GenerateValidVariableName(TSubclassOf<UActorComponent> ComponentClass, AActor* ComponentOwner)
 {
 	check(ComponentOwner);
@@ -115,26 +136,17 @@ FString FComponentEditorUtils::GenerateValidVariableName(TSubclassOf<UActorCompo
 	FString ComponentTypeName = *ComponentClass->GetName().Replace(TEXT("Component"), TEXT(""));
 	FString ComponentInstanceName;
 
-	// Make sure that none of the components currently on the actor have the same name
-	TInlineComponentArray<UActorComponent*> Components;
-	ComponentOwner->GetComponents(Components);
-
 	bool bNameIsValid = false;
 	while (!bNameIsValid)
 	{
 		bNameIsValid = true;
 
-		// Try the name with the lowest possible number
+		// Assign the lowest possible numerical suffix
 		ComponentInstanceName = FString::Printf(TEXT("%s%d"), *ComponentTypeName, Counter++);
 
-		// The name is valid if no other components in the owning actor have the same name
-		for (auto Component : Components)
+		if (!IsComponentNameAvailable(ComponentInstanceName, ComponentOwner))
 		{
-			if (Component->GetName() == ComponentInstanceName)
-			{
-				bNameIsValid = false;
-				break;
-			}
+			bNameIsValid = false;
 		}
 	}
 
diff --git a/Engine/Source/Editor/UnrealEd/Public/Kismet2/ComponentEditorUtils.h b/Engine/Source/Editor/UnrealEd/Public/Kismet2/ComponentEditorUtils.h
--- a/Engine/Source/Editor/UnrealEd/Public/Kismet2/ComponentEditorUtils.h
+++ b/Engine/Source/Editor/UnrealEd/Public/Kismet2/ComponentEditorUtils.h
@@ -16,6 +16,13 @@ public:
 	/** Test whether or not the given string is a valid variable name string for the given component instance */
 	static bool IsValidVariableNameString(const UActorComponent* InComponent, const FString& InString);
 
+	/** 
+	 * Test whether or not the given string is already the name string of a component on the the actor
+	 * Optionally excludes an existing component from the check (ex. a component currently being renamed)
+	 * @return True if the InString is an available name for a component of ComponentOwner
+	 */
+	static bool IsComponentNameAvailable(const FString& InString, const AActor* ComponentOwner, const UActorComponent* ComponentToIgnore = nullptr);
+		
 	/** Generate a valid variable name string for the given component instance */
 	static FString GenerateValidVariableName(TSubclassOf<UActorComponent> InComponentClass, AActor* ComponentOwner);
 
