[INTEGRATE] Change 2416128 by [REDACTED_PERSON][REDACTED_HANDLE] on 2015/01/22 19:04:33

	Stopped name check on component instances from causing a crash when a new component was added to an empty BP in the BP editor

diff --git a/Engine/Source/Editor/UnrealEd/Private/Kismet2/ComponentEditorUtils.cpp b/Engine/Source/Editor/UnrealEd/Private/Kismet2/ComponentEditorUtils.cpp
--- a/Engine/Source/Editor/UnrealEd/Private/Kismet2/ComponentEditorUtils.cpp
+++ b/Engine/Source/Editor/UnrealEd/Private/Kismet2/ComponentEditorUtils.cpp
@@ -109,19 +109,20 @@ bool FComponentEditorUtils::IsValidVariableNameString(const UActorComponent* InC
 
 bool FComponentEditorUtils::IsComponentNameAvailable(const FString& InString, const AActor* ComponentOwner, const UActorComponent* ComponentToIgnore)
 {
-	check(ComponentOwner);
-
 	bool bNameIsAvailable = true;
 
-	TInlineComponentArray<UActorComponent*> Components;
-	ComponentOwner->GetComponents(Components);
-
-	for (auto Component : Components)
+	if (ComponentOwner != nullptr)
 	{
-		if (Component != ComponentToIgnore && Component->GetName() == InString)
+		TInlineComponentArray<UActorComponent*> Components;
+		ComponentOwner->GetComponents(Components);
+
+		for (auto Component : Components)
 		{
-			bNameIsAvailable = false;
-			break;
+			if (Component != ComponentToIgnore && Component->GetName() == InString)
+			{
+				bNameIsAvailable = false;
+				break;
+			}
 		}
 	}
 
