[INTEGRATE] Change 2416141 by [REDACTED_PERSON][REDACTED_HANDLE].7 on 2015/01/22 19:15:32

	Processing of source content files now happens asynchronously

	The work is now distributed across multiple frames. Popup dialogs and slow tasks are now generally suppressed specifically while processing changes.

diff --git a/Engine/Source/Editor/UnrealEd/Private/AutoReimportManager.cpp b/Engine/Source/Editor/UnrealEd/Private/AutoReimportManager.cpp
--- a/Engine/Source/Editor/UnrealEd/Private/AutoReimportManager.cpp
+++ b/Engine/Source/Editor/UnrealEd/Private/AutoReimportManager.cpp
@@ -14,24 +14,21 @@
 #include "AssetRegistryModule.h"
 
 #define LOCTEXT_NAMESPACE "AutoReimportManager"
+#define CONDITIONAL_YIELD if (TimeLimit.Exceeded()) { return; }
 
 DEFINE_LOG_CATEGORY_STATIC(LogAutoReimportManager, Log, All);
 
 namespace
 {
-	/** Template helper function to count the number of elements in the specified array that pass a predicate */
-	template<typename T, typename P>
-	int32 CountIf(const TArray<T>& InArray, P Predicate)
+	/** Reduce the array to the specified accumulator */
+	template<typename T, typename P, typename A>
+	A Reduce(const TArray<T>& InArray, P Predicate, A Accumulator)
 	{
-		int32 Count = 0;
 		for (const T& X : InArray)
 		{
-			if (Predicate(X))
-			{
-				++Count;
-			}
+			Accumulator = Predicate(X, Accumulator);
 		}
-		return Count;
+		return Accumulator;
 	}
 
 	/** Retrieve a semi-colon separated string of file extensions supported by all available editor import factories */
@@ -118,7 +115,6 @@ namespace
 		FString CacheFilename = FPaths::ConvertRelativePathToFull(FPaths::GameIntermediateDir()) / TEXT("ReimportCache") / FString::Printf(TEXT("%u.bin"), CRC);
 
 		FFileCacheConfig Config(MoveTemp(Directory), MoveTemp(CacheFilename));
-		Config.BatchDelayS = 1;
 		Config.IncludeExtensions = InSupportedExtensions;
 		// We always store paths inside content folders relative to the folder
 		Config.PathType = EPathType::Relative;
@@ -158,6 +154,45 @@ namespace
 	}
 }
 
+FReimportFeedbackContext::FReimportFeedbackContext()
+{
+
+}
+
+FReimportFeedbackContext::~FReimportFeedbackContext()
+{
+	if (Notification.IsValid())
+	{
+		Notification->SetCompletionState(SNotificationItem::CS_Success);
+		Notification->ExpireAndFadeout();
+	}
+}
+
+void FReimportFeedbackContext::UpdateText(int32 Progress, int32 Total)
+{
+	FFormatOrderedArguments Args;
+	Args.Add(Progress);
+	Args.Add(Total);
+
+	const FText Message = FText::Format(LOCTEXT("ProcessingChanges", "Processing outstanding content changes ({0} of {1})"), Args);
+	if (Notification.IsValid())
+	{
+		Notification->SetText(Message);
+	}
+	else
+	{
+		FNotificationInfo Info(Message);
+		Info.ExpireDuration = 2.0f;
+		Info.bUseLargeFont = false;
+		Info.bFireAndForget = false;
+		Info.bUseSuccessFailIcons = false;
+		Info.Image = FEditorStyle::GetBrush("NoBrush");
+
+		Notification = FSlateNotificationManager::Get().AddNotification(Info);
+		Notification->SetCompletionState(SNotificationItem::CS_Pending);
+	}
+}
+
 FContentDirectoryMonitor::FContentDirectoryMonitor(const FString& InDirectory, const FString& InSupportedExtensions, const FString& InMountedContentPath)
 	: Cache(GenerateFileCacheConfig(InDirectory, InSupportedExtensions, InMountedContentPath))
 	, MountedContentPath(InMountedContentPath)
@@ -171,159 +206,180 @@ void FContentDirectoryMonitor::Destroy()
 	Cache.Destroy();
 }
 
-void FContentDirectoryMonitor::Tick(const IAssetRegistry& Registry, const FWorkLimiter& Limiter)
+void FContentDirectoryMonitor::Tick(const FTimeLimit& TimeLimit)
 {
-	Cache.Tick(Limiter);
+	Cache.Tick(TimeLimit);
 
-	if (Limiter.ShouldLimit())
+	const double Now = FPlatformTime::Seconds();
+	if (Now - LastSaveTime > ResaveIntervalS)
 	{
-		return;
+		LastSaveTime = Now;
+		Cache.WriteCache();
 	}
-	// We can't do this if the registry is loading assets as we won't necesssarily be able to react to changes effectively
-	else if (!Registry.IsLoadingAssets())
+}
+
+void FContentDirectoryMonitor::StartProcessing()
+{
+	WorkProgress = TotalWork = 0;
+
+	auto OutstandingChanges = Cache.GetOutstandingChanges();
+	if (OutstandingChanges.Num() != 0)
 	{
-		OutstandingChanges = Cache.GetOutstandingChanges();
-		if (OutstandingChanges.Num() != 0)
+		for (auto& Transaction : OutstandingChanges)
 		{
-			ProcessOutstandingChanges(Registry, Limiter);
+			switch(Transaction.Action)
+			{
+				case FFileChangeData::FCA_Added:		AddedFiles.Emplace(MoveTemp(Transaction));		break;
+				case FFileChangeData::FCA_Modified:		ModifiedFiles.Emplace(MoveTemp(Transaction));	break;
+				case FFileChangeData::FCA_Removed:		DeletedFiles.Emplace(MoveTemp(Transaction));	break;
+			}
 		}
+
+		TotalWork = AddedFiles.Num() + ModifiedFiles.Num() + DeletedFiles.Num();
 	}
-	
-	if (Limiter.ShouldLimit())
+}
+
+void FContentDirectoryMonitor::ProcessAdditions(TArray<UPackage*>& OutPackagesToSave, const FTimeLimit& TimeLimit, const TMap<FString, TArray<UFactory*>>& InFactoriesByExtension)
+{
+	if (MountedContentPath.IsEmpty())
 	{
+		// @todo: arodham: Prompt user for import location
 		return;
 	}
-	// Save out the cache file if we've got time to do so, and have not saved for a while
-	else
+	
+	for (int32 Index = 0; Index < AddedFiles.Num(); ++Index)
 	{
-		const double Now = FPlatformTime::Seconds();
-		if (Now - LastSaveTime > ResaveIntervalS)
+		auto& Addition = AddedFiles[Index];
+
+		WorkProgress += 1;
+		const FString FullFilename = Cache.GetDirectory() + Addition.Filename.Get();
+
+		bool bImportSucceeded = false;
+		bool bImportWasCancelled = false;
+
+		FString NewAssetName = ObjectTools::SanitizeObjectName(FPaths::GetBaseFilename(FullFilename));
+		FString PackagePath= PackageTools::SanitizePackageName(MountedContentPath / FPaths::GetPath(Addition.Filename.Get()) / NewAssetName);
+
+		// We can not create assets that share the name of a map file in the same location
+		if (FEditorFileUtils::IsMapPackageAsset(PackagePath))
 		{
-			LastSaveTime = Now;
-			Cache.WriteCache();
+			// @todo: Error reporting
+			continue;
 		}
-	}
-}
 
-void FContentDirectoryMonitor::ProcessOutstandingChanges(const IAssetRegistry& Registry, const FWorkLimiter& Limiter)
-{
-	// Count the number of additions so we can apprximate the work load
-	const int32 NumAdditions = CountIf(OutstandingChanges, [](const FUpdateCacheTransaction& Change){ return Change.Action == FFileChangeData::FCA_Added; });
+		if (FindPackage(nullptr, *PackagePath))
+		{
+			// @todo: Error reporting
+			continue;	
+		}
 
-	auto* ReimportManager = FReimportManager::Instance();
+		UPackage* NewPackage = CreatePackage(nullptr, *PackagePath);
+		if ( !ensure(NewPackage) )
+		{
+			// @todo: Error reporting
+			continue;
+		}
 
-	TArray<UObject*> AssetsToDelete;
-	{
-		bool bCreatedDialog = false;
-		FScopedSlowTask SlowTask(OutstandingChanges.Num() + NumAdditions, LOCTEXT("ProcessingChanges", "Processing source content changes"));
-		SlowTask.MakeDialog();
+		// Make sure the destination package is loaded
+		NewPackage->FullyLoad();
 
-		TArray<UPackage*> PackagesToSave;
+		bool bSuccess = false;
+		bool bCancelled = false;
 
-		for (auto& Change : OutstandingChanges)
-		{
-			const FString FullFilename = Cache.GetDirectory() + Change.Filename.Get();
+		UObject* NewAsset = nullptr;
 
-			if (!bCreatedDialog && FPlatformTime::Seconds() - SlowTask.StartTime > 2)
-			{
-				//SlowTask.MakeDialog();
-				bCreatedDialog = true;
-			}
-			
-			SlowTask.EnterProgressFrame();
-			switch(Change.Action)
+		const FString Ext = FPaths::GetExtension(Addition.Filename.Get(), false);
+		if (auto* Factories = InFactoriesByExtension.Find(Ext))
+		{
+			for (auto* Factory : *Factories)
 			{
-				case FFileChangeData::FCA_Modified:
-					// @todo: arodham: enable/disable specific extensions/asset types?
-					{
-						for (const auto& AssetData : FindAssetsPertainingToFile(Registry, FullFilename))
-						{
-							UObject* Asset = AssetData.GetAsset();
-							ReimportManager->Reimport(Asset);
-						}
-					}
+				NewAsset = UFactory::StaticImportObject(Factory->ResolveSupportedClass(), NewPackage, FName(*NewAssetName), RF_Public | RF_Standalone, bCancelled, *FullFilename, nullptr, Factory);
+				if (NewAsset || bCancelled)
+				{
 					break;
+				}
+			}
+		}
+		
+		if (!NewAsset || bCancelled)
+		{
+			// @todo: Error reporting
+			continue;
+		}
 
-				case FFileChangeData::FCA_Added:
-					{
-						FAssetToolsModule& AssetToolsModule = FModuleManager::Get().LoadModuleChecked<FAssetToolsModule>("AssetTools");
+		FAssetRegistryModule::AssetCreated(NewAsset);
+		GEditor->BroadcastObjectReimported(NewAsset);
 
-						TArray<FString> File;
-						File.Add(FullFilename);
+		OutPackagesToSave.Add(NewPackage);
 
-						FString ParentPackage;
-						if (MountedContentPath.IsEmpty())
-						{
-							// @todo: arodham: this is sketchy and needs revisiting. PromptUserForDirectory?
-							auto UserChosenPath = ChooseAssetImportLocation(FullFilename);
-							if (UserChosenPath.IsSet())
-							{
-								ParentPackage = UserChosenPath.GetValue();
-							}
-						}
-						else
-						{
-							ParentPackage = MountedContentPath / PackageTools::SanitizePackageName(FPaths::GetPath(Change.Filename.Get()));
-						}
+		// Refresh the supported class.  Some factories (e.g. FBX) only resolve their type after reading the file
+		// ImportAssetType = Factory->ResolveSupportedClass();
+		// @todo: analytics
 
-						if (!ParentPackage.IsEmpty())
-						{
-							auto NewAssets = AssetToolsModule.Get().ImportAssets(File, ParentPackage);
+		// Let the cache know that we've dealt with this change (it will be imported immediately)
+		Cache.CompleteTransaction(MoveTemp(Addition));
 
-							// @todo: arodham: Error reporting
+		if (TimeLimit.Exceeded())
+		{
+			AddedFiles.RemoveAt(0, Index + 1);
+			return;
+		}
+	}
 
-							for (auto* Asset : NewAssets)
-							{
-								UPackage* Package = Asset->GetOutermost();
-								if (Package)
-								{
-									PackagesToSave.Add(Package);
-								}
-							}
-						}
-					}
-					break;
+	AddedFiles.Empty();
+}
 
-				case FFileChangeData::FCA_Removed:
-					for (const auto& AssetData : FindAssetsPertainingToFile(Registry, FullFilename))
-					{
-						AssetsToDelete.Add(AssetData.GetAsset());
-					}
-					break;
+void FContentDirectoryMonitor::ProcessModifications(const IAssetRegistry& Registry, const FTimeLimit& TimeLimit, FReimportFeedbackContext& Context)
+{
+	auto* ReimportManager = FReimportManager::Instance();
 
-				default:
-					break;
-			}
+	for (int32 Index = 0; Index < ModifiedFiles.Num(); ++Index)
+	{
+		auto& Change = ModifiedFiles[Index];
 
-			// Let the cache know that we've dealt with this change
-			Cache.CompleteTransaction(MoveTemp(Change));
+		const FString FullFilename = Cache.GetDirectory() + Change.Filename.Get();
+		for (const auto& AssetData : FindAssetsPertainingToFile(Registry, FullFilename))
+		{
+			UObject* Asset = AssetData.GetAsset();
+			ReimportManager->Reimport(Asset, false /* Ask for new file */, false /* Show notification */);
 		}
 
-		SlowTask.EnterProgressFrame(NumAdditions);
-		if (PackagesToSave.Num() > 0)
+		// Let the cache know that we've dealt with this change
+		Cache.CompleteTransaction(MoveTemp(Change));
+		
+		WorkProgress += 1;
+		if (TimeLimit.Exceeded())
 		{
-			const bool bAlreadyCheckedOut = false;
-			const bool bCheckDirty = false;
-			const bool bPromptToSave = false;
-			// @todo: arodham: Error reporting?
-			FEditorFileUtils::PromptForCheckoutAndSave(PackagesToSave, bCheckDirty, bPromptToSave, nullptr, bAlreadyCheckedOut);
+			ModifiedFiles.RemoveAt(0, Index + 1);
+			return;
 		}
 	}
 
-	// Open a dialog to ask about deleted files
-	// @todo: this needs improvement
+	ModifiedFiles.Empty();
+}
 
-	if (AssetsToDelete.Num() > 0)
+void FContentDirectoryMonitor::ExtractAssetsToDelete(const IAssetRegistry& Registry, TArray<FAssetData>& OutAssetsToDelete)
+{
+	for (auto& Deletion : DeletedFiles)
 	{
-		ObjectTools::DeleteObjects(AssetsToDelete);
+		for (const auto& AssetData : FindAssetsPertainingToFile(Registry, Cache.GetDirectory() + Deletion.Filename.Get()))
+		{
+			OutAssetsToDelete.Add(AssetData);
+		}
+
+		// Let the cache know that we've dealt with this change (it will be imported in due course)
+		Cache.CompleteTransaction(MoveTemp(Deletion));
 	}
 
-	OutstandingChanges.Empty();
+	WorkProgress += DeletedFiles.Num();
+	DeletedFiles.Empty();
 }
 
 UAutoReimportManager::UAutoReimportManager(const FObjectInitializer& ObjectInitializer)
 	: Super(ObjectInitializer)
+	, CurrentState(ECurrentState::Idle)
 	, bIsProcessingChanges(false)
+	, StartProcessingDelay(1)
 {
 
 }
@@ -362,30 +418,236 @@ void UAutoReimportManager::BeginDestroy()
 	Super::BeginDestroy();
 }
 
-void UAutoReimportManager::Tick(float DeltaTime)
+int32 UAutoReimportManager::GetNumUnprocessedChanges() const
 {
-	// Only spend half a frame on this, at most
-	const FWorkLimiter Limiter(1.0 / 120);
+	auto Pred = [](const FContentDirectoryMonitor& Monitor, int32 Total){ return Total + Monitor.GetNumUnprocessedChanges(); };
+	return Reduce(MountedDirectoryMonitors, Pred, 0) + Reduce(UserDirectoryMonitors, Pred, 0);
+}
 
-	const IAssetRegistry& Registry = FModuleManager::LoadModuleChecked<FAssetRegistryModule>("AssetRegistry").Get();
+int32 UAutoReimportManager::GetTotalWork() const
+{
+	auto Pred = [](const FContentDirectoryMonitor& Monitor, int32 Total){ return Total + Monitor.GetTotalWork(); };
+	return Reduce(MountedDirectoryMonitors, Pred, 0) + Reduce(UserDirectoryMonitors, Pred, 0);
+}
+
+int32 UAutoReimportManager::GetTotalProgress() const
+{
+	auto Pred = [](const FContentDirectoryMonitor& Monitor, int32 Total){ return Total + Monitor.GetWorkProgress(); };
+	return Reduce(MountedDirectoryMonitors, Pred, 0) + Reduce(UserDirectoryMonitors, Pred, 0);
+}
+
+void UAutoReimportManager::ProcessAdditions(const FTimeLimit& TimeLimit)
+{
+	TMap<FString, TArray<UFactory*>> Factories;
+
+	TArray<FString> FactoryExtensions;
+	FactoryExtensions.Reserve(16);
+
+	// Get the list of valid factories
+	for (TObjectIterator<UClass> It ; It ; ++It)
+	{
+		UClass* CurrentClass = (*It);
+
+		if (CurrentClass->IsChildOf(UFactory::StaticClass()) && !(CurrentClass->HasAnyClassFlags(CLASS_Abstract)))
+		{
+			UFactory* Factory = Cast<UFactory>(CurrentClass->GetDefaultObject());
+			if (Factory->bEditorImport && Factory->ImportPriority >= 0)
+			{
+				FactoryExtensions.Reset();
+				Factory->GetSupportedFileExtensions(FactoryExtensions);
+
+				for (const auto& Ext : FactoryExtensions)
+				{
+					auto& Array = Factories.FindOrAdd(Ext);
+					Array.Add(Factory);
+				}
+			}
+		}
+	}
 
-	TGuardValue<bool> Guard(bIsProcessingChanges, true);
+	for (auto& Pair : Factories)
+	{
+		Pair.Value.Sort([](const UFactory& A, const UFactory& B) { return A.ImportPriority > B.ImportPriority; });
+	}
 
 	for (auto& Monitor : MountedDirectoryMonitors)
 	{
-		Monitor.Tick(Registry, Limiter);
-		if (Limiter.ShouldLimit())
+		Monitor.ProcessAdditions(PackagesToSave, TimeLimit, Factories);
+	}
+
+	if (!TimeLimit.Exceeded())
+	{
+		for (auto& Monitor : UserDirectoryMonitors)
 		{
-			return;
+			Monitor.ProcessAdditions(PackagesToSave, TimeLimit, Factories);
 		}
 	}
 
+	CONDITIONAL_YIELD
+
+	CurrentState = ECurrentState::SavePackages;
+}
+
+void UAutoReimportManager::ProcessModifications(const IAssetRegistry& Registry, const FTimeLimit& TimeLimit)
+{
+	for (auto& Monitor : MountedDirectoryMonitors)
+	{
+		Monitor.ProcessModifications(Registry, TimeLimit, *FeedbackContextOverride);
+		CONDITIONAL_YIELD
+	}
+
 	for (auto& Monitor : UserDirectoryMonitors)
 	{
-		Monitor.Tick(Registry, Limiter);
-		if (Limiter.ShouldLimit())
+		Monitor.ProcessModifications(Registry, TimeLimit, *FeedbackContextOverride);
+		CONDITIONAL_YIELD
+	}
+
+	CurrentState = ECurrentState::ProcessDeletions;
+}
+
+void UAutoReimportManager::ProcessDeletions(const IAssetRegistry& Registry)
+{
+	TArray<FAssetData> AssetsToDelete;
+
+	for (auto& Monitor : MountedDirectoryMonitors)
+	{
+		Monitor.ExtractAssetsToDelete(Registry, AssetsToDelete);
+	}
+	for (auto& Monitor : UserDirectoryMonitors)
+	{
+		Monitor.ExtractAssetsToDelete(Registry, AssetsToDelete);
+	}
+
+	if (AssetsToDelete.Num() > 0)
+	{
+		ObjectTools::DeleteAssets(AssetsToDelete);
+	}
+
+	CurrentState = ECurrentState::Cleanup;
+}
+
+
+void UAutoReimportManager::Tick(float DeltaTime)
+{
+	// Never spend more than a 60fps frame doing this work (meaning we shouldn't drop below 30fps), we can do more if we're throttling CPU usage (ie editor running in background)
+	const FTimeLimit TimeLimit(GEditor->ShouldThrottleCPUUsage() ? 1 / 6.f : 1.f / 60.f);
+
+	if (CurrentState == ECurrentState::Idle)
+	{
+		for (auto& Monitor : MountedDirectoryMonitors)
 		{
-			return;
+			Monitor.Tick(TimeLimit);
+			CONDITIONAL_YIELD
+		}
+		for (auto& Monitor : UserDirectoryMonitors)
+		{
+			Monitor.Tick(TimeLimit);
+			CONDITIONAL_YIELD
+		}
+
+		const int32 NumUnprocessedChanges = GetNumUnprocessedChanges();
+		if (NumUnprocessedChanges > 0)
+		{
+			// Check if we have any more unprocessed changes. If we have, we wait a delay before processing them.
+			if (NumUnprocessedChanges != CachedNumUnprocessedChanges)
+			{
+				CachedNumUnprocessedChanges = NumUnprocessedChanges;
+				StartProcessingDelay.Reset();
+			}
+			else
+			{
+				if (StartProcessingDelay.Exceeded())
+				{
+					// We have some stuff to do, let's kick that off
+					CurrentState = ECurrentState::Initialize;
+					FeedbackContextOverride.Reset(new FReimportFeedbackContext);
+				}
+			}
+		}
+	}
+	else if (FSlateThrottleManager::Get().IsAllowingExpensiveTasks())
+	{
+		TGuardValue<bool> GuardProcessingChanges(bIsProcessingChanges, true);
+		
+		FeedbackContextOverride->UpdateText(GetTotalProgress(), GetTotalWork());
+
+		if (CurrentState == ECurrentState::Initialize)
+		{
+			for (auto& Monitor : MountedDirectoryMonitors)
+			{
+				Monitor.StartProcessing();
+			}
+			for (auto& Monitor : UserDirectoryMonitors)
+			{
+				Monitor.StartProcessing();
+			}
+
+			CurrentState = ECurrentState::ProcessAdditions;
+
+			FeedbackContextOverride->UpdateText(GetTotalProgress(), GetTotalWork());
+			CONDITIONAL_YIELD
+		}
+
+		if (CurrentState == ECurrentState::ProcessAdditions)
+		{
+			// Override the global feedback context while we do this to avoid popping up dialogs
+			TGuardValue<FFeedbackContext*> ScopedContextOverride(GWarn, FeedbackContextOverride.Get());
+
+			ProcessAdditions(TimeLimit);
+
+			FeedbackContextOverride->UpdateText(GetTotalProgress(), GetTotalWork());
+			CONDITIONAL_YIELD
+		}
+
+		if (CurrentState == ECurrentState::SavePackages)
+		{
+			// Override the global feedback context while we do this to avoid popping up dialogs
+			TGuardValue<FFeedbackContext*> ScopedContextOverride(GWarn, FeedbackContextOverride.Get());
+
+			if (PackagesToSave.Num() > 0)
+			{
+				const bool bAlreadyCheckedOut = false;
+				const bool bCheckDirty = false;
+				const bool bPromptToSave = false;
+				// @todo: arodham: Error reporting?
+				FEditorFileUtils::PromptForCheckoutAndSave(PackagesToSave, bCheckDirty, bPromptToSave, nullptr, bAlreadyCheckedOut);
+
+				PackagesToSave.Empty();
+			}
+
+			CurrentState = ECurrentState::ProcessModifications;
+			CONDITIONAL_YIELD
+		}
+
+		const IAssetRegistry& Registry = FModuleManager::LoadModuleChecked<FAssetRegistryModule>("AssetRegistry").Get();
+		if (CurrentState == ECurrentState::ProcessModifications)
+		{
+			// Override the global feedback context while we do this to avoid popping up dialogs
+			TGuardValue<FFeedbackContext*> ScopedContextOverride(GWarn, FeedbackContextOverride.Get());
+
+			ProcessModifications(Registry, TimeLimit);
+
+			FeedbackContextOverride->UpdateText(GetTotalProgress(), GetTotalWork());
+			CONDITIONAL_YIELD
+		}
+
+		if (CurrentState == ECurrentState::ProcessDeletions)
+		{
+			ProcessDeletions(Registry);
+
+			FeedbackContextOverride->UpdateText(GetTotalProgress(), GetTotalWork());
+			CONDITIONAL_YIELD
+		}
+
+		if (CurrentState == ECurrentState::Cleanup)
+		{
+			// Finished
+			FeedbackContextOverride->UpdateText(GetTotalWork(), GetTotalWork());
+
+			CachedNumUnprocessedChanges = 0;
+			FeedbackContextOverride.Reset();
+
+			CurrentState = ECurrentState::Idle;
 		}
 	}
 }
@@ -412,7 +674,7 @@ void UAutoReimportManager::HandleLoadingSavingSettingChanged(FName PropertyName)
 		}
 		else
 		{
-			// Destroy all the existing montirs, including their file caches
+			// Destroy all the existing monitors, including their file caches
 			for (auto& Monitor : MountedDirectoryMonitors)
 			{
 				Monitor.Destroy();
@@ -429,8 +691,10 @@ void UAutoReimportManager::OnContentPathMounted(const FString& InAssetPath, cons
 
 void UAutoReimportManager::OnContentPathDismounted(const FString& InAssetPath, const FString& FileSystemPath)
 {
+	const auto FullDismountedPath = FPaths::ConvertRelativePathToFull(FileSystemPath);
+
 	MountedDirectoryMonitors.RemoveAll([&](const FContentDirectoryMonitor& Monitor){
-		return Monitor.GetMountedContentPath() == InAssetPath;
+		return Monitor.GetDirectory() == FullDismountedPath;
 	});
 }
 
@@ -543,5 +807,6 @@ void UAutoReimportManager::SetUpUserDirectories()
 	}
 }
 
+#undef CONDITIONAL_YIELD
 #undef LOCTEXT_NAMESPACE
 
diff --git a/Engine/Source/Editor/UnrealEd/Private/Editor.cpp b/Engine/Source/Editor/UnrealEd/Private/Editor.cpp
--- a/Engine/Source/Editor/UnrealEd/Private/Editor.cpp
+++ b/Engine/Source/Editor/UnrealEd/Private/Editor.cpp
@@ -2587,7 +2587,7 @@ bool FReimportManager::CanReimport( UObject* Obj ) const
 	return false;
 }
 
-bool FReimportManager::Reimport( UObject* Obj, bool bAskForNewFileIfMissing )
+bool FReimportManager::Reimport( UObject* Obj, bool bAskForNewFileIfMissing, bool bShowNotification )
 {
 	// Warn that were about to reimport, so prep for it
 	PreReimport.Broadcast( Obj );
@@ -2602,7 +2602,6 @@ bool FReimportManager::Reimport( UObject* Obj, bool bAskForNewFileIfMissing )
 			bHandlersNeedSorting = false;
 		}
 		
-		bool bShowNotification = true;
 		bool bValidSourceFilename = false;
 		TArray<FString> SourceFilenames;
 
diff --git a/Engine/Source/Editor/UnrealEd/Private/FileCache.cpp b/Engine/Source/Editor/UnrealEd/Private/FileCache.cpp
--- a/Engine/Source/Editor/UnrealEd/Private/FileCache.cpp
+++ b/Engine/Source/Editor/UnrealEd/Private/FileCache.cpp
@@ -25,7 +25,6 @@ namespace
 			{
 				if (Predicate(Array[Read], Array[DuplRead]))
 				{
-					++Read;
 					goto next;
 				}
 			}
@@ -65,7 +64,7 @@ bool FAsyncDirectoryReader::IsComplete() const
 	return PendingFiles.Num() == 0 && PendingDirectories.Num() == 0;
 }
 
-FAsyncDirectoryReader::EProgressResult FAsyncDirectoryReader::Tick(const FWorkLimiter& Limiter)
+FAsyncDirectoryReader::EProgressResult FAsyncDirectoryReader::Tick(const FTimeLimit& TimeLimit)
 {
 	auto& FileManager = IFileManager::Get();
 	const int32 RootPathLen = RootPath.Len();
@@ -73,7 +72,7 @@ FAsyncDirectoryReader::EProgressResult FAsyncDirectoryReader::Tick(const FWorkLi
 	for (int32 Index = 0; Index < PendingDirectories.Num(); ++Index)
 	{
 		ScanDirectory(PendingDirectories[Index]);
-		if (Limiter.ShouldLimit())
+		if (TimeLimit.Exceeded())
 		{
 			// We've spent too long, bail
 			PendingDirectories.RemoveAt(0, Index, false);
@@ -90,7 +89,7 @@ FAsyncDirectoryReader::EProgressResult FAsyncDirectoryReader::Tick(const FWorkLi
 		FString Filename = (PathType == EPathType::Relative ? *File + RootPathLen : *File);
 		State->Files.Emplace(MoveTemp(Filename), FileManager.GetTimeStamp(*File));
 
-		if (Index % 100 == 0 && Limiter.ShouldLimit())
+		if (Index % 100 == 0 && TimeLimit.Exceeded())
 		{
 			// We've spent too long, bail
 			PendingFiles.RemoveAt(0, Index, false);
@@ -133,7 +132,6 @@ FFileCache::FFileCache(const FFileCacheConfig& InConfig)
 	: Config(InConfig)
 	, DirectoryReader(InConfig.Directory, Config.PathType)
 	, bSavedCacheDirty(false)
-	, LastChangeTimeS(0)
 {
 	// Ensure that the extension strings are of the form ;ext1;ext2;ext3;
 	if (!Config.IncludeExtensions.IsEmpty() && Config.IncludeExtensions[0] != ';')
@@ -236,16 +234,9 @@ void FFileCache::WriteCache()
 
 TArray<FUpdateCacheTransaction> FFileCache::GetOutstandingChanges()
 {
-	if (FPlatformTime::Seconds() - LastChangeTimeS >= Config.BatchDelayS)
-	{
-		TArray<FUpdateCacheTransaction> Moved;
-		Swap(OutstandingChanges, Moved);
-		return Moved;
-	}
-	else
-	{
-		return TArray<FUpdateCacheTransaction>();
-	}
+	TArray<FUpdateCacheTransaction> Moved;
+	Swap(OutstandingChanges, Moved);
+	return Moved;
 }
 
 void FFileCache::CompleteTransaction(FUpdateCacheTransaction&& Transaction)
@@ -288,9 +279,9 @@ void FFileCache::CompleteTransaction(FUpdateCacheTransaction&& Transaction)
 	}
 }
 
-void FFileCache::Tick(const FWorkLimiter& Limiter)
+void FFileCache::Tick(const FTimeLimit& Limit)
 {
-	if (DirectoryReader.IsComplete() || DirectoryReader.Tick(Limiter) == FAsyncDirectoryReader::EProgressResult::Pending)
+	if (DirectoryReader.IsComplete() || DirectoryReader.Tick(Limit) == FAsyncDirectoryReader::EProgressResult::Pending)
 	{
 		return;
 	}
@@ -426,8 +417,6 @@ void FFileCache::OnDirectoryChanged(const TArray<FFileChangeData>& FileChanges)
 	RemoveDuplicates(OutstandingChanges, [](const FUpdateCacheTransaction& A, const FUpdateCacheTransaction& B){
 		return A.Action == B.Action && A.Filename == B.Filename;
 	});
-
-	LastChangeTimeS = FPlatformTime::Seconds();
 }
 
 bool FFileCache::IsFileApplicable(const FString& Filename)
diff --git a/Engine/Source/Editor/UnrealEd/Public/AutoReimport/AutoReimportManager.h b/Engine/Source/Editor/UnrealEd/Public/AutoReimport/AutoReimportManager.h
--- a/Engine/Source/Editor/UnrealEd/Public/AutoReimport/AutoReimportManager.h
+++ b/Engine/Source/Editor/UnrealEd/Public/AutoReimport/AutoReimportManager.h
@@ -7,6 +7,26 @@
 #include "AutoReimportManager.generated.h"
 
 class IAssetRegistry;
+class IAssetTools;
+
+/** Feedback context that overrides GWarn for import operations to prevent popup spam */
+struct FReimportFeedbackContext : public FFeedbackContext
+{
+	/** Constructor - created when we are processing changes */
+	FReimportFeedbackContext();
+	/** Destructor - called when we are done processing changes */
+	~FReimportFeedbackContext();
+
+	/** Update the text on the notification */
+	void UpdateText(int32 Progress, int32 Total);
+
+private:
+	virtual void Serialize( const TCHAR* V, ELogVerbosity::Type Verbosity, const class FName& Category ) {}
+
+private:
+	/** The notification that is shown when the context is active */
+	TSharedPtr<SNotificationItem> Notification;
+};
 
 /** Class responsible for watching a specific content directory for changes */
 class FContentDirectoryMonitor
@@ -21,8 +41,20 @@ public:
 	 */
 	FContentDirectoryMonitor(const FString& InDirectory, const FString& InSupportedExtensions, const FString& InMountedContentPath = FString());
 
-	/** Tick this directory monitor. Will potentially process changed files */
-	void Tick(const IAssetRegistry& Registry, const FWorkLimiter& Limiter);
+	/** Tick this monitor's cache to give it a chance to finish scanning for files */
+	void Tick(const FTimeLimit& TimeLimit);
+
+	/** Start processing any outstanding changes this monitor is aware of */
+	void StartProcessing();
+
+	/** Extract the files we need to import from our outstanding changes (happens first)*/ 
+	void ProcessAdditions(TArray<UPackage*>& OutPackagesToSave, const FTimeLimit& TimeLimit, const TMap<FString, TArray<UFactory*>>& InFactoriesByExtension);
+
+	/** Process the outstanding changes that we have cached */
+	void ProcessModifications(const IAssetRegistry& Registry, const FTimeLimit& TimeLimit, FReimportFeedbackContext& Context);
+
+	/** Extract the assets we need to delete from our outstanding changes (happens last) */ 
+	void ExtractAssetsToDelete(const IAssetRegistry& Registry, TArray<FAssetData>& OutAssetsToDelete);
 
 	/** Destroy this monitor including its cache */
 	void Destroy();
@@ -30,24 +62,30 @@ public:
 	/** Get the directory that this monitor applies to */
 	const FString& GetDirectory() const { return Cache.GetDirectory(); }
 
-	/** Get the mounted content path that this monitor applies to */
-	const FString& GetMountedContentPath() const { return MountedContentPath; }
+public:
 
-private:
+	/** Get the number of outstanding changes that we potentially have to process (when not already processing) */
+	int32 GetNumUnprocessedChanges() const { return Cache.GetNumOutstandingChanges(); }
 
-	/** Process the outstanding changes that we have cached */
-	void ProcessOutstandingChanges(const IAssetRegistry& Registry, const FWorkLimiter& Limiter);
+	/** Get the total amount of work this monitor has to perform in the current processing operation */
+	int32 GetTotalWork() const { return TotalWork; }
+
+	/** Get the total amount of work this monitor has performed in the current processing operation */
+	int32 GetWorkProgress() const { return WorkProgress; }
 
 private:
 
 	/** The file cache that monitors and reflects the content directory. */
 	FFileCache Cache;
 
-	/** The mounted content path for this monitor (eg /Engine/) */
+	/** The mounted content path for this monitor (eg /Game/) */
 	FString MountedContentPath;
 
 	/** A list of file system changes that are due to be processed */
-	TArray<FUpdateCacheTransaction> OutstandingChanges;
+	TArray<FUpdateCacheTransaction> AddedFiles, ModifiedFiles, DeletedFiles;
+
+	/** The number of changes we've processed out of the OutstandingChanges array */
+	int32 TotalWork, WorkProgress;
 
 	/** The last time we attempted to save the cache file */
 	double LastSaveTime;
@@ -56,7 +94,6 @@ private:
 	static const int32 ResaveIntervalS = 60;
 };
 
-
 /* Deals with auto reimporting of objects when the objects file on disk is modified*/
 UCLASS(config=Editor, transient)
 class UNREALED_API UAutoReimportManager : public UObject, public FTickableEditorObject
@@ -104,6 +141,38 @@ private:
 
 private:
 
+	/** Get the total amount of work in the current operation */
+	int32 GetTotalWork() const;
+
+	/** Get the total progress through the current operation */
+	int32 GetTotalProgress() const;
+	
+	/** Get the number of unprocessed changes that are not part of the current processing operation */
+	int32 GetNumUnprocessedChanges() const;
+
+private:
+
+	/** Process any remaining pending additions we have */
+	void ProcessAdditions(const FTimeLimit& TimeLimit);
+
+	/** Process any remaining pending modifications we have */
+	void ProcessModifications(const IAssetRegistry& Registry, const FTimeLimit& TimeLimit);
+
+	/** Process any remaining pending deletions we have */
+	void ProcessDeletions(const IAssetRegistry& Registry);
+
+	/** Enum and value to specify the current state of our processing */
+	enum class ECurrentState
+	{
+		Idle, Initialize, ProcessAdditions, SavePackages, ProcessModifications, ProcessDeletions, Cleanup
+	};
+	ECurrentState CurrentState;
+
+private:
+
+	/** Feedback context that can selectively override the global context to consume progress events for saving of assets */
+	TUniquePtr<FReimportFeedbackContext> FeedbackContextOverride;
+
 	/** Cached string of extensions that are supported by all available factories */
 	FString SupportedExtensions;
 
@@ -113,6 +182,15 @@ private:
 	/** Additional external directories to monitor for changes specified by the user */
 	TArray<FContentDirectoryMonitor> UserDirectoryMonitors;
 
+	/** A list of packages to save when we've added a bunch of assets */
+	TArray<UPackage*> PackagesToSave;
+
 	/** Set when we are processing changes to avoid reentrant logic for asset deletions etc */
 	bool bIsProcessingChanges;
+
+	/** The time when the last change to the cache was reported */
+	FTimeLimit StartProcessingDelay;
+
+	/** The cached number of unprocessed changes we currently have to process */
+	int32 CachedNumUnprocessedChanges;
 };
diff --git a/Engine/Source/Editor/UnrealEd/Public/EditorReimportHandler.h b/Engine/Source/Editor/UnrealEd/Public/EditorReimportHandler.h
--- a/Engine/Source/Editor/UnrealEd/Public/EditorReimportHandler.h
+++ b/Engine/Source/Editor/UnrealEd/Public/EditorReimportHandler.h
@@ -45,10 +45,11 @@ public:
 	 *
 	 * @param	Obj	Object to try reimporting
 	 * @param	bAskForNewFileIfMissing If the file is missing, open a dialog to ask for a new one
+	 * @param	bShowNotification True to show a notification when complete, false otherwise
 	 *
 	 * @return	true if the object was handled by one of the reimport handlers; false otherwise
 	 */
-	UNREALED_API virtual bool Reimport( UObject* Obj, bool bAskForNewFileIfMissing = false );
+	UNREALED_API virtual bool Reimport( UObject* Obj, bool bAskForNewFileIfMissing = false, bool bShowNotification = true );
 
 	/**
 	 * Convert a file path to be relative to the specified object, if it resides in the same package folder
diff --git a/Engine/Source/Editor/UnrealEd/Public/FileCache.h b/Engine/Source/Editor/UnrealEd/Public/FileCache.h
--- a/Engine/Source/Editor/UnrealEd/Public/FileCache.h
+++ b/Engine/Source/Editor/UnrealEd/Public/FileCache.h
@@ -146,28 +146,30 @@ private:
 	FUpdateCacheTransaction(FImmutableString InFilename, FFileChangeData::EFileChangeAction InAction, const FDateTime& InTimestamp = 0)
 		: Filename(MoveTemp(InFilename)), Timestamp(InTimestamp), Action(InAction)
 	{}
-
-	/** Copyable only by FFileCache*/
-	FUpdateCacheTransaction(const FUpdateCacheTransaction& In) = default;
-	FUpdateCacheTransaction& operator=(const FUpdateCacheTransaction& In) = default;
 };
 
-/** Small utility for specifying when a call to do some work should bail out, if at all. */
-struct FWorkLimiter
+/** A time limit that counts down from the time of construction, until it hits a given delay */
+struct FTimeLimit
 {
 	/** Constructor specifying that we should never bail out early */
-	FWorkLimiter() : StopAtTime(-1) {}
+	FTimeLimit() : Delay(-1) { Reset(); }
 
 	/** Constructor specifying not to run over the specified number of seconds */
-	FWorkLimiter(float NumSeconds) : StopAtTime(FPlatformTime::Seconds() + NumSeconds) {}
+	FTimeLimit(float NumSeconds) : Delay(NumSeconds) { Reset(); }
+
+	/** Check whether we have exceeded the time limit */
+	bool Exceeded() const { return Delay != -1 && FPlatformTime::Seconds() >= StartTime + Delay; }
 
-	/** Check whether we have exceeded the stop time */
-	bool ShouldLimit() const { return StopAtTime != -1 && FPlatformTime::Seconds() >= StopAtTime; }
+	/** Reset the time limit to start timing again from the current time */
+	void Reset() { StartTime = FPlatformTime::Seconds(); }
 
 private:
 
-	/** User specified time to stop processing */
-	double StopAtTime;
+	/** The delay specified by the user */
+	float Delay;
+
+	/** The time we started */
+	double StartTime;
 };
 
 /** Enum specifying whether a path should be relative or absolute */
@@ -188,7 +190,7 @@ enum EPathType
  *		while(!Reader.IsComplete())
  *		{
  *			FPlatformProcess::Sleep(1);
- *			Reader.Tick(FWorkLimiter(1));	// Do 1 second of work
+ *			Reader.Tick(FTimedSignal(1));	// Do 1 second of work
  *		}
  *		TOptional<FDirectoryState> State = Reader.GetFinalState();
  */
@@ -212,7 +214,7 @@ struct FAsyncDirectoryReader
 	bool IsComplete() const;
 
 	/** Tick this reader (discover new directories / files). Returns progress state. */
-	EProgressResult Tick(const FWorkLimiter& Limiter);
+	EProgressResult Tick(const FTimeLimit& Limit);
 
 private:
 	/** Non-recursively scan a single directory for its contents. Adds results to Pending arrays. */
@@ -238,7 +240,7 @@ private:
 struct FFileCacheConfig
 {
 	FFileCacheConfig(FString InDirectory, FString InCacheFile)
-		: Directory(InDirectory), CacheFile(InCacheFile), BatchDelayS(0.f), PathType(EPathType::Relative)
+		: Directory(InDirectory), CacheFile(InCacheFile), PathType(EPathType::Relative)
 	{}
 
 	/** String specifying the directory on disk that the cache should reflect */
@@ -253,9 +255,6 @@ struct FFileCacheConfig
 	 */
 	FString IncludeExtensions, ExcludeExtensions;
 
-	/** Optionally specified delay that should be waited before reporting changes to the watched directory. Useful for batching together large or rapid changes. */
-	float BatchDelayS;
-
 	/** Path type to return, relative to the directory or absolute. */
 	EPathType PathType;
 };
@@ -285,7 +284,7 @@ public:
 	const FString& GetDirectory() const { return Config.Directory; }
 
 	/** Tick this FileCache, optionally specifying a limiter to limit the amount of time spent in here */
-	void Tick(const FWorkLimiter& Limiter);
+	void Tick(const FTimeLimit& Limit);
 
 	/** Write out the cached file, if we have any changes to write */
 	void WriteCache();
@@ -296,6 +295,9 @@ public:
 	 */
 	void CompleteTransaction(FUpdateCacheTransaction&& Transaction);
 
+	/** Get the number of pending changes to the cache. */
+	int32 GetNumOutstandingChanges() const { return OutstandingChanges.Num(); }
+
 	/** Get all pending changes to the cache. Transactions must be returned to CompleteTransaction to update the cache. */
 	TArray<FUpdateCacheTransaction> GetOutstandingChanges();
 
@@ -338,7 +340,4 @@ private:
 
 	/** True when the cached state we have in memory is more up to date than the serialized file. Enables WriteCache() when true. */
 	bool bSavedCacheDirty;
-
-	/** The time that a change to the file-system was last logged. */
-	double LastChangeTimeS;
 };
