[INTEGRATE] Change 2416292 by [REDACTED_PERSON][REDACTED_HANDLE] on 2015/01/22 21:10:14

	Enable world component editing by default

diff --git a/Engine/Source/Editor/UnrealEd/Private/Settings/SettingsClasses.cpp b/Engine/Source/Editor/UnrealEd/Private/Settings/SettingsClasses.cpp
--- a/Engine/Source/Editor/UnrealEd/Private/Settings/SettingsClasses.cpp
+++ b/Engine/Source/Editor/UnrealEd/Private/Settings/SettingsClasses.cpp
@@ -53,7 +53,9 @@ UDestructableMeshEditorSettings::UDestructableMeshEditorSettings( const FObjectI
 
 UEditorExperimentalSettings::UEditorExperimentalSettings( const FObjectInitializer& ObjectInitializer )
 	: Super(ObjectInitializer)
+	, bInWorldBPEditing(true)
 	, bUnifiedBlueprintEditor(true)
+
 {
 }
 
