[INTEGRATE] Change 2416295 by [REDACTED_PERSON][REDACTED_HANDLE] on 2015/01/22 21:11:05

	Update empty actor icon.

diff --git a/Engine/Source/Editor/UnrealEd/Private/Factories/ActorFactory.cpp b/Engine/Source/Editor/UnrealEd/Private/Factories/ActorFactory.cpp
--- a/Engine/Source/Editor/UnrealEd/Private/Factories/ActorFactory.cpp
+++ b/Engine/Source/Editor/UnrealEd/Private/Factories/ActorFactory.cpp
@@ -1117,7 +1117,7 @@ static UBillboardComponent* CreateEditorOnlyBillboardComponent(AActor* ActorOwne
 	// Create a new billboard component to serve as a visualization of the actor until there is another primitive component
 	UBillboardComponent* BillboardComponent = ConstructObject<UBillboardComponent>(UBillboardComponent::StaticClass(), ActorOwner, NAME_None, RF_Transactional);
 
-	BillboardComponent->Sprite = LoadObject<UTexture2D>(nullptr, TEXT("/Engine/EditorResources/S_Actor.S_Actor"));
+	BillboardComponent->Sprite = LoadObject<UTexture2D>(nullptr, TEXT("/Engine/EditorResources/EmptyActor.EmptyActor"));
 	BillboardComponent->RelativeScale3D = FVector(0.5f, 0.5f, 0.5f);
 	BillboardComponent->Mobility = EComponentMobility::Movable;
 	BillboardComponent->AlwaysLoadOnClient = false;
@@ -1148,24 +1148,23 @@ AActor* UActorFactoryEmptyActor::SpawnActor( UObject* Asset, ULevel* InLevel, co
 	AActor* NewActor = nullptr;
 	if(GetDefault<UEditorExperimentalSettings>()->bInWorldBPEditing) 
 	{
-	// Spawn a temporary actor for dragging around
-	AActor* NewActor = Super::SpawnActor(Asset, InLevel, Location, Rotation, ObjectFlags, Name);
+		// Spawn a temporary actor for dragging around
+		NewActor = Super::SpawnActor(Asset, InLevel, Location, Rotation, ObjectFlags, Name);
 
-	USceneComponent* RootComponent = ConstructObject<USceneComponent>(USceneComponent::StaticClass(), NewActor, FName("Root"), RF_Transactional);
-	RootComponent->Mobility = EComponentMobility::Movable;
-	RootComponent->SetWorldLocationAndRotation(Location, Rotation);
-	NewActor->SetRootComponent(RootComponent);
+		USceneComponent* RootComponent = ConstructObject<USceneComponent>(USceneComponent::StaticClass(), NewActor, FName("Root"), RF_Transactional);
+		RootComponent->Mobility = EComponentMobility::Movable;
+		RootComponent->SetWorldLocationAndRotation(Location, Rotation);
+		NewActor->SetRootComponent(RootComponent);
 
 		UBillboardComponent* BillboardComponent	= CreateEditorOnlyBillboardComponent(NewActor, RootComponent);
 
-	NewActor->AddInstanceComponent(RootComponent);
-	NewActor->AddInstanceComponent(BillboardComponent);
+		NewActor->AddInstanceComponent(RootComponent);
+		NewActor->AddInstanceComponent(BillboardComponent);
+	}
 
 	return NewActor;
 }
 
-	return NewActor;
-}
 
 
 /*-----------------------------------------------------------------------------
