[INTEGRATE] Change 2416619 by [REDACTED_PERSON][REDACTED_HANDLE].7 on 2015/01/23 04:37:50

	Compile fix for non-unity

diff --git a/Engine/Source/Editor/UnrealEd/Private/AutoReimportManager.cpp b/Engine/Source/Editor/UnrealEd/Private/AutoReimportManager.cpp
--- a/Engine/Source/Editor/UnrealEd/Private/AutoReimportManager.cpp
+++ b/Engine/Source/Editor/UnrealEd/Private/AutoReimportManager.cpp
@@ -12,6 +12,8 @@
 #include "ARFilter.h"
 #include "ContentBrowserModule.h"
 #include "AssetRegistryModule.h"
+#include "NotificationManager.h"
+#include "SNotificationList.h"
 
 #define LOCTEXT_NAMESPACE "AutoReimportManager"
 #define CONDITIONAL_YIELD if (TimeLimit.Exceeded()) { return; }
diff --git a/Engine/Source/Editor/UnrealEd/Public/AutoReimport/AutoReimportManager.h b/Engine/Source/Editor/UnrealEd/Public/AutoReimport/AutoReimportManager.h
--- a/Engine/Source/Editor/UnrealEd/Public/AutoReimport/AutoReimportManager.h
+++ b/Engine/Source/Editor/UnrealEd/Public/AutoReimport/AutoReimportManager.h
@@ -8,6 +8,7 @@
 
 class IAssetRegistry;
 class IAssetTools;
+class SNotificationItem;
 
 /** Feedback context that overrides GWarn for import operations to prevent popup spam */
 struct FReimportFeedbackContext : public FFeedbackContext
