[INTEGRATE] Change 2416696 by [REDACTED_PERSON][REDACTED_HANDLE].7 on 2015/01/23 07:13:11

	Tidied up auto reimport manager code by splitting into multiple files

	UAutoReimportManager now uses a private implementation.
	Also refined the load balancing a little.

diff --git a/Engine/Source/Editor/UnrealEd/Private/AutoReimport/AutoReimportManager.cpp b/Engine/Source/Editor/UnrealEd/Private/AutoReimport/AutoReimportManager.cpp
new file mode 100644
--- /dev/null
+++ b/Engine/Source/Editor/UnrealEd/Private/AutoReimport/AutoReimportManager.cpp
@@ -0,0 +1,670 @@
+// Copyright 1998-2015 Epic Games, Inc. All Rights Reserved.
+
+#include "UnrealEd.h"
+
+#include "AutoReimportUtilities.h"
+#include "AutoReimport/AutoReimportManager.h"
+#include "ContentDirectoryMonitor.h"
+
+#include "ObjectTools.h"
+#include "ARFilter.h"
+#include "ContentBrowserModule.h"
+#include "AssetRegistryModule.h"
+#include "NotificationManager.h"
+#include "SNotificationList.h"
+
+#define LOCTEXT_NAMESPACE "AutoReimportManager"
+#define CONDITIONAL_YIELD if (TimeLimit.Exceeded()) { return; }
+
+DEFINE_LOG_CATEGORY_STATIC(LogAutoReimportManager, Log, All);
+
+/** Feedback context that overrides GWarn for import operations to prevent popup spam */
+struct FReimportFeedbackContext : public FFeedbackContext
+{
+	/** Constructor - created when we are processing changes */
+	FReimportFeedbackContext()
+	{
+	}
+
+	/** Destructor - called when we are done processing changes */
+	~FReimportFeedbackContext()
+	{
+		if (Notification.IsValid())
+		{
+			Notification->SetCompletionState(SNotificationItem::CS_Success);
+			Notification->ExpireAndFadeout();
+		}
+	}
+
+	/** Update the text on the notification */
+	void UpdateText(int32 Progress, int32 Total)
+	{
+		FFormatOrderedArguments Args;
+		Args.Add(Progress);
+		Args.Add(Total);
+
+		const FText Message = FText::Format(LOCTEXT("ProcessingChanges", "Processing outstanding content changes ({0} of {1})"), Args);
+		if (Notification.IsValid())
+		{
+			Notification->SetText(Message);
+		}
+		else
+		{
+			FNotificationInfo Info(Message);
+			Info.ExpireDuration = 2.0f;
+			Info.bUseLargeFont = false;
+			Info.bFireAndForget = false;
+			Info.bUseSuccessFailIcons = false;
+			Info.Image = FEditorStyle::GetBrush("NoBrush");
+
+			Notification = FSlateNotificationManager::Get().AddNotification(Info);
+			Notification->SetCompletionState(SNotificationItem::CS_Pending);
+		}
+	}
+
+private:
+	virtual void Serialize( const TCHAR* V, ELogVerbosity::Type Verbosity, const class FName& Category ) {}
+
+private:
+	/** The notification that is shown when the context is active */
+	TSharedPtr<SNotificationItem> Notification;
+
+	/** The notification content */
+	TSharedPtr<class SAutoReimportFeedback> NotificationContent;
+};
+
+/* Deals with auto reimporting of objects when the objects file on disk is modified*/
+class FAutoReimportManager : public FTickableEditorObject
+{
+public:
+
+	FAutoReimportManager();
+	~FAutoReimportManager();
+
+private:
+
+	/** FTickableEditorObject interface*/
+	virtual void Tick(float DeltaTime) override;
+	virtual bool IsTickable() const override { return true; }
+	virtual TStatId GetStatId() const override;
+
+private:
+
+	/** Called when a new asset path has been mounted */
+	void OnContentPathMounted(const FString& InAssetPath, const FString& FileSystemPath);
+
+	/** Called when a new asset path has been dismounted */
+	void OnContentPathDismounted(const FString& InAssetPath, const FString& FileSystemPath);
+
+	/** Called when an asset has been deleted */
+	void OnAssetDeleted(UObject* Object);
+
+private:
+
+	/** Callback for when an editor user setting has changed */
+	void HandleLoadingSavingSettingChanged(FName PropertyName);
+
+	/** Set up monitors to all the mounted content directories */
+	void SetUpMountedContentDirectories();
+
+	/** Set up the user directories from the editor Loading/Saving settings */
+	void SetUpUserDirectories();
+
+	/** Retrieve a semi-colon separated string of file extensions supported by all available editor import factories */
+	static FString GetAllFactoryExtensions();
+
+private:
+
+	/** Get the total amount of work in the current operation */
+	int32 GetTotalWork() const;
+
+	/** Get the total progress through the current operation */
+	int32 GetTotalProgress() const;
+	
+	/** Get the number of unprocessed changes that are not part of the current processing operation */
+	int32 GetNumUnprocessedChanges() const;
+
+private:
+
+	/** Process any remaining pending additions we have */
+	void ProcessAdditions(const FTimeLimit& TimeLimit);
+
+	/** Process any remaining pending modifications we have */
+	void ProcessModifications(const IAssetRegistry& Registry, const FTimeLimit& TimeLimit);
+
+	/** Process any remaining pending deletions we have */
+	void ProcessDeletions(const IAssetRegistry& Registry);
+
+	/** Enum and value to specify the current state of our processing */
+	enum class ECurrentState
+	{
+		Idle, Initialize, ProcessAdditions, SavePackages, ProcessModifications, ProcessDeletions, Cleanup
+	};
+	ECurrentState CurrentState;
+
+private:
+
+	/** Feedback context that can selectively override the global context to consume progress events for saving of assets */
+	TUniquePtr<FReimportFeedbackContext> FeedbackContextOverride;
+
+	/** Cached string of extensions that are supported by all available factories */
+	FString SupportedExtensions;
+
+	/** Array of directory monitors for automatically managed mounted content paths */
+	TArray<FContentDirectoryMonitor> MountedDirectoryMonitors;
+
+	/** Additional external directories to monitor for changes specified by the user */
+	TArray<FContentDirectoryMonitor> UserDirectoryMonitors;
+
+	/** A list of packages to save when we've added a bunch of assets */
+	TArray<UPackage*> PackagesToSave;
+
+	/** Set when we are processing changes to avoid reentrant logic for asset deletions etc */
+	bool bIsProcessingChanges;
+
+	/** The time when the last change to the cache was reported */
+	FTimeLimit StartProcessingDelay;
+
+	/** The cached number of unprocessed changes we currently have to process */
+	int32 CachedNumUnprocessedChanges;
+};
+
+FAutoReimportManager::FAutoReimportManager()
+	: CurrentState(ECurrentState::Idle)
+	, bIsProcessingChanges(false)
+	, StartProcessingDelay(1)
+{
+	// @todo: arodham: update this when modules are reloaded or new factory types are available?
+	SupportedExtensions = GetAllFactoryExtensions();
+
+	auto* Settings = GetMutableDefault<UEditorLoadingSavingSettings>();
+	Settings->OnSettingChanged().AddRaw(this, &FAutoReimportManager::HandleLoadingSavingSettingChanged);
+
+	FPackageName::OnContentPathMounted().AddRaw(this, &FAutoReimportManager::OnContentPathMounted);
+	FPackageName::OnContentPathDismounted().AddRaw(this, &FAutoReimportManager::OnContentPathDismounted);
+
+	IAssetRegistry& AssetRegistry = FModuleManager::LoadModuleChecked<FAssetRegistryModule>("AssetRegistry").Get();
+	AssetRegistry.OnInMemoryAssetDeleted().AddRaw(this, &FAutoReimportManager::OnAssetDeleted);
+
+
+	// Only set this up for content directories if the user has this enabled
+	if (Settings->bMonitorContentDirectories)
+	{
+		SetUpMountedContentDirectories();
+	}
+
+	// We always monitor user directories if they are set
+	SetUpUserDirectories();
+}
+
+FAutoReimportManager::~FAutoReimportManager() 
+{
+	if (auto* Settings = GetMutableDefault<UEditorLoadingSavingSettings>())
+	{
+		Settings->OnSettingChanged().RemoveAll(this);
+	}
+
+	FPackageName::OnContentPathMounted().RemoveAll(this);
+	FPackageName::OnContentPathDismounted().RemoveAll(this);
+
+	if (FAssetRegistryModule* AssetRegistryModule = FModuleManager::GetModulePtr<FAssetRegistryModule>("AssetRegistry"))
+	{
+		AssetRegistryModule->Get().OnInMemoryAssetDeleted().RemoveAll(this);
+	}
+
+	// Force a save of all the caches
+	MountedDirectoryMonitors.Empty();
+	UserDirectoryMonitors.Empty();
+}
+
+int32 FAutoReimportManager::GetNumUnprocessedChanges() const
+{
+	auto Pred = [](const FContentDirectoryMonitor& Monitor, int32 Total){ return Total + Monitor.GetNumUnprocessedChanges(); };
+	return Utils::Reduce(MountedDirectoryMonitors, Pred, 0) + Utils::Reduce(UserDirectoryMonitors, Pred, 0);
+}
+
+int32 FAutoReimportManager::GetTotalWork() const
+{
+	auto Pred = [](const FContentDirectoryMonitor& Monitor, int32 Total){ return Total + Monitor.GetTotalWork(); };
+	return Utils::Reduce(MountedDirectoryMonitors, Pred, 0) + Utils::Reduce(UserDirectoryMonitors, Pred, 0);
+}
+
+int32 FAutoReimportManager::GetTotalProgress() const
+{
+	auto Pred = [](const FContentDirectoryMonitor& Monitor, int32 Total){ return Total + Monitor.GetWorkProgress(); };
+	return Utils::Reduce(MountedDirectoryMonitors, Pred, 0) + Utils::Reduce(UserDirectoryMonitors, Pred, 0);
+}
+
+void FAutoReimportManager::ProcessAdditions(const FTimeLimit& TimeLimit)
+{
+	TMap<FString, TArray<UFactory*>> Factories;
+
+	TArray<FString> FactoryExtensions;
+	FactoryExtensions.Reserve(16);
+
+	// Get the list of valid factories
+	for (TObjectIterator<UClass> It ; It ; ++It)
+	{
+		UClass* CurrentClass = (*It);
+
+		if (CurrentClass->IsChildOf(UFactory::StaticClass()) && !(CurrentClass->HasAnyClassFlags(CLASS_Abstract)))
+		{
+			UFactory* Factory = Cast<UFactory>(CurrentClass->GetDefaultObject());
+			if (Factory->bEditorImport && Factory->ImportPriority >= 0)
+			{
+				FactoryExtensions.Reset();
+				Factory->GetSupportedFileExtensions(FactoryExtensions);
+
+				for (const auto& Ext : FactoryExtensions)
+				{
+					auto& Array = Factories.FindOrAdd(Ext);
+					Array.Add(Factory);
+				}
+			}
+		}
+	}
+
+	for (auto& Pair : Factories)
+	{
+		Pair.Value.Sort([](const UFactory& A, const UFactory& B) { return A.ImportPriority > B.ImportPriority; });
+	}
+
+	for (auto& Monitor : MountedDirectoryMonitors)
+	{
+		Monitor.ProcessAdditions(PackagesToSave, TimeLimit, Factories);
+	}
+
+	if (!TimeLimit.Exceeded())
+	{
+		for (auto& Monitor : UserDirectoryMonitors)
+		{
+			Monitor.ProcessAdditions(PackagesToSave, TimeLimit, Factories);
+		}
+	}
+
+	CONDITIONAL_YIELD
+
+	CurrentState = ECurrentState::SavePackages;
+}
+
+void FAutoReimportManager::ProcessModifications(const IAssetRegistry& Registry, const FTimeLimit& TimeLimit)
+{
+	for (auto& Monitor : MountedDirectoryMonitors)
+	{
+		Monitor.ProcessModifications(Registry, TimeLimit);
+		CONDITIONAL_YIELD
+	}
+
+	for (auto& Monitor : UserDirectoryMonitors)
+	{
+		Monitor.ProcessModifications(Registry, TimeLimit);
+		CONDITIONAL_YIELD
+	}
+
+	CurrentState = ECurrentState::ProcessDeletions;
+}
+
+void FAutoReimportManager::ProcessDeletions(const IAssetRegistry& Registry)
+{
+	TArray<FAssetData> AssetsToDelete;
+
+	for (auto& Monitor : MountedDirectoryMonitors)
+	{
+		Monitor.ExtractAssetsToDelete(Registry, AssetsToDelete);
+	}
+	for (auto& Monitor : UserDirectoryMonitors)
+	{
+		Monitor.ExtractAssetsToDelete(Registry, AssetsToDelete);
+	}
+
+	if (AssetsToDelete.Num() > 0)
+	{
+		ObjectTools::DeleteAssets(AssetsToDelete);
+	}
+
+	CurrentState = ECurrentState::Cleanup;
+}
+
+
+void FAutoReimportManager::Tick(float DeltaTime)
+{
+	// Never spend more than a 60fps frame doing this work (meaning we shouldn't drop below 30fps), we can do more if we're throttling CPU usage (ie editor running in background)
+	const FTimeLimit TimeLimit(GEditor->ShouldThrottleCPUUsage() ? 1 / 6.f : 1.f / 60.f);
+
+	if (CurrentState == ECurrentState::Idle)
+	{
+		for (auto& Monitor : MountedDirectoryMonitors)
+		{
+			Monitor.Tick(TimeLimit);
+			CONDITIONAL_YIELD
+		}
+		for (auto& Monitor : UserDirectoryMonitors)
+		{
+			Monitor.Tick(TimeLimit);
+			CONDITIONAL_YIELD
+		}
+
+		const int32 NumUnprocessedChanges = GetNumUnprocessedChanges();
+		if (NumUnprocessedChanges > 0)
+		{
+			// Check if we have any more unprocessed changes. If we have, we wait a delay before processing them.
+			if (NumUnprocessedChanges != CachedNumUnprocessedChanges)
+			{
+				CachedNumUnprocessedChanges = NumUnprocessedChanges;
+				StartProcessingDelay.Reset();
+			}
+			else
+			{
+				if (StartProcessingDelay.Exceeded())
+				{
+					// We have some stuff to do, let's kick that off
+					CurrentState = ECurrentState::Initialize;
+					FeedbackContextOverride.Reset(new FReimportFeedbackContext);
+				}
+			}
+		}
+	}
+	else if (FSlateThrottleManager::Get().IsAllowingExpensiveTasks())
+	{
+		TGuardValue<bool> GuardProcessingChanges(bIsProcessingChanges, true);
+		
+		FeedbackContextOverride->UpdateText(GetTotalProgress(), GetTotalWork());
+
+		if (CurrentState == ECurrentState::Initialize)
+		{
+			for (auto& Monitor : MountedDirectoryMonitors)
+			{
+				Monitor.StartProcessing();
+			}
+			for (auto& Monitor : UserDirectoryMonitors)
+			{
+				Monitor.StartProcessing();
+			}
+
+			CurrentState = ECurrentState::ProcessAdditions;
+
+			FeedbackContextOverride->UpdateText(GetTotalProgress(), GetTotalWork());
+			CONDITIONAL_YIELD
+		}
+
+		if (CurrentState == ECurrentState::ProcessAdditions)
+		{
+			// Override the global feedback context while we do this to avoid popping up dialogs
+			TGuardValue<FFeedbackContext*> ScopedContextOverride(GWarn, FeedbackContextOverride.Get());
+
+			ProcessAdditions(TimeLimit);
+
+			FeedbackContextOverride->UpdateText(GetTotalProgress(), GetTotalWork());
+			CONDITIONAL_YIELD
+		}
+
+		if (CurrentState == ECurrentState::SavePackages)
+		{
+			// Override the global feedback context while we do this to avoid popping up dialogs
+			TGuardValue<FFeedbackContext*> ScopedContextOverride(GWarn, FeedbackContextOverride.Get());
+
+			if (PackagesToSave.Num() > 0)
+			{
+				const bool bAlreadyCheckedOut = false;
+				const bool bCheckDirty = false;
+				const bool bPromptToSave = false;
+				// @todo: arodham: Error reporting?
+				FEditorFileUtils::PromptForCheckoutAndSave(PackagesToSave, bCheckDirty, bPromptToSave, nullptr, bAlreadyCheckedOut);
+
+				PackagesToSave.Empty();
+			}
+
+			CurrentState = ECurrentState::ProcessModifications;
+			CONDITIONAL_YIELD
+		}
+
+		const IAssetRegistry& Registry = FModuleManager::LoadModuleChecked<FAssetRegistryModule>("AssetRegistry").Get();
+		if (CurrentState == ECurrentState::ProcessModifications)
+		{
+			// Override the global feedback context while we do this to avoid popping up dialogs
+			TGuardValue<FFeedbackContext*> ScopedContextOverride(GWarn, FeedbackContextOverride.Get());
+
+			ProcessModifications(Registry, TimeLimit);
+
+			FeedbackContextOverride->UpdateText(GetTotalProgress(), GetTotalWork());
+			CONDITIONAL_YIELD
+		}
+
+		if (CurrentState == ECurrentState::ProcessDeletions)
+		{
+			ProcessDeletions(Registry);
+
+			FeedbackContextOverride->UpdateText(GetTotalProgress(), GetTotalWork());
+			CONDITIONAL_YIELD
+		}
+
+		if (CurrentState == ECurrentState::Cleanup)
+		{
+			// Finished
+			FeedbackContextOverride->UpdateText(GetTotalWork(), GetTotalWork());
+
+			CachedNumUnprocessedChanges = 0;
+			FeedbackContextOverride.Reset();
+
+			CurrentState = ECurrentState::Idle;
+		}
+	}
+}
+
+TStatId FAutoReimportManager::GetStatId() const
+{
+	RETURN_QUICK_DECLARE_CYCLE_STAT(FAutoReimportManager, STATGROUP_Tickables);
+}
+
+void FAutoReimportManager::HandleLoadingSavingSettingChanged(FName PropertyName)
+{
+	if (PropertyName == GET_MEMBER_NAME_CHECKED(UEditorLoadingSavingSettings, AutoReimportDirectories))
+	{
+		SetUpUserDirectories();
+	}
+	else if (PropertyName == GET_MEMBER_NAME_CHECKED(UEditorLoadingSavingSettings, bMonitorContentDirectories))
+	{
+		const auto* Settings = GetDefault<UEditorLoadingSavingSettings>();
+		if (Settings->bMonitorContentDirectories)
+		{
+			// Should already be empty, but just to be sure
+			MountedDirectoryMonitors.Empty();
+			SetUpMountedContentDirectories();
+		}
+		else
+		{
+			// Destroy all the existing monitors, including their file caches
+			for (auto& Monitor : MountedDirectoryMonitors)
+			{
+				Monitor.Destroy();
+			}
+			MountedDirectoryMonitors.Empty();
+		}
+	}
+}
+
+void FAutoReimportManager::OnContentPathMounted(const FString& InAssetPath, const FString& FileSystemPath)
+{
+	MountedDirectoryMonitors.Emplace(FPaths::ConvertRelativePathToFull(FileSystemPath), SupportedExtensions, InAssetPath);
+}
+
+void FAutoReimportManager::OnContentPathDismounted(const FString& InAssetPath, const FString& FileSystemPath)
+{
+	const auto FullDismountedPath = FPaths::ConvertRelativePathToFull(FileSystemPath);
+
+	MountedDirectoryMonitors.RemoveAll([&](const FContentDirectoryMonitor& Monitor){
+		return Monitor.GetDirectory() == FullDismountedPath;
+	});
+}
+
+void FAutoReimportManager::OnAssetDeleted(UObject* Object)
+{
+	if (bIsProcessingChanges)
+	{
+		return;
+	}
+
+	const IAssetRegistry& Registry = FModuleManager::LoadModuleChecked<FAssetRegistryModule>("AssetRegistry").Get();
+
+	bool bCanDelete = true;
+
+	TArray<UObject::FAssetRegistryTag> Tags;
+	Object->GetAssetRegistryTags(Tags);
+
+	TArray<FString> AbsoluteSourceFileNames;
+
+	const FName TagName = UObject::SourceFileTagName();
+
+	for (const auto& Tag : Tags)
+	{
+		// Don't compare the FName number suffix
+		if (Tag.Name.IsEqual(TagName, ENameCase::IgnoreCase, false))
+		{
+			const FString SourceFile = FReimportManager::ResolveImportFilename(Tag.Value, Object);
+			// Check there aren't any other assets referencing this thing
+			if (Utils::FindAssetsPertainingToFile(Registry, SourceFile).Num() != 0)
+			{
+				bCanDelete = false;
+				break;
+			}
+			else
+			{
+				AbsoluteSourceFileNames.Add(MoveTemp(SourceFile));
+			}
+		}
+	}
+
+	if (bCanDelete)
+	{
+		// Offer to delete the source file(s) as well
+		const FText MessageText = FText::Format(LOCTEXT("DeleteRelatedAssets", "The asset '{0}' has been deleted. Would you like to delete the source files it was imported from as well?"), FText::FromString(Object->GetName()) );
+		if (FMessageDialog::Open(EAppMsgType::YesNo, MessageText) == EAppReturnType::Yes)
+		{
+			auto& FileManager = IFileManager::Get();
+
+			for (auto& Filename : AbsoluteSourceFileNames)
+			{
+				FileManager.Delete(*Filename, false /* RequireExists */, true /* Even if read only */, true /* Quiet */);
+			}
+		}
+	}
+}
+
+void FAutoReimportManager::SetUpMountedContentDirectories()
+{
+	TArray<FString> ContentPaths;
+	FPackageName::QueryRootContentPaths(ContentPaths);
+	for (const auto& ContentPath : ContentPaths)
+	{
+		const FString Directory = FPaths::ConvertRelativePathToFull(FPackageName::LongPackageNameToFilename(ContentPath));
+
+		// We ignore the engine content directory for now
+		if (Directory != FPaths::ConvertRelativePathToFull(FPaths::EngineContentDir()))
+		{
+			MountedDirectoryMonitors.Emplace(Directory, SupportedExtensions, ContentPath);
+		}
+	}
+}
+
+void FAutoReimportManager::SetUpUserDirectories()
+{
+	TArray<bool> PersistedDirectories;
+	PersistedDirectories.AddZeroed(UserDirectoryMonitors.Num());
+
+	// Also set up user-specified directories
+	TArray<FString> NewDirectories = GetDefault<UEditorLoadingSavingSettings>()->AutoReimportDirectories;
+
+	for (int32 Index = 0; Index < NewDirectories.Num(); ++Index)
+	{
+		NewDirectories[Index] = FPaths::ConvertRelativePathToFull(NewDirectories[Index]);
+		const int32 Existing = UserDirectoryMonitors.IndexOfByPredicate([&](const FContentDirectoryMonitor& InMonitor){ return InMonitor.GetDirectory() == NewDirectories[Index]; });
+		if (Existing != INDEX_NONE)
+		{
+			PersistedDirectories[Existing] = true;
+		}
+		else
+		{
+			// Add a new one
+			UserDirectoryMonitors.Emplace(NewDirectories[Index], SupportedExtensions);
+			PersistedDirectories.Add(true);
+		}
+	}
+
+	// Delete anything that is no longer needed (we delete the cache file altogether)
+	for (int32 Index = PersistedDirectories.Num() - 1; Index >= 0; --Index)
+	{
+		if (!PersistedDirectories[Index])
+		{
+			UserDirectoryMonitors[Index].Destroy();
+			UserDirectoryMonitors.RemoveAt(Index, 1, false);
+		}
+	}
+
+	if (UserDirectoryMonitors.GetSlack() > UserDirectoryMonitors.Max() / 2)
+	{
+		UserDirectoryMonitors.Shrink();
+	}
+}
+
+FString FAutoReimportManager::GetAllFactoryExtensions()
+{
+	FString AllExtensions;
+
+	// Use a scratch buffer to avoid unnecessary re-allocation
+	FString Scratch;
+	Scratch.Reserve(32);
+
+	for (TObjectIterator<UClass> ClassIt; ClassIt; ++ClassIt)
+	{
+		UClass* Class = *ClassIt;
+		if (Class->IsChildOf(UFactory::StaticClass()) && !Class->HasAnyClassFlags(CLASS_Abstract))
+		{
+			UFactory* Factory = Cast<UFactory>(Class->GetDefaultObject());
+			if (Factory->bEditorImport)
+			{
+				for (const FString& Format : Factory->Formats)
+				{
+					int32 Index = INDEX_NONE;
+					if (Format.FindChar(';', Index) && Index > 0)
+					{
+						Scratch.GetCharArray().Reset();
+						// Include the ;
+						Scratch.AppendChars(*Format, Index + 1);
+
+						if (AllExtensions.Find(Scratch) == INDEX_NONE)
+						{
+							AllExtensions += Scratch;
+						}
+					}
+				}
+			}
+		}
+	}
+
+	return AllExtensions;
+}
+
+UAutoReimportManager::UAutoReimportManager(const FObjectInitializer& Init)
+	: Super(Init)
+{
+}
+
+UAutoReimportManager::~UAutoReimportManager()
+{
+}
+
+void UAutoReimportManager::Initialize()
+{
+	Implementation.Reset(new FAutoReimportManager);
+}
+
+void UAutoReimportManager::BeginDestroy()
+{
+	Super::BeginDestroy();
+	Implementation.Reset();
+}
+
+
+#undef CONDITIONAL_YIELD
+#undef LOCTEXT_NAMESPACE
+
diff --git a/Engine/Source/Editor/UnrealEd/Private/AutoReimport/AutoReimportUtilities.cpp b/Engine/Source/Editor/UnrealEd/Private/AutoReimport/AutoReimportUtilities.cpp
new file mode 100644
--- /dev/null
+++ b/Engine/Source/Editor/UnrealEd/Private/AutoReimport/AutoReimportUtilities.cpp
@@ -0,0 +1,44 @@
+// Copyright 1998-2015 Epic Games, Inc. All Rights Reserved.
+
+#include "UnrealEd.h"
+
+#include "AutoReimportUtilities.h"
+#include "AssetRegistryModule.h"
+
+namespace Utils
+{
+	/** Find a list of assets that were last imported from the specified filename. */
+	TArray<FAssetData> FindAssetsPertainingToFile(const IAssetRegistry& Registry, const FString& AbsoluteFilename)
+	{
+		TArray<FAssetData> Assets, Result;
+		const FString LeafName = FPaths::GetCleanFilename(AbsoluteFilename);
+		const FName TagName = UObject::SourceFileTagName();
+
+		FARFilter Filter;
+		Filter.SourceFilenames.Add(*LeafName);
+		Filter.bIncludeOnlyOnDiskAssets = true;
+		Registry.GetAssets(Filter, Assets);
+
+		for (const auto& Asset : Assets)
+		{
+			for (const auto& Pair : Asset.TagsAndValues)
+			{
+				// We don't compare numbers on FNames for this check because the tag "ReimportPath" may exist multiple times
+				const bool bCompareNumber = false;
+				if (Pair.Key.IsEqual(TagName, ENameCase::IgnoreCase, bCompareNumber))
+				{
+					// Either relative to the asset itself, or relative to the base path, or absolute.
+					// Would ideally use FReimportManager::ResolveImportFilename here, but we don't want to ask the file system whether the file exists.
+					if (AbsoluteFilename == FPaths::ConvertRelativePathToFull(FPackageName::LongPackageNameToFilename(Asset.PackagePath.ToString()) / Pair.Value) ||
+						AbsoluteFilename == FPaths::ConvertRelativePathToFull(Pair.Value))
+					{
+						Result.Add(Asset);
+						break;
+					}
+				}
+			}
+		}
+
+		return Result;
+	}
+}
\ No newline at end of file
diff --git a/Engine/Source/Editor/UnrealEd/Private/AutoReimport/AutoReimportUtilities.h b/Engine/Source/Editor/UnrealEd/Private/AutoReimport/AutoReimportUtilities.h
new file mode 100644
--- /dev/null
+++ b/Engine/Source/Editor/UnrealEd/Private/AutoReimport/AutoReimportUtilities.h
@@ -0,0 +1,102 @@
+// Copyright 1998-2015 Epic Games, Inc. All Rights Reserved.
+
+#pragma once
+
+class FAssetData;
+class IAssetRegistry;
+
+/** An immutable string with a cached CRC for efficient comparison with other strings */
+struct FImmutableString
+{
+	/** Constructible from a raw string */
+	FImmutableString(FString InString = FString()) : String(MoveTemp(InString)), CachedHash(0) {}
+	FImmutableString& operator=(FString InString) { String = MoveTemp(InString); CachedHash = 0; return *this; }
+
+	FImmutableString(const FImmutableString&) = default;
+	FImmutableString& operator=(const FImmutableString&) = default;
+
+	/** Move construction/assignment */
+	FImmutableString(FImmutableString&& In) : String(MoveTemp(In.String)), CachedHash(MoveTemp(In.CachedHash)) {}
+	FImmutableString& operator=(FImmutableString&& In) { Swap(String, In.String); Swap(CachedHash, In.CachedHash); return *this; }
+
+	/** Check for equality with another immutable string (case-sensitive) */
+	friend bool operator==(const FImmutableString& A, const FImmutableString& B)
+	{
+		return GetTypeHash(A) == GetTypeHash(B) && A.String == B.String;
+	}
+
+	/** Calculate the type hash for this string */
+	friend uint32 GetTypeHash(const FImmutableString& In)
+	{
+		if (In.CachedHash == 0 && In.String.Len() != 0)
+		{
+			In.CachedHash = GetTypeHash(In.String);
+		}
+		return In.CachedHash;
+	}
+
+	/** Get the underlying string */
+	const FString& Get() const { return String; }
+
+	/** Serialise this string */
+	friend FArchive& operator<<(FArchive& Ar, FImmutableString& String)
+	{
+		Ar << String.String;
+		if (Ar.IsSaving())
+		{
+			GetTypeHash(String);	
+		}
+		
+		Ar << String.CachedHash;
+
+		return Ar;
+	}
+
+private:
+	/** The string itself. */
+	FString String;
+
+	/** The cached hash of our string. 0 until compared with another string. */
+	mutable uint32 CachedHash;
+};
+
+/** A time limit that counts down from the time of construction, until it hits a given delay */
+struct FTimeLimit
+{
+	/** Constructor specifying that we should never bail out early */
+	FTimeLimit() : Delay(-1) { Reset(); }
+
+	/** Constructor specifying not to run over the specified number of seconds */
+	FTimeLimit(float NumSeconds) : Delay(NumSeconds) { Reset(); }
+
+	/** Check whether we have exceeded the time limit */
+	bool Exceeded() const { return Delay != -1 && FPlatformTime::Seconds() >= StartTime + Delay; }
+
+	/** Reset the time limit to start timing again from the current time */
+	void Reset() { StartTime = FPlatformTime::Seconds(); }
+
+private:
+
+	/** The delay specified by the user */
+	float Delay;
+
+	/** The time we started */
+	double StartTime;
+};
+
+namespace Utils
+{
+	/** Reduce the array to the specified accumulator */
+	template<typename T, typename P, typename A>
+	A Reduce(const TArray<T>& InArray, P Predicate, A Accumulator)
+	{
+		for (const T& X : InArray)
+		{
+			Accumulator = Predicate(X, Accumulator);
+		}
+		return Accumulator;
+	}
+
+	/** Find a list of assets that were last imported from the specified filename. */
+	TArray<FAssetData> FindAssetsPertainingToFile(const IAssetRegistry& Registry, const FString& AbsoluteFilename);
+}
\ No newline at end of file
diff --git a/Engine/Source/Editor/UnrealEd/Private/AutoReimport/ContentDirectoryMonitor.cpp b/Engine/Source/Editor/UnrealEd/Private/AutoReimport/ContentDirectoryMonitor.cpp
new file mode 100644
--- /dev/null
+++ b/Engine/Source/Editor/UnrealEd/Private/AutoReimport/ContentDirectoryMonitor.cpp
@@ -0,0 +1,208 @@
+// Copyright 1998-2015 Epic Games, Inc. All Rights Reserved.
+
+#include "UnrealEd.h"
+
+#include "ContentDirectoryMonitor.h"
+
+#include "AssetRegistryModule.h"
+#include "PackageTools.h"
+#include "ObjectTools.h"
+
+/** Generate a config from the specified options, to pass to FFileCache on construction */
+FFileCacheConfig GenerateFileCacheConfig(const FString& InPath, const FString& InSupportedExtensions, const FString& InMountedContentPath)
+{
+	FString Directory = FPaths::ConvertRelativePathToFull(InPath);
+
+	const FString& HashString = InMountedContentPath.IsEmpty() ? Directory : InMountedContentPath;
+	const uint32 CRC = FCrc::MemCrc32(*HashString, HashString.Len()*sizeof(TCHAR));	
+	FString CacheFilename = FPaths::ConvertRelativePathToFull(FPaths::GameIntermediateDir()) / TEXT("ReimportCache") / FString::Printf(TEXT("%u.bin"), CRC);
+
+	FFileCacheConfig Config(MoveTemp(Directory), MoveTemp(CacheFilename));
+	Config.IncludeExtensions = InSupportedExtensions;
+	// We always store paths inside content folders relative to the folder
+	Config.PathType = EPathType::Relative;
+
+	return Config;
+}
+
+FContentDirectoryMonitor::FContentDirectoryMonitor(const FString& InDirectory, const FString& InSupportedExtensions, const FString& InMountedContentPath)
+	: Cache(GenerateFileCacheConfig(InDirectory, InSupportedExtensions, InMountedContentPath))
+	, MountedContentPath(InMountedContentPath)
+	, LastSaveTime(0)
+{
+	
+}
+
+void FContentDirectoryMonitor::Destroy()
+{
+	Cache.Destroy();
+}
+
+void FContentDirectoryMonitor::Tick(const FTimeLimit& TimeLimit)
+{
+	Cache.Tick(TimeLimit);
+
+	const double Now = FPlatformTime::Seconds();
+	if (Now - LastSaveTime > ResaveIntervalS)
+	{
+		LastSaveTime = Now;
+		Cache.WriteCache();
+	}
+}
+
+void FContentDirectoryMonitor::StartProcessing()
+{
+	WorkProgress = TotalWork = 0;
+
+	auto OutstandingChanges = Cache.GetOutstandingChanges();
+	if (OutstandingChanges.Num() != 0)
+	{
+		for (auto& Transaction : OutstandingChanges)
+		{
+			switch(Transaction.Action)
+			{
+				case FFileChangeData::FCA_Added:		AddedFiles.Emplace(MoveTemp(Transaction));		break;
+				case FFileChangeData::FCA_Modified:		ModifiedFiles.Emplace(MoveTemp(Transaction));	break;
+				case FFileChangeData::FCA_Removed:		DeletedFiles.Emplace(MoveTemp(Transaction));	break;
+			}
+		}
+
+		TotalWork = AddedFiles.Num() + ModifiedFiles.Num() + DeletedFiles.Num();
+	}
+}
+
+void FContentDirectoryMonitor::ProcessAdditions(TArray<UPackage*>& OutPackagesToSave, const FTimeLimit& TimeLimit, const TMap<FString, TArray<UFactory*>>& InFactoriesByExtension)
+{
+	if (MountedContentPath.IsEmpty())
+	{
+		// @todo: arodham: Prompt user for import location
+		return;
+	}
+	
+	for (int32 Index = 0; Index < AddedFiles.Num(); ++Index)
+	{
+		auto& Addition = AddedFiles[Index];
+
+		WorkProgress += 1;
+		const FString FullFilename = Cache.GetDirectory() + Addition.Filename.Get();
+
+		bool bImportSucceeded = false;
+		bool bImportWasCancelled = false;
+
+		FString NewAssetName = ObjectTools::SanitizeObjectName(FPaths::GetBaseFilename(FullFilename));
+		FString PackagePath= PackageTools::SanitizePackageName(MountedContentPath / FPaths::GetPath(Addition.Filename.Get()) / NewAssetName);
+
+		// We can not create assets that share the name of a map file in the same location
+		if (FEditorFileUtils::IsMapPackageAsset(PackagePath))
+		{
+			// @todo: Error reporting
+			continue;
+		}
+
+		if (FindPackage(nullptr, *PackagePath))
+		{
+			// @todo: Error reporting
+			continue;	
+		}
+
+		UPackage* NewPackage = CreatePackage(nullptr, *PackagePath);
+		if ( !ensure(NewPackage) )
+		{
+			// @todo: Error reporting
+			continue;
+		}
+
+		// Make sure the destination package is loaded
+		NewPackage->FullyLoad();
+
+		bool bSuccess = false;
+		bool bCancelled = false;
+
+		UObject* NewAsset = nullptr;
+
+		const FString Ext = FPaths::GetExtension(Addition.Filename.Get(), false);
+		if (auto* Factories = InFactoriesByExtension.Find(Ext))
+		{
+			for (auto* Factory : *Factories)
+			{
+				NewAsset = UFactory::StaticImportObject(Factory->ResolveSupportedClass(), NewPackage, FName(*NewAssetName), RF_Public | RF_Standalone, bCancelled, *FullFilename, nullptr, Factory);
+				if (NewAsset || bCancelled)
+				{
+					break;
+				}
+			}
+		}
+		
+		if (!NewAsset || bCancelled)
+		{
+			// @todo: Error reporting
+			continue;
+		}
+
+		FAssetRegistryModule::AssetCreated(NewAsset);
+		GEditor->BroadcastObjectReimported(NewAsset);
+
+		OutPackagesToSave.Add(NewPackage);
+
+		// Refresh the supported class.  Some factories (e.g. FBX) only resolve their type after reading the file
+		// ImportAssetType = Factory->ResolveSupportedClass();
+		// @todo: analytics
+
+		// Let the cache know that we've dealt with this change (it will be imported immediately)
+		Cache.CompleteTransaction(MoveTemp(Addition));
+
+		if (TimeLimit.Exceeded())
+		{
+			AddedFiles.RemoveAt(0, Index + 1);
+			return;
+		}
+	}
+
+	AddedFiles.Empty();
+}
+
+void FContentDirectoryMonitor::ProcessModifications(const IAssetRegistry& Registry, const FTimeLimit& TimeLimit)
+{
+	auto* ReimportManager = FReimportManager::Instance();
+
+	for (int32 Index = 0; Index < ModifiedFiles.Num(); ++Index)
+	{
+		auto& Change = ModifiedFiles[Index];
+
+		const FString FullFilename = Cache.GetDirectory() + Change.Filename.Get();
+		for (const auto& AssetData : Utils::FindAssetsPertainingToFile(Registry, FullFilename))
+		{
+			UObject* Asset = AssetData.GetAsset();
+			ReimportManager->Reimport(Asset, false /* Ask for new file */, false /* Show notification */);
+		}
+
+		// Let the cache know that we've dealt with this change
+		Cache.CompleteTransaction(MoveTemp(Change));
+		
+		WorkProgress += 1;
+		if (TimeLimit.Exceeded())
+		{
+			ModifiedFiles.RemoveAt(0, Index + 1);
+			return;
+		}
+	}
+
+	ModifiedFiles.Empty();
+}
+
+void FContentDirectoryMonitor::ExtractAssetsToDelete(const IAssetRegistry& Registry, TArray<FAssetData>& OutAssetsToDelete)
+{
+	for (auto& Deletion : DeletedFiles)
+	{
+		for (const auto& AssetData : Utils::FindAssetsPertainingToFile(Registry, Cache.GetDirectory() + Deletion.Filename.Get()))
+		{
+			OutAssetsToDelete.Add(AssetData);
+		}
+
+		// Let the cache know that we've dealt with this change (it will be imported in due course)
+		Cache.CompleteTransaction(MoveTemp(Deletion));
+	}
+
+	WorkProgress += DeletedFiles.Num();
+	DeletedFiles.Empty();
+}
\ No newline at end of file
diff --git a/Engine/Source/Editor/UnrealEd/Private/AutoReimport/ContentDirectoryMonitor.h b/Engine/Source/Editor/UnrealEd/Private/AutoReimport/ContentDirectoryMonitor.h
new file mode 100644
--- /dev/null
+++ b/Engine/Source/Editor/UnrealEd/Private/AutoReimport/ContentDirectoryMonitor.h
@@ -0,0 +1,72 @@
+// Copyright 1998-2015 Epic Games, Inc. All Rights Reserved.
+
+#pragma once
+
+#include "AutoReimportUtilities.h"
+#include "FileCache.h"
+
+/** Class responsible for watching a specific content directory for changes */
+class FContentDirectoryMonitor
+{
+public:
+
+	/**
+	 * Constructor.
+	 * @param InDirectory			Content directory path to monitor. Assumed to be absolute.
+	 * @param InSupportedExtensions A string containing semi-colon separated extensions to monitor.
+	 * @param InMountedContentPath	(optional) Mounted content path (eg /Engine/, /Game/) to which InDirectory maps.
+	 */
+	FContentDirectoryMonitor(const FString& InDirectory, const FString& InSupportedExtensions, const FString& InMountedContentPath = FString());
+
+	/** Tick this monitor's cache to give it a chance to finish scanning for files */
+	void Tick(const FTimeLimit& TimeLimit);
+
+	/** Start processing any outstanding changes this monitor is aware of */
+	void StartProcessing();
+
+	/** Extract the files we need to import from our outstanding changes (happens first)*/ 
+	void ProcessAdditions(TArray<UPackage*>& OutPackagesToSave, const FTimeLimit& TimeLimit, const TMap<FString, TArray<UFactory*>>& InFactoriesByExtension);
+
+	/** Process the outstanding changes that we have cached */
+	void ProcessModifications(const IAssetRegistry& Registry, const FTimeLimit& TimeLimit);
+
+	/** Extract the assets we need to delete from our outstanding changes (happens last) */ 
+	void ExtractAssetsToDelete(const IAssetRegistry& Registry, TArray<FAssetData>& OutAssetsToDelete);
+
+	/** Destroy this monitor including its cache */
+	void Destroy();
+
+	/** Get the directory that this monitor applies to */
+	const FString& GetDirectory() const { return Cache.GetDirectory(); }
+
+public:
+
+	/** Get the number of outstanding changes that we potentially have to process (when not already processing) */
+	int32 GetNumUnprocessedChanges() const { return Cache.GetNumOutstandingChanges(); }
+
+	/** Get the total amount of work this monitor has to perform in the current processing operation */
+	int32 GetTotalWork() const { return TotalWork; }
+
+	/** Get the total amount of work this monitor has performed in the current processing operation */
+	int32 GetWorkProgress() const { return WorkProgress; }
+
+private:
+
+	/** The file cache that monitors and reflects the content directory. */
+	FFileCache Cache;
+
+	/** The mounted content path for this monitor (eg /Game/) */
+	FString MountedContentPath;
+
+	/** A list of file system changes that are due to be processed */
+	TArray<FUpdateCacheTransaction> AddedFiles, ModifiedFiles, DeletedFiles;
+
+	/** The number of changes we've processed out of the OutstandingChanges array */
+	int32 TotalWork, WorkProgress;
+
+	/** The last time we attempted to save the cache file */
+	double LastSaveTime;
+
+	/** The interval between potential re-saves of the cache file */
+	static const int32 ResaveIntervalS = 60;
+};
\ No newline at end of file
diff --git a/Engine/Source/Editor/UnrealEd/Private/AutoReimport/FileCache.cpp b/Engine/Source/Editor/UnrealEd/Private/AutoReimport/FileCache.cpp
new file mode 100644
--- /dev/null
+++ b/Engine/Source/Editor/UnrealEd/Private/AutoReimport/FileCache.cpp
@@ -0,0 +1,453 @@
+// Copyright 1998-2015 Epic Games, Inc. All Rights Reserved.
+
+
+#include "UnrealEd.h"
+#include "FileCache.h"
+#include "IDirectoryWatcher.h"
+#include "DirectoryWatcherModule.h"
+
+const FGuid FFileCacheCustomVersion::Key(0x8E7DDCB3, 0x80DA47BB, 0x9FD346A2, 0x93984DF6);
+FCustomVersionRegistration GRegisterFileCacheVersion(FFileCacheCustomVersion::Key, FFileCacheCustomVersion::Latest, TEXT("FileCacheVersion"));
+
+namespace
+{
+	/** Helper template function to remove duplicates from an array, given some predicate */
+	template<typename T, typename P>
+	void RemoveDuplicates(TArray<T>& Array, P Predicate)
+	{
+		if (Array.Num() == 0)
+			return;
+
+		const int32 NumElements = Array.Num();
+		int32 Write = 0;
+		for (int32 Read = 0; Read < NumElements; ++Read)
+		{
+			for (int32 DuplRead = Read + 1; DuplRead < NumElements; ++DuplRead)
+			{
+				if (Predicate(Array[Read], Array[DuplRead]))
+				{
+					goto next;
+				}
+			}
+
+			if (Write != Read)
+			{
+				Swap(Array[Write++], Array[Read]);
+			}
+			else
+			{
+				Write++;
+			}
+	next:
+			;
+		}
+		const int32 NumDuplicates = NumElements - Write;
+		if (NumDuplicates != 0)
+		{
+			Array.RemoveAt(Write, NumDuplicates, false);
+		}
+	}
+}
+
+FAsyncDirectoryReader::FAsyncDirectoryReader(const FString& InDirectory, EPathType InPathType)
+	: RootPath(InDirectory), PathType(InPathType)
+{
+	PendingDirectories.Add(InDirectory);
+	State.Emplace();
+}
+
+TOptional<FDirectoryState> FAsyncDirectoryReader::GetFinalState()
+{
+	TOptional<FDirectoryState> OldState;
+	Swap(OldState, State);
+
+	return OldState;
+}
+
+bool FAsyncDirectoryReader::IsComplete() const
+{
+	return PendingFiles.Num() == 0 && PendingDirectories.Num() == 0;
+}
+
+FAsyncDirectoryReader::EProgressResult FAsyncDirectoryReader::Tick(const FTimeLimit& TimeLimit)
+{
+	auto& FileManager = IFileManager::Get();
+	const int32 RootPathLen = RootPath.Len();
+
+	for (int32 Index = 0; Index < PendingDirectories.Num(); ++Index)
+	{
+		ScanDirectory(PendingDirectories[Index]);
+		if (TimeLimit.Exceeded())
+		{
+			// We've spent too long, bail
+			PendingDirectories.RemoveAt(0, Index, false);
+			return EProgressResult::Pending;
+		}
+	}
+	PendingDirectories.Empty();
+
+	for (int32 Index = 0; Index < PendingFiles.Num(); ++Index)
+	{
+		const auto& File = PendingFiles[Index];
+
+		// Store the file relative or absolute
+		FString Filename = (PathType == EPathType::Relative ? *File + RootPathLen : *File);
+		State->Files.Emplace(MoveTemp(Filename), FileManager.GetTimeStamp(*File));
+
+		if (Index % 100 == 0 && TimeLimit.Exceeded())
+		{
+			// We've spent too long, bail
+			PendingFiles.RemoveAt(0, Index, false);
+			return EProgressResult::Pending;
+		}
+	}
+	PendingFiles.Empty();
+
+	return EProgressResult::Finished;
+}
+
+void FAsyncDirectoryReader::ScanDirectory(const FString& InDirectory)
+{
+	struct FVisitor : public IPlatformFile::FDirectoryVisitor
+	{
+		TArray<FString>* PendingFiles;
+		TArray<FString>* PendingDirectories;
+
+		virtual bool Visit(const TCHAR* FilenameOrDirectory, bool bIsDirectory)
+		{
+			if (bIsDirectory)
+			{
+				PendingDirectories->Add(FString(FilenameOrDirectory));
+			}
+			else
+			{
+				PendingFiles->Add(FString(FilenameOrDirectory));
+			}
+			return true;
+		}
+	};
+
+	FVisitor Visitor;
+	Visitor.PendingFiles = &PendingFiles;
+	Visitor.PendingDirectories = &PendingDirectories;
+	IFileManager::Get().IterateDirectory(*InDirectory, Visitor);
+}
+
+FFileCache::FFileCache(const FFileCacheConfig& InConfig)
+	: Config(InConfig)
+	, DirectoryReader(InConfig.Directory, Config.PathType)
+	, bSavedCacheDirty(false)
+{
+	// Ensure that the extension strings are of the form ;ext1;ext2;ext3;
+	if (!Config.IncludeExtensions.IsEmpty() && Config.IncludeExtensions[0] != ';')
+	{
+		Config.IncludeExtensions.InsertAt(0, ';');
+	}
+	if (!Config.ExcludeExtensions.IsEmpty() && Config.ExcludeExtensions[0] != ';')
+	{
+		Config.ExcludeExtensions.InsertAt(0, ';');
+	}
+
+	// Attempt to load an existing cache file
+	auto ExistingCache = ReadCache();
+	if (ExistingCache.IsSet())
+	{
+		CachedDirectoryState = MoveTemp(ExistingCache.GetValue());
+	}
+
+	FDirectoryWatcherModule& Module = FModuleManager::LoadModuleChecked<FDirectoryWatcherModule>(TEXT("DirectoryWatcher"));
+	if (IDirectoryWatcher* DirectoryWatcher = Module.Get())
+	{
+		auto Callback = IDirectoryWatcher::FDirectoryChanged::CreateRaw(this, &FFileCache::OnDirectoryChanged);
+		DirectoryWatcher->RegisterDirectoryChangedCallback_Handle(Config.Directory, Callback, WatcherDelegate);
+	}
+}
+
+FFileCache::~FFileCache()
+{
+	UnbindWatcher();
+	WriteCache();
+}
+
+void FFileCache::Destroy()
+{
+	// Delete the cache file, and clear out everything
+	bSavedCacheDirty = false;
+	IFileManager::Get().Delete(*Config.CacheFile, false, true, true);
+
+	DirectoryReader = FAsyncDirectoryReader(Config.Directory, Config.PathType);
+	OutstandingChanges.Empty();
+	CachedDirectoryState = FDirectoryState();
+
+	UnbindWatcher();
+}
+
+void FFileCache::UnbindWatcher()
+{
+	if (WatcherDelegate == FDelegateHandle())
+	{
+		return;
+	}
+
+	if (FDirectoryWatcherModule* Module = FModuleManager::GetModulePtr<FDirectoryWatcherModule>(TEXT("DirectoryWatcher")))
+	{
+		if (IDirectoryWatcher* DirectoryWatcher = Module->Get())
+		{
+			DirectoryWatcher->UnregisterDirectoryChangedCallback_Handle(Config.Directory, WatcherDelegate);
+		}
+	}
+
+	WatcherDelegate = FDelegateHandle();
+}
+
+TOptional<FDirectoryState> FFileCache::ReadCache() const
+{
+	TOptional<FDirectoryState> Optional;
+
+	FArchive* Ar = IFileManager::Get().CreateFileReader(*Config.CacheFile);
+	if (Ar)
+	{
+		FDirectoryState Result;
+
+		*Ar << Result;
+		Ar->Close();
+		delete Ar;
+
+		Optional.Emplace(MoveTemp(Result));
+	}
+
+	return Optional;
+}
+
+void FFileCache::WriteCache()
+{
+	if (bSavedCacheDirty)
+	{
+		FArchive* Ar = IFileManager::Get().CreateFileWriter(*Config.CacheFile);
+		check(Ar);
+
+		*Ar << CachedDirectoryState;
+
+		Ar->Close();
+		delete Ar;
+
+		CachedDirectoryState.Files.Shrink();
+
+		bSavedCacheDirty = false;
+	}
+}
+
+TArray<FUpdateCacheTransaction> FFileCache::GetOutstandingChanges()
+{
+	TArray<FUpdateCacheTransaction> Moved;
+	Swap(OutstandingChanges, Moved);
+	return Moved;
+}
+
+void FFileCache::CompleteTransaction(FUpdateCacheTransaction&& Transaction)
+{
+	auto* CachedData = CachedDirectoryState.Files.Find(Transaction.Filename);
+
+	switch (Transaction.Action)
+	{
+	case FFileChangeData::FCA_Modified:
+		if (CachedData && CachedData->Timestamp < Transaction.Timestamp)
+		{
+			// Update the timestamp
+			CachedData->Timestamp = Transaction.Timestamp;
+
+			bSavedCacheDirty = true;
+		}
+		break;
+	case FFileChangeData::FCA_Added:
+		if (!CachedData)
+		{
+			// Add the file information to the cache
+			CachedDirectoryState.Files.Emplace(Transaction.Filename, Transaction.Timestamp);
+
+			bSavedCacheDirty = true;
+		}
+		break;
+	case FFileChangeData::FCA_Removed:
+		if (CachedData)
+		{
+			// Remove the file information to the cache
+			CachedDirectoryState.Files.Remove(Transaction.Filename);
+
+			bSavedCacheDirty = true;
+		}
+		break;
+
+	default:
+		checkf(false, TEXT("Invalid file cached transaction"));
+		break;
+	}
+}
+
+void FFileCache::Tick(const FTimeLimit& Limit)
+{
+	if (DirectoryReader.IsComplete() || DirectoryReader.Tick(Limit) == FAsyncDirectoryReader::EProgressResult::Pending)
+	{
+		return;
+	}
+
+	// We should only ever get here once. The directory reader has finished scanning, and we can now diff the results with what we had saved in the cache file.
+	check(DirectoryReader.IsComplete());
+	
+	TOptional<FDirectoryState> CurrentFileData = DirectoryReader.GetFinalState();
+	if (CachedDirectoryState.Files.Num() == 0)
+	{
+		// If we don't have any cached data yet, just use the file data we just harvested
+		CachedDirectoryState = MoveTemp(CurrentFileData.GetValue());
+		bSavedCacheDirty = true;
+		return;
+	}
+
+	// We already have cached data so we need to compare it with the harvested data
+	// to detect additions, modifications, and removals
+	for (const auto& FilenameAndData : CurrentFileData->Files)
+	{
+		const FString& Filename = FilenameAndData.Key.Get();
+		if (!IsFileApplicable(Filename))
+		{
+			continue;
+		}
+
+		const auto* CachedData = CachedDirectoryState.Files.Find(Filename);
+		// Have we encountered this before?
+		if (CachedData)
+		{
+			if (CachedData->Timestamp != FilenameAndData.Value.Timestamp)
+			{
+				OutstandingChanges.Add(FUpdateCacheTransaction(Filename, FFileChangeData::FCA_Modified, FilenameAndData.Value.Timestamp));
+			}
+		}
+		else
+		{
+			// No file already - create one. We don't do this if the cache didn't exist before as we have no idea what is new and what isn't
+			OutstandingChanges.Add(FUpdateCacheTransaction(Filename, FFileChangeData::FCA_Added, FilenameAndData.Value.Timestamp));
+		}
+	}
+
+	// Delete anything that doesn't exist on disk anymore
+	for (auto It = CachedDirectoryState.Files.CreateIterator(); It; ++It)
+	{
+		const FImmutableString& Filename = It.Key();
+		if (!CurrentFileData->Files.Contains(Filename) && IsFileApplicable(Filename.Get()))
+		{
+			OutstandingChanges.Add(FUpdateCacheTransaction(Filename, FFileChangeData::FCA_Removed));
+		}
+	}
+}
+
+void FFileCache::OnDirectoryChanged(const TArray<FFileChangeData>& FileChanges)
+{
+	auto& FileManager = IFileManager::Get();
+
+	OutstandingChanges.Reserve(OutstandingChanges.Num() + FileChanges.Num());
+
+	const int32 RootPathLen = Config.Directory.Len();
+	for (const auto& ThisEntry : FileChanges)
+	{
+		// If it's a directory or is not applicable, ignore it
+		if (FileManager.DirectoryExists(*ThisEntry.Filename) || !IsFileApplicable(ThisEntry.Filename))
+		{
+			continue;
+		}
+
+		FImmutableString TransactionFilename;
+		{
+			FString Temp = FPaths::ConvertRelativePathToFull(ThisEntry.Filename);
+			if (Config.PathType == EPathType::Relative)
+			{
+				TransactionFilename = FString(*Temp + RootPathLen);
+			}
+			else
+			{
+				TransactionFilename = MoveTemp(Temp);
+			}
+		}
+		
+		switch (ThisEntry.Action)
+		{
+		case FFileChangeData::FCA_Added:
+			{
+				const int32 NumRemoved = OutstandingChanges.RemoveAll([&](const FUpdateCacheTransaction& X){
+					return X.Action == FFileChangeData::FCA_Removed && X.Filename == TransactionFilename;
+				});
+
+				const auto Action = NumRemoved == 0 ? FFileChangeData::FCA_Added : FFileChangeData::FCA_Modified;
+				OutstandingChanges.Add(FUpdateCacheTransaction(TransactionFilename, Action, FileManager.GetTimeStamp(*ThisEntry.Filename)));
+			}
+			break;
+
+		case FFileChangeData::FCA_Removed:
+			{
+				bool bPreviouslyAdded = false;
+				const int32 NumRemoved = OutstandingChanges.RemoveAll([&](const FUpdateCacheTransaction& X){
+					if (X.Filename == TransactionFilename)
+					{
+						bPreviouslyAdded = X.Action == FFileChangeData::FCA_Added || bPreviouslyAdded;
+						return true;
+					}
+					return false;
+				});
+
+				if (!bPreviouslyAdded)
+				{
+					OutstandingChanges.Add(FUpdateCacheTransaction(TransactionFilename, FFileChangeData::FCA_Removed));
+				}
+			}
+			
+			break;
+
+		case FFileChangeData::FCA_Modified:
+			{
+				const bool bPreviouslyAdded = OutstandingChanges.ContainsByPredicate([&](const FUpdateCacheTransaction& X){
+					return X.Filename == TransactionFilename && X.Action == FFileChangeData::FCA_Added;
+				});
+
+				if (!bPreviouslyAdded)
+				{
+					OutstandingChanges.Add(FUpdateCacheTransaction(TransactionFilename, FFileChangeData::FCA_Modified, FileManager.GetTimeStamp(*ThisEntry.Filename)));
+				}
+			}
+			break;
+
+		default:
+			break;
+		}
+	}
+
+	RemoveDuplicates(OutstandingChanges, [](const FUpdateCacheTransaction& A, const FUpdateCacheTransaction& B){
+		return A.Action == B.Action && A.Filename == B.Filename;
+	});
+}
+
+bool FFileCache::IsFileApplicable(const FString& Filename)
+{
+	return (Config.ExcludeExtensions.IsEmpty() || !MatchExtensionString(Config.ExcludeExtensions, Filename)) &&
+		   (Config.IncludeExtensions.IsEmpty() ||  MatchExtensionString(Config.IncludeExtensions, Filename));
+}
+
+bool FFileCache::MatchExtensionString(const FString& InSource, const FString& InPath)
+{
+	int32 SlashPosition = InPath.Find(TEXT("/"), ESearchCase::CaseSensitive, ESearchDir::FromEnd);
+	// in case we are using backslashes on a platform that doesn't use backslashes
+	SlashPosition = FMath::Max(SlashPosition, InPath.Find(TEXT("\\"), ESearchCase::CaseSensitive, ESearchDir::FromEnd));
+
+	const int32 DotPosition = InPath.Find(TEXT("."), ESearchCase::CaseSensitive, ESearchDir::FromEnd);
+
+	if (DotPosition > INDEX_NONE && DotPosition > SlashPosition && DotPosition < InPath.Len() - 1)
+	{
+		FString Search;
+		Search.Reserve(InPath.Len() - DotPosition + 2);;
+		Search.InsertAt(0, ';');
+		Search += &InPath[DotPosition + 1];
+		Search.AppendChar(';');
+
+		return InSource.Find(Search, ESearchCase::IgnoreCase) != INDEX_NONE;
+	}
+
+	return false;
+}
\ No newline at end of file
diff --git a/Engine/Source/Editor/UnrealEd/Private/AutoReimport/FileCache.h b/Engine/Source/Editor/UnrealEd/Private/AutoReimport/FileCache.h
new file mode 100644
--- /dev/null
+++ b/Engine/Source/Editor/UnrealEd/Private/AutoReimport/FileCache.h
@@ -0,0 +1,265 @@
+// Copyright 1998-2015 Epic Games, Inc. All Rights Reserved.
+
+#pragma once
+
+#include "Serialization/CustomVersion.h"
+#include "AutoReimportUtilities.h"
+#include "IDirectoryWatcher.h"
+
+/** Custom serialization version for FFileCache */
+struct UNREALED_API FFileCacheCustomVersion
+{
+	static const FGuid Key;
+	enum Type {	Initial, Latest = Initial };
+};
+
+/** Structure representing specific information about a particular file */
+struct FFileData
+{
+	/** Constructiors */
+	FFileData() {}
+	FFileData(const FDateTime& InTimestamp) : Timestamp(InTimestamp) {}
+
+	/** Serializer for this type */
+	friend FArchive& operator<<(FArchive& Ar, FFileData& Data)
+	{
+		if (Ar.CustomVer(FFileCacheCustomVersion::Key) >= FFileCacheCustomVersion::Initial)
+		{
+			Ar << Data.Timestamp;
+		}
+
+		return Ar;
+	}
+
+	FDateTime Timestamp;
+};
+
+/** Structure representing the file data for a number of files in a directory */
+struct FDirectoryState
+{
+	/** Default construction */
+	FDirectoryState() {}
+
+	/** Move construction */
+	FDirectoryState(FDirectoryState&& In) : Files(MoveTemp(In.Files)) {}
+	FDirectoryState& operator=(FDirectoryState&& In) { Swap(Files, In.Files); return *this; }
+
+	/** Filename -> data map */
+	TMap<FImmutableString, FFileData> Files;
+
+	/** Serializer for this type */
+	friend FArchive& operator<<(FArchive& Ar, FDirectoryState& State)
+	{
+		Ar.UsingCustomVersion(FFileCacheCustomVersion::Key);
+
+		if (Ar.CustomVer(FFileCacheCustomVersion::Key) >= FFileCacheCustomVersion::Initial)
+		{
+			// Number of files
+			int32 Num = State.Files.Num();
+			Ar << Num;
+			if (Ar.IsLoading())
+			{
+				State.Files.Reserve(Num);
+			}
+
+			Ar << State.Files;
+		}
+
+		return Ar;
+	}
+};
+
+/** A transaction issued by FFileCache to describe a change to the cache. The change is only committed once the transaction is returned to the cache (see FFileCache::CompleteTransaction). */
+struct FUpdateCacheTransaction
+{
+	/** The path of the file to which this transaction relates */
+	FImmutableString Filename;
+
+	/** The timestamp of the file at the time of this change */
+	FDateTime Timestamp;
+
+	/** The type of action that prompted this transaction */
+	FFileChangeData::EFileChangeAction Action;
+
+	/** Publically moveable */
+	FUpdateCacheTransaction(FUpdateCacheTransaction&& In) : Filename(MoveTemp(In.Filename)), Timestamp(MoveTemp(In.Timestamp)), Action(MoveTemp(In.Action)) {}
+	FUpdateCacheTransaction& operator=(FUpdateCacheTransaction&& In) { Swap(Filename, In.Filename); Swap(Timestamp, In.Timestamp); Swap(Action, In.Action); return *this; }
+
+private:
+	friend class FFileCache;
+	
+	/** Construction responsibility is held by FFileCache */
+	FUpdateCacheTransaction(FImmutableString InFilename, FFileChangeData::EFileChangeAction InAction, const FDateTime& InTimestamp = 0)
+		: Filename(MoveTemp(InFilename)), Timestamp(InTimestamp), Action(InAction)
+	{}
+};
+
+/** Enum specifying whether a path should be relative or absolute */
+enum EPathType
+{
+	/** Paths should be cached relative to the root cache directory */
+	Relative,
+
+	/** Paths should be cached as absolute file system paths */
+	Absolute
+};
+
+/**
+ * Class responsible for 'asyncronously' scanning a folder for files and timestamps.
+ * Example usage:
+ *		FAsyncDirectoryReader Reader(TEXT("C:\\Path"), EPathType::Relative);
+ *
+ *		while(!Reader.IsComplete())
+ *		{
+ *			FPlatformProcess::Sleep(1);
+ *			Reader.Tick(FTimedSignal(1));	// Do 1 second of work
+ *		}
+ *		TOptional<FDirectoryState> State = Reader.GetFinalState();
+ */
+struct FAsyncDirectoryReader
+{
+	enum class EProgressResult
+	{
+		Finished, Pending
+	};
+
+	/** Constructor that sets up the directory reader to the specified directory */
+	FAsyncDirectoryReader(const FString& InDirectory, EPathType InPathType);
+
+	/**
+	 * Get the state of the directory once finished. Relinquishes the currently stored directory state to the client.
+	 * Returns nothing if incomplete, or if GetFinalState() has already been called.
+	 */
+	TOptional<FDirectoryState> GetFinalState();
+
+	/** Returns true when this directory reader has finished scanning the directory */
+	bool IsComplete() const;
+
+	/** Tick this reader (discover new directories / files). Returns progress state. */
+	EProgressResult Tick(const FTimeLimit& Limit);
+
+private:
+	/** Non-recursively scan a single directory for its contents. Adds results to Pending arrays. */
+	void ScanDirectory(const FString& InDirectory);
+
+	/** Path to the root directory we want to scan */
+	FString RootPath;
+
+	/** Whether we should return relative or absolute paths */
+	EPathType PathType;
+
+	/** The currently discovered state of the directory */
+	TOptional<FDirectoryState> State;
+
+	/** A list of directories we have recursively found on our travels */
+	TArray<FString> PendingDirectories;
+
+	/** A list of files we have recursively found on our travels */
+	TArray<FString> PendingFiles;
+};
+
+/** Configuration structure required to construct a FFileCache */
+struct FFileCacheConfig
+{
+	FFileCacheConfig(FString InDirectory, FString InCacheFile)
+		: Directory(InDirectory), CacheFile(InCacheFile), PathType(EPathType::Relative)
+	{}
+
+	/** String specifying the directory on disk that the cache should reflect */
+	FString Directory;
+
+	/** String specifying the file that the cache should be saved to */
+	FString CacheFile;
+
+	/**
+	 * Strings specifying extensions to include/exclude in the cache (or empty to include everything).
+	 * Specified as a list of semi-colon separated extensions (eg "png;jpg;gif;").
+	 */
+	FString IncludeExtensions, ExcludeExtensions;
+
+	/** Path type to return, relative to the directory or absolute. */
+	EPathType PathType;
+};
+
+/**
+ * A class responsible for scanning a directory, and maintaining a cache of its state (files and timestamps).
+ * Changes in the cache can be retrieved through GetOutstandingChanges(). Changes will be reported for any
+ * change in the cached state even between runs of the process.
+ */
+class UNREALED_API FFileCache
+{
+public:
+
+	/** Construction from a config */
+	FFileCache(const FFileCacheConfig& InConfig);
+
+	FFileCache(const FFileCache&) = delete;
+	FFileCache& operator=(const FFileCache&) = delete;
+
+	/** Destructor */
+	~FFileCache();
+
+	/** Destroy this cache. Cleans out the in-memory state and deletes the cache file, if present. */
+	void Destroy();
+
+	/** Get the absolute path of the directory this cache reflects */
+	const FString& GetDirectory() const { return Config.Directory; }
+
+	/** Tick this FileCache, optionally specifying a limiter to limit the amount of time spent in here */
+	void Tick(const FTimeLimit& Limit);
+
+	/** Write out the cached file, if we have any changes to write */
+	void WriteCache();
+
+	/**
+	 * Return a transaction to the cache for completion. Will update the cached state with the change
+	 * described in the transaction, and mark the cache as needing to be saved.
+	 */
+	void CompleteTransaction(FUpdateCacheTransaction&& Transaction);
+
+	/** Get the number of pending changes to the cache. */
+	int32 GetNumOutstandingChanges() const { return OutstandingChanges.Num(); }
+
+	/** Get all pending changes to the cache. Transactions must be returned to CompleteTransaction to update the cache. */
+	TArray<FUpdateCacheTransaction> GetOutstandingChanges();
+
+private:
+
+	/** Called when the directory we are monitoring has been changed in some way */
+	void OnDirectoryChanged(const TArray<FFileChangeData>& FileChanges);
+
+	/** Check whether a file is applicable to our cache */
+	bool IsFileApplicable(const FString& InPath);
+
+	/**
+	 * Check whether a filename's extension matches the specified extension list.
+	 * InExtensions must be of the form ";ext1;ext2;ext3;". Note the leading and trailing ;
+	 */
+	static bool MatchExtensionString(const FString& InExtensions, const FString& InFilename);
+
+	/** Unbind the watcher from the directory. Called on destruction. */
+	void UnbindWatcher();
+
+	/** Read the cache file data and return the contents */
+	TOptional<FDirectoryState> ReadCache() const;
+
+private:
+
+	/** Configuration settings applied on construction */
+	FFileCacheConfig Config;
+
+	/** 'Asynchronous' directory reader responsible for gathering all file/timestamp information recursively from our cache directory */
+	FAsyncDirectoryReader DirectoryReader;
+
+	/** Temporary array of transactions to be returned to the client through GetOutstandingChanges() */
+	TArray<FUpdateCacheTransaction> OutstandingChanges;
+
+	/** Our in-memory view of the cached directory state. */
+	FDirectoryState CachedDirectoryState;
+
+	/** Handle to the directory watcher delegate so we can delete it properly */
+	FDelegateHandle WatcherDelegate;
+
+	/** True when the cached state we have in memory is more up to date than the serialized file. Enables WriteCache() when true. */
+	bool bSavedCacheDirty;
+};
diff --git a/Engine/Source/Editor/UnrealEd/Private/AutoReimportManager.cpp b/Engine/Source/Editor/UnrealEd/Private/AutoReimportManager.cpp
deleted file mode 100644
diff --git a/Engine/Source/Editor/UnrealEd/Private/FileCache.cpp b/Engine/Source/Editor/UnrealEd/Private/FileCache.cpp
deleted file mode 100644
diff --git a/Engine/Source/Editor/UnrealEd/Public/AutoReimport/AutoReimportManager.h b/Engine/Source/Editor/UnrealEd/Public/AutoReimport/AutoReimportManager.h
--- a/Engine/Source/Editor/UnrealEd/Public/AutoReimport/AutoReimportManager.h
+++ b/Engine/Source/Editor/UnrealEd/Public/AutoReimport/AutoReimportManager.h
@@ -2,196 +2,25 @@
 
 #pragma once
 
-#include "FileCache.h"
-
 #include "AutoReimportManager.generated.h"
 
-class IAssetRegistry;
-class IAssetTools;
-class SNotificationItem;
-
-/** Feedback context that overrides GWarn for import operations to prevent popup spam */
-struct FReimportFeedbackContext : public FFeedbackContext
-{
-	/** Constructor - created when we are processing changes */
-	FReimportFeedbackContext();
-	/** Destructor - called when we are done processing changes */
-	~FReimportFeedbackContext();
-
-	/** Update the text on the notification */
-	void UpdateText(int32 Progress, int32 Total);
-
-private:
-	virtual void Serialize( const TCHAR* V, ELogVerbosity::Type Verbosity, const class FName& Category ) {}
-
-private:
-	/** The notification that is shown when the context is active */
-	TSharedPtr<SNotificationItem> Notification;
-};
-
-/** Class responsible for watching a specific content directory for changes */
-class FContentDirectoryMonitor
-{
-public:
-
-	/**
-	 * Constructor.
-	 * @param InDirectory			Content directory path to monitor. Assumed to be absolute.
-	 * @param InSupportedExtensions A string containing semi-colon separated extensions to monitor.
-	 * @param InMountedContentPath	(optional) Mounted content path (eg /Engine/, /Game/) to which InDirectory maps.
-	 */
-	FContentDirectoryMonitor(const FString& InDirectory, const FString& InSupportedExtensions, const FString& InMountedContentPath = FString());
-
-	/** Tick this monitor's cache to give it a chance to finish scanning for files */
-	void Tick(const FTimeLimit& TimeLimit);
-
-	/** Start processing any outstanding changes this monitor is aware of */
-	void StartProcessing();
-
-	/** Extract the files we need to import from our outstanding changes (happens first)*/ 
-	void ProcessAdditions(TArray<UPackage*>& OutPackagesToSave, const FTimeLimit& TimeLimit, const TMap<FString, TArray<UFactory*>>& InFactoriesByExtension);
-
-	/** Process the outstanding changes that we have cached */
-	void ProcessModifications(const IAssetRegistry& Registry, const FTimeLimit& TimeLimit, FReimportFeedbackContext& Context);
-
-	/** Extract the assets we need to delete from our outstanding changes (happens last) */ 
-	void ExtractAssetsToDelete(const IAssetRegistry& Registry, TArray<FAssetData>& OutAssetsToDelete);
-
-	/** Destroy this monitor including its cache */
-	void Destroy();
-
-	/** Get the directory that this monitor applies to */
-	const FString& GetDirectory() const { return Cache.GetDirectory(); }
-
-public:
-
-	/** Get the number of outstanding changes that we potentially have to process (when not already processing) */
-	int32 GetNumUnprocessedChanges() const { return Cache.GetNumOutstandingChanges(); }
-
-	/** Get the total amount of work this monitor has to perform in the current processing operation */
-	int32 GetTotalWork() const { return TotalWork; }
-
-	/** Get the total amount of work this monitor has performed in the current processing operation */
-	int32 GetWorkProgress() const { return WorkProgress; }
-
-private:
-
-	/** The file cache that monitors and reflects the content directory. */
-	FFileCache Cache;
-
-	/** The mounted content path for this monitor (eg /Game/) */
-	FString MountedContentPath;
-
-	/** A list of file system changes that are due to be processed */
-	TArray<FUpdateCacheTransaction> AddedFiles, ModifiedFiles, DeletedFiles;
-
-	/** The number of changes we've processed out of the OutstandingChanges array */
-	int32 TotalWork, WorkProgress;
-
-	/** The last time we attempted to save the cache file */
-	double LastSaveTime;
-
-	/** The interval between potential re-saves of the cache file */
-	static const int32 ResaveIntervalS = 60;
-};
-
 /* Deals with auto reimporting of objects when the objects file on disk is modified*/
 UCLASS(config=Editor, transient)
-class UNREALED_API UAutoReimportManager : public UObject, public FTickableEditorObject
+class UNREALED_API UAutoReimportManager : public UObject
 {
 	GENERATED_UCLASS_BODY()
-
 public:
+	
+	~UAutoReimportManager();
 
 	/** Initialize during engine startup */
 	void Initialize();
 
-private:
-
-	/** FTickableEditorObject interface*/
-	virtual void Tick(float DeltaTime) override;
-	virtual bool IsTickable() const override { return true; }
-	virtual TStatId GetStatId() const override;
-
 private:
 
 	/** UObject Interface */
 	virtual void BeginDestroy() override;
 
-private:
-
-	/** Called when a new asset path has been mounted */
-	void OnContentPathMounted(const FString& InAssetPath, const FString& FileSystemPath);
-
-	/** Called when a new asset path has been dismounted */
-	void OnContentPathDismounted(const FString& InAssetPath, const FString& FileSystemPath);
-
-	/** Called when an asset has been deleted */
-	void OnAssetDeleted(UObject* Object);
-
-private:
-
-	/** Callback for when an editor user setting has changed */
-	void HandleLoadingSavingSettingChanged(FName PropertyName);
-
-	/** Set up monitors to all the mounted content directories */
-	void SetUpMountedContentDirectories();
-
-	/** Set up the user directories from the editor Loading/Saving settings */
-	void SetUpUserDirectories();
-
-private:
-
-	/** Get the total amount of work in the current operation */
-	int32 GetTotalWork() const;
-
-	/** Get the total progress through the current operation */
-	int32 GetTotalProgress() const;
-	
-	/** Get the number of unprocessed changes that are not part of the current processing operation */
-	int32 GetNumUnprocessedChanges() const;
-
-private:
-
-	/** Process any remaining pending additions we have */
-	void ProcessAdditions(const FTimeLimit& TimeLimit);
-
-	/** Process any remaining pending modifications we have */
-	void ProcessModifications(const IAssetRegistry& Registry, const FTimeLimit& TimeLimit);
-
-	/** Process any remaining pending deletions we have */
-	void ProcessDeletions(const IAssetRegistry& Registry);
-
-	/** Enum and value to specify the current state of our processing */
-	enum class ECurrentState
-	{
-		Idle, Initialize, ProcessAdditions, SavePackages, ProcessModifications, ProcessDeletions, Cleanup
-	};
-	ECurrentState CurrentState;
-
-private:
-
-	/** Feedback context that can selectively override the global context to consume progress events for saving of assets */
-	TUniquePtr<FReimportFeedbackContext> FeedbackContextOverride;
-
-	/** Cached string of extensions that are supported by all available factories */
-	FString SupportedExtensions;
-
-	/** Array of directory monitors for automatically managed mounted content paths */
-	TArray<FContentDirectoryMonitor> MountedDirectoryMonitors;
-
-	/** Additional external directories to monitor for changes specified by the user */
-	TArray<FContentDirectoryMonitor> UserDirectoryMonitors;
-
-	/** A list of packages to save when we've added a bunch of assets */
-	TArray<UPackage*> PackagesToSave;
-
-	/** Set when we are processing changes to avoid reentrant logic for asset deletions etc */
-	bool bIsProcessingChanges;
-
-	/** The time when the last change to the cache was reported */
-	FTimeLimit StartProcessingDelay;
-
-	/** The cached number of unprocessed changes we currently have to process */
-	int32 CachedNumUnprocessedChanges;
+	/** Private implementation of the reimport manager */
+	TUniquePtr<class FAutoReimportManager> Implementation;
 };
diff --git a/Engine/Source/Editor/UnrealEd/Public/FileCache.h b/Engine/Source/Editor/UnrealEd/Public/FileCache.h
deleted file mode 100644
