[INTEGRATE] Change 2416912 by [REDACTED_PERSON][REDACTED_HANDLE] on 2015/01/23 11:17:07

	Added the ability for asset type actions to provide an extra overlay layer for their asset thumbnails

	UE-7970 - Add small 'play' button to EditorTutorial asset thumbnail

	Also cleaned up FAssetThumbnail::MakeThumbnailWidget so that it took a config struct (FAssetThumbnailConfig) rather than 7 or 8 defaulted parameters.

diff --git a/Engine/Source/Developer/AssetTools/Public/AssetTypeActions_Base.h b/Engine/Source/Developer/AssetTools/Public/AssetTypeActions_Base.h
--- a/Engine/Source/Developer/AssetTools/Public/AssetTypeActions_Base.h
+++ b/Engine/Source/Developer/AssetTools/Public/AssetTypeActions_Base.h
@@ -82,7 +82,12 @@ public:
 
 	virtual class UThumbnailInfo* GetThumbnailInfo(UObject* Asset) const override
 	{
-		return NULL;
+		return nullptr;
+	}
+
+	virtual TSharedPtr<class SWidget> GetThumbnailOverlay(const FAssetData& AssetData) const override
+	{
+		return nullptr;
 	}
 
 	virtual bool IsImportedAsset() const
diff --git a/Engine/Source/Developer/AssetTools/Public/IAssetTypeActions.h b/Engine/Source/Developer/AssetTools/Public/IAssetTypeActions.h
--- a/Engine/Source/Developer/AssetTools/Public/IAssetTypeActions.h
+++ b/Engine/Source/Developer/AssetTools/Public/IAssetTypeActions.h
@@ -100,6 +100,9 @@ public:
 	/** Returns the thumbnail info for the specified asset, if it has one. */
 	virtual class UThumbnailInfo* GetThumbnailInfo(UObject* Asset) const = 0;
 
+	/** Optionally returns a custom widget to overlay on top of this assets' thumbnail */
+	virtual TSharedPtr<class SWidget> GetThumbnailOverlay(const class FAssetData& AssetData) const = 0;
+
 	/** Returns additional tooltip information for the specified asset, if it has any (otherwise return the null widget) */
 	virtual FText GetAssetDescription(const class FAssetData& AssetData) const = 0;
 
diff --git a/Engine/Source/Editor/ContentBrowser/Private/AssetViewWidgets.cpp b/Engine/Source/Editor/ContentBrowser/Private/AssetViewWidgets.cpp
--- a/Engine/Source/Editor/ContentBrowser/Private/AssetViewWidgets.cpp
+++ b/Engine/Source/Editor/ContentBrowser/Private/AssetViewWidgets.cpp
@@ -983,9 +983,15 @@ void SAssetListItem::Construct( const FArguments& InArgs )
 	TSharedPtr<SWidget> Thumbnail;
 	if ( AssetItem.IsValid() && AssetThumbnail.IsValid() )
 	{
-		const bool bForceGenericThumbnail = (AssetItem->GetType() == EAssetItemType::Creation);
-		const bool bAllowFadeIn = true;
-		Thumbnail = AssetThumbnail->MakeThumbnailWidget(bAllowFadeIn, bForceGenericThumbnail, InArgs._ThumbnailLabel, InArgs._HighlightText, InArgs._ThumbnailHintColorAndOpacity, InArgs._AllowThumbnailHintLabel );
+		FAssetThumbnailConfig ThumbnailConfig;
+		ThumbnailConfig.bAllowFadeIn = true;
+		ThumbnailConfig.bAllowHintText = InArgs._AllowThumbnailHintLabel;
+		ThumbnailConfig.bForceGenericThumbnail = (AssetItem->GetType() == EAssetItemType::Creation);
+		ThumbnailConfig.bAllowAssetSpecificThumbnailOverlay = (AssetItem->GetType() != EAssetItemType::Creation);
+		ThumbnailConfig.ThumbnailLabel = InArgs._ThumbnailLabel;
+		ThumbnailConfig.HighlightedText = InArgs._HighlightText;
+		ThumbnailConfig.HintColorAndOpacity = InArgs._ThumbnailHintColorAndOpacity;
+		Thumbnail = AssetThumbnail->MakeThumbnailWidget(ThumbnailConfig);
 	}
 	else
 	{
@@ -1187,9 +1193,15 @@ void SAssetTileItem::Construct( const FArguments& InArgs )
 	TSharedPtr<SWidget> Thumbnail;
 	if ( AssetItem.IsValid() && AssetThumbnail.IsValid() )
 	{
-		const bool bForceGenericThumbnail = (AssetItem->GetType() == EAssetItemType::Creation);
-		const bool bAllowFadeIn = true;
-		Thumbnail = AssetThumbnail->MakeThumbnailWidget(bAllowFadeIn, bForceGenericThumbnail, InArgs._ThumbnailLabel, InArgs._HighlightText, InArgs._ThumbnailHintColorAndOpacity, InArgs._AllowThumbnailHintLabel);
+		FAssetThumbnailConfig ThumbnailConfig;
+		ThumbnailConfig.bAllowFadeIn = true;
+		ThumbnailConfig.bAllowHintText = InArgs._AllowThumbnailHintLabel;
+		ThumbnailConfig.bForceGenericThumbnail = (AssetItem->GetType() == EAssetItemType::Creation);
+		ThumbnailConfig.bAllowAssetSpecificThumbnailOverlay = (AssetItem->GetType() != EAssetItemType::Creation);
+		ThumbnailConfig.ThumbnailLabel = InArgs._ThumbnailLabel;
+		ThumbnailConfig.HighlightedText = InArgs._HighlightText;
+		ThumbnailConfig.HintColorAndOpacity = InArgs._ThumbnailHintColorAndOpacity;
+		Thumbnail = AssetThumbnail->MakeThumbnailWidget(ThumbnailConfig);
 	}
 	else
 	{
diff --git a/Engine/Source/Editor/DetailCustomizations/Private/WorldSettingsDetails.cpp b/Engine/Source/Editor/DetailCustomizations/Private/WorldSettingsDetails.cpp
--- a/Engine/Source/Editor/DetailCustomizations/Private/WorldSettingsDetails.cpp
+++ b/Engine/Source/Editor/DetailCustomizations/Private/WorldSettingsDetails.cpp
@@ -145,6 +145,9 @@ TSharedRef<ITableRow> FLightmapCustomNodeBuilder::MakeLightMapListViewWidget(TSh
 	UObject* LightMapObject = FindObject<UObject>(NULL, *LightMapItem->ObjectPath);
 	FAssetData LightMapAssetData(LightMapObject);
 
+	FAssetThumbnailConfig ThumbnailConfig;
+	ThumbnailConfig.bAllowFadeIn = true;
+
 	return SNew( STableRow<TSharedPtr<FLightmapItem>>, OwnerTable )
 		.Style(FEditorStyle::Get(), "ContentBrowser.AssetListView.TableRow")
 		[
@@ -164,7 +167,7 @@ TSharedRef<ITableRow> FLightmapCustomNodeBuilder::MakeLightMapListViewWidget(TSh
 					.Padding(ThumbnailBoxPadding)
 					.BorderImage( FEditorStyle::GetBrush("ContentBrowser.ThumbnailShadow") )
 					[
-						LightMapItem->Thumbnail->MakeThumbnailWidget(true)
+						LightMapItem->Thumbnail->MakeThumbnailWidget(ThumbnailConfig)
 					]
 				]
 			]
diff --git a/Engine/Source/Editor/PlacementMode/Private/SPlacementModeTools.cpp b/Engine/Source/Editor/PlacementMode/Private/SPlacementModeTools.cpp
--- a/Engine/Source/Editor/PlacementMode/Private/SPlacementModeTools.cpp
+++ b/Engine/Source/Editor/PlacementMode/Private/SPlacementModeTools.cpp
@@ -58,25 +58,9 @@ public:
 
 		Thumbnail = MakeShareable( new FAssetThumbnail( Asset, InArgs._Width, InArgs._Height, ThumbnailPool ) );
 
-		bool bAllowFadeIn = false;
-		bool bForceGenericThumbnail = false;
-		EThumbnailLabel::Type ThumbnailLabel = EThumbnailLabel::ClassName;
-		const TAttribute< FText >& HighlightedText = TAttribute< FText >( FText::GetEmpty() );
-		const TAttribute< FLinearColor >& HintColorAndOpacity = FLinearColor( 0.0f, 0.0f, 0.0f, 0.0f );
-		bool AllowHintText = true;
-		FName ClassThumbnailBrushOverride = NAME_None;
-
 		ChildSlot
 		[
-			Thumbnail->MakeThumbnailWidget(
-			bAllowFadeIn,
-			bForceGenericThumbnail,
-			ThumbnailLabel,
-			HighlightedText,
-			HintColorAndOpacity,
-			AllowHintText,
-			ClassThumbnailBrushOverride
-			)
+			Thumbnail->MakeThumbnailWidget()
 		];
 	}
 
diff --git a/Engine/Source/Editor/ReferenceViewer/Private/SReferenceNode.cpp b/Engine/Source/Editor/ReferenceViewer/Private/SReferenceNode.cpp
--- a/Engine/Source/Editor/ReferenceViewer/Private/SReferenceNode.cpp
+++ b/Engine/Source/Editor/ReferenceViewer/Private/SReferenceNode.cpp
@@ -59,12 +59,16 @@ void SReferenceNode::UpdateGraphNode()
 	TSharedRef<SWidget> ThumbnailWidget = SNullWidget::NullWidget;
 	if ( AssetThumbnail.IsValid() )
 	{
+		FAssetThumbnailConfig ThumbnailConfig;
+		ThumbnailConfig.bAllowFadeIn = bUsesThumbnail;
+		ThumbnailConfig.bForceGenericThumbnail = !bUsesThumbnail;
+
 		ThumbnailWidget =
 			SNew(SBox)
 			.WidthOverride(AssetThumbnail->GetSize().X)
 			.HeightOverride(AssetThumbnail->GetSize().Y)
 			[
-				AssetThumbnail->MakeThumbnailWidget(/*bAllowFadeIn=*/bUsesThumbnail, /*bForceGenericThumbnail=*/!bUsesThumbnail)
+				AssetThumbnail->MakeThumbnailWidget(ThumbnailConfig)
 			];
 	}
 
diff --git a/Engine/Source/Editor/UnrealEd/Private/AssetThumbnail.cpp b/Engine/Source/Editor/UnrealEd/Private/AssetThumbnail.cpp
--- a/Engine/Source/Editor/UnrealEd/Private/AssetThumbnail.cpp
+++ b/Engine/Source/Editor/UnrealEd/Private/AssetThumbnail.cpp
@@ -18,6 +18,7 @@ public:
 		, _AllowFadeIn(false)
 		, _ForceGenericThumbnail(false)
 		, _AllowHintText(true)
+		, _AllowAssetSpecificThumbnailOverlay(false)
 		, _Label(EThumbnailLabel::ClassName)
 		, _HighlightedText(FText::GetEmpty())
 		, _HintColorAndOpacity(FLinearColor(0.0f, 0.0f, 0.0f, 0.0f))
@@ -30,6 +31,7 @@ public:
 		SLATE_ARGUMENT( bool, AllowFadeIn )
 		SLATE_ARGUMENT( bool, ForceGenericThumbnail )
 		SLATE_ARGUMENT( bool, AllowHintText )
+		SLATE_ARGUMENT( bool, AllowAssetSpecificThumbnailOverlay )
 		SLATE_ARGUMENT( EThumbnailLabel::Type, Label )
 		SLATE_ATTRIBUTE( FText, HighlightedText )
 		SLATE_ATTRIBUTE( FLinearColor, HintColorAndOpacity )
@@ -57,16 +59,16 @@ public:
 
 		UClass* Class = FindObject<UClass>(ANY_PACKAGE, *AssetData.AssetClass.ToString());
 		FAssetToolsModule& AssetToolsModule = FModuleManager::LoadModuleChecked<FAssetToolsModule>(TEXT("AssetTools"));
-		TWeakPtr<IAssetTypeActions> AssetTypeActions;
+		TSharedPtr<IAssetTypeActions> AssetTypeActions;
 		if ( Class != NULL )
 		{
-			AssetTypeActions = AssetToolsModule.Get().GetAssetTypeActionsForClass(Class);
+			AssetTypeActions = AssetToolsModule.Get().GetAssetTypeActionsForClass(Class).Pin();
 		}
 
 		AssetColor = FLinearColor::White;
 		if ( AssetTypeActions.IsValid() )
 		{
-			AssetColor = AssetTypeActions.Pin()->GetTypeColor();
+			AssetColor = AssetTypeActions->GetTypeColor();
 		}
 
 		TSharedRef<SOverlay> OverlayWidget = SNew(SOverlay);
@@ -194,6 +196,19 @@ public:
 			.Padding(this, &SAssetThumbnail::GetAssetColorStripPadding)
 		];
 
+		if( InArgs._AllowAssetSpecificThumbnailOverlay )
+		{
+			// Does the asset provide an additional thumbnail overlay?
+			TSharedPtr<SWidget> AssetSpecificThumbnailOverlay = AssetTypeActions->GetThumbnailOverlay(AssetData);
+			if( AssetSpecificThumbnailOverlay.IsValid() )
+			{
+				OverlayWidget->AddSlot()
+				[
+					AssetSpecificThumbnailOverlay.ToSharedRef()
+				];
+			}
+		}
+
 		ChildSlot
 		[
 			OverlayWidget
@@ -745,27 +760,20 @@ void FAssetThumbnail::SetAsset( const FAssetData& InAssetData )
 	AssetDataChangedEvent.Broadcast();
 }
 
-TSharedRef<SWidget> FAssetThumbnail::MakeThumbnailWidget(
-	bool bAllowFadeIn,
-	bool bForceGenericThumbnail,
-	EThumbnailLabel::Type ThumbnailLabel,
-	const TAttribute< FText >& HighlightedText,
-	const TAttribute< FLinearColor >& HintColorAndOpacity,
-	bool AllowHintText,
-	FName ClassThumbnailBrushOverride
-	)
+TSharedRef<SWidget> FAssetThumbnail::MakeThumbnailWidget( const FAssetThumbnailConfig& InConfig )
 {
 	return
 		SNew(SAssetThumbnail)
 		.AssetThumbnail( SharedThis(this) )
-		.ThumbnailPool(ThumbnailPool.Pin())
-		.AllowFadeIn(bAllowFadeIn)
-		.ForceGenericThumbnail(bForceGenericThumbnail)
-		.Label( ThumbnailLabel )
-		.HighlightedText( HighlightedText )
-		.HintColorAndOpacity( HintColorAndOpacity )
-		.AllowHintText( AllowHintText )
-		.ClassThumbnailBrushOverride( ClassThumbnailBrushOverride );
+		.ThumbnailPool( ThumbnailPool.Pin() )
+		.AllowFadeIn( InConfig.bAllowFadeIn )
+		.ForceGenericThumbnail( InConfig.bForceGenericThumbnail )
+		.Label( InConfig.ThumbnailLabel )
+		.HighlightedText( InConfig.HighlightedText )
+		.HintColorAndOpacity( InConfig.HintColorAndOpacity )
+		.AllowHintText( InConfig.bAllowHintText )
+		.ClassThumbnailBrushOverride( InConfig.ClassThumbnailBrushOverride )
+		.AllowAssetSpecificThumbnailOverlay( InConfig.bAllowAssetSpecificThumbnailOverlay );
 }
 
 void FAssetThumbnail::RefreshThumbnail()
diff --git a/Engine/Source/Editor/UnrealEd/Public/AssetThumbnail.h b/Engine/Source/Editor/UnrealEd/Public/AssetThumbnail.h
--- a/Engine/Source/Editor/UnrealEd/Public/AssetThumbnail.h
+++ b/Engine/Source/Editor/UnrealEd/Public/AssetThumbnail.h
@@ -17,6 +17,31 @@ namespace EThumbnailLabel
 	};
 };
 
+/** A struct containing details about how the asset thumbnail should behave */
+struct FAssetThumbnailConfig
+{
+	FAssetThumbnailConfig()
+		: bAllowFadeIn( false )
+		, bForceGenericThumbnail( false )
+		, bAllowHintText( true )
+		, bAllowAssetSpecificThumbnailOverlay( false )
+		, ClassThumbnailBrushOverride( NAME_None )
+		, ThumbnailLabel( EThumbnailLabel::ClassName )
+		, HighlightedText( FText::GetEmpty() )
+		, HintColorAndOpacity( FLinearColor( 0.0f, 0.0f, 0.0f, 0.0f ) )
+	{
+	}
+
+	bool bAllowFadeIn;
+	bool bForceGenericThumbnail;
+	bool bAllowHintText;
+	bool bAllowAssetSpecificThumbnailOverlay;
+	FName ClassThumbnailBrushOverride;
+	EThumbnailLabel::Type ThumbnailLabel;
+	TAttribute< FText > HighlightedText;
+	TAttribute< FLinearColor > HintColorAndOpacity;
+};
+
 /**
  * Interface for rendering a thumbnail in a slate viewport                   
  */
@@ -75,15 +100,7 @@ public:
 	/**
 	 * @return A slate widget representing this thumbnail
 	 */
-	UNREALED_API TSharedRef<SWidget> MakeThumbnailWidget(
-		bool bAllowFadeIn = false,
-		bool bForceGenericThumbnail = false,
-		EThumbnailLabel::Type ThumbnailLabel = EThumbnailLabel::ClassName,
-		const TAttribute< FText >& HighlightedText = TAttribute< FText >( FText::GetEmpty() ),
-		const TAttribute< FLinearColor >& HintColorAndOpacity = FLinearColor( 0.0f, 0.0f, 0.0f, 0.0f ),
-		bool AllowHintText = true,
-		FName ClassThumbnailBrushOverride = NAME_None
-		);
+	UNREALED_API TSharedRef<SWidget> MakeThumbnailWidget( const FAssetThumbnailConfig& InConfig = FAssetThumbnailConfig() );
 
 	/** Re-renders this thumbnail */
 	UNREALED_API void RefreshThumbnail();
