[INTEGRATE] Change 2416926 by [REDACTED_PERSON][REDACTED_HANDLE] on 2015/01/23 11:26:26

	[UE-8089] - Delta is only applied to the parent-most selected component(s) (i.e. if you select a parent and a child, the delta will only be applied to the parent)

diff --git a/Engine/Source/Editor/UnrealEd/Private/LevelEditorViewport.cpp b/Engine/Source/Editor/UnrealEd/Private/LevelEditorViewport.cpp
--- a/Engine/Source/Editor/UnrealEd/Private/LevelEditorViewport.cpp
+++ b/Engine/Source/Editor/UnrealEd/Private/LevelEditorViewport.cpp
@@ -2984,13 +2984,44 @@ void FLevelEditorViewportClient::ApplyDeltaToActors(const FVector& InDrag,
 		{
 			if (GEditor->GetSelectedComponentCount() > 0)
 			{
+				auto ComponentSelection = GEditor->GetSelectedComponents();
+
+				// Only move the parent-most component(s) that are selected 
+				// Otherwise, if both a parent and child are selected and the delta is applied to both, the child will actually move 2x delta
+				TInlineComponentArray<USceneComponent*> ComponentsToMove;
 				for (FSelectionIterator It(GEditor->GetSelectedComponentIterator()); It; ++It)
 				{
 					USceneComponent* SceneComponent = CastChecked<USceneComponent>(*It);
 					if (SceneComponent)
 					{
-						ApplyDeltaToComponent(SceneComponent, InDrag, InRot, ModifiedScale);
-					}
+						auto SelectedComponent = Cast<USceneComponent>(*It);
+
+						// Check to see if any parent is selected
+						bool bParentAlsoSelected = false;
+						USceneComponent* Parent = SelectedComponent->GetAttachParent();
+						while (Parent != nullptr)
+						{
+							if (ComponentSelection->IsSelected(Parent))
+							{
+								bParentAlsoSelected = true;
+								break;
+							}
+
+							Parent = Parent->GetAttachParent();
+						}
+
+						// If no parent of this component is also in the selection set, move it!
+						if (!bParentAlsoSelected)
+						{
+							ComponentsToMove.Add(SelectedComponent);
+						}
+					}	
+				}
+
+				// Now actually apply the delta to the appropriate component(s)
+				for (auto SceneComp : ComponentsToMove)
+				{
+					ApplyDeltaToComponent(SceneComp, InDrag, InRot, ModifiedScale);
 				}
 			}
 			else
