[INTEGRATE] Change 2417028 by [REDACTED_PERSON][REDACTED_HANDLE] on 2015/01/23 13:21:10

	[UE-7993] - Component name now matches the asset name when added via drag-drop
	[UE-7994] - newly added component in BP editor or in-world now defaults to naming without a numeric suffix if there isn't a conflict

diff --git a/Engine/Source/Editor/Kismet/Private/SSCSEditor.cpp b/Engine/Source/Editor/Kismet/Private/SSCSEditor.cpp
--- a/Engine/Source/Editor/Kismet/Private/SSCSEditor.cpp
+++ b/Engine/Source/Editor/Kismet/Private/SSCSEditor.cpp
@@ -3722,7 +3722,15 @@ UActorComponent* SSCSEditor::AddNewComponent( UClass* NewComponentClass, UObject
 		ActorInstance->Modify();
 
 		// Create new component
-		FName NewComponentName(*FComponentEditorUtils::GenerateValidVariableName(NewComponentClass, ActorInstance));
+		FName NewComponentName = NAME_None;
+		if (Asset != nullptr)
+		{
+			NewComponentName = *FComponentEditorUtils::GenerateValidVariableNameFromAsset(Asset, ActorInstance);
+		}
+		else
+		{
+			NewComponentName = *FComponentEditorUtils::GenerateValidVariableName(NewComponentClass, ActorInstance);
+		}
 
 		UActorComponent* NewComponentInstance = ConstructObject<UActorComponent>(NewComponentClass, ActorInstance, NewComponentName, RF_Transactional);
 		check(NewComponentInstance != nullptr);
diff --git a/Engine/Source/Editor/UnrealEd/Private/Kismet2/ComponentEditorUtils.cpp b/Engine/Source/Editor/UnrealEd/Private/Kismet2/ComponentEditorUtils.cpp
--- a/Engine/Source/Editor/UnrealEd/Private/Kismet2/ComponentEditorUtils.cpp
+++ b/Engine/Source/Editor/UnrealEd/Private/Kismet2/ComponentEditorUtils.cpp
@@ -135,20 +135,31 @@ FString FComponentEditorUtils::GenerateValidVariableName(TSubclassOf<UActorCompo
 
 	int32 Counter = 1;
 	FString ComponentTypeName = *ComponentClass->GetName().Replace(TEXT("Component"), TEXT(""));
-	FString ComponentInstanceName;
-
-	bool bNameIsValid = false;
-	while (!bNameIsValid)
+	
+	// Try to create a name without any numerical suffix first
+	FString ComponentInstanceName = ComponentTypeName;
+	while (!IsComponentNameAvailable(ComponentInstanceName, ComponentOwner))
 	{
-		bNameIsValid = true;
-
 		// Assign the lowest possible numerical suffix
 		ComponentInstanceName = FString::Printf(TEXT("%s%d"), *ComponentTypeName, Counter++);
+	}
 
-		if (!IsComponentNameAvailable(ComponentInstanceName, ComponentOwner))
-		{
-			bNameIsValid = false;
-		}
+	return ComponentInstanceName;
+}
+
+FString FComponentEditorUtils::GenerateValidVariableNameFromAsset(UObject* Asset, AActor* ComponentOwner)
+{
+	check(ComponentOwner);
+
+	int32 Counter = 1;
+	FString AssetName = Asset->GetName();
+
+	// Try to create a name without any numerical suffix first
+	FString ComponentInstanceName = AssetName;
+	while (!IsComponentNameAvailable(ComponentInstanceName, ComponentOwner))
+	{
+		// Assign the lowest possible numerical suffix
+		ComponentInstanceName = FString::Printf(TEXT("%s%d"), *AssetName, Counter++);
 	}
 
 	return ComponentInstanceName;
diff --git a/Engine/Source/Editor/UnrealEd/Public/Kismet2/ComponentEditorUtils.h b/Engine/Source/Editor/UnrealEd/Public/Kismet2/ComponentEditorUtils.h
--- a/Engine/Source/Editor/UnrealEd/Public/Kismet2/ComponentEditorUtils.h
+++ b/Engine/Source/Editor/UnrealEd/Public/Kismet2/ComponentEditorUtils.h
@@ -26,6 +26,9 @@ public:
 	/** Generate a valid variable name string for the given component instance */
 	static FString GenerateValidVariableName(TSubclassOf<UActorComponent> InComponentClass, AActor* ComponentOwner);
 
+	/** Generate a valid variable name string for the given component instance based on the name of the asset referenced by the component */
+	static FString GenerateValidVariableNameFromAsset(UObject* Asset, AActor* ComponentOwner);
+
 	struct FTransformData
 	{
 		FTransformData(const class USceneComponent& Component)
diff --git a/Engine/Source/Runtime/Engine/Private/SCS_Node.cpp b/Engine/Source/Runtime/Engine/Private/SCS_Node.cpp
--- a/Engine/Source/Runtime/Engine/Private/SCS_Node.cpp
+++ b/Engine/Source/Runtime/Engine/Private/SCS_Node.cpp
@@ -450,12 +450,12 @@ FName USCS_Node::GenerateNewComponentName( TArray<FName>& CurrentNames, FName De
 				ComponentName = ComponentTemplate->GetClass()->GetName().Replace(TEXT("Component"), TEXT(""));
 			}
 			
+			NewName = *ComponentName;
 			int32 Counter = 1;
-			do
+			while (CurrentNames.Contains(NewName))
 			{
 				NewName = FName(*( FString::Printf(TEXT("%s%d"), *ComponentName, Counter++) ));		
 			} 
-			while( CurrentNames.Contains(NewName) );
 		}
 	}
 	return NewName;
