[INTEGRATE] Change 2417166 by [REDACTED_PERSON][REDACTED_HANDLE] on 2015/01/23 15:13:46

	Fixup attachment issues across object reinstancing

diff --git a/Engine/Source/Editor/UnrealEd/Private/Kismet2/KismetReinstanceUtilities.cpp b/Engine/Source/Editor/UnrealEd/Private/Kismet2/KismetReinstanceUtilities.cpp
--- a/Engine/Source/Editor/UnrealEd/Private/Kismet2/KismetReinstanceUtilities.cpp
+++ b/Engine/Source/Editor/UnrealEd/Private/Kismet2/KismetReinstanceUtilities.cpp
@@ -37,7 +37,7 @@ struct FReplaceReferenceHelper
 		}
 	}
 
-	static void FindAndReplaceReferences(TArray<UObject*>& SourceObjects, TSet<UObject*>* ObjectsThatShouldUseOldStuff, TArray<UObject*> &ObjectsToReplace, TMap<UObject*, UObject*>& OldToNewInstanceMap, TMap<FStringAssetReference, UObject*>& ReinstancedObjectsWeakReferenceMap)
+	static void FindAndReplaceReferences(const TArray<UObject*>& SourceObjects, TSet<UObject*>* ObjectsThatShouldUseOldStuff, const TArray<UObject*>& ObjectsToReplace, const TMap<UObject*, UObject*>& OldToNewInstanceMap, const TMap<FStringAssetReference, UObject*>& ReinstancedObjectsWeakReferenceMap)
 	{
 		// Find everything that references these objects
 		TSet<UObject *> Targets;
@@ -420,7 +420,7 @@ struct FActorReplacementHelper
 	 * Runs construction scripts on the new actor and then finishes it off by
 	 * attaching it to the same attachments that its predecessor was set with. 
 	 */
-	void Finalize();
+	void Finalize(const TMap<UObject*, UObject*>& OldToNewInstanceMap);
 
 
 private:
@@ -476,7 +476,7 @@ private:
 	TSharedPtr<AActor::FActorTransactionAnnotation> CachedActorData;
 };
 
-void FActorReplacementHelper::Finalize()
+void FActorReplacementHelper::Finalize(const TMap<UObject*, UObject*>& OldToNewInstanceMap)
 {
 	
 	// because this is an editor context it's important to use this execution guard
@@ -485,6 +485,7 @@ void FActorReplacementHelper::Finalize()
 	// run the construction script, which will use the properties we just copied over
 	if (CachedActorData.IsValid())
 	{
+		CachedActorData->ComponentInstanceData.FindAndReplaceInstances(OldToNewInstanceMap);
 		NewActor->ExecuteConstruction(TargetWorldTransform, &CachedActorData->ComponentInstanceData);
 	}
 	else
@@ -492,6 +493,24 @@ void FActorReplacementHelper::Finalize()
 		FComponentInstanceDataCache DummyComponentData;
 		NewActor->ExecuteConstruction(TargetWorldTransform, &DummyComponentData);
 	}	
+
+	if (TargetAttachParent)
+	{
+		UObject* const* NewTargetAttachParent = OldToNewInstanceMap.Find(TargetAttachParent);
+		if (NewTargetAttachParent)
+		{
+			TargetAttachParent = CastChecked<AActor>(*NewTargetAttachParent);
+		}
+	}
+	if (TargetParentComponent)
+	{
+		UObject* const* NewTargetParentComponent = OldToNewInstanceMap.Find(TargetParentComponent);
+		if (NewTargetParentComponent)
+		{
+			TargetParentComponent = CastChecked<USceneComponent>(*NewTargetParentComponent);
+		}
+	}
+
 	ApplyAttachments();
 
 	if (bSelectNewActor)
@@ -620,8 +639,21 @@ void FBlueprintCompileReinstancer::ReplaceInstancesOfClass(UClass* OldClass, UCl
 	// Set global flag to let system know we are reconstructing blueprint instances
 	TGuardValue<bool> GuardTemplateNameFlag(GIsReconstructingBlueprintInstances, true);
 
+	struct FObjectRemappingHelper
+	{
+		void OnObjectsReplaced(const TMap<UObject*, UObject*>& InReplacedObjects)
+		{
+			ReplacedObjects.Append(InReplacedObjects);
+		}
+
+		TMap<UObject*, UObject*> ReplacedObjects;
+	} ObjectRemappingHelper;
+
+	FDelegateHandle OnObjectsReplacedHandle = GEditor->OnObjectsReplaced().AddRaw(&ObjectRemappingHelper,&FObjectRemappingHelper::OnObjectsReplaced);
+
 	{ BP_SCOPED_COMPILER_EVENT_STAT(EKismetReinstancerStats_ReplaceInstancesOfClass);
 
+
 		const bool bIncludeDerivedClasses = false;
 		GetObjectsOfClass(OldClass, ObjectsToReplace, bIncludeDerivedClasses);
 	
@@ -766,6 +798,8 @@ void FBlueprintCompileReinstancer::ReplaceInstancesOfClass(UClass* OldClass, UCl
 		}
 	}
 
+	GEditor->OnObjectsReplaced().Remove(OnObjectsReplacedHandle);
+
 	// Now replace any pointers to the old archetypes/instances with pointers to the new one
 	TArray<UObject*> SourceObjects;
 	TArray<UObject*> DstObjects;
@@ -794,7 +828,7 @@ void FBlueprintCompileReinstancer::ReplaceInstancesOfClass(UClass* OldClass, UCl
 		// FArchiveReplaceObjectRef to run construction-scripts).
 		for (FActorReplacementHelper& ReplacementActor : ReplacementActors)
 		{
-			ReplacementActor.Finalize();
+			ReplacementActor.Finalize(ObjectRemappingHelper.ReplacedObjects);
 		}
 	}
 
diff --git a/Engine/Source/Runtime/Engine/Classes/Components/SceneComponent.h b/Engine/Source/Runtime/Engine/Classes/Components/SceneComponent.h
--- a/Engine/Source/Runtime/Engine/Classes/Components/SceneComponent.h
+++ b/Engine/Source/Runtime/Engine/Classes/Components/SceneComponent.h
@@ -800,6 +800,7 @@ public:
 	{}
 
 	virtual void ApplyToComponent(UActorComponent* Component) override;
+	virtual void FindAndReplaceInstances(const TMap<UObject*, UObject*>& OldToNewInstanceMap) override;
 
 	TArray<USceneComponent*> AttachedInstanceComponents;
 };
diff --git a/Engine/Source/Runtime/Engine/Private/ComponentInstanceDataCache.cpp b/Engine/Source/Runtime/Engine/Private/ComponentInstanceDataCache.cpp
--- a/Engine/Source/Runtime/Engine/Private/ComponentInstanceDataCache.cpp
+++ b/Engine/Source/Runtime/Engine/Private/ComponentInstanceDataCache.cpp
@@ -129,4 +129,15 @@ void FComponentInstanceDataCache::ApplyToActor(AActor* Actor) const
 			}
 		}
 	}
+}
+
+void FComponentInstanceDataCache::FindAndReplaceInstances(const TMap<UObject*, UObject*>& OldToNewInstanceMap) const
+{
+	for (auto ComponentInstanceDataPair : TypeToDataMap)
+	{
+		if (ComponentInstanceDataPair.Value)
+		{
+			ComponentInstanceDataPair.Value->FindAndReplaceInstances(OldToNewInstanceMap);
+		}
+	}
 }
\ No newline at end of file
diff --git a/Engine/Source/Runtime/Engine/Private/SceneComponent.cpp b/Engine/Source/Runtime/Engine/Private/SceneComponent.cpp
--- a/Engine/Source/Runtime/Engine/Private/SceneComponent.cpp
+++ b/Engine/Source/Runtime/Engine/Private/SceneComponent.cpp
@@ -1029,6 +1029,17 @@ void FSceneComponentInstanceData::ApplyToComponent(UActorComponent* Component)
 	}
 }
 
+void FSceneComponentInstanceData::FindAndReplaceInstances(const TMap<UObject*, UObject*>& OldToNewInstanceMap)
+{
+	for (USceneComponent*& ChildComponent : AttachedInstanceComponents)
+	{
+		if (UObject* const* NewChildComponent = OldToNewInstanceMap.Find(ChildComponent))
+		{
+			ChildComponent = CastChecked<USceneComponent>(*NewChildComponent);
+		}
+	}
+}
+
 FComponentInstanceDataBase* USceneComponent::GetComponentInstanceData() const
 {
 	FComponentInstanceDataBase* InstanceData = nullptr;
diff --git a/Engine/Source/Runtime/Engine/Public/ComponentInstanceDataCache.h b/Engine/Source/Runtime/Engine/Public/ComponentInstanceDataCache.h
--- a/Engine/Source/Runtime/Engine/Public/ComponentInstanceDataCache.h
+++ b/Engine/Source/Runtime/Engine/Public/ComponentInstanceDataCache.h
@@ -18,10 +18,15 @@ public:
 	virtual ~FComponentInstanceDataBase()
 	{}
 
+	/** Determines whether this component instance data matches the component */
 	bool MatchesComponent(const UActorComponent* Component) const;
 
+	/** Applies this component instance data to the supplied component */
 	virtual void ApplyToComponent(UActorComponent* Component) = 0;
 
+	/** Replaces any references to old instances during Actor reinstancing */
+	virtual void FindAndReplaceInstances(const TMap<UObject*, UObject*>& OldToNewInstanceMap) { };
+
 protected:
 	/** The name of the source component */
 	FName SourceComponentName;
@@ -45,9 +50,12 @@ public:
 
 	~FComponentInstanceDataCache();
 
-	/** Util to iterate over components and apply data to each */
+	/** Iterates over an Actor's components and applies the stored component instance data to each */
 	void ApplyToActor(AActor* Actor) const;
 
+	/** Iterates over components and replaces any object references with the reinstanced information */
+	void FindAndReplaceInstances(const TMap<UObject*, UObject*>& OldToNewInstanceMap) const;
+
 	bool HasInstanceData() const { return TypeToDataMap.Num() > 0; }
 
 private:
