[INTEGRATE] Change 2417426 by [REDACTED_PERSON][REDACTED_HANDLE] on 2015/01/23 17:58:27

	Don't duplicate InstanceComponents in to the CDO on blueprint promotion

diff --git a/Engine/Source/Editor/UnrealEd/Private/Kismet2/Kismet2.cpp b/Engine/Source/Editor/UnrealEd/Private/Kismet2/Kismet2.cpp
--- a/Engine/Source/Editor/UnrealEd/Private/Kismet2/Kismet2.cpp
+++ b/Engine/Source/Editor/UnrealEd/Private/Kismet2/Kismet2.cpp
@@ -1034,8 +1034,19 @@ UBlueprint* FKismetEditorUtilities::CreateBlueprintFromActor(const FName Bluepri
 
 			if(NewBlueprint->GeneratedClass)
 			{
+				// Since we already created SCS Nodes for the instance components, temporarily cache and clear the
+				// array to avoid creating duplicates in the new CDO
+				const TArray<UActorComponent*> TempInstanceComponents(Actor->GetInstanceComponents());
+				Actor->ClearInstanceComponents();
+
 				UObject* CDO = NewBlueprint->GeneratedClass->GetDefaultObject();
 				UEditorEngine::CopyPropertiesForUnrelatedObjects(Actor, CDO);
+
+				for (UActorComponent* Component : TempInstanceComponents)
+				{
+					Actor->AddInstanceComponent(Component);
+				}
+
 				if(AActor* CDOAsActor = Cast<AActor>(CDO))
 				{
 					if(USceneComponent* Scene = CDOAsActor->GetRootComponent())
@@ -1049,9 +1060,6 @@ UBlueprint* FKismetEditorUtilities::CreateBlueprintFromActor(const FName Bluepri
 						// Ensure the light mass information is cleaned up
 						Scene->InvalidateLightingCache();
 					}
-
-					// Clear the instance properties array as we created SCS nodes for them
-					CDOAsActor->ClearInstanceComponents();
 				}
 			}
 
