[INTEGRATE] Change 2417780 by [REDACTED_PERSON][REDACTED_HANDLE] on 2015/01/24 10:37:50

	Fix crash replacing actors when there is no level script blueprint created yet.

diff --git a/Engine/Source/Editor/UnrealEd/Private/Editor.cpp b/Engine/Source/Editor/UnrealEd/Private/Editor.cpp
--- a/Engine/Source/Editor/UnrealEd/Private/Editor.cpp
+++ b/Engine/Source/Editor/UnrealEd/Private/Editor.cpp
@@ -5172,8 +5172,13 @@ void UEditorEngine::ReplaceActors(UActorFactory* Factory, const FAssetData& Asse
 			NewActor->MarkPackageDirty();
 
 			// Replace references in the level script Blueprint with the new Actor
-			ULevelScriptBlueprint* LSB = NewActor->GetLevel()->GetLevelScriptBlueprint(true);
-			FBlueprintEditorUtils::ReplaceAllActorRefrences(LSB, OldActor, NewActor);
+			const bool bDontCreate = true;
+			ULevelScriptBlueprint* LSB = NewActor->GetLevel()->GetLevelScriptBlueprint(bDontCreate);
+			if( LSB )
+			{
+				// Only if the level script blueprint exists would there be references.  
+				FBlueprintEditorUtils::ReplaceAllActorRefrences(LSB, OldActor, NewActor);
+			}
 
 			GEditor->Layers->DisassociateActorFromLayers( OldActor );
 			World->EditorDestroyActor(OldActor, true);
