[INTEGRATE] Change 2417795 by [REDACTED_PERSON][REDACTED_HANDLE] on 2015/01/24 11:02:29

	Moved GC leak check when destroying an editor world to a new method to generate better callstacks.

diff --git a/Engine/Source/Editor/UnrealEd/Classes/Editor/EditorEngine.h b/Engine/Source/Editor/UnrealEd/Classes/Editor/EditorEngine.h
--- a/Engine/Source/Editor/UnrealEd/Classes/Editor/EditorEngine.h
+++ b/Engine/Source/Editor/UnrealEd/Classes/Editor/EditorEngine.h
@@ -2507,6 +2507,8 @@ public:
 	 void FlagModifyAllSelectedSurfacesInLevels( UWorld* InWorld );
 	 
 private:
+	/** Checks for UWorld garbage collection leaks and reports any that are found */
+	void CheckForWorldGCLeaks( UWorld* NewWorld, UPackage* WorldPackage );
 
 	/**
 	 * This destroys the given world.
diff --git a/Engine/Source/Editor/UnrealEd/Private/EditorServer.cpp b/Engine/Source/Editor/UnrealEd/Private/EditorServer.cpp
--- a/Engine/Source/Editor/UnrealEd/Private/EditorServer.cpp
+++ b/Engine/Source/Editor/UnrealEd/Private/EditorServer.cpp
@@ -1692,6 +1692,46 @@ void UEditorEngine::BSPIntersectionHelper(UWorld* InWorld, ECsgOper Operation)
 	}
 }
 
+void UEditorEngine::CheckForWorldGCLeaks( UWorld* NewWorld, UPackage* WorldPackage )
+{
+	// Make sure the old world is completely gone, except if the new world was one of it's sublevels
+	for(TObjectIterator<UWorld> It; It; ++It)
+	{
+		UWorld* RemainingWorld = *It;
+		const bool bIsNewWorld = (NewWorld && RemainingWorld == NewWorld);
+		const bool bIsPersistantWorldType = (RemainingWorld->WorldType == EWorldType::Inactive) || (RemainingWorld->WorldType == EWorldType::Preview);
+		if(!bIsNewWorld && !bIsPersistantWorldType && !WorldHasValidContext(RemainingWorld))
+		{
+			StaticExec(RemainingWorld, *FString::Printf(TEXT("OBJ REFS CLASS=WORLD NAME=%s"), *RemainingWorld->GetPathName()));
+
+			TMap<UObject*, UProperty*>	Route		= FArchiveTraceRoute::FindShortestRootPath(RemainingWorld, true, GARBAGE_COLLECTION_KEEPFLAGS);
+			FString						ErrorString	= FArchiveTraceRoute::PrintRootPath(Route, RemainingWorld);
+
+			UE_LOG(LogEditorServer, Fatal, TEXT("%s still around trying to load %s") LINE_TERMINATOR TEXT("%s"), *RemainingWorld->GetPathName(), TempFname, *ErrorString);
+		}
+	}
+
+
+	if(WorldPackage != NULL)
+	{
+		UPackage* NewWorldPackage = NewWorld ? NewWorld->GetOutermost() : nullptr;
+		for(TObjectIterator<UPackage> It; It; ++It)
+		{
+			UPackage* RemainingPackage = *It;
+			const bool bIsNewWorldPackage = (NewWorldPackage && RemainingPackage == NewWorldPackage);
+			if(!bIsNewWorldPackage && RemainingPackage == WorldPackage)
+			{
+				StaticExec(NULL, *FString::Printf(TEXT("OBJ REFS CLASS=PACKAGE NAME=%s"), *RemainingPackage->GetPathName()));
+
+				TMap<UObject*, UProperty*>	Route		= FArchiveTraceRoute::FindShortestRootPath(RemainingPackage, true, GARBAGE_COLLECTION_KEEPFLAGS);
+				FString						ErrorString	= FArchiveTraceRoute::PrintRootPath(Route, RemainingPackage);
+
+				UE_LOG(LogEditorServer, Fatal, TEXT("%s still around trying to load %s") LINE_TERMINATOR TEXT("%s"), *RemainingPackage->GetPathName(), TempFname, *ErrorString);
+			}
+		}
+	}
+}
+
 void UEditorEngine::EditorDestroyWorld( FWorldContext & Context, const FText& CleanseText, UWorld* NewWorld )
 {
 	UWorld* ContextWorld = Context.World();
@@ -1785,41 +1825,8 @@ void UEditorEngine::EditorDestroyWorld( FWorldContext & Context, const FText& Cl
 		NewWorld->RemoveFromRoot();
 	}
 
-	// Make sure the old world is completely gone, except if the new world was one of it's sublevels
-	for( TObjectIterator<UWorld> It; It; ++It )
-	{
-		UWorld* RemainingWorld = *It;
-		const bool bIsNewWorld = (NewWorld && RemainingWorld == NewWorld);
-		const bool bIsPersistantWorldType = (RemainingWorld->WorldType == EWorldType::Inactive) || (RemainingWorld->WorldType == EWorldType::Preview);
-		if (!bIsNewWorld && !bIsPersistantWorldType && !WorldHasValidContext(RemainingWorld))
-		{
-			StaticExec(RemainingWorld, *FString::Printf(TEXT("OBJ REFS CLASS=WORLD NAME=%s"), *RemainingWorld->GetPathName()));
-
-			TMap<UObject*,UProperty*>	Route		= FArchiveTraceRoute::FindShortestRootPath( RemainingWorld, true, GARBAGE_COLLECTION_KEEPFLAGS );
-			FString						ErrorString	= FArchiveTraceRoute::PrintRootPath( Route, RemainingWorld );
-
-			UE_LOG(LogEditorServer, Fatal,TEXT("%s still around trying to load %s") LINE_TERMINATOR TEXT("%s"),*RemainingWorld->GetPathName(),TempFname,*ErrorString);
-		}
-	}
+	CheckForWorldGCLeaks( NewWorld, WorldPackage );
 
-	if (WorldPackage != NULL)
-	{
-		UPackage* NewWorldPackage = NewWorld ? NewWorld->GetOutermost() : nullptr;
-		for( TObjectIterator<UPackage> It; It; ++It )
-		{
-			UPackage* RemainingPackage = *It;
-			const bool bIsNewWorldPackage = (NewWorldPackage && RemainingPackage == NewWorldPackage);
-			if (!bIsNewWorldPackage && RemainingPackage == WorldPackage)
-			{
-				StaticExec(NULL, *FString::Printf(TEXT("OBJ REFS CLASS=PACKAGE NAME=%s"), *RemainingPackage->GetPathName()));
-
-				TMap<UObject*,UProperty*>	Route		= FArchiveTraceRoute::FindShortestRootPath( RemainingPackage, true, GARBAGE_COLLECTION_KEEPFLAGS );
-				FString						ErrorString	= FArchiveTraceRoute::PrintRootPath( Route, RemainingPackage );
-
-				UE_LOG(LogEditorServer, Fatal,TEXT("%s still around trying to load %s") LINE_TERMINATOR TEXT("%s"),*RemainingPackage->GetPathName(),TempFname,*ErrorString);
-			}
-		}
-	}
 }
 
 bool UEditorEngine::ShouldAbortBecauseOfPIEWorld() const
