[INTEGRATE] Change 2417916 by [REDACTED_PERSON][REDACTED_HANDLE] on 2015/01/24 15:14:49

	Fix applying added instance components when they are attached to a node defined in the current blueprint

diff --git a/Engine/Source/Editor/UnrealEd/Private/Kismet2/Kismet2.cpp b/Engine/Source/Editor/UnrealEd/Private/Kismet2/Kismet2.cpp
--- a/Engine/Source/Editor/UnrealEd/Private/Kismet2/Kismet2.cpp
+++ b/Engine/Source/Editor/UnrealEd/Private/Kismet2/Kismet2.cpp
@@ -972,12 +972,15 @@ void FKismetEditorUtilities::AddComponentsToBlueprint(UBlueprint* Blueprint, con
 					if (ParentSCSNode->GetSCS() != SCS)
 					{
 						SCS->AddNode(SCSNode);
+						SCSNode->SetParent(ParentSCSNode);
+					}
+					else
+					{
+						ParentSCSNode->AddChildNode(SCSNode);
 					}
-
-					SCSNode->SetParent(ParentSCSNode);
 				}
 				// If we're attached to a native component
-				else if (!Components.Contains(SceneComponent->AttachParent))
+				else if (SceneComponent->AttachParent->CreationMethod == EComponentCreationMethod::Native)
 				{
 					SCS->AddNode(SCSNode);
 					SCSNode->SetParent(SceneComponent->AttachParent);
