[INTEGRATE] Change 2417935 by [REDACTED_PERSON][REDACTED_HANDLE] on 2015/01/24 16:23:14

	Added an experimental method to add a new C++ class derived from ActorComponent to an actor.  The purpose of this is to add scripting support directly on an actor without subclassing the actor.

diff --git a/Engine/Source/Editor/GameProjectGeneration/Private/GameProjectGenerationModule.cpp b/Engine/Source/Editor/GameProjectGeneration/Private/GameProjectGenerationModule.cpp
--- a/Engine/Source/Editor/GameProjectGeneration/Private/GameProjectGenerationModule.cpp
+++ b/Engine/Source/Editor/GameProjectGeneration/Private/GameProjectGenerationModule.cpp
@@ -61,9 +61,9 @@ void FGameProjectGenerationModule::OpenAddCodeToProjectDialog(const TSharedPtr<S
 	AddCodeToProjectDialogOpenedEvent.Broadcast();
 }
 
-void FGameProjectGenerationModule::OpenAddCodeToProjectDialog(const UClass* InClass, const FString& InInitialPath, const TSharedPtr<SWindow>& InParentWindow)
+void FGameProjectGenerationModule::OpenAddCodeToProjectDialog(const UClass* InClass, const FString& InInitialPath, const TSharedPtr<SWindow>& InParentWindow, bool bModal, FOnCodeAddedToProject OnCodeAddedToProject)
 {
-	GameProjectUtils::OpenAddCodeToProjectDialog(InClass, InInitialPath, InParentWindow);
+	GameProjectUtils::OpenAddCodeToProjectDialog(InClass, InInitialPath, InParentWindow, bModal, OnCodeAddedToProject);
 	AddCodeToProjectDialogOpenedEvent.Broadcast();
 }
 
diff --git a/Engine/Source/Editor/GameProjectGeneration/Private/GameProjectUtils.cpp b/Engine/Source/Editor/GameProjectGeneration/Private/GameProjectUtils.cpp
--- a/Engine/Source/Editor/GameProjectGeneration/Private/GameProjectUtils.cpp
+++ b/Engine/Source/Editor/GameProjectGeneration/Private/GameProjectUtils.cpp
@@ -721,7 +721,7 @@ bool GameProjectUtils::UpdateGameProject(const FString& ProjectFile, const FStri
 	return UpdateGameProjectFile(ProjectFile, EngineIdentifier, NULL, OutFailReason);
 }
 
-void GameProjectUtils::OpenAddCodeToProjectDialog(const UClass* InClass, const FString& InInitialPath, const TSharedPtr<SWindow>& InParentWindow)
+void GameProjectUtils::OpenAddCodeToProjectDialog(const UClass* InClass, const FString& InInitialPath, const TSharedPtr<SWindow>& InParentWindow, bool bModal, FOnCodeAddedToProject OnCodeAddedToProject )
 {
 	// If we've been given a class then we only show the second page of the dialog, so we can make the window smaller as that page doesn't have as much content
 	const FVector2D WindowSize = (InClass) ? FVector2D(940, 380) : FVector2D(940, 540);
@@ -733,7 +733,13 @@ void GameProjectUtils::OpenAddCodeToProjectDialog(const UClass* InClass, const F
 		.SizingRule( ESizingRule::FixedSize )
 		.SupportsMinimize(false) .SupportsMaximize(false);
 
-	AddCodeWindow->SetContent( SNew(SNewClassDialog).Class(InClass).InitialPath(InInitialPath) );
+	TSharedRef<SNewClassDialog> NewClassDialog = 
+		SNew(SNewClassDialog)
+		.Class(InClass)
+		.InitialPath(InInitialPath)
+		.OnCodeAddedToProject( OnCodeAddedToProject );
+
+	AddCodeWindow->SetContent( NewClassDialog );
 
 	TSharedPtr<SWindow> ParentWindow = InParentWindow;
 	if (!ParentWindow.IsValid())
@@ -745,11 +751,26 @@ void GameProjectUtils::OpenAddCodeToProjectDialog(const UClass* InClass, const F
 
 	if (ParentWindow.IsValid())
 	{
-		FSlateApplication::Get().AddWindowAsNativeChild(AddCodeWindow, ParentWindow.ToSharedRef());
+		if( bModal )
+		{
+			FSlateApplication::Get().AddModalWindow(AddCodeWindow, ParentWindow);
+		}
+		else
+		{
+			FSlateApplication::Get().AddWindowAsNativeChild(AddCodeWindow, ParentWindow.ToSharedRef());
+		}
 	}
 	else
 	{
-		FSlateApplication::Get().AddWindow(AddCodeWindow);
+		if(bModal)
+		{
+			FSlateApplication::Get().AddModalWindow(AddCodeWindow, nullptr);
+		}
+		else
+		{
+			FSlateApplication::Get().AddWindow(AddCodeWindow);
+		}
+
 	}
 }
 
diff --git a/Engine/Source/Editor/GameProjectGeneration/Private/GameProjectUtils.h b/Engine/Source/Editor/GameProjectGeneration/Private/GameProjectUtils.h
--- a/Engine/Source/Editor/GameProjectGeneration/Private/GameProjectUtils.h
+++ b/Engine/Source/Editor/GameProjectGeneration/Private/GameProjectUtils.h
@@ -175,11 +175,13 @@ public:
 	/** 
 	 * Opens a dialog to add code files to the current project. 
 	 *
-	 * @param	InClass				The class we should force the user to use as their base class type, or null to allow the user to choose their base class in the UI
+	 * @param	InClass			The class we should force the user to use as their base class type, or null to allow the user to choose their base class in the UI
 	 * @param	InInitialPath		The initial path we should use as the destination for the new class header file, or an empty string to choose a suitable default based upon the module path
 	 * @param	InParentWindow		The parent window the dialog should use, or null to choose a suitable default parent window (the main frame, if available)
+	 * @param	bModal			True if the window should be modal and force the user to make a decision before continuing or false to let the user proceed with other tasks while the window is open
+	 * @param	OnCodeAddedToProject	Callback for when code is successfully added to the project
 	 */
-	static void OpenAddCodeToProjectDialog(const UClass* InClass, const FString& InInitialPath, const TSharedPtr<SWindow>& InParentWindow);
+	static void OpenAddCodeToProjectDialog(const UClass* InClass, const FString& InInitialPath, const TSharedPtr<SWindow>& InParentWindow, bool bModal = false, FOnCodeAddedToProject OnCodeAddedToProject = FOnCodeAddedToProject() );
 
 	/** Returns true if the specified class name is properly formed and does not conflict with another class */
 	static bool IsValidClassNameForCreation(const FString& NewClassName, const FModuleContextInfo& ModuleInfo, const TSet<FString>& DisallowedHeaderNames, FText& OutFailReason);
diff --git a/Engine/Source/Editor/GameProjectGeneration/Private/SNewClassDialog.cpp b/Engine/Source/Editor/GameProjectGeneration/Private/SNewClassDialog.cpp
--- a/Engine/Source/Editor/GameProjectGeneration/Private/SNewClassDialog.cpp
+++ b/Engine/Source/Editor/GameProjectGeneration/Private/SNewClassDialog.cpp
@@ -167,6 +167,8 @@ void SNewClassDialog::Construct( const FArguments& InArgs )
 
 	const float EditableTextHeight = 26.0f;
 
+	OnCodeAddedToProject = InArgs._OnCodeAddedToProject;
+
 	ChildSlot
 	[
 		SNew(SBorder)
@@ -995,6 +997,8 @@ void SNewClassDialog::FinishClicked()
 	const TSet<FString>& DisallowedHeaderNames = FSourceCodeNavigation::GetSourceFileDatabase().GetDisallowedHeaderNames();
 	if (GameProjectUtils::AddCodeToProject(NewClassName, NewClassPath, *SelectedModuleInfo, ParentClassInfo, DisallowedHeaderNames, HeaderFilePath, CppFilePath, FailReason))
 	{
+		OnCodeAddedToProject.ExecuteIfBound( NewClassName, NewClassPath, SelectedModuleInfo->ModuleName );
+
 		// Prevent periodic validity checks. This is to prevent a brief error message about the class already existing while you are exiting.
 		bPreventPeriodicValidityChecksUntilNextChange = true;
 
diff --git a/Engine/Source/Editor/GameProjectGeneration/Private/SNewClassDialog.h b/Engine/Source/Editor/GameProjectGeneration/Private/SNewClassDialog.h
--- a/Engine/Source/Editor/GameProjectGeneration/Private/SNewClassDialog.h
+++ b/Engine/Source/Editor/GameProjectGeneration/Private/SNewClassDialog.h
@@ -20,6 +20,8 @@ public:
 	/** The initial path to use as the destination for the new class. If this is not specified, we will work out a suitable default from the available project modules */
 	SLATE_ARGUMENT(FString, InitialPath)
 
+	/** Event called when code is successfully added to the project */
+	SLATE_EVENT( FOnCodeAddedToProject, OnCodeAddedToProject )
 	SLATE_END_ARGS()
 
 	/** Constructs this widget with InArgs */
@@ -219,4 +221,7 @@ private:
 
 	/** Information about the currently selected module; used for class validation */
 	TSharedPtr<FModuleContextInfo> SelectedModuleInfo;
+
+	/** Event called when code is succesfully added to the project */
+	FOnCodeAddedToProject OnCodeAddedToProject;
 };
diff --git a/Engine/Source/Editor/GameProjectGeneration/Public/GameProjectGenerationModule.h b/Engine/Source/Editor/GameProjectGeneration/Public/GameProjectGenerationModule.h
--- a/Engine/Source/Editor/GameProjectGeneration/Public/GameProjectGenerationModule.h
+++ b/Engine/Source/Editor/GameProjectGeneration/Public/GameProjectGenerationModule.h
@@ -22,6 +22,15 @@ struct FModuleContextInfo
 };
 
 
+/**
+ * Delegate called when code is added to the project.  Passes in the created class name and class path
+ * 
+ * @param ClassName		The created class name
+ * @param ClassPath		The created class path
+ * @param ModuleName	The name of the module that the class was added to
+ */
+DECLARE_DELEGATE_ThreeParams( FOnCodeAddedToProject, const FString& /*ClassName*/, const FString& /*ClassPath*/, const FString& /*ModuleName*/);
+
 /**
  * Game Project Generation module
  */
@@ -69,11 +78,13 @@ public:
 	/** 
 	 * Opens a dialog to add code files to the current project. 
 	 *
-	 * @param	InClass				The class we should force the user to use as their base class type, or null to allow the user to choose their base class in the UI
+	 * @param	InClass			The class we should force the user to use as their base class type, or null to allow the user to choose their base class in the UI
 	 * @param	InInitialPath		The initial path we should use as the destination for the new class header file, or an empty string to choose a suitable default based upon the module path
 	 * @param	InParentWindow		The parent window the dialog should use, or null to choose a suitable default parent window (the main frame, if available)
+	 * @param	bModal			True if the window should be modal and force the user to make a decision before continuing or false to let the user proceed with other tasks while the window is open
+	 * @param	OnCodeAddedToProject	Callback for when code is successfully added to the project
 	 */
-	virtual void OpenAddCodeToProjectDialog(const UClass* InClass, const FString& InInitialPath, const TSharedPtr<SWindow>& InParentWindow = nullptr);
+	virtual void OpenAddCodeToProjectDialog(const UClass* InClass, const FString& InInitialPath, const TSharedPtr<SWindow>& InParentWindow, bool bModal = false, FOnCodeAddedToProject OnCodeAddedToProject = FOnCodeAddedToProject() );
 
 	/** Delegate for when the AddCodeToProject dialog is opened */
 	DECLARE_EVENT(FGameProjectGenerationModule, FAddCodeToProjectDialogOpenedEvent);
diff --git a/Engine/Source/Editor/Kismet/Kismet.Build.cs b/Engine/Source/Editor/Kismet/Kismet.Build.cs
--- a/Engine/Source/Editor/Kismet/Kismet.Build.cs
+++ b/Engine/Source/Editor/Kismet/Kismet.Build.cs
@@ -44,7 +44,8 @@ public class Kismet : ModuleRules
                 "JsonUtilities",
                 "DerivedDataCache",
 				"DesktopPlatform",
-				"HotReload"
+				"HotReload",
+				"GameProjectGeneration",
 			}
 			);
 
diff --git a/Engine/Source/Editor/Kismet/Private/SSCSEditor.cpp b/Engine/Source/Editor/Kismet/Private/SSCSEditor.cpp
--- a/Engine/Source/Editor/Kismet/Private/SSCSEditor.cpp
+++ b/Engine/Source/Editor/Kismet/Private/SSCSEditor.cpp
@@ -43,6 +43,9 @@
 #include "SNotificationList.h"
 #include "NotificationManager.h"
 
+#include "GameProjectGenerationModule.h"
+#include "HotReloadInterface.h"
+
 #define LOCTEXT_NAMESPACE "SSCSEditor"
 
 DEFINE_LOG_CATEGORY_STATIC(LogSCSEditor, Log, All);
@@ -3344,44 +3347,66 @@ void SSCSEditor::OnGetChildrenForTree( FSCSEditorTreeNodePtrType InNodePtr, TArr
 }
 
 
-void SSCSEditor::PerformComboAddClass(TSubclassOf<UActorComponent> ComponentClass)
+void SSCSEditor::PerformComboAddClass(TSubclassOf<UActorComponent> ComponentClass, EComponentCreateAction::Type ComponentCreateAction )
 {
 	UClass* NewClass = ComponentClass;
 
-	FEditorDelegates::LoadSelectedAssetsIfNeeded.Broadcast();
-	USelection* Selection =  GEditor->GetSelectedObjects();
+	if( ComponentCreateAction == EComponentCreateAction::CreateNewCPPClass )
+	{
+		TSharedPtr<SWindow> ParentWindow = FSlateApplication::Get().FindWidgetWindow( SharedThis(this) );
+
+		FString AddedClassName;
+		auto OnCodeAddedToProject = [&AddedClassName](const FString& ClassName, const FString& ClassPath, const FString& ModuleName )
+		{
+			if( !ClassName.IsEmpty() && !ClassPath.IsEmpty() )
+			{
+				AddedClassName = FString::Printf(TEXT("/Script/%s.%s"), *ModuleName, *ClassName);
+			}
+		};
 
-	bool bAddedComponent = false;
+		const bool bModal = true;
+		FGameProjectGenerationModule::Get().OpenAddCodeToProjectDialog( ComponentClass, FString(), ParentWindow, bModal, FOnCodeAddedToProject::CreateLambda(OnCodeAddedToProject) );
 
-	// This adds components according to the type selected in the drop down. If the user
-	// has the appropriate objects selected in the content browser then those are added,
-	// else we go down the previous route of adding components by type.
-	if ( Selection->Num() > 0 )
+		NewClass = LoadClass<UActorComponent>(nullptr, *AddedClassName, nullptr, LOAD_None, nullptr );
+	}
+	
+	if( NewClass != nullptr )
 	{
-		for(FSelectionIterator ObjectIter(*Selection);ObjectIter;++ObjectIter)
-		{
-			UObject* Object = *ObjectIter;
-			UClass*  Class	= Object->GetClass();
+		FEditorDelegates::LoadSelectedAssetsIfNeeded.Broadcast();
+		USelection* Selection =  GEditor->GetSelectedObjects();
 
-			TArray< TSubclassOf<UActorComponent> > ComponentClasses = FComponentAssetBrokerage::GetComponentsForAsset(Object);
+		bool bAddedComponent = false;
 
-			// if the selected asset supports the selected component type then go ahead and add it
-			for ( int32 ComponentIndex = 0; ComponentIndex < ComponentClasses.Num(); ComponentIndex++ )
+		// This adds components according to the type selected in the drop down. If the user
+		// has the appropriate objects selected in the content browser then those are added,
+		// else we go down the previous route of adding components by type.
+		if(Selection->Num() > 0)
+		{
+			for(FSelectionIterator ObjectIter(*Selection); ObjectIter; ++ObjectIter)
 			{
-				if ( ComponentClasses[ComponentIndex]->IsChildOf( NewClass ) )
+				UObject* Object = *ObjectIter;
+				UClass*  Class	= Object->GetClass();
+
+				TArray< TSubclassOf<UActorComponent> > ComponentClasses = FComponentAssetBrokerage::GetComponentsForAsset(Object);
+
+				// if the selected asset supports the selected component type then go ahead and add it
+				for(int32 ComponentIndex = 0; ComponentIndex < ComponentClasses.Num(); ComponentIndex++)
 				{
-					AddNewComponent( NewClass, Object );
-					bAddedComponent = true;
-					break;
+					if(ComponentClasses[ComponentIndex]->IsChildOf(NewClass))
+					{
+						AddNewComponent(NewClass, Object);
+						bAddedComponent = true;
+						break;
+					}
 				}
 			}
 		}
-	}
 
-	if ( !bAddedComponent )
-	{
-		// As the SCS splits up the scene and actor components, can now add directly
-		AddNewComponent(ComponentClass, NULL);
+		if(!bAddedComponent)
+		{
+			// As the SCS splits up the scene and actor components, can now add directly
+			AddNewComponent(NewClass, NULL);
+		}
 	}
 }
 
diff --git a/Engine/Source/Editor/Kismet/Public/SSCSEditor.h b/Engine/Source/Editor/Kismet/Public/SSCSEditor.h
--- a/Engine/Source/Editor/Kismet/Public/SSCSEditor.h
+++ b/Engine/Source/Editor/Kismet/Public/SSCSEditor.h
@@ -3,6 +3,7 @@
 #pragma once
 
 #include "BlueprintEditor.h"
+#include "SComponentClassCombo.h"
 
 class SSCSEditor;
 class USCS_Node;
@@ -625,7 +626,7 @@ protected:
 	FString GetSelectedClassText() const;
 
 	/** Add a component from the selection in the combo box */
-	void PerformComboAddClass(TSubclassOf<UActorComponent> ComponentClass);
+	void PerformComboAddClass(TSubclassOf<UActorComponent> ComponentClass, EComponentCreateAction::Type ComponentCreateAction );
 
 	/** Called to display context menu when right clicking on the widget */
 	TSharedPtr< SWidget > CreateContextMenu();
diff --git a/Engine/Source/Editor/UnrealEd/Classes/Settings/EditorExperimentalSettings.h b/Engine/Source/Editor/UnrealEd/Classes/Settings/EditorExperimentalSettings.h
--- a/Engine/Source/Editor/UnrealEd/Classes/Settings/EditorExperimentalSettings.h
+++ b/Engine/Source/Editor/UnrealEd/Classes/Settings/EditorExperimentalSettings.h
@@ -37,6 +37,10 @@ public:
 	UPROPERTY(EditAnywhere, config, Category = Tools, meta = ( DisplayName = "Single Layout Blueprint Editor" ))
 	bool bUnifiedBlueprintEditor;
 
+	/** Ability to add C++ and blueprinted components to actors */
+	UPROPERTY(EditAnywhere, config, Category=Tools)
+	bool bScriptableComponentsOnActors;
+
 	/** The Messaging Debugger provides a visual utility for debugging the messaging system. */
 	UPROPERTY(EditAnywhere, config, Category=Tools, meta=(DisplayName="Messaging Debugger"))
 	bool bMessagingDebugger;
diff --git a/Engine/Source/Editor/UnrealEd/Private/SComponentClassCombo.cpp b/Engine/Source/Editor/UnrealEd/Private/SComponentClassCombo.cpp
--- a/Engine/Source/Editor/UnrealEd/Private/SComponentClassCombo.cpp
+++ b/Engine/Source/Editor/UnrealEd/Private/SComponentClassCombo.cpp
@@ -11,6 +11,7 @@
 #include "EditorClassUtils.h"
 #include "SSearchBox.h"
 #include "Engine/Selection.h"
+#include "HotReloadInterface.h"
 
 #define LOCTEXT_NAMESPACE "ComponentClassCombo"
 
@@ -23,6 +24,9 @@ void SComponentClassCombo::Construct(const FArguments& InArgs)
 	UpdateComponentClassList();
 	GenerateFilteredComponentList(FString());
 
+	IHotReloadInterface& HotReloadSupport = FModuleManager::LoadModuleChecked<IHotReloadInterface>("HotReload");
+	HotReloadSupport.OnHotReload().AddSP( this, &SComponentClassCombo::OnProjectHotReloaded );
+	
 	SAssignNew(ComponentClassListView, SListView<FComponentClassComboEntryPtr>)
 		.ListItemsSource( &FilteredComponentClassList )
 		.OnSelectionChanged( this, &SComponentClassCombo::OnAddComponentSelectionChanged )
@@ -174,7 +178,7 @@ void SComponentClassCombo::OnAddComponentSelectionChanged( FComponentClassComboE
 
 		if ( InItem->IsClass() )
 		{
-			OnComponentClassSelected.ExecuteIfBound(InItem->GetComponentClass());
+			OnComponentClassSelected.ExecuteIfBound(InItem->GetComponentClass(), InItem->GetComponentCreateAction() );
 
 			// Neither do we want the combo dropdown staying open once the user has clicked on a valid option
 			SetIsOpen( false, false );
@@ -322,11 +326,21 @@ void SComponentClassCombo::UpdateComponentClassList()
 
 	const UEdGraphSchema_K2* K2Schema = GetDefault<UEdGraphSchema_K2>();
 
+	if( GetDefault<UEditorExperimentalSettings>()->bScriptableComponentsOnActors )
+	{
+		// Add new C++ component class
+		FComponentClassComboEntryPtr NewClassHeader = MakeShareable(new FComponentClassComboEntry(LOCTEXT("ScriptingHeading", "Scripting").ToString()));
+		ComponentClassList.Add(NewClassHeader);
+
+		FComponentClassComboEntryPtr NewCPPClass = MakeShareable(new FComponentClassComboEntry(LOCTEXT("ScriptingHeading", "Scripting").ToString(), UActorComponent::StaticClass(), true, EComponentCreateAction::CreateNewCPPClass));
+		ComponentClassList.Add(NewCPPClass);
+	}
+	
 	TArray<FComponentClassComboEntryPtr> SortedClassList;
 	for (TObjectIterator<UClass> It; It; ++It)
 	{
 		UClass* Class = *It;
-		// If this is a subclass of SceneComponent, not abstract, and tagged as spawnable from Kismet
+		// If this is a subclass of Actor Component, not abstract, and tagged as spawnable from Kismet
 		if (Class->IsChildOf(UActorComponent::StaticClass()) && !Class->HasAnyClassFlags(CLASS_Abstract) && Class->HasMetaData(FBlueprintMetadata::MD_BlueprintSpawnableComponent)) //@TODO: Fold this logic together with the one in UEdGraphSchema_K2::GetAddComponentClasses
 		{
 			TArray<FString> ClassGroupNames;
@@ -335,13 +349,15 @@ void SComponentClassCombo::UpdateComponentClassList()
 			if (ClassGroupNames.Contains(SortComboEntry::CommonClassGroup))
 			{
 				FString ClassGroup = SortComboEntry::CommonClassGroup;
-				FComponentClassComboEntryPtr NewEntry(new FComponentClassComboEntry(ClassGroup, Class, ClassGroupNames.Num() <= 1));
+				FComponentClassComboEntryPtr NewEntry(new FComponentClassComboEntry(ClassGroup, Class, ClassGroupNames.Num() <= 1, EComponentCreateAction::SpawnExistingClass));
 				SortedClassList.Add(NewEntry);
 			}
 			if (ClassGroupNames.Num() && !ClassGroupNames[0].Equals(SortComboEntry::CommonClassGroup))
 			{
+				const bool bIncludeInFilter = true;
+
 				FString ClassGroup = ClassGroupNames[0];
-				FComponentClassComboEntryPtr NewEntry(new FComponentClassComboEntry(ClassGroup, Class, true));
+				FComponentClassComboEntryPtr NewEntry(new FComponentClassComboEntry(ClassGroup, Class, bIncludeInFilter, EComponentCreateAction::SpawnExistingClass));
 				SortedClassList.Add(NewEntry);
 			}
 		}
@@ -383,51 +399,69 @@ FReply SComponentClassCombo::OnButtonClicked()
 	return SComboButton::OnButtonClicked();
 }
 
+void SComponentClassCombo::OnProjectHotReloaded( bool bWasTriggeredAutomatically )
+{
+	UpdateComponentClassList();
+
+	ComponentClassListView->RequestListRefresh();
+}
+
+
 FText SComponentClassCombo::GetFriendlyComponentName(FComponentClassComboEntryPtr Entry) const
 {
 	// Get a user friendly string from the component name
-	FString FriendlyComponentName = GetSanitizedComponentName(Entry->GetComponentClass());
-
-	// Search the selected assets and look for any that can be used as a source asset for this type of component
-	// If there is one we append the asset name to the component name, if there are many we append "Multiple Assets"
-	FString AssetName;
-	UObject* PreviousMatchingAsset = NULL;
+	FString FriendlyComponentName;
 
-	FEditorDelegates::LoadSelectedAssetsIfNeeded.Broadcast();
-	USelection* Selection = GEditor->GetSelectedObjects();
-	for (FSelectionIterator ObjectIter(*Selection); ObjectIter; ++ObjectIter)
+	if( Entry->GetComponentCreateAction() == EComponentCreateAction::CreateNewCPPClass )
+	{
+		FriendlyComponentName = LOCTEXT("NewCPPBehaviorFriendlyName", "New C++ Behavior").ToString();
+	}
+	else
 	{
-		UObject* Object = *ObjectIter;
-		UClass* Class = Object->GetClass();
+		FriendlyComponentName = GetSanitizedComponentName(Entry->GetComponentClass());
+
+
+		// Search the selected assets and look for any that can be used as a source asset for this type of component
+		// If there is one we append the asset name to the component name, if there are many we append "Multiple Assets"
+		FString AssetName;
+		UObject* PreviousMatchingAsset = NULL;
 
-		TArray<TSubclassOf<UActorComponent> > ComponentClasses = FComponentAssetBrokerage::GetComponentsForAsset(Object);
-		for (int32 ComponentIndex = 0; ComponentIndex < ComponentClasses.Num(); ComponentIndex++)
+		FEditorDelegates::LoadSelectedAssetsIfNeeded.Broadcast();
+		USelection* Selection = GEditor->GetSelectedObjects();
+		for(FSelectionIterator ObjectIter(*Selection); ObjectIter; ++ObjectIter)
 		{
-			if (ComponentClasses[ComponentIndex]->IsChildOf(Entry->GetComponentClass()))
+			UObject* Object = *ObjectIter;
+			UClass* Class = Object->GetClass();
+
+			TArray<TSubclassOf<UActorComponent> > ComponentClasses = FComponentAssetBrokerage::GetComponentsForAsset(Object);
+			for(int32 ComponentIndex = 0; ComponentIndex < ComponentClasses.Num(); ComponentIndex++)
 			{
-				if (AssetName.IsEmpty())
-				{
-					// If there is no previous asset then we just accept the name
-					AssetName = Object->GetName();
-					PreviousMatchingAsset = Object;
-				}
-				else
+				if(ComponentClasses[ComponentIndex]->IsChildOf(Entry->GetComponentClass()))
 				{
-					// if there is a previous asset then check that we didn't just find multiple appropriate components
-					// in a single asset - if the asset differs then we don't display the name, just "Multiple Assets"
-					if (PreviousMatchingAsset != Object)
+					if(AssetName.IsEmpty())
 					{
-						AssetName = LOCTEXT("MultipleAssetsForComponentAnnotation", "Multiple Assets").ToString();
+						// If there is no previous asset then we just accept the name
+						AssetName = Object->GetName();
 						PreviousMatchingAsset = Object;
 					}
+					else
+					{
+						// if there is a previous asset then check that we didn't just find multiple appropriate components
+						// in a single asset - if the asset differs then we don't display the name, just "Multiple Assets"
+						if(PreviousMatchingAsset != Object)
+						{
+							AssetName = LOCTEXT("MultipleAssetsForComponentAnnotation", "Multiple Assets").ToString();
+							PreviousMatchingAsset = Object;
+						}
+					}
 				}
 			}
 		}
-	}
 
-	if (!AssetName.IsEmpty())
-	{
-		FriendlyComponentName += FString(" (") + AssetName + FString(")");
+		if(!AssetName.IsEmpty())
+		{
+			FriendlyComponentName += FString(" (") + AssetName + FString(")");
+		}
 	}
 	return FText::FromString(FriendlyComponentName);
 }
diff --git a/Engine/Source/Editor/UnrealEd/Public/SComponentClassCombo.h b/Engine/Source/Editor/UnrealEd/Public/SComponentClassCombo.h
--- a/Engine/Source/Editor/UnrealEd/Public/SComponentClassCombo.h
+++ b/Engine/Source/Editor/UnrealEd/Public/SComponentClassCombo.h
@@ -6,21 +6,35 @@ typedef TSharedPtr<class FComponentClassComboEntry> FComponentClassComboEntryPtr
 
 //////////////////////////////////////////////////////////////////////////
 
+namespace EComponentCreateAction
+{
+	enum Type
+	{
+		/** Create a new C++ class based off the specified ActorComponent class and then add it to the tree */
+		CreateNewCPPClass,
+		/** Create a new blueprint class based off the specified ActorComponent class and then add it to the tree */
+		CreateNewBlueprintClass,
+		/** Spawn a new instance of the specified ActorComponent class and then add it to the tree */
+		SpawnExistingClass,
+	};
+}
 
-DECLARE_DELEGATE_OneParam(FComponentClassSelected, TSubclassOf<UActorComponent>);
+DECLARE_DELEGATE_TwoParams(FComponentClassSelected, TSubclassOf<UActorComponent>, EComponentCreateAction::Type );
 
 class FComponentClassComboEntry: public TSharedFromThis<FComponentClassComboEntry>
 {
 public:
-	FComponentClassComboEntry( const FString& InHeadingText, TSubclassOf<UActorComponent> InComponentClass, bool InIncludedInFilter )
+	FComponentClassComboEntry( const FString& InHeadingText, TSubclassOf<UActorComponent> InComponentClass, bool InIncludedInFilter, EComponentCreateAction::Type InComponentCreateAction )
 		: ComponentClass(InComponentClass)
 		, HeadingText(InHeadingText)
 		, bIncludedInFilter(InIncludedInFilter)
+		, ComponentCreateAction(InComponentCreateAction)
 	{}
 
 	FComponentClassComboEntry( const FString& InHeadingText )
 		: ComponentClass(NULL)
 		, HeadingText(InHeadingText)
+		, bIncludedInFilter(false)
 	{}
 
 	FComponentClassComboEntry()
@@ -40,19 +54,27 @@ public:
 	{
 		return (ComponentClass==NULL && HeadingText.IsEmpty());
 	}
+	
 	bool IsClass() const
 	{
 		return (ComponentClass!=NULL);
 	}
+	
 	bool IsIncludedInFilter() const
 	{
 		return bIncludedInFilter;
 	}
+	
+	EComponentCreateAction::Type GetComponentCreateAction() const
+	{
+		return ComponentCreateAction;
+	}
 
 private:
 	TSubclassOf<UActorComponent> ComponentClass;
 	FString HeadingText;
 	bool bIncludedInFilter;
+	EComponentCreateAction::Type ComponentCreateAction;
 };
 
 //////////////////////////////////////////////////////////////////////////
@@ -104,6 +126,8 @@ public:
 protected:
 	virtual FReply OnButtonClicked();
 
+	/** Called when a project is hot reloaded to refresh the components list */
+	void OnProjectHotReloaded( bool bWasTriggeredAutomatically );
 private:
 
 	FText GetFriendlyComponentName(FComponentClassComboEntryPtr Entry) const;
