[INTEGRATE] Change 2418597 by [REDACTED_PERSON][REDACTED_HANDLE].7 on 2015/01/26 10:04:00

	More work on the auto reimport manager
	Refined some of the batch processing, errors are now routed to the message log, and presented on the context notification.

diff --git a/Engine/Source/Editor/UnrealEd/Classes/Settings/EditorLoadingSavingSettings.h b/Engine/Source/Editor/UnrealEd/Classes/Settings/EditorLoadingSavingSettings.h
--- a/Engine/Source/Editor/UnrealEd/Classes/Settings/EditorLoadingSavingSettings.h
+++ b/Engine/Source/Editor/UnrealEd/Classes/Settings/EditorLoadingSavingSettings.h
@@ -35,6 +35,10 @@ public:
 	UPROPERTY(EditAnywhere, config, Category=AutoReimport, meta=(DisplayName="Additional Directories to Monitor"))
 	TArray<FString> AutoReimportDirectories;
 
+	/** Internal setting to control whether we should ask the user whether we should automatically delete source files when their assets are deleted */
+	UPROPERTY(config)
+	bool bDeleteSourceFilesWithAssets;
+
 private:
 
 	/** Deprecated properties - we automatically monitor all source file types that are supported by in memory factories now */
diff --git a/Engine/Source/Editor/UnrealEd/Private/AssetDeleteModel.cpp b/Engine/Source/Editor/UnrealEd/Private/AssetDeleteModel.cpp
--- a/Engine/Source/Editor/UnrealEd/Private/AssetDeleteModel.cpp
+++ b/Engine/Source/Editor/UnrealEd/Private/AssetDeleteModel.cpp
@@ -6,6 +6,7 @@
 #include "ObjectTools.h"
 #include "AssetRegistryModule.h"
 #include "AssetEditorManager.h"
+#include "AutoReimport/AutoReimportUtilities.h"
 
 #define LOCTEXT_NAMESPACE "FAssetDeleteModel"
 
@@ -117,6 +118,23 @@ void FAssetDeleteModel::Tick( const float InDeltaTime )
 	}
 }
 
+void FAssetDeleteModel::DeleteSourceContentFiles()
+{
+	const IAssetRegistry& Registry = FModuleManager::LoadModuleChecked<FAssetRegistryModule>("AssetRegistry").Get();
+	
+	TArray<FString> FilesToDelete;
+	for ( const TSharedPtr< FPendingDelete >& PendingDelete : PendingDeletes )
+	{
+		Utils::ExtractSourceFilePaths(PendingDelete->GetObject(), FilesToDelete);
+	}
+
+	IFileManager& FileManager = IFileManager::Get();
+	for (const auto& Path : FilesToDelete)
+	{
+		FileManager.Delete(*Path, false /* RequireExists */, true /* Even if read only */, true /* Quiet */);
+	}
+}
+
 bool FAssetDeleteModel::CanDelete() const
 {
 	return !CanForceDelete();
diff --git a/Engine/Source/Editor/UnrealEd/Private/AutoReimport/AutoReimportManager.cpp b/Engine/Source/Editor/UnrealEd/Private/AutoReimport/AutoReimportManager.cpp
--- a/Engine/Source/Editor/UnrealEd/Private/AutoReimport/AutoReimportManager.cpp
+++ b/Engine/Source/Editor/UnrealEd/Private/AutoReimport/AutoReimportManager.cpp
@@ -10,76 +10,95 @@
 #include "ARFilter.h"
 #include "ContentBrowserModule.h"
 #include "AssetRegistryModule.h"
-#include "NotificationManager.h"
-#include "SNotificationList.h"
+#include "ReimportFeedbackContext.h"
+#include "MessageLogModule.h"
 
 #define LOCTEXT_NAMESPACE "AutoReimportManager"
-#define CONDITIONAL_YIELD if (TimeLimit.Exceeded()) { return; }
+#define yield if (TimeLimit.Exceeded()) return
 
 DEFINE_LOG_CATEGORY_STATIC(LogAutoReimportManager, Log, All);
 
-/** Feedback context that overrides GWarn for import operations to prevent popup spam */
-struct FReimportFeedbackContext : public FFeedbackContext
+
+enum class EStateMachineNode { CallOnce, CallMany };
+
+/** A simple state machine class that calls functions mapped on enum values. If any function returns a new enum type, it moves onto that function */
+template<typename TState>
+struct FStateMachine
 {
-	/** Constructor - created when we are processing changes */
-	FReimportFeedbackContext()
+	typedef TFunction<TOptional<TState>(const FTimeLimit&)> FunctionType;
+
+	/** Constructor that specifies the initial state of the machine */
+	FStateMachine(TState InitialState) : CurrentState(InitialState){}
+
+	/** Add an enum->function mapping for this state machine */
+	void Add(TState State, EStateMachineNode NodeType, const FunctionType& Function)
 	{
+		Nodes.Add(State, FStateMachineNode(NodeType, Function));
 	}
 
-	/** Destructor - called when we are done processing changes */
-	~FReimportFeedbackContext()
+	/** Set a new state for this machine */
+	void SetState(TState NewState)
 	{
-		if (Notification.IsValid())
-		{
-			Notification->SetCompletionState(SNotificationItem::CS_Success);
-			Notification->ExpireAndFadeout();
-		}
+		CurrentState = NewState;
 	}
 
-	/** Update the text on the notification */
-	void UpdateText(int32 Progress, int32 Total)
+	/** Tick this state machine with the given time limit. Will continuously enumerate the machine until TimeLimit is reached */
+	void Tick(const FTimeLimit& TimeLimit)
 	{
-		FFormatOrderedArguments Args;
-		Args.Add(Progress);
-		Args.Add(Total);
-
-		const FText Message = FText::Format(LOCTEXT("ProcessingChanges", "Processing outstanding content changes ({0} of {1})"), Args);
-		if (Notification.IsValid())
-		{
-			Notification->SetText(Message);
-		}
-		else
+		while (!TimeLimit.Exceeded())
 		{
-			FNotificationInfo Info(Message);
-			Info.ExpireDuration = 2.0f;
-			Info.bUseLargeFont = false;
-			Info.bFireAndForget = false;
-			Info.bUseSuccessFailIcons = false;
-			Info.Image = FEditorStyle::GetBrush("NoBrush");
-
-			Notification = FSlateNotificationManager::Get().AddNotification(Info);
-			Notification->SetCompletionState(SNotificationItem::CS_Pending);
+			const auto& State = Nodes[CurrentState];
+			TOptional<TState> NewState = State.Endpoint(TimeLimit);
+			if (NewState.IsSet())
+			{
+				CurrentState = NewState.GetValue();
+			}
+			else if (State.Type == EStateMachineNode::CallOnce)
+			{
+				break;
+			}
 		}
 	}
 
 private:
-	virtual void Serialize( const TCHAR* V, ELogVerbosity::Type Verbosity, const class FName& Category ) {}
+
+	/** The current state of this machine */
+	TState CurrentState;
 
 private:
-	/** The notification that is shown when the context is active */
-	TSharedPtr<SNotificationItem> Notification;
 
-	/** The notification content */
-	TSharedPtr<class SAutoReimportFeedback> NotificationContent;
+	struct FStateMachineNode
+	{
+		FStateMachineNode(EStateMachineNode InType, const FunctionType& InEndpoint) : Endpoint(InEndpoint), Type(InType) {}
+
+		/** The function endpoint for this node */
+		FunctionType Endpoint;
+
+		/** Whether this endpoint should be called multiple times in a frame, or just once */
+		EStateMachineNode Type;
+	};
+
+	/** A map of enum value -> callback information */
+	TMap<TState, FStateMachineNode> Nodes;
 };
 
+/** Enum and value to specify the current state of our processing */
+enum class ECurrentState
+{
+	Idle, PendingResponse, ProcessAdditions, SavePackages, ProcessModifications, ProcessDeletions
+};
+uint32 GetTypeHash(ECurrentState State)
+{
+	return (int32)State;
+}
+
 /* Deals with auto reimporting of objects when the objects file on disk is modified*/
-class FAutoReimportManager : public FTickableEditorObject
+class FAutoReimportManager : public FTickableEditorObject, public FGCObject, public TSharedFromThis<FAutoReimportManager>
 {
 public:
 
 	FAutoReimportManager();
-	~FAutoReimportManager();
+	void Destroy();
 
 private:
 
@@ -88,6 +107,9 @@ private:
 	virtual bool IsTickable() const override { return true; }
 	virtual TStatId GetStatId() const override;
 
+	/** FGCObject interface*/
+	virtual void AddReferencedObjects(FReferenceCollector& Collector) override;
+
 private:
 
 	/** Called when a new asset path has been mounted */
@@ -96,9 +118,6 @@ private:
 	/** Called when a new asset path has been dismounted */
 	void OnContentPathDismounted(const FString& InAssetPath, const FString& FileSystemPath);
 
-	/** Called when an asset has been deleted */
-	void OnAssetDeleted(UObject* Object);
-
 private:
 
 	/** Callback for when an editor user setting has changed */
@@ -114,38 +133,45 @@ private:
 	static FString GetAllFactoryExtensions();
 
 private:
-
-	/** Get the total amount of work in the current operation */
-	int32 GetTotalWork() const;
-
-	/** Get the total progress through the current operation */
-	int32 GetTotalProgress() const;
 	
 	/** Get the number of unprocessed changes that are not part of the current processing operation */
 	int32 GetNumUnprocessedChanges() const;
 
+	/** Get the progress text to display on the context notification */
+	FText GetProgressText() const;
+
+private:
+
+	/** A state machine holding information about the current state of the manager */
+	FStateMachine<ECurrentState> StateMachine;
+
 private:
 
+	/** Idle processing */
+	TOptional<ECurrentState> Idle(const FTimeLimit& Limit);
+	
 	/** Process any remaining pending additions we have */
-	void ProcessAdditions(const FTimeLimit& TimeLimit);
+	TOptional<ECurrentState> ProcessAdditions(const FTimeLimit& TimeLimit);
+
+	/** Save any packages that were created inside ProcessAdditions */
+	TOptional<ECurrentState> SavePackages();
 
 	/** Process any remaining pending modifications we have */
-	void ProcessModifications(const IAssetRegistry& Registry, const FTimeLimit& TimeLimit);
+	TOptional<ECurrentState> ProcessModifications(const FTimeLimit& TimeLimit);
 
 	/** Process any remaining pending deletions we have */
-	void ProcessDeletions(const IAssetRegistry& Registry);
+	TOptional<ECurrentState> ProcessDeletions();
 
-	/** Enum and value to specify the current state of our processing */
-	enum class ECurrentState
-	{
-		Idle, Initialize, ProcessAdditions, SavePackages, ProcessModifications, ProcessDeletions, Cleanup
-	};
-	ECurrentState CurrentState;
+	/** Wait for a user's input. Just updates the progress text for now */
+	TOptional<ECurrentState> PendingResponse();
+
+	/** Cleanup an operation that just processed some changes */
+	void Cleanup();
 
 private:
 
 	/** Feedback context that can selectively override the global context to consume progress events for saving of assets */
-	TUniquePtr<FReimportFeedbackContext> FeedbackContextOverride;
+	TSharedPtr<FReimportFeedbackContext> FeedbackContextOverride;
 
 	/** Cached string of extensions that are supported by all available factories */
 	FString SupportedExtensions;
@@ -170,9 +196,9 @@ private:
 };
 
 FAutoReimportManager::FAutoReimportManager()
-	: CurrentState(ECurrentState::Idle)
+	: StateMachine(ECurrentState::Idle)
 	, bIsProcessingChanges(false)
-	, StartProcessingDelay(1)
+	, StartProcessingDelay(0.5)
 {
 	// @todo: arodham: update this when modules are reloaded or new factory types are available?
 	SupportedExtensions = GetAllFactoryExtensions();
@@ -183,9 +209,11 @@ FAutoReimportManager::FAutoReimportManager()
 	FPackageName::OnContentPathMounted().AddRaw(this, &FAutoReimportManager::OnContentPathMounted);
 	FPackageName::OnContentPathDismounted().AddRaw(this, &FAutoReimportManager::OnContentPathDismounted);
 
-	IAssetRegistry& AssetRegistry = FModuleManager::LoadModuleChecked<FAssetRegistryModule>("AssetRegistry").Get();
-	AssetRegistry.OnInMemoryAssetDeleted().AddRaw(this, &FAutoReimportManager::OnAssetDeleted);
-
+	FMessageLogModule& MessageLogModule = FModuleManager::LoadModuleChecked<FMessageLogModule>("MessageLog");
+	if (!MessageLogModule.IsRegisteredLogListing("AssetReimport"))
+	{
+		MessageLogModule.RegisterLogListing("AssetReimport", LOCTEXT("AssetReimportLabel", "Asset Reimport"));
+	}
 
 	// Only set this up for content directories if the user has this enabled
 	if (Settings->bMonitorContentDirectories)
@@ -195,9 +223,16 @@ FAutoReimportManager::FAutoReimportManager()
 
 	// We always monitor user directories if they are set
 	SetUpUserDirectories();
+
+	StateMachine.Add(ECurrentState::Idle,					EStateMachineNode::CallOnce,	[this](const FTimeLimit& T){ return this->Idle(T);					});
+	StateMachine.Add(ECurrentState::PendingResponse,		EStateMachineNode::CallOnce,	[this](const FTimeLimit& T){ return this->PendingResponse();		});
+	StateMachine.Add(ECurrentState::ProcessAdditions,		EStateMachineNode::CallMany,	[this](const FTimeLimit& T){ return this->ProcessAdditions(T);		});
+	StateMachine.Add(ECurrentState::SavePackages,			EStateMachineNode::CallOnce,	[this](const FTimeLimit& T){ return this->SavePackages();			});
+	StateMachine.Add(ECurrentState::ProcessModifications,	EStateMachineNode::CallMany,	[this](const FTimeLimit& T){ return this->ProcessModifications(T);	});
+	StateMachine.Add(ECurrentState::ProcessDeletions,		EStateMachineNode::CallMany,	[this](const FTimeLimit& T){ return this->ProcessDeletions();		});
 }
 
-FAutoReimportManager::~FAutoReimportManager() 
+void FAutoReimportManager::Destroy()
 {
 	if (auto* Settings = GetMutableDefault<UEditorLoadingSavingSettings>())
 	{
@@ -223,20 +258,30 @@ int32 FAutoReimportManager::GetNumUnprocessedChanges() const
 	return Utils::Reduce(MountedDirectoryMonitors, Pred, 0) + Utils::Reduce(UserDirectoryMonitors, Pred, 0);
 }
 
-int32 FAutoReimportManager::GetTotalWork() const
+FText FAutoReimportManager::GetProgressText() const
 {
-	auto Pred = [](const FContentDirectoryMonitor& Monitor, int32 Total){ return Total + Monitor.GetTotalWork(); };
-	return Utils::Reduce(MountedDirectoryMonitors, Pred, 0) + Utils::Reduce(UserDirectoryMonitors, Pred, 0);
-}
+	FFormatOrderedArguments Args;
+	{
+		auto Pred = [](const FContentDirectoryMonitor& Monitor, int32 Total){ return Total + Monitor.GetWorkProgress(); };
+		const int32 Progress = Utils::Reduce(MountedDirectoryMonitors, Pred, 0) + Utils::Reduce(UserDirectoryMonitors, Pred, 0);
 
-int32 FAutoReimportManager::GetTotalProgress() const
-{
-	auto Pred = [](const FContentDirectoryMonitor& Monitor, int32 Total){ return Total + Monitor.GetWorkProgress(); };
-	return Utils::Reduce(MountedDirectoryMonitors, Pred, 0) + Utils::Reduce(UserDirectoryMonitors, Pred, 0);
+		Args.Add(Progress);
+	}
+
+	{
+		auto Pred = [](const FContentDirectoryMonitor& Monitor, int32 Total){ return Total + Monitor.GetTotalWork(); };
+		const int32 Total = Utils::Reduce(MountedDirectoryMonitors, Pred, 0) + Utils::Reduce(UserDirectoryMonitors, Pred, 0);
+
+		Args.Add(Total);
+	}
+
+	return FText::Format(LOCTEXT("ProcessingChanges", "Processing outstanding content changes ({0} of {1})..."), Args);
 }
 
-void FAutoReimportManager::ProcessAdditions(const FTimeLimit& TimeLimit)
+TOptional<ECurrentState> FAutoReimportManager::ProcessAdditions(const FTimeLimit& TimeLimit)
 {
+	FeedbackContextOverride->GetContent()->SetMainText(GetProgressText());
+
 	TMap<FString, TArray<UFactory*>> Factories;
 
 	TArray<FString> FactoryExtensions;
@@ -271,41 +316,69 @@ void FAutoReimportManager::ProcessAdditions(const FTimeLimit& TimeLimit)
 
 	for (auto& Monitor : MountedDirectoryMonitors)
 	{
-		Monitor.ProcessAdditions(PackagesToSave, TimeLimit, Factories);
+		Monitor.ProcessAdditions(PackagesToSave, TimeLimit, Factories, *FeedbackContextOverride);
+		yield TOptional<ECurrentState>();
 	}
 
-	if (!TimeLimit.Exceeded())
+	for (auto& Monitor : UserDirectoryMonitors)
 	{
-		for (auto& Monitor : UserDirectoryMonitors)
-		{
-			Monitor.ProcessAdditions(PackagesToSave, TimeLimit, Factories);
-		}
+		Monitor.ProcessAdditions(PackagesToSave, TimeLimit, Factories, *FeedbackContextOverride);
+		yield TOptional<ECurrentState>();
 	}
 
-	CONDITIONAL_YIELD
+	return ECurrentState::SavePackages;
+}
+
+TOptional<ECurrentState> FAutoReimportManager::SavePackages()
+{
+	// We don't override the context specifically when saving packages so the user gets proper feedback
+	//TGuardValue<FFeedbackContext*> ScopedContextOverride(GWarn, FeedbackContextOverride.Get());
+
+	FeedbackContextOverride->GetContent()->SetMainText(GetProgressText());
+
+	if (PackagesToSave.Num() > 0)
+	{
+		const bool bAlreadyCheckedOut = false;
+		const bool bCheckDirty = false;
+		const bool bPromptToSave = false;
+		FEditorFileUtils::PromptForCheckoutAndSave(PackagesToSave, bCheckDirty, bPromptToSave, nullptr, bAlreadyCheckedOut);
+
+		PackagesToSave.Empty();
+	}
 
-	CurrentState = ECurrentState::SavePackages;
+	return ECurrentState::ProcessModifications;
 }
 
-void FAutoReimportManager::ProcessModifications(const IAssetRegistry& Registry, const FTimeLimit& TimeLimit)
+TOptional<ECurrentState> FAutoReimportManager::ProcessModifications(const FTimeLimit& TimeLimit)
 {
+	// Override the global feedback context while we do this to avoid popping up dialogs
+	TGuardValue<FFeedbackContext*> ScopedContextOverride(GWarn, FeedbackContextOverride.Get());
+
+	FeedbackContextOverride->GetContent()->SetMainText(GetProgressText());
+
+	const IAssetRegistry& Registry = FModuleManager::LoadModuleChecked<FAssetRegistryModule>("AssetRegistry").Get();
+
 	for (auto& Monitor : MountedDirectoryMonitors)
 	{
-		Monitor.ProcessModifications(Registry, TimeLimit);
-		CONDITIONAL_YIELD
+		Monitor.ProcessModifications(Registry, TimeLimit, *FeedbackContextOverride);
+		yield TOptional<ECurrentState>();
 	}
 
 	for (auto& Monitor : UserDirectoryMonitors)
 	{
-		Monitor.ProcessModifications(Registry, TimeLimit);
-		CONDITIONAL_YIELD
+		Monitor.ProcessModifications(Registry, TimeLimit, *FeedbackContextOverride);
+		yield TOptional<ECurrentState>();
 	}
 
-	CurrentState = ECurrentState::ProcessDeletions;
+	return ECurrentState::ProcessDeletions;
 }
 
-void FAutoReimportManager::ProcessDeletions(const IAssetRegistry& Registry)
+TOptional<ECurrentState> FAutoReimportManager::ProcessDeletions()
 {
+	const IAssetRegistry& Registry = FModuleManager::LoadModuleChecked<FAssetRegistryModule>("AssetRegistry").Get();
+
+	FeedbackContextOverride->GetContent()->SetMainText(GetProgressText());
+
 	TArray<FAssetData> AssetsToDelete;
 
 	for (auto& Monitor : MountedDirectoryMonitors)
@@ -322,133 +395,82 @@ void FAutoReimportManager::ProcessDeletions(const IAssetRegistry& Registry)
 		ObjectTools::DeleteAssets(AssetsToDelete);
 	}
 
-	CurrentState = ECurrentState::Cleanup;
+	Cleanup();
+	return ECurrentState::Idle;
 }
 
-
-void FAutoReimportManager::Tick(float DeltaTime)
+TOptional<ECurrentState> FAutoReimportManager::PendingResponse()
 {
-	// Never spend more than a 60fps frame doing this work (meaning we shouldn't drop below 30fps), we can do more if we're throttling CPU usage (ie editor running in background)
-	const FTimeLimit TimeLimit(GEditor->ShouldThrottleCPUUsage() ? 1 / 6.f : 1.f / 60.f);
+	FeedbackContextOverride->GetContent()->SetMainText(GetProgressText());
+	return TOptional<ECurrentState>();
+}
 
-	if (CurrentState == ECurrentState::Idle)
+TOptional<ECurrentState> FAutoReimportManager::Idle(const FTimeLimit& TimeLimit)
+{
+	for (auto& Monitor : MountedDirectoryMonitors)
 	{
-		for (auto& Monitor : MountedDirectoryMonitors)
-		{
-			Monitor.Tick(TimeLimit);
-			CONDITIONAL_YIELD
-		}
-		for (auto& Monitor : UserDirectoryMonitors)
-		{
-			Monitor.Tick(TimeLimit);
-			CONDITIONAL_YIELD
-		}
-
-		const int32 NumUnprocessedChanges = GetNumUnprocessedChanges();
-		if (NumUnprocessedChanges > 0)
-		{
-			// Check if we have any more unprocessed changes. If we have, we wait a delay before processing them.
-			if (NumUnprocessedChanges != CachedNumUnprocessedChanges)
-			{
-				CachedNumUnprocessedChanges = NumUnprocessedChanges;
-				StartProcessingDelay.Reset();
-			}
-			else
-			{
-				if (StartProcessingDelay.Exceeded())
-				{
-					// We have some stuff to do, let's kick that off
-					CurrentState = ECurrentState::Initialize;
-					FeedbackContextOverride.Reset(new FReimportFeedbackContext);
-				}
-			}
-		}
+		Monitor.Tick(TimeLimit);
+		yield TOptional<ECurrentState>();
 	}
-	else if (FSlateThrottleManager::Get().IsAllowingExpensiveTasks())
+	for (auto& Monitor : UserDirectoryMonitors)
 	{
-		TGuardValue<bool> GuardProcessingChanges(bIsProcessingChanges, true);
-		
-		FeedbackContextOverride->UpdateText(GetTotalProgress(), GetTotalWork());
-
-		if (CurrentState == ECurrentState::Initialize)
-		{
-			for (auto& Monitor : MountedDirectoryMonitors)
-			{
-				Monitor.StartProcessing();
-			}
-			for (auto& Monitor : UserDirectoryMonitors)
-			{
-				Monitor.StartProcessing();
-			}
-
-			CurrentState = ECurrentState::ProcessAdditions;
-
-			FeedbackContextOverride->UpdateText(GetTotalProgress(), GetTotalWork());
-			CONDITIONAL_YIELD
-		}
+		Monitor.Tick(TimeLimit);
+		yield TOptional<ECurrentState>();
+	}
 
-		if (CurrentState == ECurrentState::ProcessAdditions)
+	const int32 NumUnprocessedChanges = GetNumUnprocessedChanges();
+	if (NumUnprocessedChanges > 0)
+	{
+		// Check if we have any more unprocessed changes. If we have, we wait a delay before processing them.
+		if (NumUnprocessedChanges != CachedNumUnprocessedChanges)
 		{
-			// Override the global feedback context while we do this to avoid popping up dialogs
-			TGuardValue<FFeedbackContext*> ScopedContextOverride(GWarn, FeedbackContextOverride.Get());
-
-			ProcessAdditions(TimeLimit);
-
-			FeedbackContextOverride->UpdateText(GetTotalProgress(), GetTotalWork());
-			CONDITIONAL_YIELD
+			CachedNumUnprocessedChanges = NumUnprocessedChanges;
+			StartProcessingDelay.Reset();
 		}
-
-		if (CurrentState == ECurrentState::SavePackages)
+		else
 		{
-			// Override the global feedback context while we do this to avoid popping up dialogs
-			TGuardValue<FFeedbackContext*> ScopedContextOverride(GWarn, FeedbackContextOverride.Get());
-
-			if (PackagesToSave.Num() > 0)
+			if (StartProcessingDelay.Exceeded())
 			{
-				const bool bAlreadyCheckedOut = false;
-				const bool bCheckDirty = false;
-				const bool bPromptToSave = false;
-				// @todo: arodham: Error reporting?
-				FEditorFileUtils::PromptForCheckoutAndSave(PackagesToSave, bCheckDirty, bPromptToSave, nullptr, bAlreadyCheckedOut);
-
-				PackagesToSave.Empty();
-			}
+				// Create a new feedback context override
+				FeedbackContextOverride = MakeShareable(new FReimportFeedbackContext);
+				FeedbackContextOverride->Initialize(SNew(SReimportFeedback, GetProgressText()));
 
-			CurrentState = ECurrentState::ProcessModifications;
-			CONDITIONAL_YIELD
-		}
-
-		const IAssetRegistry& Registry = FModuleManager::LoadModuleChecked<FAssetRegistryModule>("AssetRegistry").Get();
-		if (CurrentState == ECurrentState::ProcessModifications)
-		{
-			// Override the global feedback context while we do this to avoid popping up dialogs
-			TGuardValue<FFeedbackContext*> ScopedContextOverride(GWarn, FeedbackContextOverride.Get());
-
-			ProcessModifications(Registry, TimeLimit);
+				// Deal with changes to the file system second
+				for (auto& Monitor : MountedDirectoryMonitors)
+				{
+					Monitor.StartProcessing();
+				}
+				for (auto& Monitor : UserDirectoryMonitors)
+				{
+					Monitor.StartProcessing();
+				}
 
-			FeedbackContextOverride->UpdateText(GetTotalProgress(), GetTotalWork());
-			CONDITIONAL_YIELD
+				return ECurrentState::ProcessAdditions;
+			}
 		}
+	}
 
-		if (CurrentState == ECurrentState::ProcessDeletions)
-		{
-			ProcessDeletions(Registry);
-
-			FeedbackContextOverride->UpdateText(GetTotalProgress(), GetTotalWork());
-			CONDITIONAL_YIELD
-		}
+	return TOptional<ECurrentState>();
+}
 
-		if (CurrentState == ECurrentState::Cleanup)
-		{
-			// Finished
-			FeedbackContextOverride->UpdateText(GetTotalWork(), GetTotalWork());
+void FAutoReimportManager::Cleanup()
+{
+	CachedNumUnprocessedChanges = 0;
 
-			CachedNumUnprocessedChanges = 0;
-			FeedbackContextOverride.Reset();
+	FeedbackContextOverride->Destroy();
+	FeedbackContextOverride = nullptr;
+}
 
-			CurrentState = ECurrentState::Idle;
-		}
+void FAutoReimportManager::Tick(float DeltaTime)
+{
+	if (FeedbackContextOverride.IsValid())
+	{
+		FeedbackContextOverride->Tick();
 	}
+	
+	// Never spend more than a 60fps frame doing this work (meaning we shouldn't drop below 30fps), we can do more if we're throttling CPU usage (ie editor running in background)
+	const FTimeLimit TimeLimit(GEditor->ShouldThrottleCPUUsage() ? 1 / 6.f : 1.f / 60.f);
+	StateMachine.Tick(TimeLimit);
 }
 
 TStatId FAutoReimportManager::GetStatId() const
@@ -456,6 +478,14 @@ TStatId FAutoReimportManager::GetStatId() const
 	RETURN_QUICK_DECLARE_CYCLE_STAT(FAutoReimportManager, STATGROUP_Tickables);
 }
 
+void FAutoReimportManager::AddReferencedObjects(FReferenceCollector& Collector)
+{
+	for (auto* Package : PackagesToSave)
+	{
+		Collector.AddReferencedObject(Package);
+	}
+}
+
 void FAutoReimportManager::HandleLoadingSavingSettingChanged(FName PropertyName)
 {
 	if (PropertyName == GET_MEMBER_NAME_CHECKED(UEditorLoadingSavingSettings, AutoReimportDirectories))
@@ -497,59 +527,6 @@ void FAutoReimportManager::OnContentPathDismounted(const FString& InAssetPath, c
 	});
 }
 
-void FAutoReimportManager::OnAssetDeleted(UObject* Object)
-{
-	if (bIsProcessingChanges)
-	{
-		return;
-	}
-
-	const IAssetRegistry& Registry = FModuleManager::LoadModuleChecked<FAssetRegistryModule>("AssetRegistry").Get();
-
-	bool bCanDelete = true;
-
-	TArray<UObject::FAssetRegistryTag> Tags;
-	Object->GetAssetRegistryTags(Tags);
-
-	TArray<FString> AbsoluteSourceFileNames;
-
-	const FName TagName = UObject::SourceFileTagName();
-
-	for (const auto& Tag : Tags)
-	{
-		// Don't compare the FName number suffix
-		if (Tag.Name.IsEqual(TagName, ENameCase::IgnoreCase, false))
-		{
-			const FString SourceFile = FReimportManager::ResolveImportFilename(Tag.Value, Object);
-			// Check there aren't any other assets referencing this thing
-			if (Utils::FindAssetsPertainingToFile(Registry, SourceFile).Num() != 0)
-			{
-				bCanDelete = false;
-				break;
-			}
-			else
-			{
-				AbsoluteSourceFileNames.Add(MoveTemp(SourceFile));
-			}
-		}
-	}
-
-	if (bCanDelete)
-	{
-		// Offer to delete the source file(s) as well
-		const FText MessageText = FText::Format(LOCTEXT("DeleteRelatedAssets", "The asset '{0}' has been deleted. Would you like to delete the source files it was imported from as well?"), FText::FromString(Object->GetName()) );
-		if (FMessageDialog::Open(EAppMsgType::YesNo, MessageText) == EAppReturnType::Yes)
-		{
-			auto& FileManager = IFileManager::Get();
-
-			for (auto& Filename : AbsoluteSourceFileNames)
-			{
-				FileManager.Delete(*Filename, false /* RequireExists */, true /* Even if read only */, true /* Quiet */);
-			}
-		}
-	}
-}
-
 void FAutoReimportManager::SetUpMountedContentDirectories()
 {
 	TArray<FString> ContentPaths;
@@ -655,13 +632,17 @@ UAutoReimportManager::~UAutoReimportManager()
 
 void UAutoReimportManager::Initialize()
 {
-	Implementation.Reset(new FAutoReimportManager);
+	Implementation = MakeShareable(new FAutoReimportManager);
 }
 
 void UAutoReimportManager::BeginDestroy()
 {
 	Super::BeginDestroy();
-	Implementation.Reset();
+	if (Implementation.IsValid())
+	{
+		Implementation->Destroy();
+		Implementation = nullptr;
+	}
 }
 
 
diff --git a/Engine/Source/Editor/UnrealEd/Private/AutoReimport/AutoReimportUtilities.cpp b/Engine/Source/Editor/UnrealEd/Private/AutoReimport/AutoReimportUtilities.cpp
--- a/Engine/Source/Editor/UnrealEd/Private/AutoReimport/AutoReimportUtilities.cpp
+++ b/Engine/Source/Editor/UnrealEd/Private/AutoReimport/AutoReimportUtilities.cpp
@@ -1,16 +1,15 @@
 // Copyright 1998-2015 Epic Games, Inc. All Rights Reserved.
 
 #include "UnrealEd.h"
-
 #include "AutoReimportUtilities.h"
-#include "AssetRegistryModule.h"
+#include "EditorReimportHandler.h"
 
 namespace Utils
 {
-	/** Find a list of assets that were last imported from the specified filename. */
 	TArray<FAssetData> FindAssetsPertainingToFile(const IAssetRegistry& Registry, const FString& AbsoluteFilename)
 	{
-		TArray<FAssetData> Assets, Result;
+		TArray<FAssetData> Assets;
+
 		const FString LeafName = FPaths::GetCleanFilename(AbsoluteFilename);
 		const FName TagName = UObject::SourceFileTagName();
 
@@ -19,8 +18,7 @@ namespace Utils
 		Filter.bIncludeOnlyOnDiskAssets = true;
 		Registry.GetAssets(Filter, Assets);
 
-		for (const auto& Asset : Assets)
-		{
+		Assets.RemoveAll([&](const FAssetData& Asset){
 			for (const auto& Pair : Asset.TagsAndValues)
 			{
 				// We don't compare numbers on FNames for this check because the tag "ReimportPath" may exist multiple times
@@ -32,13 +30,30 @@ namespace Utils
 					if (AbsoluteFilename == FPaths::ConvertRelativePathToFull(FPackageName::LongPackageNameToFilename(Asset.PackagePath.ToString()) / Pair.Value) ||
 						AbsoluteFilename == FPaths::ConvertRelativePathToFull(Pair.Value))
 					{
-						Result.Add(Asset);
-						break;
+						return false;
 					}
 				}
 			}
-		}
+			return true;
+		});
+		
+		return Assets;
+	}
+
+	void ExtractSourceFilePaths(UObject* Object, TArray<FString>& OutSourceFiles)
+	{
+		TArray<UObject::FAssetRegistryTag> TagList;
+		Object->GetAssetRegistryTags(TagList);
 
-		return Result;
+		const FName TagName = UObject::SourceFileTagName();
+		for (const auto& Tag : TagList)
+		{
+			// We don't compare numbers on FNames for this check because the tag "ReimportPath" may exist multiple times
+			const bool bCompareNumber = false;
+			if (Tag.Name.IsEqual(TagName, ENameCase::IgnoreCase, bCompareNumber))
+			{
+				OutSourceFiles.Add(FReimportManager::ResolveImportFilename(Tag.Value, Object));
+			}
+		}
 	}
 }
\ No newline at end of file
diff --git a/Engine/Source/Editor/UnrealEd/Private/AutoReimport/AutoReimportUtilities.h b/Engine/Source/Editor/UnrealEd/Private/AutoReimport/AutoReimportUtilities.h
--- a/Engine/Source/Editor/UnrealEd/Private/AutoReimport/AutoReimportUtilities.h
+++ b/Engine/Source/Editor/UnrealEd/Private/AutoReimport/AutoReimportUtilities.h
@@ -97,6 +97,9 @@ namespace Utils
 		return Accumulator;
 	}
 
-	/** Find a list of assets that were last imported from the specified filename. */
+	/** Find a list of assets that were once imported from the specified filename */
 	TArray<FAssetData> FindAssetsPertainingToFile(const IAssetRegistry& Registry, const FString& AbsoluteFilename);
+
+	/** Extract any source file paths from the specified object */
+	void ExtractSourceFilePaths(UObject* Object, TArray<FString>& OutSourceFiles);
 }
\ No newline at end of file
diff --git a/Engine/Source/Editor/UnrealEd/Private/AutoReimport/ContentDirectoryMonitor.cpp b/Engine/Source/Editor/UnrealEd/Private/AutoReimport/ContentDirectoryMonitor.cpp
--- a/Engine/Source/Editor/UnrealEd/Private/AutoReimport/ContentDirectoryMonitor.cpp
+++ b/Engine/Source/Editor/UnrealEd/Private/AutoReimport/ContentDirectoryMonitor.cpp
@@ -7,6 +7,9 @@
 #include "AssetRegistryModule.h"
 #include "PackageTools.h"
 #include "ObjectTools.h"
+#include "ReimportFeedbackContext.h"
+
+#define LOCTEXT_NAMESPACE "ContentDirectoryMonitor"
 
 /** Generate a config from the specified options, to pass to FFileCache on construction */
 FFileCacheConfig GenerateFileCacheConfig(const FString& InPath, const FString& InSupportedExtensions, const FString& InMountedContentPath)
@@ -71,7 +74,7 @@ void FContentDirectoryMonitor::StartProcessing()
 	}
 }
 
-void FContentDirectoryMonitor::ProcessAdditions(TArray<UPackage*>& OutPackagesToSave, const FTimeLimit& TimeLimit, const TMap<FString, TArray<UFactory*>>& InFactoriesByExtension)
+void FContentDirectoryMonitor::ProcessAdditions(TArray<UPackage*>& OutPackagesToSave, const FTimeLimit& TimeLimit, const TMap<FString, TArray<UFactory*>>& InFactoriesByExtension, FReimportFeedbackContext& Context)
 {
 	if (MountedContentPath.IsEmpty())
 	{
@@ -79,80 +82,87 @@ void FContentDirectoryMonitor::ProcessAdditions(TArray<UPackage*>& OutPackagesTo
 		return;
 	}
 	
+	bool bCancelled = false;
 	for (int32 Index = 0; Index < AddedFiles.Num(); ++Index)
 	{
 		auto& Addition = AddedFiles[Index];
 
 		WorkProgress += 1;
-		const FString FullFilename = Cache.GetDirectory() + Addition.Filename.Get();
 
-		bool bImportSucceeded = false;
-		bool bImportWasCancelled = false;
+		if (bCancelled)
+		{
+			// Just update the cache immediately if the user cancelled
+			Cache.CompleteTransaction(MoveTemp(Addition));
+			continue;
+		}
+
+		const FString FullFilename = Cache.GetDirectory() + Addition.Filename.Get();
 
 		FString NewAssetName = ObjectTools::SanitizeObjectName(FPaths::GetBaseFilename(FullFilename));
-		FString PackagePath= PackageTools::SanitizePackageName(MountedContentPath / FPaths::GetPath(Addition.Filename.Get()) / NewAssetName);
+		FString PackagePath = PackageTools::SanitizePackageName(MountedContentPath / FPaths::GetPath(Addition.Filename.Get()) / NewAssetName);
 
-		// We can not create assets that share the name of a map file in the same location
 		if (FEditorFileUtils::IsMapPackageAsset(PackagePath))
 		{
-			// @todo: Error reporting
-			continue;
+			Context.AddMessage(EMessageSeverity::Error, FText::Format(LOCTEXT("Error_ExistingMap", "Attempted to create asset with the same name as a map ({0})."), FText::FromString(PackagePath)));
 		}
-
-		if (FindPackage(nullptr, *PackagePath))
+		else if (FindPackage(nullptr, *PackagePath))
 		{
-			// @todo: Error reporting
-			continue;	
+			Context.AddMessage(EMessageSeverity::Error, FText::Format(LOCTEXT("Error_ExistingAsset", "Attempted to create a new asset that already exists ({0})."), FText::FromString(PackagePath)));
 		}
-
-		UPackage* NewPackage = CreatePackage(nullptr, *PackagePath);
-		if ( !ensure(NewPackage) )
+		else
 		{
-			// @todo: Error reporting
-			continue;
-		}
-
-		// Make sure the destination package is loaded
-		NewPackage->FullyLoad();
-
-		bool bSuccess = false;
-		bool bCancelled = false;
+			UPackage* NewPackage = CreatePackage(nullptr, *PackagePath);
+			if ( !ensure(NewPackage) )
+			{
+				Context.AddMessage(EMessageSeverity::Error, FText::Format(LOCTEXT("Error_FailedToCreateAsset", "Failed to create new asset ({0}) for file ({1})."), FText::FromString(NewAssetName), FText::FromString(FullFilename)));
+			}
+			else
+			{
+				// Make sure the destination package is loaded
+				NewPackage->FullyLoad();
 
-		UObject* NewAsset = nullptr;
+				bool bSuccess = false;
+				UObject* NewAsset = nullptr;
 
-		const FString Ext = FPaths::GetExtension(Addition.Filename.Get(), false);
-		if (auto* Factories = InFactoriesByExtension.Find(Ext))
-		{
-			for (auto* Factory : *Factories)
-			{
-				NewAsset = UFactory::StaticImportObject(Factory->ResolveSupportedClass(), NewPackage, FName(*NewAssetName), RF_Public | RF_Standalone, bCancelled, *FullFilename, nullptr, Factory);
-				if (NewAsset || bCancelled)
+				const FString Ext = FPaths::GetExtension(Addition.Filename.Get(), false);
+				if (auto* Factories = InFactoriesByExtension.Find(Ext))
 				{
-					break;
+					for (auto* Factory : *Factories)
+					{
+						NewAsset = UFactory::StaticImportObject(Factory->ResolveSupportedClass(), NewPackage, FName(*NewAssetName), RF_Public | RF_Standalone, bCancelled, *FullFilename, nullptr, Factory);
+						if (NewAsset || bCancelled)
+						{
+							break;
+						}
+					}
 				}
-			}
-		}
-		
-		if (!NewAsset || bCancelled)
-		{
-			// @todo: Error reporting
-			continue;
-		}
+				
+				if (!bCancelled)
+				{
+					if (!NewAsset)
+					{
+						Context.AddMessage(EMessageSeverity::Error, FText::Format(LOCTEXT("Error_FailedToImportAsset", "Failed to import file {0}."), FText::FromString(FullFilename)));
+						continue;
+					}
 
-		FAssetRegistryModule::AssetCreated(NewAsset);
-		GEditor->BroadcastObjectReimported(NewAsset);
+					FAssetRegistryModule::AssetCreated(NewAsset);
+					GEditor->BroadcastObjectReimported(NewAsset);
 
-		OutPackagesToSave.Add(NewPackage);
+					OutPackagesToSave.Add(NewPackage);
+				}
 
-		// Refresh the supported class.  Some factories (e.g. FBX) only resolve their type after reading the file
-		// ImportAssetType = Factory->ResolveSupportedClass();
-		// @todo: analytics
+				// Refresh the supported class.  Some factories (e.g. FBX) only resolve their type after reading the file
+				// ImportAssetType = Factory->ResolveSupportedClass();
+				// @todo: analytics
+			}
+		}
 
 		// Let the cache know that we've dealt with this change (it will be imported immediately)
 		Cache.CompleteTransaction(MoveTemp(Addition));
 
-		if (TimeLimit.Exceeded())
+		if (!bCancelled && TimeLimit.Exceeded())
 		{
+			// Remove the ones we've processed
 			AddedFiles.RemoveAt(0, Index + 1);
 			return;
 		}
@@ -161,7 +171,7 @@ void FContentDirectoryMonitor::ProcessAdditions(TArray<UPackage*>& OutPackagesTo
 	AddedFiles.Empty();
 }
 
-void FContentDirectoryMonitor::ProcessModifications(const IAssetRegistry& Registry, const FTimeLimit& TimeLimit)
+void FContentDirectoryMonitor::ProcessModifications(const IAssetRegistry& Registry, const FTimeLimit& TimeLimit, FReimportFeedbackContext& Context)
 {
 	auto* ReimportManager = FReimportManager::Instance();
 
@@ -173,7 +183,10 @@ void FContentDirectoryMonitor::ProcessModifications(const IAssetRegistry& Regist
 		for (const auto& AssetData : Utils::FindAssetsPertainingToFile(Registry, FullFilename))
 		{
 			UObject* Asset = AssetData.GetAsset();
-			ReimportManager->Reimport(Asset, false /* Ask for new file */, false /* Show notification */);
+			if (!ReimportManager->Reimport(Asset, false /* Ask for new file */, false /* Show notification */))
+			{
+				Context.AddMessage(EMessageSeverity::Error, FText::Format(LOCTEXT("Error_FailedToReimportAsset", "Failed to reimport asset {0}."), FText::FromString(Asset->GetName())));
+			}
 		}
 
 		// Let the cache know that we've dealt with this change
@@ -205,4 +218,6 @@ void FContentDirectoryMonitor::ExtractAssetsToDelete(const IAssetRegistry& Regis
 
 	WorkProgress += DeletedFiles.Num();
 	DeletedFiles.Empty();
-}
\ No newline at end of file
+}
+
+#undef LOCTEXT_NAMESPACE
diff --git a/Engine/Source/Editor/UnrealEd/Private/AutoReimport/ContentDirectoryMonitor.h b/Engine/Source/Editor/UnrealEd/Private/AutoReimport/ContentDirectoryMonitor.h
--- a/Engine/Source/Editor/UnrealEd/Private/AutoReimport/ContentDirectoryMonitor.h
+++ b/Engine/Source/Editor/UnrealEd/Private/AutoReimport/ContentDirectoryMonitor.h
@@ -4,6 +4,7 @@
 
 #include "AutoReimportUtilities.h"
 #include "FileCache.h"
+#include "MessageLog.h"
 
 /** Class responsible for watching a specific content directory for changes */
 class FContentDirectoryMonitor
@@ -25,10 +26,10 @@ public:
 	void StartProcessing();
 
 	/** Extract the files we need to import from our outstanding changes (happens first)*/ 
-	void ProcessAdditions(TArray<UPackage*>& OutPackagesToSave, const FTimeLimit& TimeLimit, const TMap<FString, TArray<UFactory*>>& InFactoriesByExtension);
+	void ProcessAdditions(TArray<UPackage*>& OutPackagesToSave, const FTimeLimit& TimeLimit, const TMap<FString, TArray<UFactory*>>& InFactoriesByExtension, class FReimportFeedbackContext& Context);
 
 	/** Process the outstanding changes that we have cached */
-	void ProcessModifications(const IAssetRegistry& Registry, const FTimeLimit& TimeLimit);
+	void ProcessModifications(const IAssetRegistry& Registry, const FTimeLimit& TimeLimit, class FReimportFeedbackContext& Context);
 
 	/** Extract the assets we need to delete from our outstanding changes (happens last) */ 
 	void ExtractAssetsToDelete(const IAssetRegistry& Registry, TArray<FAssetData>& OutAssetsToDelete);
diff --git a/Engine/Source/Editor/UnrealEd/Private/AutoReimport/ReimportFeedbackContext.cpp b/Engine/Source/Editor/UnrealEd/Private/AutoReimport/ReimportFeedbackContext.cpp
new file mode 100644
--- /dev/null
+++ b/Engine/Source/Editor/UnrealEd/Private/AutoReimport/ReimportFeedbackContext.cpp
@@ -0,0 +1,309 @@
+// Copyright 1998-2015 Epic Games, Inc. All Rights Reserved.
+
+#pragma once
+
+#include "UnrealEd.h"
+
+#include "ReimportFeedbackContext.h"
+#include "NotificationManager.h"
+#include "SNotificationList.h"
+#include "INotificationWidget.h"
+
+
+class SWidgetStack : public SVerticalBox
+{
+	SLATE_BEGIN_ARGS(SWidgetStack){}
+	SLATE_END_ARGS()
+
+	void Construct(const FArguments& InArgs, int32 InMaxNumVisible)
+	{
+		MaxNumVisible = InMaxNumVisible;
+		SlideCurve = FCurveSequence(0.f, .5f, ECurveEaseFunction::QuadOut);
+
+		StartSlideOffset = LerpSlideOffset = 0;
+		StartSizeOffset = LerpSizeOffset = 0;
+	}
+
+	FVector2D ComputeTotalSize() const
+	{
+		FVector2D Size;
+		for (int32 Index = 0; Index < FMath::Min(NumSlots(), MaxNumVisible); ++Index)
+		{
+			const FVector2D& ChildSize = Children[Index].GetWidget()->GetDesiredSize();
+			if (ChildSize.X > Size.X)
+			{
+				Size.X = ChildSize.X;
+			}
+			Size.Y += ChildSize.Y + Children[Index].SlotPadding.Get().GetTotalSpaceAlong<Orient_Vertical>();
+		}
+		return Size;
+	}
+
+	virtual FVector2D ComputeDesiredSize() const override
+	{
+		FVector2D Size = ComputeTotalSize();
+		Size.Y -= LerpSizeOffset;
+		return Size;
+	}
+
+	virtual void OnArrangeChildren( const FGeometry& AllottedGeometry, FArrangedChildren& ArrangedChildren ) const
+	{
+		if (Children.Num() == 0)
+		{
+			return;
+		}
+
+		float PositionSoFar = AllottedGeometry.GetLocalSize().Y + LerpSlideOffset;
+
+		for (int32 Index = 0; Index < NumSlots(); ++Index)
+		{
+			const SBoxPanel::FSlot& CurChild = Children[Index];
+			const EVisibility ChildVisibility = CurChild.GetWidget()->GetVisibility();
+
+			if (ChildVisibility != EVisibility::Collapsed)
+			{
+				const FVector2D ChildDesiredSize = CurChild.GetWidget()->GetDesiredSize();
+
+				const FMargin SlotPadding(CurChild.SlotPadding.Get());
+				const FVector2D SlotSize(AllottedGeometry.Size.X, ChildDesiredSize.Y + SlotPadding.GetTotalSpaceAlong<Orient_Vertical>());
+
+				const AlignmentArrangeResult XAlignmentResult = AlignChild<Orient_Horizontal>( SlotSize.X, CurChild, SlotPadding );
+				const AlignmentArrangeResult YAlignmentResult = AlignChild<Orient_Vertical>( SlotSize.Y, CurChild, SlotPadding );
+
+				ArrangedChildren.AddWidget( ChildVisibility, AllottedGeometry.MakeChild(
+					CurChild.GetWidget(),
+					FVector2D( XAlignmentResult.Offset, PositionSoFar - SlotSize.Y + YAlignmentResult.Offset ),
+					FVector2D( XAlignmentResult.Size, YAlignmentResult.Size )
+					));
+
+				PositionSoFar -= SlotSize.Y;
+			}
+		}
+	}
+
+	void Remove(const TSharedRef<SWidget>& InWidget)
+	{
+
+	}
+
+	void Add(const TSharedRef<SWidget>& InWidget)
+	{
+		InsertSlot(0)
+		.AutoHeight()
+		[
+			SNew(SWidgetStackItem)
+			[
+				InWidget
+			]
+		];
+
+		auto Widget = Children[0].GetWidget();
+		Widget->SlatePrepass();
+		
+		StartSlideOffset = LerpSlideOffset + Widget->GetDesiredSize().Y;
+		LerpSlideOffset = StartSlideOffset;
+
+		StartSizeOffset = ComputeTotalSize().Y - GetDesiredSize().Y;
+		LerpSizeOffset = StartSizeOffset;
+
+		SlideCurve.Play();
+
+		if (Children.Num() > MaxNumVisible)
+		{
+			auto& Curve = StaticCastSharedRef<SWidgetStackItem>(Children[MaxNumVisible].GetWidget())->OpacityCurve;
+			if (Curve.IsPlaying())
+			{
+				Curve.Reverse();
+			}
+			else
+			{
+				Curve.PlayReverse();
+			}
+		}
+	}
+	
+	virtual void Tick(const FGeometry& AllottedGeometry, const double InCurrentTime, const float InDeltaTime)
+	{
+		const float Alpha = 1.f - SlideCurve.GetLerp();
+
+		LerpSlideOffset = StartSlideOffset * Alpha;
+		LerpSizeOffset = StartSizeOffset * Alpha;
+
+		// Delete any widgets that are now offscreen
+		if (Children.Num() != 0)
+		{
+			float PositionSoFar = AllottedGeometry.GetLocalSize().Y + LerpSlideOffset;
+
+			int32 Index = 0;
+			for (; PositionSoFar > 0 && Index < NumSlots(); ++Index)
+			{
+				const SBoxPanel::FSlot& CurChild = Children[Index];
+				const EVisibility ChildVisibility = CurChild.GetWidget()->GetVisibility();
+
+				if (ChildVisibility != EVisibility::Collapsed)
+				{
+					const FVector2D ChildDesiredSize = CurChild.GetWidget()->GetDesiredSize();
+					PositionSoFar -= ChildDesiredSize.Y + CurChild.SlotPadding.Get().GetTotalSpaceAlong<Orient_Vertical>();
+				}
+			}
+
+			while (Children.Num() > MaxNumVisible + 1)
+			{
+				Children.RemoveAt(Children.Num() - 1);
+			}
+		}
+	}
+
+	class SWidgetStackItem : public SCompoundWidget
+	{
+		SLATE_BEGIN_ARGS(SWidgetStackItem){}
+			SLATE_DEFAULT_SLOT(FArguments, Content)
+			SLATE_EVENT(FSimpleDelegate, OnFadeOut)
+		SLATE_END_ARGS()
+
+		void Construct(const FArguments& InArgs)
+		{
+			OpacityCurve = FCurveSequence(0.f, .5f, ECurveEaseFunction::QuadOut);
+			OpacityCurve.Play();
+
+			ChildSlot
+			[
+				SNew(SBorder)
+				.BorderImage(FEditorStyle::GetBrush("NoBorder"))
+				.ColorAndOpacity(this, &SWidgetStackItem::GetColorAndOpacity)
+				.Padding(0)
+				[
+					InArgs._Content.Widget
+				]
+			];
+		}
+		
+		virtual void Tick(const FGeometry& AllottedGeometry, const double InCurrentTime, const float InDeltaTime)
+		{
+			if (OnFadeOut.IsBound() && OpacityCurve.IsAtStart() && OpacityCurve.IsInReverse())
+			{
+				OnFadeOut.Execute();
+				OnFadeOut = FSimpleDelegate();
+			}
+		}
+
+		FLinearColor GetColorAndOpacity() const
+		{
+			return FLinearColor(1.f, 1.f, 1.f, OpacityCurve.GetLerp());
+		}
+
+		FCurveSequence OpacityCurve;
+		FSimpleDelegate OnFadeOut;
+	};
+
+	FCurveSequence SlideCurve;
+
+	float StartSlideOffset, LerpSlideOffset;
+	float StartSizeOffset, LerpSizeOffset;
+
+	int32 MaxNumVisible;
+};
+
+
+void SReimportFeedback::Add(const TSharedRef<SWidget>& Widget)
+{
+	WidgetStack->Add(Widget);
+}
+
+void SReimportFeedback::SetMainText(FText InText)
+{
+	MainText = InText;
+}
+
+FText SReimportFeedback::GetMainText() const
+{
+	return MainText;
+}
+
+void SReimportFeedback::Construct(const FArguments& InArgs, FText InMainText)
+{
+	MainText = InMainText;
+
+	ChildSlot
+	[
+		SNew(SBorder)
+		.Padding(FMargin(10))
+		.BorderImage(FCoreStyle::Get().GetBrush("NotificationList.ItemBackground"))
+		[
+			SNew(SVerticalBox)
+
+			+ SVerticalBox::Slot()
+			.Padding(FMargin(0))
+			.AutoHeight()
+			[
+				SNew(STextBlock)
+				.Text(this, &SReimportFeedback::GetMainText)
+				.Font(FCoreStyle::Get().GetFontStyle(TEXT("NotificationList.FontBold")))
+			]
+
+			+ SVerticalBox::Slot()
+			.Padding(FMargin(0, 5, 0, 0))
+			.AutoHeight()
+			[
+				SAssignNew(WidgetStack, SWidgetStack, 1)
+			]
+		]
+	];
+}
+
+FReimportFeedbackContext::FReimportFeedbackContext()
+	: MessageLog("AssetReimport")
+{}
+
+void FReimportFeedbackContext::Initialize(TSharedRef<SReimportFeedback> Widget)
+{
+	ShowNotificationDelay = FTimeLimit(.5f);
+	NotificationContent = Widget;
+	MessageLog.NewPage(Widget->GetMainText());
+}
+
+void FReimportFeedbackContext::Tick()
+{
+	if (!Notification.IsValid() && ShowNotificationDelay.Exceeded())
+	{
+		FNotificationInfo Info(SharedThis(this));
+		Info.ExpireDuration = 1.f;
+		Info.bFireAndForget = false;
+
+		Notification = FSlateNotificationManager::Get().AddNotification(Info);
+	}
+}
+
+void FReimportFeedbackContext::Destroy()
+{
+	MessageLog.Notify();
+
+	if (Notification.IsValid())
+	{
+		NotificationContent->SetVisibility(EVisibility::HitTestInvisible);
+		//Notification->SetCompletionState(SNotificationItem::CS_Success);
+		Notification->ExpireAndFadeout();
+	}
+}
+
+void FReimportFeedbackContext::AddMessage(EMessageSeverity::Type Severity, const FText& Message)
+{
+	MessageLog.Message(Severity, Message);
+	if (Severity >= EMessageSeverity::Error)
+	{
+		AddWidget(SNew(STextBlock).Text(Message));
+	}
+}
+
+void FReimportFeedbackContext::AddWidget(const TSharedRef<SWidget>& Widget)
+{
+	NotificationContent->Add(Widget);
+}
+
+void FReimportFeedbackContext::ProgressReported(const float TotalProgressInterp, FText DisplayMessage)
+{
+	if (!DisplayMessage.IsEmpty())
+	{
+		NotificationContent->Add(SNew(STextBlock).Text(DisplayMessage));
+	}
+}
\ No newline at end of file
diff --git a/Engine/Source/Editor/UnrealEd/Private/AutoReimport/ReimportFeedbackContext.h b/Engine/Source/Editor/UnrealEd/Private/AutoReimport/ReimportFeedbackContext.h
new file mode 100644
--- /dev/null
+++ b/Engine/Source/Editor/UnrealEd/Private/AutoReimport/ReimportFeedbackContext.h
@@ -0,0 +1,83 @@
+// Copyright 1998-2015 Epic Games, Inc. All Rights Reserved.
+
+#pragma once
+
+#include "FeedbackContext.h"
+#include "SNotificationList.h"
+#include "INotificationWidget.h"
+
+class SWidgetStack;
+
+/** Feedback context that overrides GWarn for import operations to prevent popup spam */
+class SReimportFeedback : public SCompoundWidget
+{
+public:
+	SLATE_BEGIN_ARGS(SReimportFeedback){}
+	SLATE_END_ARGS()
+
+	/** Construct this widget */
+	void Construct(const FArguments& InArgs, FText InMainText);
+
+	/** Add a widget to this feedback's widget stack */
+	void Add(const TSharedRef<SWidget>& Widget);
+
+	/** Set the main text of this widget */
+	void SetMainText(FText InText);
+	
+	/** Get the main text of this widget */
+	FText GetMainText() const;
+
+private:
+	/** Cached main text for the notification */
+	FText MainText;
+
+	/** The widget stack, displaying contextural information about the current state of the process */
+	TSharedPtr<SWidgetStack> WidgetStack;
+};
+
+/** Feedback context that overrides GWarn for import operations to prevent popup spam */
+class FReimportFeedbackContext : public FFeedbackContext, public INotificationWidget, public TSharedFromThis<FReimportFeedbackContext>
+{
+public:
+	/** Constructor */
+	FReimportFeedbackContext();
+
+	/** Initialize this reimport context with the specified widget */
+	void Initialize(TSharedRef<SReimportFeedback> Widget);
+
+	/** Destroy this reimport context */
+	void Destroy();
+
+	/** Tick this context - potentially opening a new notification if the process has taken a while */
+	void Tick();
+
+	TSharedPtr<class SReimportFeedback> GetContent() { return NotificationContent; }
+
+	/** Add a message to the context (and message log) */
+	void AddMessage(EMessageSeverity::Type Severity, const FText& Message);
+
+	/** Add a custom widget to the context */
+	void AddWidget(const TSharedRef<SWidget>& Widget);
+
+	/** INotificationWidget and FFeedbackContext overrides */
+	virtual void Serialize( const TCHAR* V, ELogVerbosity::Type Verbosity, const class FName& Category ) override {}
+	virtual void OnSetCompletionState(SNotificationItem::ECompletionState State) override {}
+	virtual TSharedRef<SWidget> AsWidget() override { return NotificationContent.ToSharedRef(); };
+	virtual void ProgressReported(const float TotalProgressInterp, FText DisplayMessage) override;
+
+private:
+	/** Message log for output of errors and log messages */
+	FMessageLog MessageLog;
+
+	/** True if we need to notifiy about new messages on the message log */
+	bool NotifyMessageLog;
+
+	/** The time when the last change to the cache was reported */
+	FTimeLimit ShowNotificationDelay;
+
+	/** The notification that is shown when the context is active */
+	TSharedPtr<class SNotificationItem> Notification;
+
+	/** The notification content */
+	TSharedPtr<class SReimportFeedback> NotificationContent;
+};
\ No newline at end of file
diff --git a/Engine/Source/Editor/UnrealEd/Private/Dialogs/SDeleteAssetsDialog.cpp b/Engine/Source/Editor/UnrealEd/Private/Dialogs/SDeleteAssetsDialog.cpp
--- a/Engine/Source/Editor/UnrealEd/Private/Dialogs/SDeleteAssetsDialog.cpp
+++ b/Engine/Source/Editor/UnrealEd/Private/Dialogs/SDeleteAssetsDialog.cpp
@@ -164,6 +164,8 @@ TSharedRef<SWidget> SDeleteAssetsDialog::BuildProgressDialog()
 
 TSharedRef<SWidget> SDeleteAssetsDialog::BuildDeleteDialog()
 {
+	const auto* LoadingSavingSettings = GetDefault<UEditorLoadingSavingSettings>();
+
 	FFormatNamedArguments Args;
 	Args.Add( TEXT( "OnDiskReferences" ), FText::AsNumber( DeleteModel->GetAssetReferences().Num() ) );
 
@@ -325,6 +327,22 @@ TSharedRef<SWidget> SDeleteAssetsDialog::BuildDeleteDialog()
 				]
 			]
 
+			+ SVerticalBox::Slot()
+			.AutoHeight()
+			.Padding(6,4)
+			[
+				SAssignNew(DeleteSourceFilesCheckbox, SCheckBox)
+				.IsChecked(LoadingSavingSettings->bDeleteSourceFilesWithAssets ? ECheckBoxState::Checked : ECheckBoxState::Unchecked)
+				[
+					SNew(STextBlock)
+					.Text(LOCTEXT("DeleteSourceFiles", "Also delete related source content files"))
+					.ToolTip(
+						SNew(SToolTip)
+						.Text(LOCTEXT("DeleteSourceFiles_Tooltip", "When checked, any source content files that were used to create the selected asset(s) will also be deleted. For instance, deleting a texture will also delete its associated *.png file."))
+					)
+				]
+			]
+
 			+ SVerticalBox::Slot()
 			.AutoHeight()
 			.Padding(0,4)
@@ -577,6 +595,9 @@ TSharedRef<ITableRow> SDeleteAssetsDialog::HandleGenerateAssetRow( TSharedPtr<FP
 
 FReply SDeleteAssetsDialog::Delete()
 {
+	auto* Settings = GetMutableDefault<UEditorLoadingSavingSettings>();
+	Settings->bDeleteSourceFilesWithAssets = DeleteSourceFilesCheckbox->GetCheckedState() == ECheckBoxState::Checked;
+
 	ParentWindow.Get()->RequestDestroyWindow();
 
 	if (DeleteModel->IsAnythingReferencedInMemoryByUndo())
@@ -584,6 +605,10 @@ FReply SDeleteAssetsDialog::Delete()
 		GEditor->Trans->Reset(LOCTEXT("DeleteSelectedItem", "Delete Selected Item"));
 	}
 
+	if (Settings->bDeleteSourceFilesWithAssets)
+	{
+		DeleteModel->DeleteSourceContentFiles();
+	}
 	DeleteModel->DoDelete();
 
 	return FReply::Handled();
@@ -598,6 +623,9 @@ FReply SDeleteAssetsDialog::Cancel()
 
 FReply SDeleteAssetsDialog::ForceDelete()
 {
+	auto* Settings = GetMutableDefault<UEditorLoadingSavingSettings>();
+	Settings->bDeleteSourceFilesWithAssets = DeleteSourceFilesCheckbox->GetCheckedState() == ECheckBoxState::Checked;
+
 	ParentWindow.Get()->RequestDestroyWindow();
 
 	if( DeleteModel->IsAnythingReferencedInMemoryByUndo() )
@@ -605,6 +633,10 @@ FReply SDeleteAssetsDialog::ForceDelete()
 		GEditor->Trans->Reset( LOCTEXT("DeleteSelectedItem", "Delete Selected Item") );
 	}
 
+	if (Settings->bDeleteSourceFilesWithAssets)
+	{
+		DeleteModel->DeleteSourceContentFiles();
+	}
 	DeleteModel->DoForceDelete();
 
 	return FReply::Handled();
@@ -634,6 +666,13 @@ FReply SDeleteAssetsDialog::ReplaceReferences()
 
 	if ( EAppReturnType::Ok == OpenMsgDlgInt( EAppMsgType::OkCancel, Message, Title ) )
 	{
+		auto* Settings = GetMutableDefault<UEditorLoadingSavingSettings>();
+		Settings->bDeleteSourceFilesWithAssets = DeleteSourceFilesCheckbox->GetCheckedState() == ECheckBoxState::Checked;
+		if (Settings->bDeleteSourceFilesWithAssets)
+		{
+			DeleteModel->DeleteSourceContentFiles();
+		}
+
 		ParentWindow.Get()->RequestDestroyWindow();
 		DeleteModel->DoReplaceReferences( ConsolidationAsset );
 	}
diff --git a/Engine/Source/Editor/UnrealEd/Private/Settings/SettingsClasses.cpp b/Engine/Source/Editor/UnrealEd/Private/Settings/SettingsClasses.cpp
--- a/Engine/Source/Editor/UnrealEd/Private/Settings/SettingsClasses.cpp
+++ b/Engine/Source/Editor/UnrealEd/Private/Settings/SettingsClasses.cpp
@@ -94,6 +94,7 @@ void UEditorExperimentalSettings::PostEditChangeProperty( struct FPropertyChange
 UEditorLoadingSavingSettings::UEditorLoadingSavingSettings( const FObjectInitializer& ObjectInitializer )
 	: Super(ObjectInitializer)
 	, bMonitorContentDirectories(false)
+	, bDeleteSourceFilesWithAssets(true)
 {
 	TextDiffToolPath.FilePath = TEXT("P4Merge.exe");
 }
diff --git a/Engine/Source/Editor/UnrealEd/Public/AssetDeleteModel.h b/Engine/Source/Editor/UnrealEd/Public/AssetDeleteModel.h
--- a/Engine/Source/Editor/UnrealEd/Public/AssetDeleteModel.h
+++ b/Engine/Source/Editor/UnrealEd/Public/AssetDeleteModel.h
@@ -124,6 +124,9 @@ public:
 	/** Returns true if the package is one of the pending deleted assets. */
 	bool IsAssetInPendingDeletes( const FName& PackageName ) const;
 
+	/** Deletes any source content files referenced by the assets */
+	void DeleteSourceContentFiles();
+
 	/** Returns true if it is valid to delete the current objects with no problems. */
 	bool CanDelete() const;
 
diff --git a/Engine/Source/Editor/UnrealEd/Public/AutoReimport/AutoReimportManager.h b/Engine/Source/Editor/UnrealEd/Public/AutoReimport/AutoReimportManager.h
--- a/Engine/Source/Editor/UnrealEd/Public/AutoReimport/AutoReimportManager.h
+++ b/Engine/Source/Editor/UnrealEd/Public/AutoReimport/AutoReimportManager.h
@@ -22,5 +22,5 @@ private:
 	virtual void BeginDestroy() override;
 
 	/** Private implementation of the reimport manager */
-	TUniquePtr<class FAutoReimportManager> Implementation;
+	TSharedPtr<class FAutoReimportManager> Implementation;
 };
diff --git a/Engine/Source/Editor/UnrealEd/Public/Dialogs/SDeleteAssetsDialog.h b/Engine/Source/Editor/UnrealEd/Public/Dialogs/SDeleteAssetsDialog.h
--- a/Engine/Source/Editor/UnrealEd/Public/Dialogs/SDeleteAssetsDialog.h
+++ b/Engine/Source/Editor/UnrealEd/Public/Dialogs/SDeleteAssetsDialog.h
@@ -134,6 +134,7 @@ private:
 	TSharedPtr< SListView< TSharedPtr<FPendingDelete> > > ObjectsToDeleteList;
 	TSharedPtr< SBorder > AssetReferenceNameBorderWidget;
 	TSharedPtr< SComboButton > ConsolidationPickerComboButton;
+	TSharedPtr< SCheckBox > DeleteSourceFilesCheckbox;
 
 	/** The selected asset we're going to consolidate the would be deleted assets into. */
 	FAssetData ConsolidationAsset;
