[INTEGRATE] Change 2418768 by [REDACTED_PERSON][REDACTED_HANDLE].7 on 2015/01/26 12:22:18

	Fixed fbx import factories not working correctly with the auto reimport manager

diff --git a/Engine/Source/Editor/UnrealEd/Private/AutoReimport/ContentDirectoryMonitor.cpp b/Engine/Source/Editor/UnrealEd/Private/AutoReimport/ContentDirectoryMonitor.cpp
--- a/Engine/Source/Editor/UnrealEd/Private/AutoReimport/ContentDirectoryMonitor.cpp
+++ b/Engine/Source/Editor/UnrealEd/Private/AutoReimport/ContentDirectoryMonitor.cpp
@@ -122,21 +122,32 @@ void FContentDirectoryMonitor::ProcessAdditions(TArray<UPackage*>& OutPackagesTo
 				NewPackage->FullyLoad();
 
 				bool bSuccess = false;
-				UObject* NewAsset = nullptr;
-
+				
+				// Find a relevant factory for this file
 				const FString Ext = FPaths::GetExtension(Addition.Filename.Get(), false);
+				UFactory* FactoryType = nullptr;
 				if (auto* Factories = InFactoriesByExtension.Find(Ext))
 				{
-					for (auto* Factory : *Factories)
+					FactoryType = (*Factories)[0];
+
+					UFactory* const* FoundCompatibleFactoryType = Factories->FindByPredicate([&](UFactory* InFactory){
+						return InFactory->FactoryCanImport(FullFilename);
+					});
+
+					if (FoundCompatibleFactoryType)
 					{
-						NewAsset = UFactory::StaticImportObject(Factory->ResolveSupportedClass(), NewPackage, FName(*NewAssetName), RF_Public | RF_Standalone, bCancelled, *FullFilename, nullptr, Factory);
-						if (NewAsset || bCancelled)
-						{
-							break;
-						}
+						FactoryType = *FoundCompatibleFactoryType;
 					}
 				}
-				
+
+				UObject* NewAsset = nullptr;
+				UFactory* FactoryInstance = FactoryType ? ConstructObject<UFactory>(FactoryType->GetClass()) : nullptr;
+				if (FactoryInstance && FactoryInstance->ConfigureProperties())
+				{
+					UClass* ImportAssetType = FactoryInstance->SupportedClass;
+					NewAsset = UFactory::StaticImportObject(ImportAssetType, NewPackage, FName(*NewAssetName), RF_Public | RF_Standalone, bCancelled, *FullFilename, nullptr, FactoryInstance);
+				}
+
 				if (!bCancelled)
 				{
 					if (!NewAsset)
